AsyncLoading2: Preparation for optional package data support
#rnx

diff --git a/Engine/Source/Runtime/CoreUObject/Private/Serialization/AsyncLoading2.cpp b/Engine/Source/Runtime/CoreUObject/Private/Serialization/AsyncLoading2.cpp
--- a/Engine/Source/Runtime/CoreUObject/Private/Serialization/AsyncLoading2.cpp
+++ b/Engine/Source/Runtime/CoreUObject/Private/Serialization/AsyncLoading2.cpp
@@ -860,47 +860,61 @@ public:
 	}
 };
 
-struct FPackageImportStore
+struct FAsyncPackageHeaderData
 {
-	IPackageStore& PackageStore;
-	FGlobalImportStore& GlobalImportStore;
-	FLoadedPackageStore& LoadedPackageStore;
-	const FAsyncPackageDesc2& Desc;
-	const TArrayView<const FPackageId>& ImportedPackageIds;
+	uint32 CookedHeaderSize = 0;
+	TOptional<FZenPackageVersioningInfo> VersioningInfo;
+	FNameMap NameMap;
+	FName PackageName;
+	// Backed by IoBuffer
+	const FZenPackageSummary* PackageSummary = nullptr;
 	TArrayView<const uint64> ImportedPublicExportHashes;
 	TArrayView<const FPackageObjectIndex> ImportMap;
+	TArrayView<const FExportMapEntry> ExportMap;
+	TArrayView<const uint8> ArcsData;
+	// Backed by allocation in FAsyncPackageData
+	TArrayView<FPackageId> ImportedPackageIds;
+	TArrayView<FExportBundleHeader> ExportBundleHeaders;
+	TArrayView<FExportBundleEntry> ExportBundleEntries;
 
-	FPackageImportStore(IPackageStore& InPackageStore, FGlobalImportStore& InGlobalImportStore, FLoadedPackageStore& InLoadedPackageStore, const FAsyncPackageDesc2& InDesc, const TArrayView<const FPackageId>& InImportedPackageIds)
-		: PackageStore(InPackageStore)
-		, GlobalImportStore(InGlobalImportStore)
-		, LoadedPackageStore(InLoadedPackageStore)
-		, Desc(InDesc)
-		, ImportedPackageIds(InImportedPackageIds)
+	void OnReleaseHeaderBuffer()
 	{
+		PackageSummary = nullptr;
+		ImportedPublicExportHashes = TArrayView<const uint64>();
+		ImportMap = TArrayView<const FPackageObjectIndex>();
+		ExportMap = TArrayView<const FExportMapEntry>();
+		ArcsData = TArrayView<const uint8>();
 	}
+};
 
-	~FPackageImportStore()
+struct FPackageImportStore
+{
+	FGlobalImportStore& GlobalImportStore;
+	FLoadedPackageStore& LoadedPackageStore;
+
+	FPackageImportStore(FGlobalImportStore& InGlobalImportStore, FLoadedPackageStore& InLoadedPackageStore)
+		: GlobalImportStore(InGlobalImportStore)
+		, LoadedPackageStore(InLoadedPackageStore)
 	{
-		check(ImportMap.Num() == 0);
 	}
 
-	inline bool IsValidLocalImportIndex(FPackageIndex LocalIndex)
+	inline bool IsValidLocalImportIndex(const TArrayView<const FPackageObjectIndex>& ImportMap, FPackageIndex LocalIndex)
 	{
 		check(ImportMap.Num() > 0);
 		return LocalIndex.IsImport() && LocalIndex.ToImport() < ImportMap.Num();
 	}
 
-	inline UObject* FindOrGetImportObjectFromLocalIndex(FPackageIndex LocalIndex)
+	inline UObject* FindOrGetImportObjectFromLocalIndex(const FAsyncPackageHeaderData& Header, FPackageIndex LocalIndex)
 	{
 		check(LocalIndex.IsImport());
-		check(ImportMap.Num() > 0);
+		check(Header.ImportMap.Num() > 0);
 		const int32 LocalImportIndex = LocalIndex.ToImport();
-		check(LocalImportIndex < ImportMap.Num());
-		const FPackageObjectIndex GlobalIndex = ImportMap[LocalIndex.ToImport()];
-		return FindOrGetImportObject(GlobalIndex);
+		check(LocalImportIndex < Header.ImportMap.Num());
+		const FPackageObjectIndex GlobalIndex = Header.ImportMap[LocalIndex.ToImport()];
+		return FindOrGetImportObject(Header, GlobalIndex);
 	}
 
-	inline UObject* FindOrGetImportObject(FPackageObjectIndex GlobalIndex)
+	inline UObject* FindOrGetImportObject(const FAsyncPackageHeaderData& Header, FPackageObjectIndex GlobalIndex)
 	{
 		check(GlobalIndex.IsImport());
 		UObject* Object = nullptr;
@@ -910,7 +924,7 @@ struct FPackageImportStore
 		}
 		else if (GlobalIndex.IsPackageImport())
 		{
-			Object = GlobalImportStore.FindPublicExportObject(FPublicExportKey::FromPackageImport(GlobalIndex, ImportedPackageIds, ImportedPublicExportHashes));
+			Object = GlobalImportStore.FindPublicExportObject(FPublicExportKey::FromPackageImport(GlobalIndex, Header.ImportedPackageIds, Header.ImportedPublicExportHashes));
 		}
 		else
 		{
@@ -919,9 +933,9 @@ struct FPackageImportStore
 		return Object;
 	}
 
-	bool GetUnresolvedCDOs(TArray<UClass*, TInlineAllocator<8>>& Classes)
+	void GetUnresolvedCDOs(const FAsyncPackageHeaderData& Header, TArray<UClass*, TInlineAllocator<8>>& Classes)
 	{
-		for (const FPackageObjectIndex& Index : ImportMap)
+		for (const FPackageObjectIndex& Index : Header.ImportMap)
 		{
 			if (!Index.IsScriptImport())
 			{
@@ -947,7 +961,6 @@ struct FPackageImportStore
 				}
 			}
 		}
-		return Classes.Num() > 0;
 	}
 
 	inline void StoreGlobalObject(FPackageId PackageId, uint64 ExportHash, UObject* Object)
@@ -998,7 +1011,7 @@ private:
 	}
 
 public:
-	void AddPackageReferences()
+	void AddImportedPackageReferences(const TArrayView<const FPackageId>& ImportedPackageIds)
 	{
 		for (const FPackageId& ImportedPackageId : ImportedPackageIds)
 		{
@@ -1008,6 +1021,10 @@ public:
 				PinPublicExportsForGC(PackageRef.GetPackage());
 			}
 		}
+	}
+
+	void AddPackageReference(const FAsyncPackageDesc2& Desc)
+	{
 		if (Desc.bCanBeImported)
 		{
 			FLoadedPackageRef& PackageRef = LoadedPackageStore.GetPackageRef(Desc.UPackageId);
@@ -1019,7 +1036,7 @@ public:
 		}
 	}
 
-	void ReleasePackageReferences()
+	void ReleaseImportedPackageReferences(const FAsyncPackageDesc2& Desc, const TArrayView<const FPackageId>& ImportedPackageIds)
 	{
 		for (const FPackageId& ImportedPackageId : ImportedPackageIds)
 		{
@@ -1029,6 +1046,10 @@ public:
 				UnpinPublicExportsForGC(PackageRef.GetPackage());
 			}
 		}
+	}
+
+	void ReleasePackageReference(const FAsyncPackageDesc2& Desc)
+	{
 		if (Desc.bCanBeImported)
 		{
 			// clear own reference, and possible all async flags if no remaining ref count
@@ -1098,7 +1119,7 @@ public:
 
 	virtual int64 TotalSize() override
 	{
-		return CookedHeaderSize + (ActiveFPLB->EndFastPathLoadBuffer - ActiveFPLB->OriginalFastPathLoadBuffer);
+		return HeaderData->CookedHeaderSize + (ActiveFPLB->EndFastPathLoadBuffer - ActiveFPLB->OriginalFastPathLoadBuffer);
 	}
 
 	virtual int64 Tell() override
@@ -1165,7 +1186,7 @@ public:
 	{
 		UE_ASYNC_PACKAGE_LOG(Fatal, *PackageDesc, TEXT("ObjectSerializationError"),
 			TEXT("%s: Bad import index %d/%d."),
-			CurrentExport ? *CurrentExport->GetFullName() : TEXT("null"), ImportIndex, ImportStore->ImportMap.Num());
+			CurrentExport ? *CurrentExport->GetFullName() : TEXT("null"), ImportIndex, HeaderData->ImportMap.Num());
 		Object = nullptr;
 	}
 
@@ -1187,8 +1208,8 @@ public:
 				Object = Exports[ExportIndex].Object;
 
 #if ALT2_LOG_VERBOSE
-				const FExportMapEntry& Export = ExportMap[ExportIndex];
-				FName ObjectName = NameMap->GetName(Export.ObjectName);
+				const FExportMapEntry& Export = HeaderData->ExportMap[ExportIndex];
+				FName ObjectName = HeaderData->NameMap.GetName(Export.ObjectName);
 				UE_ASYNC_PACKAGE_CLOG_VERBOSE(!Object, VeryVerbose, *PackageDesc,
 					TEXT("FExportArchive: Object"), TEXT("Export %s at index %d is null."),
 					*ObjectName.ToString(), 
@@ -1202,9 +1223,9 @@ public:
 		}
 		else
 		{
-			if (ImportStore->IsValidLocalImportIndex(Index))
+			if (ImportStore->IsValidLocalImportIndex(HeaderData->ImportMap, Index))
 			{
-				Object = ImportStore->FindOrGetImportObjectFromLocalIndex(Index);
+				Object = ImportStore->FindOrGetImportObjectFromLocalIndex(*HeaderData, Index);
 
 				UE_ASYNC_PACKAGE_CLOG_VERBOSE(!Object, Log, *PackageDesc,
 					TEXT("FExportArchive: Object"), TEXT("Import index %d is null"),
@@ -1248,7 +1269,7 @@ public:
 	{
 		UE_ASYNC_PACKAGE_LOG(Fatal, *PackageDesc, TEXT("ObjectSerializationError"),
 			TEXT("%s: Bad name index %d/%d."),
-			CurrentExport ? *CurrentExport->GetFullName() : TEXT("null"), NameIndex, NameMap->Num());
+			CurrentExport ? *CurrentExport->GetFullName() : TEXT("null"), NameIndex, HeaderData->NameMap.Num());
 		Name = FName();
 		SetCriticalError();
 	}
@@ -1262,7 +1283,7 @@ public:
 		Ar << Number;
 
 		FMappedName MappedName = FMappedName::Create(NameIndex, Number, FMappedName::EType::Package);
-		if (!NameMap->TryGetName(MappedName, Name))
+		if (!HeaderData->NameMap.TryGetName(MappedName, Name))
 		{
 			HandleBadNameIndex(NameIndex, Name);
 		}
@@ -1282,11 +1303,9 @@ private:
 	FAsyncPackageDesc2* PackageDesc = nullptr;
 	FPackageImportStore* ImportStore = nullptr;
 	TArray<FExternalReadCallback>* ExternalReadDependencies;
-	const FNameMap* NameMap = nullptr;
+	const FAsyncPackageHeaderData* HeaderData = nullptr;
 	TArrayView<const FExportObject> Exports;
-	const FExportMapEntry* ExportMap = nullptr;
 	UObject* CurrentExport = nullptr;
-	uint32 CookedHeaderSize = 0;
 	uint64 CookedSerialOffset = 0;
 	uint64 CookedSerialSize = 0;
 	uint64 BufferSerialOffset = 0;
@@ -1562,20 +1581,35 @@ enum EEventLoadNode2 : uint8
 	ExportBundle_NumPhases,
 };
 
+struct FAsyncPackageExportToBundleMapping
+{
+	uint64 ExportHash;
+	int32 BundleIndex[FExportBundleEntry::ExportCommandType_Count];
+};
+
 struct FAsyncPackageData
 {
 	uint8* MemoryBuffer = nullptr;
-	FPackageStoreExportInfo ExportInfo;
-	uint64 ExportBundleHeadersSize = 0;
-	uint64 ExportBundleEntriesSize = 0;
-	uint8* ExportBundlesMetaMemory = nullptr;
-	const FExportBundleHeader* ExportBundleHeaders = nullptr;
-	const FExportBundleEntry* ExportBundleEntries = nullptr;
 	TArrayView<FExportObject> Exports;
 	TArrayView<FAsyncPackage2*> ImportedAsyncPackages;
 	TArrayView<FEventLoadNode2> ExportBundleNodes;
-	TArrayView<const FPackageId> ImportedPackageIds;
 	TArrayView<const FSHAHash> ShaderMapHashes;
+	TArrayView<FAsyncPackageExportToBundleMapping> ExportToBundleMappings;
+	int32 ExportBundleCount = 0;
+};
+
+struct FAsyncPackageSerializationState
+{
+	FIoRequest IoRequest;
+	const uint8* AllExportDataPtr = nullptr;
+	const uint8* CurrentExportDataPtr = nullptr;
+
+	void ReleaseIoRequest()
+	{
+		IoRequest.Release();
+		AllExportDataPtr = nullptr;
+		CurrentExportDataPtr = nullptr;
+	}
 };
 
 /**
@@ -1665,16 +1699,14 @@ struct FAsyncPackage2
 
 private:
 
-	struct FExportToBundleMapping
-	{
-		uint64 ExportHash;
-		int32 BundleIndex[FExportBundleEntry::ExportCommandType_Count];
-	};
+	void ImportPackagesRecursiveInner(FIoBatch& IoBatch, IPackageStore& PackageStore, const TArrayView<const FPackageId>& ImportedPackageIds, int32& ImportedPackageIndex);
 
 	uint8 PackageNodesMemory[EEventLoadNode2::Package_NumPhases * sizeof(FEventLoadNode2)];
 	/** Basic information associated with this package */
 	FAsyncPackageDesc2 Desc;
 	FAsyncPackageData Data;
+	FAsyncPackageHeaderData HeaderData;
+	FAsyncPackageSerializationState SerializationState;
 	/** Cached async loading thread object this package was created by */
 	FAsyncLoadingThread2& AsyncLoadingThread;
 	FAsyncLoadEventGraphAllocator& GraphAllocator;
@@ -1803,17 +1835,6 @@ private:
 	using FCompletionCallback = TUniquePtr<FLoadPackageAsyncDelegate>;
 	TArray<FCompletionCallback, TInlineAllocator<2>> CompletionCallbacks;
 
-	FIoRequest IoRequest;
-	const uint8* AllExportDataPtr = nullptr;
-	const uint8* CurrentExportDataPtr = nullptr;
-	uint32 CookedHeaderSize = 0;
-
-	TArrayView<const FExportMapEntry> ExportMap;
-	TArrayView<const uint64> PublicExportHashes;
-	TArrayView<const uint8> ArcsData;
-	FNameMap NameMap;
-	TArray<FExportToBundleMapping> ExportToBundleMappings;
-
 public:
 
 	FAsyncLoadingThread2& GetAsyncLoadingThread()
@@ -1835,14 +1856,14 @@ public:
 	static EAsyncPackageState::Type Event_PostLoadExportBundle(FAsyncLoadingThreadState2& ThreadState, FAsyncPackage2* Package, int32 ExportBundleIndex);
 	static EAsyncPackageState::Type Event_DeferredPostLoadExportBundle(FAsyncLoadingThreadState2& ThreadState, FAsyncPackage2* Package, int32 ExportBundleIndex);
 	
-	void EventDrivenCreateExport(int32 LocalExportIndex);
-	bool EventDrivenSerializeExport(int32 LocalExportIndex, FExportArchive& Ar);
+	void EventDrivenCreateExport(const FAsyncPackageHeaderData& Header, const TArrayView<FExportObject>& Exports, int32 LocalExportIndex);
+	bool EventDrivenSerializeExport(const FAsyncPackageHeaderData& Header, const TArrayView<FExportObject>& Exports, int32 LocalExportIndex, FExportArchive& Ar);
 
-	UObject* EventDrivenIndexToObject(FPackageObjectIndex Index, bool bCheckSerialized);
+	UObject* EventDrivenIndexToObject(const FAsyncPackageHeaderData& Header, const TArrayView<const FExportObject>& Exports, FPackageObjectIndex Index, bool bCheckSerialized);
 	template<class T>
-	T* CastEventDrivenIndexToObject(FPackageObjectIndex Index, bool bCheckSerialized)
+	T* CastEventDrivenIndexToObject(const FAsyncPackageHeaderData& Header, const TArrayView<const FExportObject>& Exports, FPackageObjectIndex Index, bool bCheckSerialized)
 	{
-		UObject* Result = EventDrivenIndexToObject(Index, bCheckSerialized);
+		UObject* Result = EventDrivenIndexToObject(Header, Exports, Index, bCheckSerialized);
 		if (!Result)
 		{
 			return nullptr;
@@ -1860,7 +1881,7 @@ public:
 private:
 	void CreatePackageNodes(const FAsyncLoadEventSpec* EventSpecs);
 	void CreateExportBundleNodes(const FAsyncLoadEventSpec* EventSpecs);
-	void SetupSerializedArcs();
+	void SetupSerializedArcs(const FAsyncPackageHeaderData& Header, const TArrayView<FEventLoadNode2>& ExportBundleNodes, const TArrayView<FAsyncPackage2*>& ImportedAsyncPackages);
 	void SetupScriptDependencies();
 	bool HaveAllDependenciesReachedStateDebug(FAsyncPackage2* Package, TSet<FAsyncPackage2*>& VisitedPackages, EAsyncPackageLoadingState2 WaitForPackageState);
 	bool HaveAllDependenciesReachedState(FAllDependenciesState FAsyncPackage2::* StateMemberPtr, EAsyncPackageLoadingState2 WaitForPackageState, uint32 CurrentTick);
@@ -2386,55 +2407,63 @@ private:
 		UE_ASYNC_PACKAGE_DEBUG(AsyncPackage->Desc);
 		
 		FAsyncPackageData& Data = AsyncPackage->Data;
-		Data.ExportInfo = PackageStoreEntry.ExportInfo;
-
-		const int32 ExportBundleNodeCount = Data.ExportInfo.ExportBundleCount * EEventLoadNode2::ExportBundle_NumPhases;
-		const int32 ImportedPackageCount = PackageStoreEntry.ImportedPackageIds.Num();
+		FAsyncPackageHeaderData& HeaderData = AsyncPackage->HeaderData;
+
+		const int32 ExportCount = PackageStoreEntry.ExportInfo.ExportCount;
+		const int32 ExportBundleCount = PackageStoreEntry.ExportInfo.ExportBundleCount;
+		const uint64 ExportBundleHeadersMemSize = Align(sizeof(FExportBundleHeader) * ExportBundleCount, 8);
+		const int32 ExportBundleEntriesCount = ExportCount * FExportBundleEntry::ExportCommandType_Count;
+		const uint64 ExportBundleEntriesMemSize = Align(sizeof(FExportBundleEntry) * ExportBundleEntriesCount, 8);
+		const int32 ImportedPackagesCount = PackageStoreEntry.ImportedPackageIds.Num();
+		const uint64 ImportedPackageIdsMemSize = Align(sizeof(FPackageId) * ImportedPackagesCount, 8);
+		AsyncPackage->Data.ExportBundleCount = ExportBundleCount;
+		const int32 ExportBundleNodeCount = AsyncPackage->Data.ExportBundleCount * EEventLoadNode2::ExportBundle_NumPhases;
 		const int32 ShaderMapHashesCount = PackageStoreEntry.ShaderMapHashes.Num();
-		
-		Data.ExportBundleHeadersSize = sizeof(FExportBundleHeader) * Data.ExportInfo.ExportBundleCount;
-		Data.ExportBundleEntriesSize = sizeof(FExportBundleEntry) * Data.ExportInfo.ExportCount * FExportBundleEntry::ExportCommandType_Count;
-		const uint64 ExportBundlesMetaSize = Data.ExportBundleHeadersSize + Data.ExportBundleEntriesSize;
 
-		const uint64 ExportBundlesMetaMemSize = Align(ExportBundlesMetaSize, 8);
-		const uint64 ExportsMemSize = Align(sizeof(FExportObject) * Data.ExportInfo.ExportCount, 8);
-		const uint64 ImportedPackagesMemSize = Align(sizeof(FAsyncPackage2*) * ImportedPackageCount, 8);
+		const uint64 ImportedPackagesMemSize = Align(sizeof(FAsyncPackage2*) * ImportedPackagesCount, 8);
+		const uint64 ExportsMemSize = Align(sizeof(FExportObject) * ExportCount, 8);
 		const uint64 ExportBundleNodesMemSize = Align(sizeof(FEventLoadNode2) * ExportBundleNodeCount, 8);
-		const uint64 ImportedPackageIdsMemSize = Align(sizeof(FPackageId) * ImportedPackageCount, 8);
 		const uint64 ShaderMapHashesMemSize = Align(sizeof(FSHAHash) * ShaderMapHashesCount, 8);
+		const uint64 ExportToBundleMappingMemSize = Align(sizeof(FAsyncPackageExportToBundleMapping) * ExportCount, 8);
 		const uint64 MemoryBufferSize =
-			ExportBundlesMetaMemSize +
-			ExportsMemSize +
+			ExportBundleHeadersMemSize +
+			ExportBundleEntriesMemSize +
+			ImportedPackageIdsMemSize +
 			ImportedPackagesMemSize +
+			ExportsMemSize +
 			ExportBundleNodesMemSize +
-			ImportedPackageIdsMemSize +
-			ShaderMapHashesMemSize;
+			ShaderMapHashesMemSize +
+			ExportToBundleMappingMemSize;
 
 		Data.MemoryBuffer = reinterpret_cast<uint8*>(FMemory::Malloc(MemoryBufferSize));
 
 		uint8* DataPtr = Data.MemoryBuffer;
 
-		Data.ExportBundlesMetaMemory = DataPtr;
-		Data.ExportBundleHeaders = reinterpret_cast<const FExportBundleHeader*>(Data.ExportBundlesMetaMemory);
-		Data.ExportBundleEntries = reinterpret_cast<const FExportBundleEntry*>(Data.ExportBundleHeaders + Data.ExportInfo.ExportBundleCount);
-
-		DataPtr += ExportBundlesMetaMemSize;
-		Data.Exports = MakeArrayView(reinterpret_cast<FExportObject*>(DataPtr), Data.ExportInfo.ExportCount);
+		Data.Exports = MakeArrayView(reinterpret_cast<FExportObject*>(DataPtr), ExportCount);
 		DataPtr += ExportsMemSize;
-		Data.ImportedAsyncPackages = MakeArrayView(reinterpret_cast<FAsyncPackage2**>(DataPtr), 0);
-		DataPtr += ImportedPackagesMemSize;
 		Data.ExportBundleNodes = MakeArrayView(reinterpret_cast<FEventLoadNode2*>(DataPtr), ExportBundleNodeCount);
 		DataPtr += ExportBundleNodesMemSize;
-		Data.ImportedPackageIds = MakeArrayView(reinterpret_cast<const FPackageId*>(DataPtr), ImportedPackageCount);
-		FMemory::Memcpy((void*)Data.ImportedPackageIds.GetData(), PackageStoreEntry.ImportedPackageIds.GetData(), sizeof(FPackageId) * ImportedPackageCount);
-		DataPtr += ImportedPackageIdsMemSize;
 		Data.ShaderMapHashes = MakeArrayView(reinterpret_cast<const FSHAHash*>(DataPtr), ShaderMapHashesCount);
 		FMemory::Memcpy((void*)Data.ShaderMapHashes.GetData(), PackageStoreEntry.ShaderMapHashes.GetData(), sizeof(FSHAHash) * ShaderMapHashesCount);
-	
+		DataPtr += ShaderMapHashesMemSize;
+		Data.ImportedAsyncPackages = MakeArrayView(reinterpret_cast<FAsyncPackage2**>(DataPtr), 0);
+		DataPtr += ImportedPackagesMemSize;
+		Data.ExportToBundleMappings = MakeArrayView(reinterpret_cast<FAsyncPackageExportToBundleMapping*>(DataPtr), ExportCount);
+		DataPtr += ExportToBundleMappingMemSize;
+
+		HeaderData.ExportBundleHeaders = MakeArrayView(reinterpret_cast<FExportBundleHeader*>(DataPtr), ExportBundleCount);
+		DataPtr += ExportBundleHeadersMemSize;
+		HeaderData.ExportBundleEntries = MakeArrayView(reinterpret_cast<FExportBundleEntry*>(DataPtr), ExportBundleEntriesCount);
+		DataPtr += ExportBundleEntriesMemSize;
+		HeaderData.ImportedPackageIds = MakeArrayView(reinterpret_cast<FPackageId*>(DataPtr), ImportedPackagesCount);
+		FMemory::Memcpy((void*)HeaderData.ImportedPackageIds.GetData(), PackageStoreEntry.ImportedPackageIds.GetData(), sizeof(FPackageId) * ImportedPackagesCount);
+		DataPtr += ImportedPackageIdsMemSize;
+		AsyncPackage->ImportStore.AddImportedPackageReferences(HeaderData.ImportedPackageIds);
+		AsyncPackage->ImportStore.AddPackageReference(AsyncPackage->Desc);
+		check(DataPtr - Data.MemoryBuffer == MemoryBufferSize);
 		AsyncPackage->CreateExportBundleNodes(EventSpecs.GetData());
 
-		AsyncPackage->ImportStore.AddPackageReferences();
-		AsyncPackage->ConstructedObjects.Reserve(Data.ExportInfo.ExportCount + 1); // +1 for UPackage, may grow dynamically beoynd that
+		AsyncPackage->ConstructedObjects.Reserve(Data.Exports.Num() + 1); // +1 for UPackage, may grow dynamically beoynd that
 		for (FExportObject& Export : Data.Exports)
 		{
 			Export = FExportObject();
@@ -2560,7 +2589,7 @@ void FAsyncLoadingThread2::UpdatePackagePriority(FAsyncPackage2* Package, int32
 {
 	TRACE_CPUPROFILER_EVENT_SCOPE(UpdatePackagePriority);
 	Package->Desc.Priority = NewPriority;
-	Package->IoRequest.UpdatePriority(NewPriority);
+	Package->SerializationState.IoRequest.UpdatePriority(NewPriority);
 }
 
 FAsyncPackage2* FAsyncLoadingThread2::FindOrInsertPackage(FAsyncPackageDesc2& Desc, bool& bInserted, TUniquePtr<FLoadPackageAsyncDelegate>&& PackageLoadedDelegate)
@@ -3077,11 +3106,11 @@ FUObjectSerializeContext* FAsyncPackage2::GetSerializeContext()
 	return FUObjectThreadContext::Get().GetSerializeContext();
 }
 
-void FAsyncPackage2::SetupSerializedArcs()
+void FAsyncPackage2::SetupSerializedArcs(const FAsyncPackageHeaderData& Header, const TArrayView<FEventLoadNode2>& ExportBundleNodes, const TArrayView<FAsyncPackage2*>& ImportedAsyncPackages)
 {
 	TRACE_CPUPROFILER_EVENT_SCOPE(SetupExternalArcs);
 
-	FSimpleArchive ArcsArchive(ArcsData.GetData(), ArcsData.Num());
+	FSimpleArchive ArcsArchive(Header.ArcsData.GetData(), Header.ArcsData.Num());
 	int32 InternalArcsCount;
 	ArcsArchive << InternalArcsCount;
 	for (int32 InternalArcIndex = 0; InternalArcIndex < InternalArcsCount; ++InternalArcIndex)
@@ -3096,10 +3125,10 @@ void FAsyncPackage2::SetupSerializedArcs()
 		{
 			uint32 ToNodeIndex = ToNodeIndexBase + Phase;
 			uint32 FromNodeIndex = FromNodeIndexBase + Phase;
-			Data.ExportBundleNodes[ToNodeIndex].DependsOn(&Data.ExportBundleNodes[FromNodeIndex]);
+			ExportBundleNodes[ToNodeIndex].DependsOn(&ExportBundleNodes[FromNodeIndex]);
 		}
 	}
-	for (const FAsyncPackage2* ImportedPackage : Data.ImportedAsyncPackages)
+	for (const FAsyncPackage2* ImportedPackage : ImportedAsyncPackages)
 	{
 		int32 ExternalArcCount;
 		ArcsArchive << ExternalArcCount;
@@ -3117,14 +3146,14 @@ void FAsyncPackage2::SetupSerializedArcs()
 			ArcsArchive << ToExportBundleIndex;
 			if (ImportedPackage)
 			{
-				check(FromImportIndex < ImportStore.ImportMap.Num());
+				check(FromImportIndex < Header.ImportMap.Num());
 				check(FromCommandType < FExportBundleEntry::ExportCommandType_Count);
-				check(ToExportBundleIndex < Data.ExportInfo.ExportBundleCount);
+				check(ToExportBundleIndex < ExportBundleNodes.Num());
 				
-				FPackageObjectIndex GlobalImportIndex = ImportStore.ImportMap[FromImportIndex];
+				FPackageObjectIndex GlobalImportIndex = Header.ImportMap[FromImportIndex];
 				FPackageImportReference PackageImportRef = GlobalImportIndex.ToPackageImportRef();
-				const uint64 ImportedPublicExportHash = ImportStore.ImportedPublicExportHashes[PackageImportRef.GetImportedPublicExportHashIndex()];
-				for (const FExportToBundleMapping& ExportToBundleMapping : ImportedPackage->ExportToBundleMappings)
+				const uint64 ImportedPublicExportHash = Header.ImportedPublicExportHashes[PackageImportRef.GetImportedPublicExportHashIndex()];
+				for (const FAsyncPackageExportToBundleMapping& ExportToBundleMapping : ImportedPackage->Data.ExportToBundleMappings)
 				{
 					if (ExportToBundleMapping.ExportHash == ImportedPublicExportHash)
 					{
@@ -3139,7 +3168,7 @@ void FAsyncPackage2::SetupSerializedArcs()
 							{
 								uint32 ToNodeIndex = ToNodeIndexBase + Phase;
 								uint32 FromNodeIndex = FromNodeIndexBase + Phase;
-								Data.ExportBundleNodes[ToNodeIndex].DependsOn(&ImportedPackage->Data.ExportBundleNodes[FromNodeIndex]);
+								ExportBundleNodes[ToNodeIndex].DependsOn(&ImportedPackage->Data.ExportBundleNodes[FromNodeIndex]);
 							}
 						}
 						break;
@@ -3158,7 +3187,8 @@ void FAsyncPackage2::SetupScriptDependencies()
 	// During initial laod, if a CDO called LoadObject for this package it may depend on other CDOs later in the list.
 	// Then collect them here, and wait for them to be created before allowing this package to proceed.
 	TArray<UClass*, TInlineAllocator<8>> UnresolvedCDOs;
-	if (ImportStore.GetUnresolvedCDOs(UnresolvedCDOs))
+	ImportStore.GetUnresolvedCDOs(HeaderData, UnresolvedCDOs);
+	if (!UnresolvedCDOs.IsEmpty())
 	{
 		AsyncLoadingThread.AddPendingCDOs(this, UnresolvedCDOs);
 	}
@@ -3230,29 +3260,9 @@ void FGlobalImportStore::FindAllScriptObjects()
 		ScriptObjects.Num(), (float)ScriptObjects.GetAllocatedSize() / 1024.f);
 }
 
-void FAsyncPackage2::ImportPackagesRecursive(FIoBatch& IoBatch, IPackageStore& PackageStore)
+void FAsyncPackage2::ImportPackagesRecursiveInner(FIoBatch& IoBatch, IPackageStore& PackageStore, const TArrayView<const FPackageId>& ImportedPackageIds, int32& ImportedPackageIndex)
 {
-	if (AsyncPackageLoadingState >= EAsyncPackageLoadingState2::ImportPackages)
-	{
-		return;
-	}
-	check(AsyncPackageLoadingState == EAsyncPackageLoadingState2::NewPackage);
-
-	const int32 ImportedPackageCount = Data.ImportedPackageIds.Num();
-	if (!ImportedPackageCount)
-	{
-		AsyncPackageLoadingState = EAsyncPackageLoadingState2::ImportPackagesDone;
-		return;
-	}
-	else
-	{
-		AsyncPackageLoadingState = EAsyncPackageLoadingState2::ImportPackages;
-	}
-
-	int32 ImportedPackageIndex = 0;
-
-	Data.ImportedAsyncPackages = MakeArrayView(Data.ImportedAsyncPackages.GetData(), Data.ImportedPackageIds.Num());
-	for (const FPackageId& ImportedPackageId : Data.ImportedPackageIds)
+	for (const FPackageId& ImportedPackageId : ImportedPackageIds)
 	{
 		FName ImportedPackageUPackageName = FName();
 		FPackageId ImportedPackageIdToLoad = ImportedPackageId;
@@ -3400,6 +3410,33 @@ void FAsyncPackage2::ImportPackagesRecursive(FIoBatch& IoBatch, IPackageStore& P
 			}
 		}
 	}
+}
+
+void FAsyncPackage2::ImportPackagesRecursive(FIoBatch& IoBatch, IPackageStore& PackageStore)
+{
+	if (AsyncPackageLoadingState >= EAsyncPackageLoadingState2::ImportPackages)
+	{
+		return;
+	}
+	check(AsyncPackageLoadingState == EAsyncPackageLoadingState2::NewPackage);
+
+	int32 ImportedPackageCount = HeaderData.ImportedPackageIds.Num();
+	if (!ImportedPackageCount)
+	{
+		AsyncPackageLoadingState = EAsyncPackageLoadingState2::ImportPackagesDone;
+		return;
+	}
+	else
+	{
+		AsyncPackageLoadingState = EAsyncPackageLoadingState2::ImportPackages;
+	}
+
+	int32 ImportedPackageIndex = 0;
+
+	Data.ImportedAsyncPackages = MakeArrayView(Data.ImportedAsyncPackages.GetData(), ImportedPackageCount);
+	
+	ImportPackagesRecursiveInner(IoBatch, PackageStore, HeaderData.ImportedPackageIds, ImportedPackageIndex);
+
 	UE_ASYNC_PACKAGE_LOG_VERBOSE(VeryVerbose, Desc, TEXT("ImportPackages: ImportsDone"),
 		TEXT("All imported packages are now being loaded."));
 
@@ -3421,7 +3458,7 @@ void FAsyncPackage2::StartLoading(FIoBatch& IoBatch)
 	TRACE_COUNTER_SET(AsyncLoadingPendingIoRequests, LocalPendingIoRequestsCounter);
 
 	FIoReadOptions ReadOptions;
-	IoRequest = IoBatch.ReadWithCallback(CreatePackageDataChunkId(Desc.PackageIdToLoad),
+	SerializationState.IoRequest = IoBatch.ReadWithCallback(CreatePackageDataChunkId(Desc.PackageIdToLoad),
 		ReadOptions,
 		Desc.Priority,
 		[this](TIoStatusOr<FIoBuffer> Result)
@@ -3463,6 +3500,50 @@ void FAsyncPackage2::StartLoading(FIoBatch& IoBatch)
 	}
 }
 
+static void ReadAsyncPackageHeader(FAsyncPackageSerializationState& SerializationState, FAsyncPackageHeaderData& HeaderData)
+{
+	const uint8* PackageHeaderDataPtr = SerializationState.IoRequest.GetResultOrDie().Data();
+	const FZenPackageSummary* PackageSummary = reinterpret_cast<const FZenPackageSummary*>(PackageHeaderDataPtr);
+	HeaderData.PackageSummary = PackageSummary;
+
+	TArrayView<const uint8> PackageHeaderDataView(PackageHeaderDataPtr + sizeof(FZenPackageSummary), PackageSummary->HeaderSize - sizeof(FZenPackageSummary));
+	FMemoryReaderView PackageHeaderDataReader(PackageHeaderDataView);
+	if (PackageSummary->bHasVersioningInfo)
+	{
+		HeaderData.VersioningInfo.Emplace();
+		PackageHeaderDataReader << HeaderData.VersioningInfo.GetValue();
+	}
+
+	{
+		TRACE_CPUPROFILER_EVENT_SCOPE(LoadPackageNameMap);
+		HeaderData.NameMap.Load(PackageHeaderDataReader, FMappedName::EType::Package);
+	}
+	HeaderData.PackageName = HeaderData.NameMap.GetName(PackageSummary->Name);
+
+	HeaderData.CookedHeaderSize = PackageSummary->CookedHeaderSize;
+	HeaderData.ImportedPublicExportHashes = TArrayView<const uint64>(
+		reinterpret_cast<const uint64*>(PackageHeaderDataPtr + PackageSummary->ImportedPublicExportHashesOffset),
+		(PackageSummary->ImportMapOffset - PackageSummary->ImportedPublicExportHashesOffset) / sizeof(uint64));
+	HeaderData.ImportMap = TArrayView<const FPackageObjectIndex>(
+		reinterpret_cast<const FPackageObjectIndex*>(PackageHeaderDataPtr + PackageSummary->ImportMapOffset),
+		(PackageSummary->ExportMapOffset - PackageSummary->ImportMapOffset) / sizeof(FPackageObjectIndex));
+	HeaderData.ExportMap = TArrayView<const FExportMapEntry>(
+		reinterpret_cast<const FExportMapEntry*>(PackageHeaderDataPtr + PackageSummary->ExportMapOffset),
+		(PackageSummary->ExportBundleEntriesOffset - PackageSummary->ExportMapOffset) / sizeof(FExportMapEntry));
+
+	const uint64 ExportBundleHeadersOffset = PackageSummary->GraphDataOffset;
+	const uint64 ExportBundleHeadersSize = sizeof(FExportBundleHeader) * HeaderData.ExportBundleHeaders.Num();
+	const uint64 ExportBundleEntriesSize = sizeof(FExportBundleEntry) * HeaderData.ExportBundleEntries.Num();
+	const uint64 ArcsDataOffset = ExportBundleHeadersOffset + ExportBundleHeadersSize;
+	const uint64 ArcsDataSize = PackageSummary->HeaderSize - ArcsDataOffset;
+	check(ExportBundleEntriesSize == PackageSummary->GraphDataOffset - PackageSummary->ExportBundleEntriesOffset);
+	HeaderData.ArcsData = MakeArrayView(PackageHeaderDataPtr + ArcsDataOffset, ArcsDataSize);
+	FMemory::Memcpy(HeaderData.ExportBundleHeaders.GetData(), PackageHeaderDataPtr + ExportBundleHeadersOffset, ExportBundleHeadersSize);
+	FMemory::Memcpy(HeaderData.ExportBundleEntries.GetData(), PackageHeaderDataPtr + PackageSummary->ExportBundleEntriesOffset, ExportBundleEntriesSize);
+
+	SerializationState.AllExportDataPtr = PackageHeaderDataPtr + PackageSummary->HeaderSize;
+}
+
 EAsyncPackageState::Type FAsyncPackage2::Event_ProcessPackageSummary(FAsyncLoadingThreadState2& ThreadState, FAsyncPackage2* Package, int32)
 {
 	TRACE_CPUPROFILER_EVENT_SCOPE(Event_ProcessPackageSummary);
@@ -3485,77 +3566,34 @@ EAsyncPackageState::Type FAsyncPackage2::Event_ProcessPackageSummary(FAsyncLoadi
 	{
 		check(Package->ExportBundleEntryIndex == 0);
 
-		const uint8* PackageHeaderData = Package->IoRequest.GetResultOrDie().Data();
-		const FZenPackageSummary* PackageSummary = reinterpret_cast<const FZenPackageSummary*>(PackageHeaderData);
-		
-		TArrayView<const uint8> PackageHeaderDataView(PackageHeaderData + sizeof(FZenPackageSummary), PackageSummary->HeaderSize - sizeof(FZenPackageSummary));
-		FMemoryReaderView PackageHeaderDataReader(PackageHeaderDataView);
-
-		FZenPackageVersioningInfo VersioningInfo;
-		if (PackageSummary->bHasVersioningInfo)
-		{
-			PackageHeaderDataReader << VersioningInfo;
-		}
-
-		{
-			TRACE_CPUPROFILER_EVENT_SCOPE(LoadPackageNameMap);
-			Package->NameMap.Load(PackageHeaderDataReader, FMappedName::EType::Package);
-		}
+		ReadAsyncPackageHeader(Package->SerializationState, Package->HeaderData);
 
+		for (int32 ExportBundleIndex = 0; ExportBundleIndex < Package->Data.ExportBundleCount; ++ExportBundleIndex)
 		{
-			FName PackageName = Package->NameMap.GetName(PackageSummary->Name);
-			check(Package->Desc.PackageNameToLoad.IsNone());
-			Package->Desc.PackageNameToLoad = PackageName;
-			check(Package->Desc.PackageIdToLoad == FPackageId::FromName(Package->Desc.PackageNameToLoad));
-			// Imported packages won't have a UPackage name set unless they were redirected, in which case they will have the source package name
-			if (Package->Desc.UPackageName.IsNone())
+			const FAsyncPackageHeaderData* HeaderData = &Package->HeaderData;
+			const FExportBundleHeader& ExportBundle = HeaderData->ExportBundleHeaders[ExportBundleIndex];
+			for (int32 EntryIndex = ExportBundle.FirstEntryIndex, EntryEnd = ExportBundle.FirstEntryIndex + ExportBundle.EntryCount; EntryIndex < EntryEnd; ++EntryIndex)
 			{
-				Package->Desc.UPackageName = PackageName;
-			}
-			check(Package->Desc.UPackageId == FPackageId::FromName(Package->Desc.UPackageName));
-		}
-
-		Package->CookedHeaderSize = PackageSummary->CookedHeaderSize;
-		Package->ImportStore.ImportedPublicExportHashes = TArrayView<const uint64>(
-			reinterpret_cast<const uint64*>(PackageHeaderData + PackageSummary->ImportedPublicExportHashesOffset),
-			(PackageSummary->ImportMapOffset - PackageSummary->ImportedPublicExportHashesOffset) / sizeof(uint64));
-		Package->ImportStore.ImportMap = TArrayView<const FPackageObjectIndex>(
-			reinterpret_cast<const FPackageObjectIndex*>(PackageHeaderData + PackageSummary->ImportMapOffset),
-			(PackageSummary->ExportMapOffset - PackageSummary->ImportMapOffset) / sizeof(FPackageObjectIndex));
-		Package->ExportMap = TArrayView<const FExportMapEntry>(
-			reinterpret_cast<const FExportMapEntry*>(PackageHeaderData + PackageSummary->ExportMapOffset),
-			(PackageSummary->ExportBundleEntriesOffset - PackageSummary->ExportMapOffset) / sizeof(FExportMapEntry));
-		check(Package->ExportMap.Num() == Package->Data.ExportInfo.ExportCount);
-
-		uint64 ExportBundleHeadersOffset = PackageSummary->GraphDataOffset;
-		uint64 ArcsDataOffset = ExportBundleHeadersOffset + Package->Data.ExportBundleHeadersSize;
-		uint64 ArcsDataSize = PackageSummary->HeaderSize - ArcsDataOffset;
-		Package->ArcsData = TArrayView<const uint8>(PackageHeaderData + ArcsDataOffset, ArcsDataSize);
-		FMemory::Memcpy(Package->Data.ExportBundlesMetaMemory, PackageHeaderData + ExportBundleHeadersOffset, Package->Data.ExportBundleHeadersSize);
-		FMemory::Memcpy(Package->Data.ExportBundlesMetaMemory + Package->Data.ExportBundleHeadersSize, PackageHeaderData + PackageSummary->ExportBundleEntriesOffset, Package->Data.ExportBundleEntriesSize);
-
-		Package->CreateUPackage(PackageSummary, PackageSummary->bHasVersioningInfo ? &VersioningInfo : nullptr);
-
-		Package->AllExportDataPtr = PackageHeaderData + PackageSummary->HeaderSize;
-
-		Package->ExportToBundleMappings.SetNum(Package->Data.ExportInfo.ExportCount);
-		for (int32 ExportBundleIndex = 0, ExportBundleCount = Package->Data.ExportInfo.ExportBundleCount; ExportBundleIndex < ExportBundleCount; ++ExportBundleIndex)
-		{
-			const FExportBundleHeader* ExportBundle = Package->Data.ExportBundleHeaders + ExportBundleIndex;
-			const FExportBundleEntry* BundleEntries = Package->Data.ExportBundleEntries + ExportBundle->FirstEntryIndex;
-			const FExportBundleEntry* BundleEntry = BundleEntries + Package->ExportBundleEntryIndex;
-			const FExportBundleEntry* BundleEntryEnd = BundleEntries + ExportBundle->EntryCount;
-			while (BundleEntry < BundleEntryEnd)
-			{
-				const FExportMapEntry& ExportMapEntry = Package->ExportMap[BundleEntry->LocalExportIndex];
-				FExportToBundleMapping& ExportToBundleMapping = Package->ExportToBundleMappings[BundleEntry->LocalExportIndex];
+				const FExportBundleEntry& BundleEntry = HeaderData->ExportBundleEntries[EntryIndex];
+				const FExportMapEntry& ExportMapEntry = HeaderData->ExportMap[BundleEntry.LocalExportIndex];
+				FAsyncPackageExportToBundleMapping& ExportToBundleMapping = Package->Data.ExportToBundleMappings[BundleEntry.LocalExportIndex];
 				ExportToBundleMapping.ExportHash = ExportMapEntry.PublicExportHash;
-				ExportToBundleMapping.BundleIndex[BundleEntry->CommandType] = ExportBundleIndex;
-				++BundleEntry;
+				ExportToBundleMapping.BundleIndex[BundleEntry.CommandType] = ExportBundleIndex;
 			}
 		}
 
-		TRACE_LOADTIME_PACKAGE_SUMMARY(Package, Package->Desc.PackageNameToLoad, PackageSummary->HeaderSize, Package->ImportStore.ImportMap.Num(), Package->Data.ExportInfo.ExportCount);
+		check(Package->Desc.PackageNameToLoad.IsNone());
+		Package->Desc.PackageNameToLoad = Package->HeaderData.PackageName;
+		check(Package->Desc.PackageIdToLoad == FPackageId::FromName(Package->Desc.PackageNameToLoad));
+		// Imported packages won't have a UPackage name set unless they were redirected, in which case they will have the source package name
+		if (Package->Desc.UPackageName.IsNone())
+		{
+			Package->Desc.UPackageName = Package->HeaderData.PackageName;
+		}
+		check(Package->Desc.UPackageId == FPackageId::FromName(Package->Desc.UPackageName));
+		Package->CreateUPackage(Package->HeaderData.PackageSummary, Package->HeaderData.VersioningInfo.GetPtrOrNull());
+
+		TRACE_LOADTIME_PACKAGE_SUMMARY(Package, Package->Desc.PackageNameToLoad, Package->HeaderData.PackageSummary->HeaderSize, Package->HeaderData.ImportMap.Num(), Package->HeaderData.ExportMap.Num());
 	}
 
 	Package->AsyncPackageLoadingState = EAsyncPackageLoadingState2::SetupDependencies;
@@ -3574,10 +3612,12 @@ EAsyncPackageState::Type FAsyncPackage2::Event_SetupDependencies(FAsyncLoadingTh
 		{
 			Package->SetupScriptDependencies();
 		}
-		Package->SetupSerializedArcs();
+		TArrayView<FEventLoadNode2> ExportBundleNodesView = Package->Data.ExportBundleNodes;
+		TArrayView<FAsyncPackage2*> ImportedAsyncPackagesView = Package->Data.ImportedAsyncPackages;
+		Package->SetupSerializedArcs(Package->HeaderData, ExportBundleNodesView.Left(Package->HeaderData.ExportBundleHeaders.Num() * EEventLoadNode2::ExportBundle_NumPhases), ImportedAsyncPackagesView.Left(Package->HeaderData.ImportedPackageIds.Num()));
 	}
 	Package->AsyncPackageLoadingState = EAsyncPackageLoadingState2::ProcessExportBundles;
-	for (int32 ExportBundleIndex = 0, ExportBundleCount = Package->Data.ExportInfo.ExportBundleCount; ExportBundleIndex < ExportBundleCount; ++ExportBundleIndex)
+	for (int32 ExportBundleIndex = 0; ExportBundleIndex < Package->Data.ExportBundleCount; ++ExportBundleIndex)
 	{
 		Package->GetExportBundleNode(ExportBundle_Process, ExportBundleIndex).ReleaseBarrier();
 	}
@@ -3616,18 +3656,21 @@ EAsyncPackageState::Type FAsyncPackage2::Event_ProcessExportBundle(FAsyncLoading
 #endif
 	};
 
-	check(InExportBundleIndex < Package->Data.ExportInfo.ExportBundleCount);
+	check(InExportBundleIndex < Package->Data.ExportBundleCount);
 	
 	if (!Package->bLoadHasFailed)
 	{
-		const FExportBundleHeader* ExportBundle = Package->Data.ExportBundleHeaders + InExportBundleIndex;
-		const FIoBuffer& IoBuffer = Package->IoRequest.GetResultOrDie();
-		const uint64 AllExportDataSize = IoBuffer.DataSize() - (Package->AllExportDataPtr - IoBuffer.Data());
+		const FAsyncPackageHeaderData* HeaderData = &Package->HeaderData;
+		const FExportBundleHeader* ExportBundle = &Package->HeaderData.ExportBundleHeaders[InExportBundleIndex];
+		FAsyncPackageSerializationState* SerializationState = &Package->SerializationState;
+		const TArrayView<FExportObject>& Exports = Package->Data.Exports;
+		const FIoBuffer& IoBuffer = SerializationState->IoRequest.GetResultOrDie();
+		const uint64 AllExportDataSize = IoBuffer.DataSize() - (SerializationState->AllExportDataPtr - IoBuffer.Data());
 		if (Package->ExportBundleEntryIndex == 0)
 		{
-			Package->CurrentExportDataPtr = Package->AllExportDataPtr + ExportBundle->SerialOffset;
+			SerializationState->CurrentExportDataPtr = SerializationState->AllExportDataPtr + ExportBundle->SerialOffset;
 		}
-		FExportArchive Ar(Package->AllExportDataPtr, Package->CurrentExportDataPtr, AllExportDataSize);
+		FExportArchive Ar(SerializationState->AllExportDataPtr, SerializationState->CurrentExportDataPtr, AllExportDataSize);
 		{
 			Ar.SetUEVer(Package->LinkerRoot->GetLinkerPackageVersion());
 			Ar.SetLicenseeUEVer(Package->LinkerRoot->GetLinkerLicenseeVersion());
@@ -3647,12 +3690,10 @@ EAsyncPackageState::Type FAsyncPackage2::Event_ProcessExportBundle(FAsyncLoading
 			Ar.ArAllowLazyLoading = true;
 
 			// FExportArchive special fields
-			Ar.CookedHeaderSize = Package->CookedHeaderSize;
 			Ar.PackageDesc = &Package->Desc;
-			Ar.NameMap = &Package->NameMap;
+			Ar.HeaderData = HeaderData;
 			Ar.ImportStore = &Package->ImportStore;
-			Ar.Exports = Package->Data.Exports;
-			Ar.ExportMap = Package->ExportMap.GetData();
+			Ar.Exports = Exports;
 			Ar.ExternalReadDependencies = &Package->ExternalReadDependencies;
 
 			// Check if the package is instanced
@@ -3662,32 +3703,30 @@ EAsyncPackageState::Type FAsyncPackage2::Event_ProcessExportBundle(FAsyncLoading
 				Package->Desc.UPackageName.ToString(Ar.InstancedPackageInstanceName);
 			}
 		}
-		const FExportBundleEntry* BundleEntries = Package->Data.ExportBundleEntries + ExportBundle->FirstEntryIndex;
-		const FExportBundleEntry* BundleEntry = BundleEntries + Package->ExportBundleEntryIndex;
-		const FExportBundleEntry* BundleEntryEnd = BundleEntries + ExportBundle->EntryCount;
-		check(BundleEntry <= BundleEntryEnd);
-		while (BundleEntry < BundleEntryEnd)
+
+		while (Package->ExportBundleEntryIndex < int32(ExportBundle->EntryCount))
 		{
+			const FExportBundleEntry& BundleEntry = HeaderData->ExportBundleEntries[ExportBundle->FirstEntryIndex + Package->ExportBundleEntryIndex];
 			if (ThreadState.IsTimeLimitExceeded(TEXT("Event_ProcessExportBundle")))
 			{
 				return EAsyncPackageState::TimeOut;
 			}
-			const FExportMapEntry& ExportMapEntry = Package->ExportMap[BundleEntry->LocalExportIndex];
-			FExportObject& Export = Package->Data.Exports[BundleEntry->LocalExportIndex];
+			const FExportMapEntry& ExportMapEntry = HeaderData->ExportMap[BundleEntry.LocalExportIndex];
+			FExportObject& Export = Exports[BundleEntry.LocalExportIndex];
 			Export.bFiltered = FilterExport(ExportMapEntry.FilterFlags);
 
-			if (BundleEntry->CommandType == FExportBundleEntry::ExportCommandType_Create)
+			if (BundleEntry.CommandType == FExportBundleEntry::ExportCommandType_Create)
 			{
-				Package->EventDrivenCreateExport(BundleEntry->LocalExportIndex);
+				Package->EventDrivenCreateExport(*HeaderData, Exports, BundleEntry.LocalExportIndex);
 			}
 			else
 			{
-				check(BundleEntry->CommandType == FExportBundleEntry::ExportCommandType_Serialize);
+				check(BundleEntry.CommandType == FExportBundleEntry::ExportCommandType_Serialize);
 
 				const uint64 CookedSerialSize = ExportMapEntry.CookedSerialSize;
 				UObject* Object = Export.Object;
 
-				check(Package->CurrentExportDataPtr + CookedSerialSize <= IoBuffer.Data() + IoBuffer.DataSize());
+				check(SerializationState->CurrentExportDataPtr + CookedSerialSize <= IoBuffer.Data() + IoBuffer.DataSize());
 				check(Object || Export.bFiltered || Export.bExportLoadFailed);
 
 				Ar.ExportBufferBegin(Object, ExportMapEntry.CookedSerialOffset, ExportMapEntry.CookedSerialSize);
@@ -3698,7 +3737,7 @@ EAsyncPackageState::Type FAsyncPackage2::Event_ProcessExportBundle(FAsyncLoading
 					TEXT("%s: Serial size mismatch: Expected read size %d, Remaining archive size: %d"),
 					Object ? *Object->GetFullName() : TEXT("null"), CookedSerialSize, uint64(Ar.TotalSize() - Pos));
 
-				const bool bSerialized = Package->EventDrivenSerializeExport(BundleEntry->LocalExportIndex, Ar);
+				const bool bSerialized = Package->EventDrivenSerializeExport(*HeaderData, Exports, BundleEntry.LocalExportIndex, Ar);
 				if (!bSerialized)
 				{
 					Ar.Skip(CookedSerialSize);
@@ -3712,20 +3751,19 @@ EAsyncPackageState::Type FAsyncPackage2::Event_ProcessExportBundle(FAsyncLoading
 
 				check((Object && !Object->HasAnyFlags(RF_NeedLoad)) || Export.bFiltered || Export.bExportLoadFailed);
 
-				Package->CurrentExportDataPtr += CookedSerialSize;
+				SerializationState->CurrentExportDataPtr += CookedSerialSize;
 			}
-			++BundleEntry;
 			++Package->ExportBundleEntryIndex;
 		}
 	}
 	
 	Package->ExportBundleEntryIndex = 0;
 
-	if (++Package->ProcessedExportBundlesCount == Package->Data.ExportInfo.ExportBundleCount)
+	if (++Package->ProcessedExportBundlesCount == Package->Data.ExportBundleCount)
 	{
 		Package->ProcessedExportBundlesCount = 0;
-		Package->ImportStore.ImportMap = TArrayView<const FPackageObjectIndex>();
-		Package->IoRequest.Release();
+		Package->HeaderData.OnReleaseHeaderBuffer();
+		Package->SerializationState.ReleaseIoRequest();
 
 		if (Package->ExternalReadDependencies.Num() == 0)
 		{
@@ -3743,7 +3781,7 @@ EAsyncPackageState::Type FAsyncPackage2::Event_ProcessExportBundle(FAsyncLoading
 	return EAsyncPackageState::Complete;
 }
 
-UObject* FAsyncPackage2::EventDrivenIndexToObject(FPackageObjectIndex Index, bool bCheckSerialized)
+UObject* FAsyncPackage2::EventDrivenIndexToObject(const FAsyncPackageHeaderData& Header, const TArrayView<const FExportObject>& Exports, FPackageObjectIndex Index, bool bCheckSerialized)
 {
 	UObject* Result = nullptr;
 	if (Index.IsNull())
@@ -3752,11 +3790,11 @@ UObject* FAsyncPackage2::EventDrivenIndexToObject(FPackageObjectIndex Index, boo
 	}
 	if (Index.IsExport())
 	{
-		Result = Data.Exports[Index.ToExport()].Object;
+		Result = Exports[Index.ToExport()].Object;
 	}
 	else if (Index.IsImport())
 	{
-		Result = ImportStore.FindOrGetImportObject(Index);
+		Result = ImportStore.FindOrGetImportObject(Header, Index);
 		UE_CLOG(!Result, LogStreaming, Warning, TEXT("Missing %s import 0x%llX for package %s"),
 			Index.IsScriptImport() ? TEXT("script") : TEXT("package"),
 			Index.Value(),
@@ -3789,12 +3827,12 @@ UObject* FAsyncPackage2::EventDrivenIndexToObject(FPackageObjectIndex Index, boo
 }
 
 
-void FAsyncPackage2::EventDrivenCreateExport(int32 LocalExportIndex)
+void FAsyncPackage2::EventDrivenCreateExport(const FAsyncPackageHeaderData& Header, const TArrayView<FExportObject>& Exports, int32 LocalExportIndex)
 {
 	TRACE_CPUPROFILER_EVENT_SCOPE(CreateExport);
 
-	const FExportMapEntry& Export = ExportMap[LocalExportIndex];
-	FExportObject& ExportObject = Data.Exports[LocalExportIndex];
+	const FExportMapEntry& Export = Header.ExportMap[LocalExportIndex];
+	FExportObject& ExportObject = Exports[LocalExportIndex];
 	UObject*& Object = ExportObject.Object;
 	check(!Object);
 
@@ -3803,7 +3841,7 @@ void FAsyncPackage2::EventDrivenCreateExport(int32 LocalExportIndex)
 	FName ObjectName;
 	{
 		TRACE_CPUPROFILER_EVENT_SCOPE(ObjectNameFixup);
-		ObjectName = NameMap.GetName(Export.ObjectName);
+		ObjectName = Header.NameMap.GetName(Export.ObjectName);
 	}
 
 	if (ExportObject.bFiltered | ExportObject.bExportLoadFailed)
@@ -3823,8 +3861,8 @@ void FAsyncPackage2::EventDrivenCreateExport(int32 LocalExportIndex)
 	// LLM_SCOPED_TAG_WITH_OBJECT_IN_SET(CastEventDrivenIndexToObject<UClass>(Export.ClassIndex, false), ELLMTagSet::AssetClasses);
 
 	bool bIsCompleteyLoaded = false;
-	UClass* LoadClass = Export.ClassIndex.IsNull() ? UClass::StaticClass() : CastEventDrivenIndexToObject<UClass>(Export.ClassIndex, true);
-	UObject* ThisParent = Export.OuterIndex.IsNull() ? LinkerRoot : EventDrivenIndexToObject(Export.OuterIndex, false);
+	UClass* LoadClass = Export.ClassIndex.IsNull() ? UClass::StaticClass() : CastEventDrivenIndexToObject<UClass>(Header, Exports, Export.ClassIndex, true);
+	UObject* ThisParent = Export.OuterIndex.IsNull() ? LinkerRoot : EventDrivenIndexToObject(Header, Exports, Export.OuterIndex, false);
 
 	if (!LoadClass)
 	{
@@ -3841,7 +3879,7 @@ void FAsyncPackage2::EventDrivenCreateExport(int32 LocalExportIndex)
 	check(!dynamic_cast<UObjectRedirector*>(ThisParent));
 	if (!Export.SuperIndex.IsNull())
 	{
-		ExportObject.SuperObject = EventDrivenIndexToObject(Export.SuperIndex, false);
+		ExportObject.SuperObject = EventDrivenIndexToObject(Header, Exports, Export.SuperIndex, false);
 		if (!ExportObject.SuperObject)
 		{
 			UE_ASYNC_PACKAGE_LOG(Error, Desc, TEXT("CreateExport"), TEXT("Could not find SuperStruct object for %s"), *ObjectName.ToString());
@@ -3851,7 +3889,7 @@ void FAsyncPackage2::EventDrivenCreateExport(int32 LocalExportIndex)
 	}
 	// Find the Archetype object for the one we are loading.
 	check(!Export.TemplateIndex.IsNull());
-	ExportObject.TemplateObject = EventDrivenIndexToObject(Export.TemplateIndex, true);
+	ExportObject.TemplateObject = EventDrivenIndexToObject(Header, Exports, Export.TemplateIndex, true);
 	if (!ExportObject.TemplateObject)
 	{
 		UE_ASYNC_PACKAGE_LOG(Error, Desc, TEXT("CreateExport"), TEXT("Could not find template object for %s"), *ObjectName.ToString());
@@ -3975,9 +4013,9 @@ void FAsyncPackage2::EventDrivenCreateExport(int32 LocalExportIndex)
 	check(Object);
 	EInternalObjectFlags FlagsToSet = EInternalObjectFlags::Async;
 	
-	if (Desc.bCanBeImported && Export.PublicExportHash)
+	if (Desc.bCanBeImported && Export.PublicExportHash &&
+		(Object->HasAnyFlags(RF_Public) || WITH_IOSTORE_IN_EDITOR))
 	{
-		check(Object->HasAnyFlags(RF_Public));
 		FlagsToSet |= EInternalObjectFlags::LoaderImport;
 		ImportStore.StoreGlobalObject(Desc.UPackageId, Export.PublicExportHash, Object);
 
@@ -3992,13 +4030,13 @@ void FAsyncPackage2::EventDrivenCreateExport(int32 LocalExportIndex)
 	Object->SetInternalFlags(FlagsToSet);
 }
 
-bool FAsyncPackage2::EventDrivenSerializeExport(int32 LocalExportIndex, FExportArchive& Ar)
+bool FAsyncPackage2::EventDrivenSerializeExport(const FAsyncPackageHeaderData& Header, const TArrayView<FExportObject>& Exports, int32 LocalExportIndex, FExportArchive& Ar)
 {
 	LLM_SCOPE(ELLMTag::UObject);
 	TRACE_CPUPROFILER_EVENT_SCOPE(SerializeExport);
 
-	const FExportMapEntry& Export = ExportMap[LocalExportIndex];
-	FExportObject& ExportObject = Data.Exports[LocalExportIndex];
+	const FExportMapEntry& Export = Header.ExportMap[LocalExportIndex];
+	FExportObject& ExportObject = Exports[LocalExportIndex];
 	UObject* Object = ExportObject.Object;
 	check(Object || (ExportObject.bFiltered | ExportObject.bExportLoadFailed));
 
@@ -4009,17 +4047,17 @@ bool FAsyncPackage2::EventDrivenSerializeExport(int32 LocalExportIndex, FExportA
 		if (ExportObject.bExportLoadFailed)
 		{
 			UE_ASYNC_PACKAGE_LOG(Warning, Desc, TEXT("SerializeExport"),
-				TEXT("Skipped failed export %s"), *NameMap.GetName(Export.ObjectName).ToString());
+				TEXT("Skipped failed export %s"), *Header.NameMap.GetName(Export.ObjectName).ToString());
 		}
 		else if (ExportObject.bFiltered)
 		{
 			UE_ASYNC_PACKAGE_LOG_VERBOSE(Verbose, Desc, TEXT("SerializeExport"),
-				TEXT("Skipped filtered export %s"), *NameMap.GetName(Export.ObjectName).ToString());
+				TEXT("Skipped filtered export %s"), *Header.NameMap.GetName(Export.ObjectName).ToString());
 		}
 		else
 		{
 			UE_ASYNC_PACKAGE_LOG_VERBOSE(VeryVerbose, Desc, TEXT("SerializeExport"),
-				TEXT("Skipped already serialized export %s"), *NameMap.GetName(Export.ObjectName).ToString());
+				TEXT("Skipped already serialized export %s"), *Header.NameMap.GetName(Export.ObjectName).ToString());
 		}
 		return false;
 	}
@@ -4288,7 +4326,7 @@ void FAsyncPackage2::ConditionalBeginPostLoad()
 	WaitForAllDependenciesToReachState(&FAsyncPackage2::AllDependenciesSerializedState, EAsyncPackageLoadingState2::ExportsDone, AsyncLoadingThread.ConditionalBeginPostLoadTick,
 		[](FAsyncPackage2* Package)
 		{
-			for (int32 ExportBundleIndex = 0, ExportBundleCount = Package->Data.ExportInfo.ExportBundleCount; ExportBundleIndex < ExportBundleCount; ++ExportBundleIndex)
+			for (int32 ExportBundleIndex = 0; ExportBundleIndex < Package->Data.ExportBundleCount; ++ExportBundleIndex)
 			{
 				Package->GetExportBundleNode(EEventLoadNode2::ExportBundle_PostLoad, ExportBundleIndex).ReleaseBarrier();
 			}
@@ -4335,7 +4373,7 @@ EAsyncPackageState::Type FAsyncPackage2::Event_PostLoadExportBundle(FAsyncLoadin
 		}
 	}*/
 	
-	check(InExportBundleIndex < Package->Data.ExportInfo.ExportBundleCount);
+	check(InExportBundleIndex < Package->Data.ExportBundleCount);
 
 	EAsyncPackageState::Type LoadingState = EAsyncPackageState::Complete;
 
@@ -4352,24 +4390,24 @@ EAsyncPackageState::Type FAsyncPackage2::Event_PostLoadExportBundle(FAsyncLoadin
 		const bool bAsyncPostLoadEnabled = FAsyncLoadingThreadSettings::Get().bAsyncPostLoadEnabled;
 		const bool bIsMultithreaded = Package->AsyncLoadingThread.IsMultithreaded();
 
-		const FExportBundleHeader* ExportBundle = Package->Data.ExportBundleHeaders + InExportBundleIndex;
-		const FExportBundleEntry* BundleEntries = Package->Data.ExportBundleEntries + ExportBundle->FirstEntryIndex;
-		const FExportBundleEntry* BundleEntry = BundleEntries + Package->ExportBundleEntryIndex;
-		const FExportBundleEntry* BundleEntryEnd = BundleEntries + ExportBundle->EntryCount;
-		check(BundleEntry <= BundleEntryEnd);
-		while (BundleEntry < BundleEntryEnd)
+		const FAsyncPackageHeaderData* HeaderData = &Package->HeaderData;
+		const FExportBundleHeader* ExportBundle = &Package->HeaderData.ExportBundleHeaders[InExportBundleIndex];
+		const TArrayView<FExportObject>& Exports = Package->Data.Exports;
+
+		while (Package->ExportBundleEntryIndex < int32(ExportBundle->EntryCount))
 		{
+			const FExportBundleEntry& BundleEntry = HeaderData->ExportBundleEntries[ExportBundle->FirstEntryIndex + Package->ExportBundleEntryIndex];
 			if (ThreadState.IsTimeLimitExceeded(TEXT("Event_PostLoadExportBundle")))
 			{
 				LoadingState = EAsyncPackageState::TimeOut;
 				break;
 			}
 			
-			if (BundleEntry->CommandType == FExportBundleEntry::ExportCommandType_Serialize)
+			if (BundleEntry.CommandType == FExportBundleEntry::ExportCommandType_Serialize)
 			{
 				do
 				{
-					FExportObject& Export = Package->Data.Exports[BundleEntry->LocalExportIndex];
+					FExportObject& Export = Exports[BundleEntry.LocalExportIndex];
 					if (Export.bFiltered | Export.bExportLoadFailed)
 					{
 						break;
@@ -4395,7 +4433,6 @@ EAsyncPackageState::Type FAsyncPackage2::Event_PostLoadExportBundle(FAsyncLoadin
 					}
 				} while (false);
 			}
-			++BundleEntry;
 			++Package->ExportBundleEntryIndex;
 		}
 
@@ -4410,7 +4447,7 @@ EAsyncPackageState::Type FAsyncPackage2::Event_PostLoadExportBundle(FAsyncLoadin
 
 	Package->ExportBundleEntryIndex = 0;
 
-	if (++Package->ProcessedExportBundlesCount == Package->Data.ExportInfo.ExportBundleCount)
+	if (++Package->ProcessedExportBundlesCount == Package->Data.ExportBundleCount)
 	{
 		Package->ProcessedExportBundlesCount = 0;
 		if (Package->LinkerRoot && !Package->bLoadHasFailed)
@@ -4424,7 +4461,7 @@ EAsyncPackageState::Type FAsyncPackage2::Event_PostLoadExportBundle(FAsyncLoadin
 
 		check(Package->AsyncPackageLoadingState == EAsyncPackageLoadingState2::PostLoad);
 		Package->AsyncPackageLoadingState = EAsyncPackageLoadingState2::DeferredPostLoad;
-		for (int32 ExportBundleIndex = 0, ExportBundleCount = Package->Data.ExportInfo.ExportBundleCount; ExportBundleIndex < ExportBundleCount; ++ExportBundleIndex)
+		for (int32 ExportBundleIndex = 0; ExportBundleIndex < Package->Data.ExportBundleCount; ++ExportBundleIndex)
 		{
 			Package->GetExportBundleNode(ExportBundle_DeferredPostLoad, ExportBundleIndex).ReleaseBarrier();
 		}
@@ -4442,7 +4479,7 @@ EAsyncPackageState::Type FAsyncPackage2::Event_DeferredPostLoadExportBundle(FAsy
 
 	FAsyncPackageScope2 PackageScope(Package);
 
-	check(InExportBundleIndex < Package->Data.ExportInfo.ExportBundleCount);
+	check(InExportBundleIndex < Package->Data.ExportBundleCount);
 	EAsyncPackageState::Type LoadingState = EAsyncPackageState::Complete;
 
 	if (Package->bLoadHasFailed)
@@ -4455,24 +4492,24 @@ EAsyncPackageState::Type FAsyncPackage2::Event_DeferredPostLoadExportBundle(FAsy
 		TGuardValue<bool> GuardIsRoutingPostLoad(PackageScope.ThreadContext.IsRoutingPostLoad, true);
 		FAsyncLoadingTickScope2 InAsyncLoadingTick(Package->AsyncLoadingThread);
 
-		const FExportBundleHeader* ExportBundle = Package->Data.ExportBundleHeaders + InExportBundleIndex;
-		const FExportBundleEntry* BundleEntries = Package->Data.ExportBundleEntries + ExportBundle->FirstEntryIndex;
-		const FExportBundleEntry* BundleEntry = BundleEntries + Package->ExportBundleEntryIndex;
-		const FExportBundleEntry* BundleEntryEnd = BundleEntries + ExportBundle->EntryCount;
-		check(BundleEntry <= BundleEntryEnd);
-		while (BundleEntry < BundleEntryEnd)
+		const FAsyncPackageHeaderData* HeaderData = &Package->HeaderData;
+		const FExportBundleHeader* ExportBundle = &Package->HeaderData.ExportBundleHeaders[InExportBundleIndex];
+		const TArrayView<FExportObject>& Exports = Package->Data.Exports;
+
+		while (Package->ExportBundleEntryIndex < int32(ExportBundle->EntryCount))
 		{
+			const FExportBundleEntry& BundleEntry = HeaderData->ExportBundleEntries[ExportBundle->FirstEntryIndex + Package->ExportBundleEntryIndex];
 			if (ThreadState.IsTimeLimitExceeded(TEXT("Event_DeferredPostLoadExportBundle")))
 			{
 				LoadingState = EAsyncPackageState::TimeOut;
 				break;
 			}
 
-			if (BundleEntry->CommandType == FExportBundleEntry::ExportCommandType_Serialize)
+			if (BundleEntry.CommandType == FExportBundleEntry::ExportCommandType_Serialize)
 			{
 				do
 				{
-					FExportObject& Export = Package->Data.Exports[BundleEntry->LocalExportIndex];
+					FExportObject& Export = Exports[BundleEntry.LocalExportIndex];
 					if (Export.bFiltered | Export.bExportLoadFailed)
 					{
 						break;
@@ -4493,7 +4530,6 @@ EAsyncPackageState::Type FAsyncPackage2::Event_DeferredPostLoadExportBundle(FAsy
 					}
 				} while (false);
 			}
-			++BundleEntry;
 			++Package->ExportBundleEntryIndex;
 		}
 	}
@@ -4505,7 +4541,7 @@ EAsyncPackageState::Type FAsyncPackage2::Event_DeferredPostLoadExportBundle(FAsy
 
 	Package->ExportBundleEntryIndex = 0;
 
-	if (++Package->ProcessedExportBundlesCount == Package->Data.ExportInfo.ExportBundleCount)
+	if (++Package->ProcessedExportBundlesCount == Package->Data.ExportBundleCount)
 	{
 		Package->ProcessedExportBundlesCount = 0;
 		check(Package->AsyncPackageLoadingState == EAsyncPackageLoadingState2::DeferredPostLoad);
@@ -4524,7 +4560,7 @@ FEventLoadNode2& FAsyncPackage2::GetPackageNode(EEventLoadNode2 Phase)
 
 FEventLoadNode2& FAsyncPackage2::GetExportBundleNode(EEventLoadNode2 Phase, uint32 ExportBundleIndex)
 {
-	check(ExportBundleIndex < uint32(Data.ExportInfo.ExportBundleCount));
+	check(ExportBundleIndex < uint32(Data.ExportBundleCount));
 	uint32 ExportBundleNodeIndex = ExportBundleIndex * EEventLoadNode2::ExportBundle_NumPhases + Phase;
 	return Data.ExportBundleNodes[ExportBundleNodeIndex];
 }
@@ -4690,9 +4726,8 @@ EAsyncPackageState::Type FAsyncLoadingThread2::ProcessLoadedPackagesFromGameThre
 
 			TArray<UObject*> CDODefaultSubobjects;
 			// Clear async loading flags (we still want RF_Async, but EInternalObjectFlags::AsyncLoading can be cleared)
-			for (int32 FinalizeIndex = 0, ExportCount = Package->Data.ExportInfo.ExportCount; FinalizeIndex < ExportCount; ++FinalizeIndex)
+			for (const FExportObject& Export : Package->Data.Exports)
 			{
-				const FExportObject& Export = Package->Data.Exports[FinalizeIndex];
 				if (Export.bFiltered | Export.bExportLoadFailed)
 				{
 					continue;
@@ -5573,7 +5608,7 @@ FAsyncPackage2::FAsyncPackage2(
 : Desc(InDesc)
 , AsyncLoadingThread(InAsyncLoadingThread)
 , GraphAllocator(InGraphAllocator)
-, ImportStore(*AsyncLoadingThread.PackageStore.Get(), AsyncLoadingThread.GlobalImportStore, AsyncLoadingThread.LoadedPackageStore, Desc, Data.ImportedPackageIds)
+, ImportStore(AsyncLoadingThread.GlobalImportStore, AsyncLoadingThread.LoadedPackageStore)
 {
 	TRACE_LOADTIME_NEW_ASYNC_PACKAGE(this);
 	AddRequestID(Desc.RequestID);
@@ -5596,7 +5631,7 @@ void FAsyncPackage2::CreatePackageNodes(const FAsyncLoadEventSpec* EventSpecs)
 void FAsyncPackage2::CreateExportBundleNodes(const FAsyncLoadEventSpec* EventSpecs)
 {
 	const int32 BarrierCount = 1;
-	for (int32 ExportBundleIndex = 0; ExportBundleIndex < Data.ExportInfo.ExportBundleCount; ++ExportBundleIndex)
+	for (int32 ExportBundleIndex = 0; ExportBundleIndex < Data.ExportBundleCount; ++ExportBundleIndex)
 	{
 		uint32 NodeIndex = EEventLoadNode2::ExportBundle_NumPhases * ExportBundleIndex;
 		for (int32 Phase = 0; Phase < EEventLoadNode2::ExportBundle_NumPhases; ++Phase)
@@ -5611,7 +5646,8 @@ FAsyncPackage2::~FAsyncPackage2()
 	TRACE_LOADTIME_DESTROY_ASYNC_PACKAGE(this);
 	UE_ASYNC_PACKAGE_LOG(Verbose, Desc, TEXT("AsyncThread: Deleted"), TEXT("Package deleted."));
 
-	ImportStore.ReleasePackageReferences();
+	ImportStore.ReleaseImportedPackageReferences(Desc, HeaderData.ImportedPackageIds);
+	ImportStore.ReleasePackageReference(Desc);
 
 	checkf(RefCount == 0, TEXT("RefCount is not 0 when deleting package %s"),
 		*Desc.PackageNameToLoad.ToString());
@@ -5891,7 +5927,8 @@ EAsyncPackageState::Type FAsyncPackage2::ProcessExternalReads(EExternalReadActio
 
 EAsyncPackageState::Type FAsyncPackage2::PostLoadInstances(FAsyncLoadingThreadState2& ThreadState)
 {
-	while (PostLoadInstanceIndex < Data.ExportInfo.ExportCount && !ThreadState.IsTimeLimitExceeded(TEXT("PostLoadInstances")))
+	const int32 ExportCount = Data.Exports.Num();
+	while (PostLoadInstanceIndex < ExportCount && !ThreadState.IsTimeLimitExceeded(TEXT("PostLoadInstances")))
 	{
 		const FExportObject& Export = Data.Exports[PostLoadInstanceIndex++];
 
@@ -5901,12 +5938,13 @@ EAsyncPackageState::Type FAsyncPackage2::PostLoadInstances(FAsyncLoadingThreadSt
 			ObjClass->PostLoadInstance(Export.Object);
 		}
 	}
-	return PostLoadInstanceIndex == Data.ExportInfo.ExportCount ? EAsyncPackageState::Complete : EAsyncPackageState::TimeOut;
+	return PostLoadInstanceIndex == ExportCount ? EAsyncPackageState::Complete : EAsyncPackageState::TimeOut;
 }
 
 EAsyncPackageState::Type FAsyncPackage2::CreateClusters(FAsyncLoadingThreadState2& ThreadState)
 {
-	while (DeferredClusterIndex < Data.ExportInfo.ExportCount && !ThreadState.IsTimeLimitExceeded(TEXT("CreateClusters")))
+	const int32 ExportCount = Data.Exports.Num();
+	while (DeferredClusterIndex < ExportCount && !ThreadState.IsTimeLimitExceeded(TEXT("CreateClusters")))
 	{
 		const FExportObject& Export = Data.Exports[DeferredClusterIndex++];
 
@@ -5916,7 +5954,7 @@ EAsyncPackageState::Type FAsyncPackage2::CreateClusters(FAsyncLoadingThreadState
 		}
 	}
 
-	return DeferredClusterIndex == Data.ExportInfo.ExportCount ? EAsyncPackageState::Complete : EAsyncPackageState::TimeOut;
+	return DeferredClusterIndex == ExportCount ? EAsyncPackageState::Complete : EAsyncPackageState::TimeOut;
 }
 
 void FAsyncPackage2::FinishUPackage()
