Fix handling of invalid package paths from clients

diff --git a/Engine/Source/Runtime/CoreUObject/Private/Misc/PackageName.cpp b/Engine/Source/Runtime/CoreUObject/Private/Misc/PackageName.cpp
--- a/Engine/Source/Runtime/CoreUObject/Private/Misc/PackageName.cpp
+++ b/Engine/Source/Runtime/CoreUObject/Private/Misc/PackageName.cpp
@@ -1038,6 +1038,11 @@ bool FPackageName::IsValidTextForLongPackageName(FStringView InLongPackageName,
 		if (OutReason) *OutReason = EErrorCode::LongPackageNames_PathWithTrailingSlash;
 		return false;
 	}
+	if (InLongPackageName.Contains(TEXT("//")))
+	{
+		if (OutReason) *OutReason = EErrorCode::LongPackageNames_PathWithDoubleSlash;
+		return false;
+	}
 	// Check for invalid characters
 	if (DoesPackageNameContainInvalidCharacters(InLongPackageName, OutReason))
 	{
@@ -2400,6 +2405,8 @@ FText FPackageName::FormatErrorAsText(FStringView InPath, EErrorCode ErrorCode)
 		return FText::Format(NSLOCTEXT("Core", "LongPackageNames_PathWithNoStartingSlash", "Input '{InPath}' does not start with a '/', which is required for LongPackageNames."), Args);
 	case EErrorCode::LongPackageNames_PathWithTrailingSlash:
 		return FText::Format(NSLOCTEXT("Core", "LongPackageNames_PathWithTrailingSlash", "Input '{InPath}' ends with a '/', which is invalid for LongPackageNames."), Args);
+	case EErrorCode::LongPackageNames_PathWithDoubleSlash:
+		return FText::Format(NSLOCTEXT("Core", "LongPackageNames_PathWithDoubleSlash", "Input '{InPath}' contains '//', which is invalid for LongPackageNames."), Args);
 	default:
 		UE_LOG(LogPackageName, Warning, TEXT("FPackageName::FormatErrorAsText: Invalid ErrorCode %d"), static_cast<int32>(ErrorCode));
 		return FText::Format(NSLOCTEXT("Core", "PackageNameUndocumentedError", "Input '{InPath} caused undocumented internal error."), Args);
diff --git a/Engine/Source/Runtime/CoreUObject/Public/Misc/PackageName.h b/Engine/Source/Runtime/CoreUObject/Public/Misc/PackageName.h
--- a/Engine/Source/Runtime/CoreUObject/Public/Misc/PackageName.h
+++ b/Engine/Source/Runtime/CoreUObject/Public/Misc/PackageName.h
@@ -48,6 +48,7 @@ public:
 		LongPackageNames_PathTooShort,
 		LongPackageNames_PathWithNoStartingSlash,
 		LongPackageNames_PathWithTrailingSlash,
+		LongPackageNames_PathWithDoubleSlash,
 	};
 
 	/**
