CoreRedirects: Fix handling of struct and enum redirects to allow registering redirects by FullObjectPath in addition to ObjectName.
#rnx

diff --git a/Engine/Source/Runtime/CoreUObject/Private/UObject/CoreRedirects.cpp b/Engine/Source/Runtime/CoreUObject/Private/UObject/CoreRedirects.cpp
--- a/Engine/Source/Runtime/CoreUObject/Private/UObject/CoreRedirects.cpp
+++ b/Engine/Source/Runtime/CoreUObject/Private/UObject/CoreRedirects.cpp
@@ -78,43 +78,43 @@ void FCoreRedirectObjectName::Reset()
 
 bool FCoreRedirectObjectName::Matches(const FCoreRedirectObjectName& Other, bool bCheckSubstring) const
 {
-	if (bCheckSubstring)
+	return Matches(Other, EMatchFlags::CheckSubString);
+}
+
+bool FCoreRedirectObjectName::Matches(const FCoreRedirectObjectName& Other, EMatchFlags MatchFlags) const
+{
+	bool bPartialLHS = !EnumHasAnyFlags(MatchFlags, EMatchFlags::DisallowPartialLHSMatch);
+	bool bPartialRHS = EnumHasAnyFlags(MatchFlags, EMatchFlags::AllowPartialRHSMatch);
+	bool bSubString = EnumHasAllFlags(MatchFlags, EMatchFlags::CheckSubString);
+
+	auto FieldMatches = [bPartialLHS, bPartialRHS, bSubString](FName LHS, FName RHS)
 	{
-		// Much slower
-		auto NameContains = [](FName Text, FName LookFor)
-		{
-			return WriteToString<FName::StringBufferSize>(Text).ToView().Contains(WriteToString<FName::StringBufferSize>(LookFor));
-		};
-		if (ObjectName != NAME_None && !NameContains(Other.ObjectName, ObjectName))
+		if (LHS.IsNone() || RHS.IsNone())
 		{
-			return false;
+			return LHS == RHS || ((!LHS.IsNone() || bPartialLHS) && (!RHS.IsNone() || bPartialRHS));
 		}
-		if (OuterName != NAME_None && !NameContains(Other.OuterName, OuterName))
+		else if (bSubString)
 		{
-			return false;
+			// Much slower
+			return WriteToString<FName::StringBufferSize>(RHS).ToView().Contains(WriteToString<FName::StringBufferSize>(LHS));
 		}
-		if (PackageName != NAME_None && !NameContains(Other.PackageName, PackageName))
+		else
 		{
-			return false;
+			return LHS == RHS;
 		}
+	};
+	if (!FieldMatches(ObjectName, Other.ObjectName))
+	{
+		return false;
 	}
-	else
+	if (!FieldMatches(OuterName, Other.OuterName))
 	{
-		// Check all names that are non empty
-		if (ObjectName != Other.ObjectName && ObjectName != NAME_None)
-		{
-			return false;
-		}
-		if (OuterName != Other.OuterName && OuterName != NAME_None)
-		{
-			return false;
-		}
-		if (PackageName != Other.PackageName && PackageName != NAME_None)
-		{
-			return false;
-		}
+		return false;
+	}
+	if (!FieldMatches(PackageName, Other.PackageName))
+	{
+		return false;
 	}
-
 	return true;
 }
 
@@ -378,15 +378,30 @@ static FORCEINLINE bool CheckRedirectFlagsMatch(ECoreRedirectFlags FlagsA, ECore
 	return bTypesOverlap && bCategoriesMatch;
 }
 
-bool FCoreRedirect::Matches(ECoreRedirectFlags InFlags, const FCoreRedirectObjectName& InName) const
+bool FCoreRedirect::Matches(ECoreRedirectFlags InFlags, const FCoreRedirectObjectName& InName,
+	ECoreRedirectMatchFlags MatchFlags) const
 {
 	// Check flags for Type/Category match
 	if (!CheckRedirectFlagsMatch(InFlags, RedirectFlags))
 	{
 		return false;
 	}
+	return Matches(InName, MatchFlags);
+}
+
 
-	return OldName.Matches(InName, IsSubstringMatch());
+bool FCoreRedirect::Matches(const FCoreRedirectObjectName& InName, ECoreRedirectMatchFlags MatchFlags) const
+{
+	FCoreRedirectObjectName::EMatchFlags NameMatchFlags = FCoreRedirectObjectName::EMatchFlags::None;
+	if (IsSubstringMatch())
+	{
+		NameMatchFlags |= FCoreRedirectObjectName::EMatchFlags::CheckSubString;
+	}
+	if (EnumHasAllFlags(MatchFlags, ECoreRedirectMatchFlags::AllowPartialMatch))
+	{
+		NameMatchFlags |= FCoreRedirectObjectName::EMatchFlags::AllowPartialRHSMatch;
+	}
+	return OldName.Matches(InName, NameMatchFlags);
 }
 
 bool FCoreRedirect::HasValueChanges() const
@@ -404,7 +419,7 @@ FCoreRedirectObjectName FCoreRedirect::RedirectName(const FCoreRedirectObjectNam
 	FCoreRedirectObjectName ModifyName(OldObjectName);
 
 	// Convert names that are different and non empty
-	if (OldName.ObjectName != NewName.ObjectName)
+	if (OldName.ObjectName != NewName.ObjectName && !ModifyName.ObjectName.IsNone())
 	{
 		if (IsSubstringMatch())
 		{
@@ -416,7 +431,7 @@ FCoreRedirectObjectName FCoreRedirect::RedirectName(const FCoreRedirectObjectNam
 		}
 	}
 	// If package name and object name are specified, copy outer also it was set to null explicitly
-	if (OldName.OuterName != NewName.OuterName || (NewName.PackageName != NAME_None && NewName.ObjectName != NAME_None))
+	if ((OldName.OuterName != NewName.OuterName || (NewName.PackageName != NAME_None && NewName.ObjectName != NAME_None)) && !ModifyName.OuterName.IsNone())
 	{
 		if (IsSubstringMatch())
 		{
@@ -427,7 +442,7 @@ FCoreRedirectObjectName FCoreRedirect::RedirectName(const FCoreRedirectObjectNam
 			ModifyName.OuterName = NewName.OuterName;
 		}
 	}
-	if (OldName.PackageName != NewName.PackageName)
+	if (OldName.PackageName != NewName.PackageName && !ModifyName.PackageName.IsNone())
 	{
 		if (IsSubstringMatch())
 		{
@@ -442,6 +457,23 @@ FCoreRedirectObjectName FCoreRedirect::RedirectName(const FCoreRedirectObjectNam
 	return ModifyName;
 }
 
+bool FCoreRedirect::IdenticalMatchRules(const FCoreRedirect& Other) const
+{
+	ECoreRedirectFlags TypeFlags = RedirectFlags & ECoreRedirectFlags::Type_AllMask;
+	if (TypeFlags == ECoreRedirectFlags::Type_Struct ||
+		TypeFlags == ECoreRedirectFlags::Type_Enum)
+	{
+		// Struct and Enum remappings are requested based solely on the Name - not considering package and outer -
+		// so they are identical if only the name matches
+		return RedirectFlags == Other.RedirectFlags && OldName.ObjectName == Other.OldName.ObjectName;
+	}
+	else
+	{
+		return RedirectFlags == Other.RedirectFlags && OldName == Other.OldName;
+	}
+}
+
+
 bool FCoreRedirects::bInitialized = false;
 #if WITH_COREREDIRECTS_MULTITHREAD_WARNING
 bool FCoreRedirects::bIsInMultithreadedPhase = false;
@@ -535,12 +567,13 @@ void FCoreRedirects::EnterMultithreadedPhase()
 }
 #endif
 
-bool FCoreRedirects::RedirectNameAndValues(ECoreRedirectFlags Type, const FCoreRedirectObjectName& OldObjectName, FCoreRedirectObjectName& NewObjectName, const FCoreRedirect** FoundValueRedirect)
+bool FCoreRedirects::RedirectNameAndValues(ECoreRedirectFlags Type, const FCoreRedirectObjectName& OldObjectName,
+	FCoreRedirectObjectName& NewObjectName, const FCoreRedirect** FoundValueRedirect, ECoreRedirectMatchFlags MatchFlags)
 {
 	NewObjectName = OldObjectName;
 	TArray<const FCoreRedirect*> FoundRedirects;
 
-	if (GetMatchingRedirects(Type, OldObjectName, FoundRedirects))
+	if (GetMatchingRedirects(Type, OldObjectName, FoundRedirects, MatchFlags))
 	{
 		// Sort them based on match
 		FoundRedirects.Sort([&OldObjectName](const FCoreRedirect& A, const FCoreRedirect& B) { return A.OldName.MatchScore(OldObjectName) > B.OldName.MatchScore(OldObjectName); });
@@ -556,7 +589,7 @@ bool FCoreRedirects::RedirectNameAndValues(ECoreRedirectFlags Type, const FCoreR
 			}
 
 			// Only apply if name match is still valid, if it already renamed part of it it may not apply any more. Don't want to check flags as those were checked in the gather step
-			if (Redirect->OldName.Matches(NewObjectName, Redirect->IsSubstringMatch()))
+			if (Redirect->Matches(NewObjectName, MatchFlags))
 			{
 				if (FoundValueRedirect && (Redirect->HasValueChanges() || Redirect->OverrideClassName.IsValid()))
 				{
@@ -582,21 +615,23 @@ bool FCoreRedirects::RedirectNameAndValues(ECoreRedirectFlags Type, const FCoreR
 	return NewObjectName != OldObjectName;
 }
 
-FCoreRedirectObjectName FCoreRedirects::GetRedirectedName(ECoreRedirectFlags Type, const FCoreRedirectObjectName& OldObjectName)
+FCoreRedirectObjectName FCoreRedirects::GetRedirectedName(ECoreRedirectFlags Type,
+	const FCoreRedirectObjectName& OldObjectName, ECoreRedirectMatchFlags MatchFlags)
 {
 	FCoreRedirectObjectName NewObjectName;
 
-	RedirectNameAndValues(Type, OldObjectName, NewObjectName, nullptr);
+	RedirectNameAndValues(Type, OldObjectName, NewObjectName, nullptr, MatchFlags);
 
 	return NewObjectName;
 }
 
-const TMap<FString, FString>* FCoreRedirects::GetValueRedirects(ECoreRedirectFlags Type, const FCoreRedirectObjectName& OldObjectName)
+const TMap<FString, FString>* FCoreRedirects::GetValueRedirects(ECoreRedirectFlags Type,
+	const FCoreRedirectObjectName& OldObjectName, ECoreRedirectMatchFlags MatchFlags)
 {
 	FCoreRedirectObjectName NewObjectName;
 	const FCoreRedirect* FoundRedirect = nullptr;
 
-	RedirectNameAndValues(Type, OldObjectName, NewObjectName, &FoundRedirect);
+	RedirectNameAndValues(Type, OldObjectName, NewObjectName, &FoundRedirect, MatchFlags);
 
 	if (FoundRedirect && FoundRedirect->ValueChanges.Num() > 0)
 	{
@@ -606,13 +641,17 @@ const TMap<FString, FString>* FCoreRedirects::GetValueRedirects(ECoreRedirectFla
 	return nullptr;
 }
 
-bool FCoreRedirects::GetMatchingRedirects(ECoreRedirectFlags SearchFlags, const FCoreRedirectObjectName& OldObjectName, TArray<const FCoreRedirect*>& FoundRedirects)
+bool FCoreRedirects::GetMatchingRedirects(ECoreRedirectFlags SearchFlags, const FCoreRedirectObjectName& OldObjectName,
+	TArray<const FCoreRedirect*>& FoundRedirects, ECoreRedirectMatchFlags MatchFlags)
 {
 	// Look for all redirects that match the given names and flags
 	bool bFound = false;
 	
-	// If we're not explicitly searching for packages or looking for removed things, add the implicit (Type=Package,Category=None) redirects
-	const bool bSearchPackageRedirects = !(SearchFlags & ECoreRedirectFlags::Type_Package) && !(SearchFlags & ECoreRedirectFlags::Category_Removed);
+	// If we're not explicitly searching for packages, and not looking for removed things, and not searching for partial matches
+	// based on ObjectName only, add the implicit (Type=Package,Category=None) redirects
+	const bool bSearchPackageRedirects = !(SearchFlags & ECoreRedirectFlags::Type_Package) &&
+		!(SearchFlags & ECoreRedirectFlags::Category_Removed) &&
+			(!(MatchFlags & ECoreRedirectMatchFlags::AllowPartialMatch) || !OldObjectName.PackageName.IsNone());
 
 	// Determine list of maps to look over, need to handle being passed multiple types in a bit mask
 	for (const TPair<ECoreRedirectFlags, FRedirectNameMap>& Pair : RedirectTypeMap)
@@ -628,7 +667,7 @@ bool FCoreRedirects::GetMatchingRedirects(ECoreRedirectFlags SearchFlags, const
 			{
 				for (const FCoreRedirect& CheckRedirect : *RedirectsForName)
 				{
-					if (CheckRedirect.Matches(PairFlags, OldObjectName))
+					if (CheckRedirect.Matches(PairFlags, OldObjectName, MatchFlags))
 					{
 						bFound = true;
 						FoundRedirects.Add(&CheckRedirect);
@@ -661,7 +700,12 @@ bool FCoreRedirects::FindPreviousNames(ECoreRedirectFlags SearchFlags, const FCo
 			{
 				for (const FCoreRedirect& Redirect : RedirectPair.Value)
 				{
-					if (Redirect.NewName.Matches(NewObjectName, Redirect.IsSubstringMatch()))
+					FCoreRedirectObjectName::EMatchFlags MatchFlags = FCoreRedirectObjectName::EMatchFlags::None;
+					if (Redirect.IsSubstringMatch())
+					{
+						MatchFlags |= FCoreRedirectObjectName::EMatchFlags::CheckSubString;
+					}
+					if (Redirect.NewName.Matches(NewObjectName, MatchFlags))
 					{
 						// Construct a reverse redirect
 						FCoreRedirect ReverseRedirect = FCoreRedirect(Redirect);
diff --git a/Engine/Source/Runtime/CoreUObject/Private/UObject/LinkerLoad.cpp b/Engine/Source/Runtime/CoreUObject/Private/UObject/LinkerLoad.cpp
--- a/Engine/Source/Runtime/CoreUObject/Private/UObject/LinkerLoad.cpp
+++ b/Engine/Source/Runtime/CoreUObject/Private/UObject/LinkerLoad.cpp
@@ -5614,7 +5614,8 @@ TArray<FName> FLinkerLoad::FindPreviousNamesForClass(FString CurrentClassPath, b
 FName FLinkerLoad::FindNewNameForEnum(const FName OldEnumName)
 {
 	FCoreRedirectObjectName OldName = FCoreRedirectObjectName(OldEnumName, NAME_None, NAME_None);
-	FCoreRedirectObjectName NewName = FCoreRedirects::GetRedirectedName(ECoreRedirectFlags::Type_Enum, OldName);
+	FCoreRedirectObjectName NewName = FCoreRedirects::GetRedirectedName(
+		ECoreRedirectFlags::Type_Enum, OldName, ECoreRedirectMatchFlags::AllowPartialMatch);
 
 	if (NewName != OldName)
 	{
@@ -5626,7 +5627,8 @@ FName FLinkerLoad::FindNewNameForEnum(const FName OldEnumName)
 FName FLinkerLoad::FindNewNameForStruct(const FName OldStructName)
 {
 	FCoreRedirectObjectName OldName = FCoreRedirectObjectName(OldStructName, NAME_None, NAME_None);
-	FCoreRedirectObjectName NewName = FCoreRedirects::GetRedirectedName(ECoreRedirectFlags::Type_Struct, OldName);
+	FCoreRedirectObjectName NewName = FCoreRedirects::GetRedirectedName(
+		ECoreRedirectFlags::Type_Struct, OldName, ECoreRedirectMatchFlags::AllowPartialMatch);
 
 	if (NewName != OldName)
 	{
diff --git a/Engine/Source/Runtime/CoreUObject/Public/UObject/CoreRedirects.h b/Engine/Source/Runtime/CoreUObject/Public/UObject/CoreRedirects.h
--- a/Engine/Source/Runtime/CoreUObject/Public/UObject/CoreRedirects.h
+++ b/Engine/Source/Runtime/CoreUObject/Public/UObject/CoreRedirects.h
@@ -48,6 +48,14 @@ enum class ECoreRedirectFlags : uint32
 };
 ENUM_CLASS_FLAGS(ECoreRedirectFlags);
 
+enum class ECoreRedirectMatchFlags
+{
+	None = 0,
+	/** The passed-in CoreRedirectObjectName has null fields in Package, Outer, or Name, and should still be allowed to match
+	 against redirectors that were created with a full Package.[Outer:]Name. */
+	AllowPartialMatch = (1 << 0),
+};
+ENUM_CLASS_FLAGS(ECoreRedirectMatchFlags);
 
 /**
  * An object path extracted into component names for matching. TODO merge with FSoftObjectPath?
@@ -102,8 +110,26 @@ struct COREUOBJECT_API FCoreRedirectObjectName
 		return !(*this == Other);
 	}
 
-	/** Returns true if the passed in name matches requirements, will ignore names that are none */
-	bool Matches(const FCoreRedirectObjectName& Other, bool bCheckSubstring = false) const;
+	/** Flags for the Matches function. These flags overlap but are lower-level than ECoreRedirectMatchFlags. */
+	enum class EMatchFlags
+	{
+		None = 0,
+		/** Do not match if LHS (aka *this) has null fields that RHS (aka Other) does not. Default is to match. */
+		DisallowPartialLHSMatch = (1 << 0),
+		/** Match even if RHS (aka Other) has null fields that LHS (aka *this) does not. Default is to NOT match. */
+		AllowPartialRHSMatch = (1 << 1),
+		/**
+		 * LHS fields (aka *this) are searchstrings; RHS (aka Other) fields are searched for that substring.
+		 * Default is to require a complete string match LHS == RHS.
+		 * This flag makes the match more expensive and should be avoided when possible.
+		 */
+		CheckSubString = (1 << 2),
+	};
+	/** Returns true if the passed in name matches requirements. */
+	bool Matches(const FCoreRedirectObjectName& Other, EMatchFlags MatchFlags = EMatchFlags::None) const;
+
+	UE_DEPRECATED(5.1, "Use EMatchFlags::CheckSubString to pass in bCheckSubstring=true.")
+	bool Matches(const FCoreRedirectObjectName& Other, bool bCheckSubstring) const;
 
 	/** Returns integer of degree of match. 0 if doesn't match at all, higher integer for better matches */
 	int32 MatchScore(const FCoreRedirectObjectName& Other) const;
@@ -142,6 +168,7 @@ struct COREUOBJECT_API FCoreRedirectObjectName
 	/** Turn it back into an FString */
 	static FString CombineNames(FName NewName, FName NewOuter, FName NewPackage);
 };
+ENUM_CLASS_FLAGS(FCoreRedirectObjectName::EMatchFlags);
 
 /** 
  * A single redirection from an old name to a new name, parsed out of an ini file
@@ -190,7 +217,11 @@ struct COREUOBJECT_API FCoreRedirect
 	const TCHAR* ParseValueChanges(const TCHAR* Buffer);
 
 	/** Returns true if the passed in name and flags match requirements */
-	bool Matches(ECoreRedirectFlags InFlags, const FCoreRedirectObjectName& InName) const;
+	bool Matches(ECoreRedirectFlags InFlags, const FCoreRedirectObjectName& InName,
+		ECoreRedirectMatchFlags MatchFlags = ECoreRedirectMatchFlags::None) const;
+	/** Returns true if the passed in name matches requirements */
+	bool Matches(const FCoreRedirectObjectName& InName,
+		ECoreRedirectMatchFlags MatchFlags = ECoreRedirectMatchFlags::None) const;
 
 	/** Returns true if this has value redirects */
 	bool HasValueChanges() const;
@@ -202,10 +233,7 @@ struct COREUOBJECT_API FCoreRedirect
 	FCoreRedirectObjectName RedirectName(const FCoreRedirectObjectName& OldObjectName) const;
 
 	/** See if search criteria is identical */
-	bool IdenticalMatchRules(const FCoreRedirect& Other) const
-	{
-		return RedirectFlags == Other.RedirectFlags && OldName == Other.OldName;
-	}
+	bool IdenticalMatchRules(const FCoreRedirect& Other) const;
 
 	/** Returns the name used as the key into the acceleration map */
 	FName GetSearchKey() const
@@ -223,13 +251,17 @@ struct COREUOBJECT_API FCoreRedirects
 	static void Initialize();
 
 	/** Returns a redirected version of the object name. If there are no valid redirects, it will return the original name */
-	static FCoreRedirectObjectName GetRedirectedName(ECoreRedirectFlags Type, const FCoreRedirectObjectName& OldObjectName);
+	static FCoreRedirectObjectName GetRedirectedName(ECoreRedirectFlags Type,
+		const FCoreRedirectObjectName& OldObjectName, ECoreRedirectMatchFlags MatchFlags = ECoreRedirectMatchFlags::None);
 
 	/** Returns map of String->String value redirects for the object name, or nullptr if none found */
-	static const TMap<FString, FString>* GetValueRedirects(ECoreRedirectFlags Type, const FCoreRedirectObjectName& OldObjectName);
+	static const TMap<FString, FString>* GetValueRedirects(ECoreRedirectFlags Type,
+		const FCoreRedirectObjectName& OldObjectName, ECoreRedirectMatchFlags MatchFlags = ECoreRedirectMatchFlags::None);
 
 	/** Performs both a name redirect and gets a value redirect struct if it exists. Returns true if either redirect found */
-	static bool RedirectNameAndValues(ECoreRedirectFlags Type, const FCoreRedirectObjectName& OldObjectName, FCoreRedirectObjectName& NewObjectName, const FCoreRedirect** FoundValueRedirect);
+	static bool RedirectNameAndValues(ECoreRedirectFlags Type, const FCoreRedirectObjectName& OldObjectName,
+		FCoreRedirectObjectName& NewObjectName, const FCoreRedirect** FoundValueRedirect,
+		ECoreRedirectMatchFlags MatchFlags = ECoreRedirectMatchFlags::None);
 
 	/** Returns true if this name has been registered as explicitly missing */
 	static bool IsKnownMissing(ECoreRedirectFlags Type, const FCoreRedirectObjectName& ObjectName);
@@ -258,7 +290,8 @@ struct COREUOBJECT_API FCoreRedirects
 	static bool FindPreviousNames(ECoreRedirectFlags Type, const FCoreRedirectObjectName& NewObjectName, TArray<FCoreRedirectObjectName>& PreviousNames);
 
 	/** Returns list of all core redirects that match requirements */
-	static bool GetMatchingRedirects(ECoreRedirectFlags Type, const FCoreRedirectObjectName& OldObjectName, TArray<const FCoreRedirect*>& FoundRedirects);
+	static bool GetMatchingRedirects(ECoreRedirectFlags Type, const FCoreRedirectObjectName& OldObjectName,
+		TArray<const FCoreRedirect*>& FoundRedirects, ECoreRedirectMatchFlags MatchFlags = ECoreRedirectMatchFlags::None);
 
 	/** Parse all redirects out of a given ini file */
 	static bool ReadRedirectsFromIni(const FString& IniName);
