CoreRedirects: Fix bug in CL 19809751 that blocked matching of enum value redirects by the full objectpath.
#rnx

diff --git a/Engine/Source/Runtime/CoreUObject/Private/UObject/CoreRedirects.cpp b/Engine/Source/Runtime/CoreUObject/Private/UObject/CoreRedirects.cpp
--- a/Engine/Source/Runtime/CoreUObject/Private/UObject/CoreRedirects.cpp
+++ b/Engine/Source/Runtime/CoreUObject/Private/UObject/CoreRedirects.cpp
@@ -460,15 +460,19 @@ FCoreRedirectObjectName FCoreRedirect::RedirectName(const FCoreRedirectObjectNam
 bool FCoreRedirect::IdenticalMatchRules(const FCoreRedirect& Other) const
 {
 	ECoreRedirectFlags TypeFlags = RedirectFlags & ECoreRedirectFlags::Type_AllMask;
-	if (TypeFlags == ECoreRedirectFlags::Type_Struct ||
-		TypeFlags == ECoreRedirectFlags::Type_Enum)
+	if (TypeFlags == ECoreRedirectFlags::Type_Struct)
 	{
-		// Struct and Enum remappings are requested based solely on the Name - not considering package and outer -
+		// Struct remappings are requested based solely on the Name - not considering package and outer -
 		// so they are identical if only the name matches
 		return RedirectFlags == Other.RedirectFlags && OldName.ObjectName == Other.OldName.ObjectName;
 	}
 	else
 	{
+		// Enum remappings are requested in some locations based solely on the Name - not considering package and outer -
+		// but they are also requested from UEnum::GetIndexByNameString -> GetValueRedirects by full name
+		// Return IdenticalMatchRules=true only if the full name matches.
+		// TODO: Change FCoreRedirects::AddSingleRedirect to remove the package and outer name so that we can return IdenticalMatchRules=true for partial name enum matches
+
 		return RedirectFlags == Other.RedirectFlags && OldName == Other.OldName;
 	}
 }
