Do not selectively process async package when using the old async loading EDL
Do not pop when peeking flush requests, they are removed from the GT

diff --git a/Engine/Source/Runtime/CoreUObject/Private/Serialization/AsyncLoading.cpp b/Engine/Source/Runtime/CoreUObject/Private/Serialization/AsyncLoading.cpp
--- a/Engine/Source/Runtime/CoreUObject/Private/Serialization/AsyncLoading.cpp
+++ b/Engine/Source/Runtime/CoreUObject/Private/Serialization/AsyncLoading.cpp
@@ -4225,18 +4225,6 @@ FFlushRequest FAsyncLoadingThread::PeekFlushRequest()
 		}
 	}
 
-	if (FlushRequest)
-	{
-		// Validate that the flush request is still valid and pending
-		if (!ContainsRequestInternal(FlushRequest.GetId()))
-		{
-			FScopeLock Lock(&FlushRequestCritical);
-			FFlushRequest Top = FlushRequests.Pop();
-			ensure(Top == FlushRequest);
-			FlushRequest.Reset();
-		}
-	}
-
 	// Adjust request id
 	AdjustFlushRequest(FlushRequest);
 	return FlushRequest;
@@ -4310,7 +4298,10 @@ EAsyncPackageState::Type FAsyncLoadingThread::ProcessLoadedPackages(bool bUseTim
 
 		FAsyncPackage* Package = LoadedPackagesToProcess[PackageIndex];
 		// if we are currently flushing a specific load request, only post load related package objects
-		if (ShouldProcessPackage(Package, FlushRequest))
+		if (ShouldProcessPackage(Package, FlushRequest) 
+			// Always process loaded packages when using the EDL 
+			// This is because objects created/loaded are passed to package for processing through the serialize context without consideration to which their belong
+			|| GEventDrivenLoaderEnabled)
 		{
 			SCOPED_LOADTIMER(ProcessLoadedPackagesTime);
 			Result = Package->PostLoadDeferredObjects(TickStartTime, bUseTimeLimit, TimeLimit);
