Fix error language
#rnx

diff --git a/Engine/Source/Runtime/CoreUObject/Private/UObject/LinkerLoad.cpp b/Engine/Source/Runtime/CoreUObject/Private/UObject/LinkerLoad.cpp
--- a/Engine/Source/Runtime/CoreUObject/Private/UObject/LinkerLoad.cpp
+++ b/Engine/Source/Runtime/CoreUObject/Private/UObject/LinkerLoad.cpp
@@ -1424,9 +1424,9 @@ FLinkerLoad::ELinkerStatus FLinkerLoad::UpdateFromPackageFileSummary()
 			bCustomVersionIsLatest = false;
 			if (Diff.Type == ECustomVersionDifference::Missing)
 			{
-				// Loading a package with custom integration that we don't know about!
+				// Loading a package with custom version that we don't know about!
 				// Temporarily just warn and continue. @todo: this needs to be fixed properly
-				UE_LOG(LogLinker, Warning, TEXT("Package %s was saved with a custom integration that is not present. Tag %s  Version %d"), *GetDebugName(),
+				UE_LOG(LogLinker, Warning, TEXT("Package %s was saved with a custom version that is not present. Tag %s  Version %d"), *GetDebugName(),
 					*Diff.Version->Key.ToString(), Diff.Version->Version);
 			}
 			else if (Diff.Type == ECustomVersionDifference::Invalid)
