EditorPackageLoader: Fixed issue with recursive imports of cooked packages when all three loaders are involved in loading a package.
#rnx

diff --git a/Engine/Source/Runtime/CoreUObject/Private/UObject/UObjectGlobals.cpp b/Engine/Source/Runtime/CoreUObject/Private/UObject/UObjectGlobals.cpp
--- a/Engine/Source/Runtime/CoreUObject/Private/UObject/UObjectGlobals.cpp
+++ b/Engine/Source/Runtime/CoreUObject/Private/UObject/UObjectGlobals.cpp
@@ -1176,12 +1176,27 @@ UPackage* LoadPackageInternal(UPackage* InOuter, const FPackagePath& PackagePath
 		}
 
 		GSyncLoadUsingAsyncLoaderCount++;
+
+#if WITH_EDITOR
+		// Temporarily clear the list of loaded objects before entering async loading. Otherwise the (uncooked) async loader might try to preload them which in turn can lead to
+		// recursive imports of the same cooked package.
+		FUObjectSerializeContext* LoadContext = FUObjectThreadContext::Get().GetSerializeContext();
+		TArray<UObject*>& ThreadObjLoaded = LoadContext->PRIVATE_GetObjectsLoadedInternalUseOnly();
+		TArray<UObject*> SavedObjLoaded;
+		Swap(SavedObjLoaded, ThreadObjLoaded);
+#endif
 		int32 RequestID = LoadPackageAsync(PackagePath, PackageName);
 
 		if (RequestID != INDEX_NONE)
 		{
 			FlushAsyncLoading(RequestID);
 		}
+
+#if WITH_EDITOR
+		check(ThreadObjLoaded.IsEmpty());
+		Swap(SavedObjLoaded, ThreadObjLoaded);
+#endif
+
 		GSyncLoadUsingAsyncLoaderCount--;
 
 		return (InOuter ? InOuter : FindObjectFast<UPackage>(nullptr, PackageName));
