BulkDataRegistry: Fix code that can create duplicate BulkData identifiers. When loading a duplicate from an existing package, give a warning and update the BulkData during Serialize. Add resaveondemand=bulkdataduplicates argument in resave packages to resave these existing packages.
#rnx

diff --git a/Engine/Source/Editor/UnrealEd/Classes/Commandlets/ResavePackagesCommandlet.h b/Engine/Source/Editor/UnrealEd/Classes/Commandlets/ResavePackagesCommandlet.h
--- a/Engine/Source/Editor/UnrealEd/Classes/Commandlets/ResavePackagesCommandlet.h
+++ b/Engine/Source/Editor/UnrealEd/Classes/Commandlets/ResavePackagesCommandlet.h
@@ -129,6 +129,7 @@ protected:
 	FString ForceHLODSetupAsset;
 	FString HLODSkipToMap;
 	bool bForceUATEnvironmentVariableSet;
+	bool bResaveOnDemand = false;
 
 	/** Running count of packages that got modified and will need to be resaved */
 	int32 PackagesConsideredForResave;
@@ -154,6 +155,11 @@ protected:
 	/** A queue containing source control operations to be performed in batches. */
 	TPimplPtr<class FQueuedSourceControlOperations> SourceControlQueue;
 
+	/** If running resaveondemand, only packages reported into this set are resaved. */
+	TSet<FName> ResaveOnDemandPackages;
+	TSet<FName> ResaveOnDemandSystems;
+	FCriticalSection ResaveOnDemandPackagesLock;
+
 	/**
 	 * Evaluates the command-line to determine which maps to check.  By default all maps are checked
 	 * Provides child classes with a chance to initialize any variables, parse the command line, etc.
@@ -167,6 +173,8 @@ protected:
 
 	void ParseSourceControlOptions(const TArray<FString>& Tokens);
 
+	void OnAddResaveOnDemandPackage(FName SystemName, FName PackageName);
+
 	/** Loads and saves a single package */
 	virtual void LoadAndSaveOnePackage(const FString& Filename);
 
@@ -253,6 +261,9 @@ protected:
 	// Print out a message only if running in very verbose mode
 	void VerboseMessage(const FString& Message);
 
+	/** Parse commandline to decide whether resaveondemand is activated for the given system. */
+	TSet<FName> ParseResaveOnDemandSystems();
+
 public:		
 	//~ Begin UCommandlet Interface
 	virtual int32 Main(const FString& Params) override;
diff --git a/Engine/Source/Editor/UnrealEd/Private/Commandlets/ContentCommandlets.cpp b/Engine/Source/Editor/UnrealEd/Private/Commandlets/ContentCommandlets.cpp
--- a/Engine/Source/Editor/UnrealEd/Private/Commandlets/ContentCommandlets.cpp
+++ b/Engine/Source/Editor/UnrealEd/Private/Commandlets/ContentCommandlets.cpp
@@ -54,7 +54,9 @@
 #include "AssetCompilingManager.h"
 #include "PackageHelperFunctions.h"
 #include "PackageTools.h"
+#include "String/ParseTokens.h"
 #include "UObject/PackageTrailer.h"
+#include "UObject/SavePackage.h"
 
 DEFINE_LOG_CATEGORY(LogContentCommandlet);
 
@@ -341,6 +343,22 @@ int32 UResavePackagesCommandlet::InitializeResaveParameters( const TArray<FStrin
 	// This option will filter the package list and only save packages that are redirectors, or that reference redirectors
 	const bool bFixupRedirects = (Switches.Contains(TEXT("FixupRedirects")) || Switches.Contains(TEXT("FixupRedirectors")));
 
+	// ResaveOnDemand is a generic way for systems to specify during load that the commandlet should resave a package
+	ResaveOnDemandSystems = ParseResaveOnDemandSystems();
+	if (!ResaveOnDemandSystems.IsEmpty())
+	{
+		TStringBuilder<64> SystemNames;
+		for (FName SystemName : ResaveOnDemandSystems)
+		{
+			SystemNames << SystemName << TEXT(",");
+		}
+		SystemNames.RemoveSuffix(1); // Remove the trailing ,
+		UE_LOG(LogContentCommandlet, Display, TEXT("ResaveOnDemand=%s. Only saving packages that are reported via UE::SavePackageUtilities::OnAddResaveOnDemandPackage."),
+			*SystemNames);
+		bResaveOnDemand = true;
+		UE::SavePackageUtilities::OnAddResaveOnDemandPackage.BindUObject(this, &UResavePackagesCommandlet::OnAddResaveOnDemandPackage);
+	}
+
 	// This option allows the dependency graph and soft object path redirect map to be populated. This is useful if you want soft object references to redirectors to be followed to the destination asset at save time.
 	const bool bSearchAllAssets = Switches.Contains(TEXT("SearchAllAssets"));
 
@@ -568,6 +586,43 @@ void UResavePackagesCommandlet::ParseSourceControlOptions(const TArray<FString>&
 	}
 }
 
+void UResavePackagesCommandlet::OnAddResaveOnDemandPackage(FName SystemName, FName PackageName)
+{
+	FScopeLock ScopeLock(&ResaveOnDemandPackagesLock);
+	if (ResaveOnDemandSystems.Contains(SystemName))
+	{
+		ResaveOnDemandPackages.Add(PackageName);
+	}
+}
+
+
+TSet<FName> UResavePackagesCommandlet::ParseResaveOnDemandSystems()
+{
+	using namespace UE::String;
+	TSet<FName> SystemNames;
+
+	const TCHAR* CommandLineStream = FCommandLine::Get();
+	for (;;)
+	{
+		FString Token = FParse::Token(CommandLineStream, false /* bUseEscape */);
+		if (Token.IsEmpty())
+		{
+			break;
+		}
+		FString TokenSystemNames;
+		if (FParse::Value(*Token, TEXT("-resaveondemand="), TokenSystemNames))
+		{
+			UE::String::ParseTokensMultiple(TokenSystemNames, TConstArrayView<TCHAR>({ '+', ',' }),
+				[&SystemNames](FStringView TokenSystemName)
+				{
+					SystemNames.Add(FName(TokenSystemName));
+				}, EParseTokensOptions::Trim | EParseTokensOptions::SkipEmpty);
+		}
+	}
+	return SystemNames;
+}
+
+
 bool UResavePackagesCommandlet::ShouldSkipPackage(const FString& Filename)
 {
 	return false;
@@ -752,6 +807,16 @@ void UResavePackagesCommandlet::LoadAndSaveOnePackage(const FString& Filename)
 				Package->SetPackageFlags(PKG_FilterEditorOnly);
 			}
 
+			if (bResaveOnDemand && bSavePackage)
+			{
+				// Finish all async loading for the package, in case the systems marking it for load do so only from the async loading
+				FAssetCompilingManager::Get().FinishAllCompilation();
+				FlushAsyncLoading();
+				// Skip saving the package if no systems requested it
+				FScopeLock ScopeLock(&ResaveOnDemandPackagesLock);
+				bSavePackage = ResaveOnDemandPackages.Contains(Package->GetFName());
+			}
+
 			if (bSavePackage == true)
 			{
 				bool bIsEmpty = true;
diff --git a/Engine/Source/Editor/UnrealEd/Private/EditorDomain/BulkDataRegistryEditorDomain.cpp b/Engine/Source/Editor/UnrealEd/Private/EditorDomain/BulkDataRegistryEditorDomain.cpp
--- a/Engine/Source/Editor/UnrealEd/Private/EditorDomain/BulkDataRegistryEditorDomain.cpp
+++ b/Engine/Source/Editor/UnrealEd/Private/EditorDomain/BulkDataRegistryEditorDomain.cpp
@@ -155,31 +155,131 @@ void FBulkDataRegistryEditorDomain::OnEndLoadPackage(const FEndLoadPackageContex
 	}
 }
 
-void FBulkDataRegistryEditorDomain::Register(UPackage* Owner, const UE::Serialization::FEditorBulkData& BulkData)
+UE::BulkDataRegistry::ERegisterResult
+FBulkDataRegistryEditorDomain::TryRegister(UPackage* Owner, const UE::Serialization::FEditorBulkData& BulkData)
 {
 	if (!BulkData.GetIdentifier().IsValid())
 	{
-		return;
+		return UE::BulkDataRegistry::ERegisterResult::Success;
 	}
 
-	FName PackageName = NAME_None;
+	bool bAllowedToReadWritePayloadIdFromCache = false;
+	FName PackageName;
 	UE::Serialization::FEditorBulkData CopyBulk(BulkData.CopyTornOff());
-	if (Owner
-		&& Owner->GetFileSize() // We only record the BulkDataList for disk packages
-		&& !Owner->GetHasBeenEndLoaded() // We only record BulkDatas that are loaded before the package finishes loading
-		&& CopyBulk.CanSaveForRegistry()
-		)
+	UE::Serialization::FEditorBulkData PendingPackageBulk;
+	if (Owner)
 	{
 		PackageName = Owner->GetFName();
-		AddPendingPackageBulkData(PackageName, CopyBulk);
+		if (Owner->GetFileSize() // We only cache the BulkDataList for disk packages
+			&& !Owner->GetHasBeenEndLoaded() // We only cache BulkDatas that are loaded before the package finishes loading
+			&& CopyBulk.CanSaveForRegistry()
+			)
+		{
+			bAllowedToReadWritePayloadIdFromCache = true;
+			PendingPackageBulk = CopyBulk;
+		}
 	}
 
 	{
 		FWriteScopeLock RegistryScopeLock(RegistryLock);
 		check(bActive); // Registrations should not come in after we destruct
-		Registry.Add(BulkData.GetIdentifier(), FRegisteredBulk(MoveTemp(CopyBulk), PackageName));
+		UE::BulkDataRegistry::Private::FRegisteredBulk& RegisteredBulk = Registry.FindOrAdd(BulkData.GetIdentifier());
+		if (RegisteredBulk.bRegistered)
+		{
+			bool bAllowSharing = false;
+			// If the original BulkData has left memory and a new bulkdata registers with the same package owner,
+			// assume that it is the same BulkData being reloaded and allow the sharing
+			bAllowSharing |= (!RegisteredBulk.bInMemory && PackageName == RegisteredBulk.PackageName);
+			if (!bAllowSharing)
+			{
+				return UE::BulkDataRegistry::ERegisterResult::AlreadyExists;
+			}
+
+			if (RegisteredBulk.PackageName.IsNone())
+			{
+				RegisteredBulk.PackageName = PackageName;
+			}
+			RegisteredBulk.bRegistered = true;
+			RegisteredBulk.bInMemory = true;
+			RegisteredBulk.bAllowedToWritePayloadIdToCache = false;
+			RegisteredBulk.bPayloadAvailable = true;
+
+			// For updated registrations, skip reading payloadid from cache, and call ResaveSizeTracker.Update rather than Register.
+			ResaveSizeTracker.UpdateRegistrationData(Owner, BulkData);
+			return UE::BulkDataRegistry::ERegisterResult::Success;
+		}
+		else
+		{
+			bool bCachedLocationMatches = false;
+			UE::Serialization::FEditorBulkData& TargetBulkData = RegisteredBulk.BulkData;
+
+			if (TargetBulkData.GetIdentifier().IsValid())
+			{
+				// This BulkData was added when loading the cached list of BulkData for the package, but has not been registered by an in-memory bulkdata
+				check(TargetBulkData.GetIdentifier() == BulkData.GetIdentifier());
+				bCachedLocationMatches = BulkData.LocationMatches(TargetBulkData);
+			}
+
+			if (!bCachedLocationMatches || !BulkData.HasPlaceholderPayloadId() || TargetBulkData.HasPlaceholderPayloadId())
+			{
+				// Copy the new BulkData over the value we got from the cache; the new BulkData is more authoritative
+				TargetBulkData = MoveTemp(CopyBulk);
+				RegisteredBulk.bAllowedToWritePayloadIdToCache = bAllowedToReadWritePayloadIdFromCache;
+			}
+			else
+			{
+				// Otherwise Keep the TargetBulkData, since it matches the location and has already calculated the PayloadId
+				RegisteredBulk.bAllowedToWritePayloadIdToCache = false;
+			}
+
+			RegisteredBulk.PackageName = PackageName;
+			RegisteredBulk.bRegistered = true;
+			RegisteredBulk.bInMemory = true;
+			RegisteredBulk.bPayloadAvailable = true;
+		}
+	}
+	if (bAllowedToReadWritePayloadIdFromCache)
+	{
+		AddPendingPackageBulkData(PackageName, MoveTemp(PendingPackageBulk));
 	}
 	ResaveSizeTracker.Register(Owner, BulkData);
+	return UE::BulkDataRegistry::ERegisterResult::Success;
+}
+
+void FBulkDataRegistryEditorDomain::UpdateRegistrationData(UPackage* Owner, const UE::Serialization::FEditorBulkData& BulkData)
+{
+	if (!BulkData.GetIdentifier().IsValid())
+	{
+		UE_LOG(LogBulkDataRegistry, Warning, TEXT("UpdateRegistrationData called with invalid BulkData for Owner %s."),
+			Owner ? *Owner->GetName() : TEXT("<unknown>"));
+		return;
+	}
+
+	FName PackageName = Owner ? Owner->GetFName() : NAME_None;
+	{
+		FWriteScopeLock RegistryScopeLock(RegistryLock);
+		check(bActive); // Registrations should not come in after we destruct
+		UE::BulkDataRegistry::Private::FRegisteredBulk& RegisteredBulk = Registry.FindOrAdd(BulkData.GetIdentifier());
+		// Add the owner if we previously did not have an owner, otherwise keep the first owner
+		if (RegisteredBulk.PackageName.IsNone())
+		{
+			RegisteredBulk.PackageName = PackageName;
+		}
+		RegisteredBulk.BulkData = BulkData.CopyTornOff();
+		RegisteredBulk.bAllowedToWritePayloadIdToCache = false;
+		RegisteredBulk.bRegistered = true;
+		RegisteredBulk.bInMemory = true;
+		RegisteredBulk.bPayloadAvailable = true;
+	}
+	ResaveSizeTracker.UpdateRegistrationData(Owner, BulkData);
+}
+
+void FBulkDataRegistryEditorDomain::Unregister(const UE::Serialization::FEditorBulkData& BulkData)
+{
+	const FGuid& Key = BulkData.GetIdentifier();
+	FWriteScopeLock RegistryScopeLock(RegistryLock);
+	check(bActive); // Deregistrations should not come in after we destruct
+	Registry.Remove(Key);
 }
 
 void FBulkDataRegistryEditorDomain::OnExitMemory(const UE::Serialization::FEditorBulkData& BulkData)
@@ -192,8 +292,11 @@ void FBulkDataRegistryEditorDomain::OnExitMemory(const UE::Serialization::FEdito
 	{
 		if (Existing->BulkData.IsMemoryOnlyPayload())
 		{
-			Registry.Remove(Key);
+			Existing->BulkData.Reset();
+			Existing->BulkData.TearOff(); // Keep the TearOff flag after resetting
+			Existing->bPayloadAvailable = false;
 		}
+		Existing->bInMemory = false;
 	}
 }
 
@@ -208,7 +311,7 @@ TFuture<UE::BulkDataRegistry::FMetaData> FBulkDataRegistryEditorDomain::GetMeta(
 		{
 			Existing = Registry.Find(BulkDataId);
 		}
-		if (!Existing)
+		if (!Existing || !Existing->bPayloadAvailable)
 		{
 			TPromise<UE::BulkDataRegistry::FMetaData> Promise;
 			Promise.SetValue(UE::BulkDataRegistry::FMetaData{ false, FIoHash(), 0 });
@@ -260,7 +363,7 @@ TFuture<UE::BulkDataRegistry::FData> FBulkDataRegistryEditorDomain::GetData(cons
 			{
 				Existing = Registry.Find(BulkDataId);
 			}
-			if (!Existing)
+			if (!Existing || !Existing->bPayloadAvailable)
 			{
 				TPromise<UE::BulkDataRegistry::FData> Result;
 				Result.SetValue(UE::BulkDataRegistry::FData{ false, FCompressedBuffer() });
@@ -322,6 +425,30 @@ TFuture<UE::BulkDataRegistry::FData> FBulkDataRegistryEditorDomain::GetData(cons
 		});
 }
 
+bool FBulkDataRegistryEditorDomain::TryGetBulkData(const FGuid& BulkDataId, UE::Serialization::FEditorBulkData* OutBulk,
+	FName* OutOwner)
+{
+	FReadScopeLock RegistryScopeLock(RegistryLock);
+	FRegisteredBulk* Existing = nullptr;
+	if (bActive)
+	{
+		Existing = Registry.Find(BulkDataId);
+	}
+	if (Existing)
+	{
+		if (OutBulk)
+		{
+			*OutBulk = Existing->BulkData;
+		}
+		if (OutOwner)
+		{
+			*OutOwner = Existing->PackageName;
+		}
+		return true;
+	}
+	return false;
+}
+
 uint64 FBulkDataRegistryEditorDomain::GetBulkDataResaveSize(FName PackageName)
 {
 	return ResaveSizeTracker.GetBulkDataResaveSize(PackageName);
@@ -343,7 +470,7 @@ void FBulkDataRegistryEditorDomain::Tick(float DeltaTime)
 	TickCook(DeltaTime, false /* bTickComplete */);
 }
 
-void FBulkDataRegistryEditorDomain::AddPendingPackageBulkData(FName PackageName, const UE::Serialization::FEditorBulkData& BulkData)
+void FBulkDataRegistryEditorDomain::AddPendingPackageBulkData(FName PackageName, UE::Serialization::FEditorBulkData&& BulkData)
 {
 	FScopeLock PendingPackageScopeLock(&PendingPackageLock);
 	check(bActive); // Registrations should not come in after we destruct, and AsyncTasks should check bActive before calling
@@ -356,7 +483,7 @@ void FBulkDataRegistryEditorDomain::AddPendingPackageBulkData(FName PackageName,
 	{
 		return;
 	}
-	PendingPackage->AddBulkData(BulkData);
+	PendingPackage->AddBulkData(MoveTemp(BulkData));
 }
 
 void FBulkDataRegistryEditorDomain::AddTempLoadedPayload(const FGuid& RegistryKey, uint64 PayloadSize)
@@ -427,7 +554,7 @@ void FBulkDataRegistryEditorDomain::ReadPayloadIdsFromCache(FName PackageName, T
 			const FGuid& BulkDataId = NewPending->GetBulkDataId();
 			if (Reader.IsError() || CachedBulkData.GetIdentifier() != BulkDataId)
 			{
-				UE_LOG(LogEditorDomain, Warning, TEXT("Corrupt cache data for BulkDataPayloadId %s."), WriteToString<192>(PackageName, TEXT("/"), BulkDataId).ToString());
+				UE_LOG(LogBulkDataRegistry, Warning, TEXT("Corrupt cache data for BulkDataPayloadId %s."), WriteToString<192>(PackageName, TEXT("/"), BulkDataId).ToString());
 				return;
 			}
 
@@ -457,11 +584,8 @@ void FBulkDataRegistryEditorDomain::ReadPayloadIdsFromCache(FName PackageName, T
 
 			UE::Serialization::FEditorBulkData& ExistingBulkData = ExistingRegisteredBulk->BulkData;
 			check(ExistingBulkData.GetIdentifier() == BulkDataId);
-			if (ExistingBulkData.HasPlaceholderPayloadId())
+			if (ExistingBulkData.HasPlaceholderPayloadId() && CachedBulkData.LocationMatches(ExistingBulkData))
 			{
-				// BULKDATAREGISTRY_TODO: Implement LocationMatches
-				//&& CachedBulkData.LocationMatches(ExistingBulkData);
-
 				ExistingBulkData = CachedBulkData;
 			}
 		});
@@ -609,16 +733,17 @@ void FPendingPackage::ReadCache()
 			{
 				TargetBulkData = BulkData;
 				TargetRegisteredBulk.PackageName = PackageName;
+				TargetRegisteredBulk.bPayloadAvailable = true;
 			}
 			else
 			{
 				check(TargetBulkData.GetIdentifier() == BulkDataId);
-				// BULKDATAREGISTRY_TODO: Implement LocationMatches
-				// bCachedLocationMatches = BulkData.LocationMatches(TargetBulkData);
+				bCachedLocationMatches = BulkData.LocationMatches(TargetBulkData);
 				if (bCachedLocationMatches && !BulkData.HasPlaceholderPayloadId() && TargetBulkData.HasPlaceholderPayloadId())
 				{
 					TargetBulkData = BulkData;
 					TargetRegisteredBulk.PackageName = PackageName;
+					TargetRegisteredBulk.bPayloadAvailable = true;
 				}
 			}
 
@@ -723,8 +848,7 @@ void FUpdatePayloadWorker::DoWork()
 				break;
 			}
 
-			// BULKDATAREGISTRY_TODO: Implement LocationMatches
-			if (false) // !BulkData.LocationMatches(BulkData));
+			if (!BulkData.LocationMatches(RegisteredBulk->BulkData))
 			{
 				// Some caller has assigned a new BulkData. We need to abandon the BulkData we just loaded and give our callers the
 				// information about the new one
@@ -735,7 +859,7 @@ void FUpdatePayloadWorker::DoWork()
 				continue;
 			}
 
-			// Store the new payload in the Registry's entry for the BulkData; new MetaData requests will no longer need to wait for it
+			// Store the new PayloadId in the Registry's entry for the BulkData; new MetaData requests will no longer need to wait for it
 			RegisteredBulk->BulkData = BulkData;
 
 			// Mark that the next GetData call should remove the temporary payload
@@ -743,9 +867,10 @@ void FUpdatePayloadWorker::DoWork()
 			BulkDataRegistry->AddTempLoadedPayload(BulkData.GetIdentifier(), BulkData.GetPayloadSize());
 			BulkDataRegistry->PruneTempLoadedPayloads();
 
-			if (!RegisteredBulk->PackageName.IsNone())
+			if (RegisteredBulk->bAllowedToWritePayloadIdToCache)
 			{
 				BulkDataRegistry->WritePayloadIdToCache(RegisteredBulk->PackageName, BulkData);
+				RegisteredBulk->bAllowedToWritePayloadIdToCache = false;
 			}
 			break;
 		}
diff --git a/Engine/Source/Editor/UnrealEd/Private/EditorDomain/BulkDataRegistryEditorDomain.h b/Engine/Source/Editor/UnrealEd/Private/EditorDomain/BulkDataRegistryEditorDomain.h
--- a/Engine/Source/Editor/UnrealEd/Private/EditorDomain/BulkDataRegistryEditorDomain.h
+++ b/Engine/Source/Editor/UnrealEd/Private/EditorDomain/BulkDataRegistryEditorDomain.h
@@ -27,15 +27,17 @@ class FBulkDataRegistryEditorDomain;
 /** Struct for storage of a BulkData in the registry, including the BulkData itself and data about cache status. */
 struct FRegisteredBulk
 {
-	FRegisteredBulk() = default;
-	FRegisteredBulk(UE::Serialization::FEditorBulkData&& InBulkData, FName InPackageName = NAME_None)
-		:BulkData(MoveTemp(InBulkData)), PackageName(InPackageName)
+	FRegisteredBulk()
+		: bHasTempPayload(false), bAllowedToWritePayloadIdToCache(false), bRegistered(false), bInMemory(false), bPayloadAvailable(false)
 	{
 	}
-
 	UE::Serialization::FEditorBulkData BulkData;
 	FName PackageName;
-	bool bHasTempPayload = false;
+	bool bHasTempPayload : 1;
+	bool bAllowedToWritePayloadIdToCache : 1;
+	bool bRegistered : 1;
+	bool bInMemory : 1;
+	bool bPayloadAvailable : 1;
 };
 
 /** Serialize an array of BulkDatas into or out of bytes saved/load from the registry's persistent cache. */
@@ -50,9 +52,9 @@ public:
 	FPendingPackage(const FPendingPackage& Other) = delete;
 
 	void Cancel();
-	void AddBulkData(const UE::Serialization::FEditorBulkData& BulkData)
+	void AddBulkData(UE::Serialization::FEditorBulkData&& BulkData)
 	{
-		BulkDatas.Add(BulkData);
+		BulkDatas.Add(MoveTemp(BulkData));
 	}
 	void OnEndLoad(bool& bOutShouldRemove, bool& bOutShouldWriteCache);
 	bool IsLoadInProgress() const
@@ -165,10 +167,16 @@ public:
 	virtual ~FBulkDataRegistryEditorDomain();
 
 	// IBulkDataRegistry interface
-	virtual void Register(UPackage* Owner, const UE::Serialization::FEditorBulkData& BulkData) override;
+	virtual UE::BulkDataRegistry::ERegisterResult
+		TryRegister(UPackage* Owner, const UE::Serialization::FEditorBulkData& BulkData) override;
+	virtual void UpdateRegistrationData(UPackage* Owner, const UE::Serialization::FEditorBulkData& BulkData) override;
+	virtual void Unregister(const UE::Serialization::FEditorBulkData& BulkData) override;
 	virtual void OnExitMemory(const UE::Serialization::FEditorBulkData& BulkData) override;
 	virtual TFuture<UE::BulkDataRegistry::FMetaData> GetMeta(const FGuid& BulkDataId) override;
 	virtual TFuture<UE::BulkDataRegistry::FData> GetData(const FGuid& BulkDataId) override;
+	virtual bool TryGetBulkData(const FGuid& BulkDataId, UE::Serialization::FEditorBulkData* OutBulk = nullptr,
+		FName* OutOwner = nullptr) override;
+
 	virtual uint64 GetBulkDataResaveSize(FName PackageName) override;
 
 	// FTickableEditorObject/FTickableCookObject interface
@@ -178,7 +186,7 @@ public:
 	virtual TStatId GetStatId() const override { return TStatId(); }
 
 private:
-	void AddPendingPackageBulkData(FName PackageName, const UE::Serialization::FEditorBulkData& BulkData);
+	void AddPendingPackageBulkData(FName PackageName, UE::Serialization::FEditorBulkData&& BulkData);
 	void PollPendingPackages(bool bWaitForCooldown);
 	void AddTempLoadedPayload(const FGuid& RegistryKey, uint64 PayloadSize);
 	void PruneTempLoadedPayloads();
diff --git a/Engine/Source/Runtime/CoreUObject/Private/Serialization/BulkDataRegistry.cpp b/Engine/Source/Runtime/CoreUObject/Private/Serialization/BulkDataRegistry.cpp
--- a/Engine/Source/Runtime/CoreUObject/Private/Serialization/BulkDataRegistry.cpp
+++ b/Engine/Source/Runtime/CoreUObject/Private/Serialization/BulkDataRegistry.cpp
@@ -16,6 +16,22 @@
 #include "UObject/PackageResourceManager.h"
 #include "UObject/UObjectGlobals.h"
 
+DEFINE_LOG_CATEGORY(LogBulkDataRegistry);
+
+const TCHAR* LexToString(UE::BulkDataRegistry::ERegisterResult Value)
+{
+	using namespace UE::BulkDataRegistry;
+	switch (Value)
+	{
+	case ERegisterResult::Success:
+		return TEXT("Success");
+	case ERegisterResult::AlreadyExists:
+		return TEXT("AlreadyExists");
+	default:
+		return TEXT("InvalidResultCode");
+	}
+}
+
 namespace UE::BulkDataRegistry::Private
 {
 
@@ -29,7 +45,13 @@ public:
 	FBulkDataRegistryNull() = default;
 	virtual ~FBulkDataRegistryNull() {}
 
-	virtual void Register(UPackage* Owner, const UE::Serialization::FEditorBulkData& BulkData) override {}
+	virtual UE::BulkDataRegistry::ERegisterResult
+		TryRegister(UPackage* Owner, const UE::Serialization::FEditorBulkData& BulkData) override
+	{
+		return UE::BulkDataRegistry::ERegisterResult::Success;
+	}
+	virtual void UpdateRegistrationData(UPackage* Owner, const UE::Serialization::FEditorBulkData& BulkData) override {}
+	virtual void Unregister(const UE::Serialization::FEditorBulkData& BulkData) override {}
 	virtual void OnExitMemory(const UE::Serialization::FEditorBulkData& BulkData) override {}
 	virtual TFuture<UE::BulkDataRegistry::FMetaData> GetMeta(const FGuid& BulkDataId) override
 	{
@@ -43,6 +65,11 @@ public:
 		Promise.SetValue(UE::BulkDataRegistry::FData{ false, FCompressedBuffer() });
 		return Promise.GetFuture();
 	}
+	virtual bool TryGetBulkData(const FGuid & BulkDataId, UE::Serialization::FEditorBulkData* OutBulk = nullptr,
+		FName* OutOwner = nullptr) override
+	{
+		return false;
+	}
 	virtual uint64 GetBulkDataResaveSize(FName PackageName) override
 	{
 		return 0;
@@ -55,9 +82,15 @@ public:
 	FBulkDataRegistryTrackBulkDataToResave() = default;
 	virtual ~FBulkDataRegistryTrackBulkDataToResave() {}
 
-	virtual void Register(UPackage* Owner, const UE::Serialization::FEditorBulkData& BulkData) override
+	virtual UE::BulkDataRegistry::ERegisterResult
+		TryRegister(UPackage* Owner, const UE::Serialization::FEditorBulkData& BulkData) override
 	{
 		ResaveSizeTracker.Register(Owner, BulkData);
+		return UE::BulkDataRegistry::ERegisterResult::Success;
+	}
+	virtual void UpdateRegistrationData(UPackage* Owner, const UE::Serialization::FEditorBulkData& BulkData) override
+	{
+		ResaveSizeTracker.UpdateRegistrationData(Owner, BulkData);
 	}
 	virtual uint64 GetBulkDataResaveSize(FName PackageName) override
 	{
@@ -192,6 +225,12 @@ void FResaveSizeTracker::Register(UPackage* Owner, const UE::Serialization::FEdi
 	PackageBulkResaveSize.FindOrAdd(Owner->GetFName()) += BulkData.GetPayloadSize();
 }
 
+void FResaveSizeTracker::UpdateRegistrationData(UPackage* Owner, const UE::Serialization::FEditorBulkData& BulkData)
+{
+	// Not yet implemented; we keep the values only from the first registration
+	// Implementing this would require keeping more data; we will need to know the previous payload size to subtract it
+}
+
 uint64 FResaveSizeTracker::GetBulkDataResaveSize(FName PackageName)
 {
 	FReadScopeLock ScopeLock(Lock);
diff --git a/Engine/Source/Runtime/CoreUObject/Private/Serialization/EditorBulkData.cpp b/Engine/Source/Runtime/CoreUObject/Private/Serialization/EditorBulkData.cpp
--- a/Engine/Source/Runtime/CoreUObject/Private/Serialization/EditorBulkData.cpp
+++ b/Engine/Source/Runtime/CoreUObject/Private/Serialization/EditorBulkData.cpp
@@ -16,6 +16,9 @@
 #include "UObject/LinkerSave.h"
 #include "UObject/Object.h"
 #include "UObject/ObjectSaveContext.h"
+#include "UObject/PropertyPortFlags.h"
+#include "UObject/UObjectGlobals.h"
+#include "UObject/SavePackage.h"
 #include "Virtualization/VirtualizationSystem.h"
 
 //#if WITH_EDITORONLY_DATA
@@ -376,7 +379,7 @@ FEditorBulkData& FEditorBulkData::operator=(FEditorBulkData&& Other)
 
 	Other.Reset();
 
-	Register(nullptr);
+	Register(nullptr, TEXT("operator=(FEditorBulkData &&)"), false /* bAllowUpdateId */);
 
 	return *this;
 }
@@ -429,7 +432,7 @@ FEditorBulkData& FEditorBulkData::operator=(const FEditorBulkData& Other)
 	}
 	else
 	{
-		Register(nullptr);
+		Register(nullptr, TEXT("operator=(const FEditorBulkData&)"), false /* bAllowUpdateId */);
 	}
 	return *this;
 }
@@ -441,7 +444,7 @@ FEditorBulkData::~FEditorBulkData()
 		AttachedAr->DetachBulkData(this, false);
 	}
 
-	Unregister();
+	OnExitMemory();
 }
 
 FEditorBulkData::FEditorBulkData(const FEditorBulkData& Other, ETornOff)
@@ -456,27 +459,9 @@ void FEditorBulkData::TearOff()
 	EnumAddFlags(Flags, EFlags::IsTornOff);
 }
 
-void FEditorBulkData::Register(UObject* Owner)
-{
-#if WITH_EDITOR
-	if (BulkDataId.IsValid() && PayloadSize > 0 && !EnumHasAnyFlags(Flags, EFlags::IsTornOff))
-	{
-		IBulkDataRegistry::Get().Register(Owner ? Owner->GetPackage() : nullptr, *this);
-		EnumAddFlags(Flags, EFlags::HasRegistered);
-	}
-#endif
-}
-
-void FEditorBulkData::Unregister()
+void FEditorBulkData::UpdateRegistrationOwner(UObject* Owner)
 {
-#if WITH_EDITOR
-	if (EnumHasAnyFlags(Flags, EFlags::HasRegistered))
-	{
-		check(!EnumHasAnyFlags(Flags, EFlags::IsTornOff));
-		IBulkDataRegistry::Get().OnExitMemory(*this);
-		EnumRemoveFlags(Flags, EFlags::HasRegistered);
-	}
-#endif
+	UpdateRegistrationData(Owner, TEXT("UpdateRegistrationOwner"), false /* bAllowUpdateId */);
 }
 
 static FGuid CreateUniqueGuid(const FGuid& NonUniqueGuid, const UObject* Owner, const TCHAR* DebugName)
@@ -504,6 +489,155 @@ static FGuid CreateUniqueGuid(const FGuid& NonUniqueGuid, const UObject* Owner,
 	}
 }
 
+void FEditorBulkData::Register(UObject* Owner, const TCHAR* ErrorLogCaller, bool bAllowUpdateId)
+{
+#if WITH_EDITOR
+	if (BulkDataId.IsValid() && PayloadSize > 0 && !EnumHasAnyFlags(Flags, EFlags::IsTornOff))
+	{
+		UPackage* OwnerPackage = Owner ? Owner->GetPackage() : nullptr;
+		UE::BulkDataRegistry::ERegisterResult Result = IBulkDataRegistry::Get().TryRegister(OwnerPackage, *this);
+		if (Result == UE::BulkDataRegistry::ERegisterResult::Success)
+		{
+			EnumAddFlags(Flags, EFlags::HasRegistered);
+		}
+		else
+		{
+			UE::BulkDataRegistry::ERegisterResult SecondResult = UE::BulkDataRegistry::ERegisterResult::InvalidResultCode;
+			if (Result == UE::BulkDataRegistry::ERegisterResult::AlreadyExists && bAllowUpdateId)
+			{
+				FGuid OldBulkDataId(BulkDataId);
+				BulkDataId = CreateUniqueGuid(OldBulkDataId, Owner, TEXT("BulkDataRegistryGuidCollision"));
+				SecondResult = IBulkDataRegistry::Get().TryRegister(OwnerPackage, *this);
+				if (SecondResult == UE::BulkDataRegistry::ERegisterResult::Success)
+				{
+					EnumAddFlags(Flags, EFlags::HasRegistered);
+					LogRegisterError(Result, Owner, OldBulkDataId, ErrorLogCaller, true/* bHandledByCreateUniqueGuid */);
+				}
+				else
+				{
+					BulkDataId = OldBulkDataId;
+				}
+			}
+			if (SecondResult != UE::BulkDataRegistry::ERegisterResult::Success)
+			{
+				LogRegisterError(Result, Owner, BulkDataId, ErrorLogCaller, false /* bHandledByCreateUniqueGuid */);
+			}
+		}
+	}
+#endif
+}
+
+void FEditorBulkData::OnExitMemory()
+{
+#if WITH_EDITOR
+	if (EnumHasAnyFlags(Flags, EFlags::HasRegistered))
+	{
+		check(!EnumHasAnyFlags(Flags, EFlags::IsTornOff));
+		IBulkDataRegistry::Get().OnExitMemory(*this);
+		EnumRemoveFlags(Flags, EFlags::HasRegistered);
+	}
+#endif
+}
+
+void FEditorBulkData::Unregister()
+{
+#if WITH_EDITOR
+	if (EnumHasAnyFlags(Flags, EFlags::HasRegistered))
+	{
+		check(!EnumHasAnyFlags(Flags, EFlags::IsTornOff));
+		IBulkDataRegistry::Get().Unregister(*this);
+		EnumRemoveFlags(Flags, EFlags::HasRegistered);
+	}
+#endif
+}
+
+void FEditorBulkData::UpdateRegistrationData(UObject* Owner, const TCHAR* ErrorLogCaller, bool bAllowUpdateId)
+{
+#if WITH_EDITOR
+	UPackage* OwnerPackage = Owner ? Owner->GetPackage() : nullptr;
+	bool bShouldRegister = BulkDataId.IsValid() && PayloadSize > 0 && !EnumHasAnyFlags(Flags, EFlags::IsTornOff);
+	if (EnumHasAnyFlags(Flags, EFlags::HasRegistered))
+	{
+		if (bShouldRegister)
+		{
+			IBulkDataRegistry::Get().UpdateRegistrationData(OwnerPackage, *this);
+		}
+		else
+		{
+			Unregister();
+		}
+	}
+	else
+	{
+		if (bShouldRegister)
+		{
+			Register(Owner, ErrorLogCaller, bAllowUpdateId);
+		}
+	}
+#endif
+}
+
+void FEditorBulkData::LogRegisterError(UE::BulkDataRegistry::ERegisterResult Value, UObject* Owner,
+	const FGuid& FailedBulkDataId, const TCHAR* CallerName, bool bHandledByCreateUniqueGuid) const
+{
+#if WITH_EDITOR
+	if (Value == UE::BulkDataRegistry::ERegisterResult::Success)
+	{
+		return;
+	}
+	bool bMessageLogged = false;
+	FString OwnerPathName = Owner ? Owner->GetPathName() : TEXT("<unknown>");
+	UPackage* OwnerPackage = Owner ? Owner->GetPackage() : nullptr;
+	CallerName = CallerName ? CallerName : TEXT("<unknown>");
+	if (Value == UE::BulkDataRegistry::ERegisterResult::AlreadyExists)
+	{
+		FEditorBulkData OtherBulkData;
+		FName OtherOwnerPackageFName;
+		bool bOtherExists = IBulkDataRegistry::Get().TryGetBulkData(FailedBulkDataId, &OtherBulkData, &OtherOwnerPackageFName);
+		FString OtherOwnerPackageName = !OtherOwnerPackageFName.IsNone() ? OtherOwnerPackageFName.ToString() : TEXT("<unknown>");
+
+		if (bHandledByCreateUniqueGuid)
+		{
+			if (UE::SavePackageUtilities::OnAddResaveOnDemandPackage.IsBound() && OwnerPackage)
+			{
+				static FName BulkDataDuplicatesSystemName(TEXT("BulkDataDuplicates"));
+				UE::SavePackageUtilities::OnAddResaveOnDemandPackage.Execute(BulkDataDuplicatesSystemName, OwnerPackage->GetFName());
+				UE_LOG(LogBulkDataRegistry, Display, TEXT("AddPackageToResave %s: BulkData %s collided with an ID in package %s."),
+					*OwnerPathName, *FailedBulkDataId.ToString(), *OtherOwnerPackageName);
+			}
+			else
+			{
+				FString SilenceWarningMessage;
+				if (OwnerPackage)
+				{
+					SilenceWarningMessage = TEXT(" To silence this warning, run the ResavePackagesCommandlet with \"-autocheckout -resaveondemand=bulkdataduplicates\".");
+				}
+				UE_LOG(LogBulkDataRegistry, Warning, TEXT("%s updated BulkData %s on load because it collided with an ID in package %s.%s"),
+					*OwnerPathName, *FailedBulkDataId.ToString(), *OtherOwnerPackageName, *SilenceWarningMessage);
+			}
+			bMessageLogged = true;
+		}
+		else if (bOtherExists)
+		{
+			UE_LOG(LogBulkDataRegistry, Warning,
+				TEXT("%s failed to register BulkData %s from %s in the BulkDataRegistry: it already exists with owner %s.\n")
+				TEXT("            Payload: %s\n")
+				TEXT("    ExistingPayload: %s"),
+				*OwnerPathName, *FailedBulkDataId.ToString(), CallerName, *OtherOwnerPackageName,
+				*WriteToString<64>(GetPayloadId()), *WriteToString<64>(OtherBulkData.GetPayloadId()));
+			bMessageLogged = true;
+		}
+	}
+	if (!bMessageLogged)
+	{
+		UE_LOG(LogBulkDataRegistry, Warning,
+			TEXT("%s failed to register BulkData %s from %s in the BulkDataRegistry: %s."),
+			*OwnerPathName, *FailedBulkDataId.ToString(), CallerName, LexToString(Value));
+	}
+#endif
+}
+
+
 void FEditorBulkData::CreateFromBulkData(FUntypedBulkData& InBulkData, const FGuid& InGuid, UObject* Owner)
 {
 	TRACE_CPUPROFILER_EVENT_SCOPE(FEditorBulkData::CreateFromBulkData);
@@ -554,7 +688,10 @@ void FEditorBulkData::CreateFromBulkData(FUntypedBulkData& InBulkData, const FGu
 	{
 		EnumAddFlags(Flags, EFlags::LegacyKeyWasGuidDerived);
 	}
-	Register(Owner);
+	// bAllowUpdateId=true because we can encounter duplicates even with the newly
+	// created BulkDataId, if the user duplicated (before the fix for duplication) the owner
+	// of this bulkdata, without ever resaving this bulkdata's package
+	Register(Owner, TEXT("CreateFromBulkData"), true /* bAllowUpdateId */);
 }
 
 void FEditorBulkData::CreateLegacyUniqueIdentifier(UObject* Owner)
@@ -563,7 +700,10 @@ void FEditorBulkData::CreateLegacyUniqueIdentifier(UObject* Owner)
 	{
 		Unregister();
 		BulkDataId = CreateUniqueGuid(BulkDataId, Owner, TEXT("Unknown"));
-		Register(Owner);
+		// bAllowUpdateId=true because we can encounter duplicates even with the newly
+		// created BulkDataId, if the user duplicated (before the fix for duplication) the owner
+		// of this bulkdata, without ever resaving this bulkdata's package
+		Register(Owner, TEXT("CreateLegacyUniqueIdentifier"), true /* bAllowUpdateId */);
 	}
 }
 
@@ -622,7 +762,7 @@ void FEditorBulkData::Serialize(FArchive& Ar, UObject* Owner, bool bAllowRegiste
 				
 				Payload = CompressedPayload.Decompress();	
 
-				Register(Owner);
+				Register(Owner, TEXT("Serialize/Transacting"), false /* bAllowUpdateId */);
 			}
 		}
 	}
@@ -682,7 +822,8 @@ void FEditorBulkData::Serialize(FArchive& Ar, UObject* Owner, bool bAllowRegiste
 		{
 			checkf(Ar.IsCooking() == false, TEXT("FEditorBulkData::Serialize should not be called during a cook"));
 
-			const EFlags UpdatedFlags = BuildFlagsForSerialization(Ar, bKeepFileDataByReference);
+			EFlags UpdatedFlags = BuildFlagsForSerialization(Ar, bKeepFileDataByReference);
+			EnumRemoveFlags(UpdatedFlags, EFlags::TransientFlags);
 
 			// Go back in the archive and update the flags in the archive, we will only apply the updated flags to the current
 			// object later if we detect that the package saved successfully.
@@ -791,6 +932,12 @@ void FEditorBulkData::Serialize(FArchive& Ar, UObject* Owner, bool bAllowRegiste
 		}
 		else if (Ar.IsLoading())
 		{
+			if (Ar.HasAllPortFlags(PPF_Duplicate) && BulkDataId.IsValid())
+			{
+				// When duplicating BulkDatas we need to create a new BulkDataId to respect the uniqueness contract
+				BulkDataId = CreateUniqueGuid(BulkDataId, Owner, TEXT("PPF_Duplicate serialization"));
+			}
+
 			OffsetInFile = INDEX_NONE;
 			PackagePath.Empty();
 			PackageSegment = EPackageSegment::Header;
@@ -872,7 +1019,7 @@ void FEditorBulkData::Serialize(FArchive& Ar, UObject* Owner, bool bAllowRegiste
 
 			if (bAllowRegister)
 			{
-				Register(Owner);
+				Register(Owner, TEXT("Serialize/Persistent"), true /* bAllowUpdateId */);
 			}
 		}
 	}
@@ -1239,7 +1386,7 @@ void FEditorBulkData::PushData(const FPackagePath& InPackagePath)
 			OffsetInFile = INDEX_NONE;
 
 			// Update our information in the registry
-			Register(nullptr);
+			UpdateRegistrationData(nullptr, TEXT("PushData"), false /* bAllowUpdateId */);
 		}
 	}	
 }
@@ -1274,13 +1421,14 @@ bool FEditorBulkData::IsMemoryOnlyPayload() const
 
 void FEditorBulkData::Reset()
 {
+	// Unregister rather than allowing the Registry to keep our record, since we are changing the payload
+	Unregister();
 	// Note that we do not reset the BulkDataId
 	if (AttachedAr != nullptr)
 	{
 		AttachedAr->DetachBulkData(this, false);
 	}
 
-	Unregister();
 	PayloadContentId.Reset();
 	Payload.Reset();
 	PayloadSize = 0;
@@ -1321,14 +1469,19 @@ void FEditorBulkData::DetachFromDisk(FArchive* Ar, bool bEnsurePayloadIsLoaded)
 		OffsetInFile = INDEX_NONE;
 
 		EnumRemoveFlags(Flags, EFlags::ReferencesLegacyFile | EFlags::ReferencesWorkspaceDomain | EFlags::LegacyFileIsCompressed);
-
-		if (PayloadSize > 0)
+		if (EnumHasAnyFlags(Flags, EFlags::HasRegistered))
 		{
-			Register(nullptr);
-		}
-		else
-		{
-			Unregister();
+			if (bEnsurePayloadIsLoaded)
+			{
+				// We are e.g. renaming a package and should switch registration over to our in-memory Payload
+				UpdateRegistrationData(nullptr, TEXT("DetachFromDisk"), false /* bAllowUpdateId */);
+			}
+			else
+			{
+				// The package is unloading; we should instruct the bulkdataregistry to keep our old packagepath information,
+				// but will drop our is-registered flag
+				OnExitMemory();
+			}
 		}
 	}
 
@@ -1382,7 +1535,7 @@ void FEditorBulkData::SerializeToLegacyPath(FLinkerSave& LinkerSave, FCompressed
 				this->PackagePath = InPackagePath;
 				check(!this->PackagePath.IsEmpty()); // LinkerSave guarantees a valid PackagePath if we're updating loaded path
 				this->OffsetInFile = DataStartOffset;
-				this->Flags = UpdatedFlags;
+				this->Flags = (this->Flags & EFlags::TransientFlags) | (UpdatedFlags & ~EFlags::TransientFlags);
 
 				if (CanUnloadData())
 				{
@@ -1391,8 +1544,7 @@ void FEditorBulkData::SerializeToLegacyPath(FLinkerSave& LinkerSave, FCompressed
 				}
 
 				// Update our information in the registry
-				// TODO: Pass Owner into Register once the AssetRegistry has been fixed to use the updated PackageGuid from the save
-				Register(nullptr);
+				UpdateRegistrationData(Owner, TEXT("SerializeToLegacyPath"), false /* bAllowUpdateId */);
 			};
 
 			LinkerSave.PostSaveCallbacks.Add(MoveTemp(OnSavePackage));
@@ -1410,7 +1562,7 @@ void FEditorBulkData::SerializeToLegacyPath(FLinkerSave& LinkerSave, FCompressed
 
 void FEditorBulkData::SerializeToPackageTrailer(FLinkerSave& LinkerSave, FCompressedBuffer PayloadToSerialize, EFlags UpdatedFlags, UObject* Owner)
 {
-	auto OnPayloadWritten = [this, UpdatedFlags](FLinkerSave& LinkerSave, const FPackageTrailer& Trailer) mutable
+	auto OnPayloadWritten = [this, UpdatedFlags, Owner](FLinkerSave& LinkerSave, const FPackageTrailer& Trailer) mutable
 	{
 		checkf(LinkerSave.IsCooking() == false, TEXT("FEditorBulkData::Serialize should not be called during a cook"));
 
@@ -1421,7 +1573,7 @@ void FEditorBulkData::SerializeToPackageTrailer(FLinkerSave& LinkerSave, FCompre
 		// disk so that we can update the object's members to be redirected to the saved file.
 		if (!LinkerSave.GetFilename().IsEmpty())
 		{
-			auto OnSavePackage = [this, PayloadOffset, UpdatedFlags](const FPackagePath& InPackagePath, FObjectPostSaveContext ObjectSaveContext)
+			auto OnSavePackage = [this, PayloadOffset, UpdatedFlags, Owner](const FPackagePath& InPackagePath, FObjectPostSaveContext ObjectSaveContext)
 			{
 				if (!ObjectSaveContext.IsUpdatingLoadedPath())
 				{
@@ -1431,7 +1583,7 @@ void FEditorBulkData::SerializeToPackageTrailer(FLinkerSave& LinkerSave, FCompre
 				this->PackagePath = InPackagePath;
 				check(!this->PackagePath.IsEmpty()); // LinkerSave guarantees a valid PackagePath if we're updating loaded path
 				this->OffsetInFile = PayloadOffset;
-				this->Flags = UpdatedFlags;
+				this->Flags = (this->Flags & EFlags::TransientFlags) | (UpdatedFlags & ~EFlags::TransientFlags);
 
 				if (CanUnloadData())
 				{
@@ -1440,8 +1592,7 @@ void FEditorBulkData::SerializeToPackageTrailer(FLinkerSave& LinkerSave, FCompre
 				}
 
 				// Update our information in the registry
-				// TODO: Pass Owner into Register once the AssetRegistry has been fixed to use the updated PackageGuid from the save
-				Register(nullptr);
+				UpdateRegistrationData(Owner, TEXT("SerializeToPackageTrailer"), false /* bAllowUpdateId */);
 			};
 
 			LinkerSave.PostSaveCallbacks.Add(MoveTemp(OnSavePackage));
@@ -1455,6 +1606,9 @@ void FEditorBulkData::UpdatePayloadImpl(FSharedBuffer&& InPayload, FIoHash&& InP
 {
 	TRACE_CPUPROFILER_EVENT_SCOPE(FEditorBulkData::UpdatePayloadImpl);
 
+	// Unregister before calling DetachBulkData; DetachFromDisk calls OnExitMemory which is incorrect since
+	// we are changing data rather than leaving memory
+	Unregister();
 	if (AttachedAr != nullptr)
 	{
 		AttachedAr->DetachBulkData(this, false);
@@ -1464,7 +1618,7 @@ void FEditorBulkData::UpdatePayloadImpl(FSharedBuffer&& InPayload, FIoHash&& InP
 
 	// We only take the payload if it has a length to avoid potentially holding onto a
 	// 0 byte allocation in the FSharedBuffer
-	if(InPayload.GetSize() > 0)
+	if (InPayload.GetSize() > 0)
 	{ 
 		Payload = MoveTemp(InPayload).MakeOwned();
 	}
@@ -1486,18 +1640,11 @@ void FEditorBulkData::UpdatePayloadImpl(FSharedBuffer&& InPayload, FIoHash&& InP
 	PackageSegment = EPackageSegment::Header;
 	OffsetInFile = INDEX_NONE;
 
-	if (PayloadSize > 0)
+	if (PayloadSize > 0 && !BulkDataId.IsValid())
 	{
-		if (!BulkDataId.IsValid())
-		{
-			BulkDataId = FGuid::NewGuid();
-		}
-		Register(Owner);
-	}
-	else
-	{
-		Unregister();
+		BulkDataId = FGuid::NewGuid();
 	}
+	UpdateRegistrationData(Owner, TEXT("UpdatePayloadImpl"), false /* bAllowUpdateId */);
 }
 
 FCompressedBuffer FEditorBulkData::GetDataInternal() const
@@ -1648,6 +1795,51 @@ void FEditorBulkData::UpdatePayloadId()
 	UpdateKeyIfNeeded();
 }
 
+bool FEditorBulkData::LocationMatches(const FEditorBulkData& Other) const
+{
+	if (GetIdentifier() != Other.GetIdentifier())
+	{
+		// Different identifiers return false, even if location is the same
+		return false;
+	}
+	if (GetPayloadSize() == 0 || Other.GetPayloadSize() == 0)
+	{
+		// True if both of them have an empty payload, false if only one does
+		return GetPayloadSize() == 0 && Other.GetPayloadSize() == 0;
+	}
+	else if (IsDataVirtualized() || Other.IsDataVirtualized())
+	{
+		// Data is virtualized and will be looked up by PayloadContentId; use PayloadContentId
+		return IsDataVirtualized() && Other.IsDataVirtualized() &&
+			PayloadContentId == Other.PayloadContentId;
+	}
+	else if (!PackagePath.IsEmpty() || !Other.PackagePath.IsEmpty())
+	{
+		if (Other.PackagePath != PackagePath)
+		{
+			return false;
+		}
+		if ((HasPayloadSidecarFile() && CVarShouldLoadFromSidecar.GetValueOnAnyThread()) ||
+			CVarShouldLoadFromTrailer.GetValueOnAnyThread())
+		{
+			// Loading from the SidecarFile or PackageTrailer uses the PayloadContentId as
+			// the identifier to look up the location in the file
+			return PayloadContentId == Other.PayloadContentId;
+		}
+		else
+		{
+			// Loading from PackagePath without sidecar or trailer uses the OffsetInFile
+			return OffsetInFile == Other.OffsetInFile;
+		}
+	}
+	else
+	{
+		// BulkData is memory-only and therefore PayloadContentId is available, use PayloadContentId
+		return PayloadContentId == Other.PayloadContentId;
+	}
+}
+
+
 #if UE_ENABLE_VIRTUALIZATION_TOGGLE
 
 void FEditorBulkData::SetVirtualizationOptOut(bool bOptOut)
diff --git a/Engine/Source/Runtime/CoreUObject/Private/UObject/SavePackage/SavePackageUtilities.cpp b/Engine/Source/Runtime/CoreUObject/Private/UObject/SavePackage/SavePackageUtilities.cpp
--- a/Engine/Source/Runtime/CoreUObject/Private/UObject/SavePackage/SavePackageUtilities.cpp
+++ b/Engine/Source/Runtime/CoreUObject/Private/UObject/SavePackage/SavePackageUtilities.cpp
@@ -901,6 +901,10 @@ int32 GetNumPackagesSaved()
 #endif
 }
 
+#if WITH_EDITOR
+FAddResaveOnDemandPackage OnAddResaveOnDemandPackage;
+#endif
+
 } // end namespace UE::SavePackageUtilities
 
 FObjectSaveContextData::FObjectSaveContextData(UPackage* Package, const ITargetPlatform* InTargetPlatform, const TCHAR* InTargetFilename, uint32 InSaveFlags)
diff --git a/Engine/Source/Runtime/CoreUObject/Public/Serialization/BulkDataRegistry.h b/Engine/Source/Runtime/CoreUObject/Public/Serialization/BulkDataRegistry.h
--- a/Engine/Source/Runtime/CoreUObject/Public/Serialization/BulkDataRegistry.h
+++ b/Engine/Source/Runtime/CoreUObject/Public/Serialization/BulkDataRegistry.h
@@ -2,6 +2,8 @@
 
 #pragma once
 
+#include "HAL/Platform.h"
+
 #if WITH_EDITOR
 
 #include "Async/Future.h"
@@ -11,8 +13,8 @@
 #include "Containers/ArrayView.h"
 #include "Delegates/Delegate.h"
 #include "HAL/CriticalSection.h"
-#include "HAL/Platform.h"
 #include "IO/IoHash.h"
+#include "Misc/Optional.h"
 
 class IBulkDataRegistry;
 class UPackage;
@@ -20,8 +22,26 @@ struct FEndLoadPackageContext;
 struct FGuid;
 namespace UE::Serialization { class FEditorBulkData; }
 
+COREUOBJECT_API DECLARE_LOG_CATEGORY_EXTERN(LogBulkDataRegistry, Log, All);
 DECLARE_DELEGATE_RetVal(IBulkDataRegistry*, FSetBulkDataRegistry);
 
+#endif
+
+namespace UE::BulkDataRegistry
+{
+	// Define ERegisterResult even if !WITH_EDITOR to make it easier to use in EditorBulkData.cpp
+	enum class ERegisterResult : uint8
+	{
+		Success,
+		AlreadyExists,
+		InvalidResultCode,
+		// Update LexToString when adding new values
+	};
+}
+
+#if WITH_EDITOR
+const TCHAR* LexToString(UE::BulkDataRegistry::ERegisterResult Value);
+
 namespace UE::BulkDataRegistry
 {
 	/** Results of GetMeta call. */
@@ -43,7 +63,6 @@ namespace UE::BulkDataRegistry
 		/** The discovered data. Empty if data was not found. */
 		FCompressedBuffer Buffer;
 	};
-
 }
 
 /** Registers BulkDatas so that they can be referenced by guid during builds later in the editor process. */
@@ -56,7 +75,12 @@ public:
 	COREUOBJECT_API static IBulkDataRegistry& Get();
 	
 	/** Register a BulkData with the registry. Its payload and metadata will be fetchable by its GetIdentifier. */
-	virtual void Register(UPackage* Owner, const UE::Serialization::FEditorBulkData& BulkData) = 0;
+	virtual UE::BulkDataRegistry::ERegisterResult
+		TryRegister(UPackage* Owner, const UE::Serialization::FEditorBulkData& BulkData) = 0;
+	/** Change existing registration data to have the new values (keeping old Owner), or add if it doesn't exist. */
+	virtual void UpdateRegistrationData(UPackage* Owner, const UE::Serialization::FEditorBulkData& BulkData) = 0;
+	/** Unregister the BulkData associated with a guid, because it is being modified. */
+	virtual void Unregister(const UE::Serialization::FEditorBulkData& BulkData) = 0;
 	/** Report that a BulkData is leaving memory and its in-memory payload (if it had one) is no longer available. */
 	virtual void OnExitMemory(const UE::Serialization::FEditorBulkData& BulkData) = 0;
 
@@ -69,6 +93,13 @@ public:
 	 */
 	virtual TFuture<UE::BulkDataRegistry::FData> GetData(const FGuid& BulkDataId) = 0;
 
+	/**
+	 * Returns the EditorBulkData and Owner for given id.
+	 * Returns true if and only if the BulkDataId is a known id. OutOwner may be null even if returning true.
+	 */
+	virtual bool TryGetBulkData(const FGuid& BulkDataId, UE::Serialization::FEditorBulkData* OutBulk = nullptr,
+		FName* OutOwner = nullptr) = 0;
+
 	/**
 	 * Report whether the Package had BulkDatas during load that upgrade or otherwise exist in memoryonly and
 	 * cannot save all its BulkDatas by reference when resaved. This function only returns the correct
@@ -99,6 +130,7 @@ public:
 	COREUOBJECT_API ~FResaveSizeTracker();
 
 	COREUOBJECT_API void Register(UPackage* Owner, const UE::Serialization::FEditorBulkData& BulkData);
+	COREUOBJECT_API void UpdateRegistrationData(UPackage* Owner, const UE::Serialization::FEditorBulkData& BulkData);
 	COREUOBJECT_API uint64 GetBulkDataResaveSize(FName PackageName);
 
 private:
diff --git a/Engine/Source/Runtime/CoreUObject/Public/Serialization/EditorBulkData.h b/Engine/Source/Runtime/CoreUObject/Public/Serialization/EditorBulkData.h
--- a/Engine/Source/Runtime/CoreUObject/Public/Serialization/EditorBulkData.h
+++ b/Engine/Source/Runtime/CoreUObject/Public/Serialization/EditorBulkData.h
@@ -13,6 +13,7 @@
 
 class FArchive;
 class UObject;
+namespace UE::BulkDataRegistry { enum class ERegisterResult : uint8; }
 
 // When enabled it will be possible for specific editor bulkdata objects to opt out of being virtualized
 // This is a development feature and not expected to be used!
@@ -299,14 +300,15 @@ public:
 	bool IsMemoryOnlyPayload() const;
 	/** Load the payload and set the correct payload id, if the bulkdata has a PlaceholderPayloadId. */
 	void UpdatePayloadId();
+	/** Return whether *this has the same source for the bulkdata (e.g. identical file locations if from file) as Other */
+	bool LocationMatches(const FEditorBulkData& Other) const;
 
 	/**
-	 * Register this BulkData with the BulkData registry, if it is valid for registration.
-	 * This function is called automatically when necessary by EditorBulkData internals, 
-	 * but external callers can call it to provide information about the owner if it was previously
-	 * registered anonymously by UpdatePayload.
+	 * Update the Owner of this BulkData in the BulkDataRegistry to include the Owner information.
+	 * Has no effect if the BulkData is not valid for registration.
+	 * The Owner information is lost and this function must be called again if the BulkData's payload information is modified.
 	 */
-	void Register(UObject* Owner);
+	void UpdateRegistrationOwner(UObject* Owner);
 
 #if UE_ENABLE_VIRTUALIZATION_TOGGLE
 	UE_DEPRECATED(5.0, "SetVirtualizationOptOut is an internal feature for development and will be removed without warning!")
@@ -439,7 +441,18 @@ private:
 		return EnumHasAnyFlags(InFlags, EFlags::StoredInPackageTrailer);
 	}
 
+	void Register(UObject* Owner, const TCHAR* LogCallerName, bool bAllowUpdateId);
+	/** Used when we need to clear the registration for our bulkdataId because its ownership is transferring elsewhere. */
 	void Unregister();
+	/**
+	 * Used when we want to keep registration of our BulkDataId tied to our Payload location, but this BulkData is exiting memory
+	 * and the registry needs to know so it can drop its copy of the payload and allow us to reregister later if reloaded.
+	 */
+	void OnExitMemory();
+	/** Check whether we should be registered, and register/unregister/update as necessary to match. */
+	void UpdateRegistrationData(UObject* Owner, const TCHAR* LogCallerName, bool bAllowUpdateId);
+	void LogRegisterError(UE::BulkDataRegistry::ERegisterResult Value, UObject* Owner, const FGuid& FailedBulkDataId,
+		const TCHAR* CallerName, bool bHandledbyCreateUniqueGuid) const;
 
 	/**
 	 * Checks to make sure that the payload we are saving is what we expect it to be. If not then we need to log an error to the
diff --git a/Engine/Source/Runtime/CoreUObject/Public/UObject/SavePackage.h b/Engine/Source/Runtime/CoreUObject/Public/UObject/SavePackage.h
--- a/Engine/Source/Runtime/CoreUObject/Public/UObject/SavePackage.h
+++ b/Engine/Source/Runtime/CoreUObject/Public/UObject/SavePackage.h
@@ -5,6 +5,7 @@
 #include "Containers/Array.h"
 #include "Containers/Set.h"
 #include "Containers/Map.h"
+#include "Delegates/Delegate.h"
 #include "Misc/DateTime.h"
 #include "ObjectMacros.h"
 #include "Serialization/FileRegions.h"
@@ -175,6 +176,16 @@ namespace UE::SavePackageUtilities
 	COREUOBJECT_API void StartSavingEDLCookInfoForVerification();
 	COREUOBJECT_API void VerifyEDLCookInfo(bool bFullReferencesExpected = true);
 	COREUOBJECT_API void EDLCookInfoAddIterativelySkippedPackage(FName LongPackageName);
+
+#if WITH_EDITOR
+	/** void FAddResaveOnDemandPackage(FName SystemName, FName PackageName); */
+	DECLARE_DELEGATE_TwoParams(FAddResaveOnDemandPackage, FName, FName);
+	/**
+	 * Low-level systems execute this delegate during automated resave-on-demand to request that a package be resaved.
+	 * Automated resave managers like the ResavePackagesCommandlet subscribe to it to know which packages to save.
+	 */
+	extern COREUOBJECT_API FAddResaveOnDemandPackage OnAddResaveOnDemandPackage;
+#endif
 }
 
 COREUOBJECT_API DECLARE_LOG_CATEGORY_EXTERN(LogSavePackage, Log, All);
