Avoid copying sparse class archetype state for ABPGC as they don't need it and there's a race condition

#rnx

diff --git a/Engine/Source/Runtime/CoreUObject/Private/UObject/Class.cpp b/Engine/Source/Runtime/CoreUObject/Private/UObject/Class.cpp
--- a/Engine/Source/Runtime/CoreUObject/Private/UObject/Class.cpp
+++ b/Engine/Source/Runtime/CoreUObject/Private/UObject/Class.cpp
@@ -5196,16 +5196,10 @@ void* UClass::CreateSparseClassData()
 		SparseClassData = FMemory::Malloc(SparseClassDataStruct->GetStructureSize(), SparseClassDataStruct->GetMinAlignment());
 		SparseClassDataStruct->InitializeStruct(SparseClassData);
 	}
+
 	if (SparseClassData)
 	{
-		// initialize per class data from the archetype if we have one
-		const void* SparseArchetypeData = GetArchetypeForSparseClassData();
-		UScriptStruct* SparseClassDataArchetypeStruct = GetSparseClassDataArchetypeStruct();
-
-		if (SparseArchetypeData && SparseClassDataStruct->IsChildOf(SparseClassDataArchetypeStruct))
-		{
-			SparseClassDataArchetypeStruct->CopyScriptStruct(SparseClassData, SparseArchetypeData);
-		}
+		InitializeSparseClassDataFromArchetype();
 	}
 
 	return SparseClassData;
@@ -5221,6 +5215,20 @@ void UClass::CleanupSparseClassData()
 	}
 }
 
+void UClass::InitializeSparseClassDataFromArchetype()
+{
+	check(SparseClassData);
+
+	// initialize per class data from the archetype if we have one
+	const void* SparseArchetypeData = GetArchetypeForSparseClassData();
+	UScriptStruct* SparseClassDataArchetypeStruct = GetSparseClassDataArchetypeStruct();
+
+	if (SparseArchetypeData && SparseClassDataStruct->IsChildOf(SparseClassDataArchetypeStruct))
+	{
+		SparseClassDataArchetypeStruct->CopyScriptStruct(SparseClassData, SparseArchetypeData);
+	}
+}
+
 const void* UClass::GetSparseClassData(const EGetSparseClassDataMethod GetMethod)
 {
 	if (SparseClassData)
diff --git a/Engine/Source/Runtime/CoreUObject/Public/UObject/Class.h b/Engine/Source/Runtime/CoreUObject/Public/UObject/Class.h
--- a/Engine/Source/Runtime/CoreUObject/Public/UObject/Class.h
+++ b/Engine/Source/Runtime/CoreUObject/Public/UObject/Class.h
@@ -2517,6 +2517,8 @@ protected:
 
 	void CleanupSparseClassData();
 
+	virtual void InitializeSparseClassDataFromArchetype();
+
 private:
 #if WITH_EDITOR
 	/** Provides access to attributes of the underlying C++ class. Should never be unset. */
diff --git a/Engine/Source/Runtime/Engine/Classes/Animation/AnimBlueprintGeneratedClass.h b/Engine/Source/Runtime/Engine/Classes/Animation/AnimBlueprintGeneratedClass.h
--- a/Engine/Source/Runtime/Engine/Classes/Animation/AnimBlueprintGeneratedClass.h
+++ b/Engine/Source/Runtime/Engine/Classes/Animation/AnimBlueprintGeneratedClass.h
@@ -636,6 +636,11 @@ public:
 	virtual void PostLoadDefaultObject(UObject* Object) override;
 	virtual void PostLoad() override;
 	// End of UClass interface
+
+protected:
+	// UClass interface
+	virtual void InitializeSparseClassDataFromArchetype() override;
+	// End of UClass interface
 };
 
 template<typename NodeType>
diff --git a/Engine/Source/Runtime/Engine/Private/Animation/AnimBlueprintGeneratedClass.cpp b/Engine/Source/Runtime/Engine/Private/Animation/AnimBlueprintGeneratedClass.cpp
--- a/Engine/Source/Runtime/Engine/Private/Animation/AnimBlueprintGeneratedClass.cpp
+++ b/Engine/Source/Runtime/Engine/Private/Animation/AnimBlueprintGeneratedClass.cpp
@@ -634,6 +634,11 @@ void UAnimBlueprintGeneratedClass::PostLoad()
 	});
 }
 
+void UAnimBlueprintGeneratedClass::InitializeSparseClassDataFromArchetype()
+{
+	// ABPGC sparse data is generated per-instance, so there is nothing to copy
+}
+
 void UAnimBlueprintGeneratedClass::GenerateAnimationBlueprintFunctions()
 {
 	AnimBlueprintFunctions.Empty();
