CoreRedirect duplicate struct redirect errors:  CL 19809751 made duplicates unnecessary as the full-path redirects will now also apply to Object-Name-only redirect requests, and added an error for any duplicates. Error severity is not necessary; downgrade it to a warning.
Remove the now-unnecessary duplicates to remove the warning.
#rnx

diff --git a/Engine/Plugins/AI/MassAI/Config/BaseMassAI.ini b/Engine/Plugins/AI/MassAI/Config/BaseMassAI.ini
--- a/Engine/Plugins/AI/MassAI/Config/BaseMassAI.ini
+++ b/Engine/Plugins/AI/MassAI/Config/BaseMassAI.ini
@@ -39,7 +39,5 @@
 +PropertyRedirects=(OldName="MassZoneGraphFindEscapeTarget.FleeTag",NewName="DisturbanceAnnotationTag")
 
 ;MassAIBehavior
-+StructRedirects=(OldName="MassZoneGraphAnnotationTagsFragment",NewName="MassZoneGraphAnnotationFragment")
-+StructRedirects=(OldName="MassStateTreeActivated",NewName="MassStateTreeActivatedTag")
 +StructRedirects=(OldName="/Script/MassAIBehavior.MassZoneGraphAnnotationTagsFragment",NewName="/Script/MassAIBehavior.MassZoneGraphAnnotationFragment")
 +StructRedirects=(OldName="/Script/MassAIBehavior.MassStateTreeActivated",NewName="/Script/MassAIBehavior.MassStateTreeActivatedTag")
diff --git a/Engine/Plugins/AI/MassCrowd/Config/BaseMassCrowd.ini b/Engine/Plugins/AI/MassCrowd/Config/BaseMassCrowd.ini
--- a/Engine/Plugins/AI/MassCrowd/Config/BaseMassCrowd.ini
+++ b/Engine/Plugins/AI/MassCrowd/Config/BaseMassCrowd.ini
@@ -2,7 +2,6 @@
 +PackageRedirects=(OldName="/Script/MassCrowdAI5", NewName="/Script/MassCrowd")
 +StructRedirects=(OldName="/Script/MassAIBehavior.MassZoneGraphFindWanderTarget",NewName="/Script/MassCrowd.MassZoneGraphFindWanderTarget")
 
-+StructRedirects=(OldName="TagFragment_MassCrowd",NewName="MassCrowdTag")
 +StructRedirects=(OldName="/Script/MassCrowd.TagFragment_MassCrowd",NewName="/Script/MassCrowd.MassCrowdTag")
 
 +PropertyRedirects=(OldName="MassCrowdServerRepresentationTrait.Config", NewName="Params")
diff --git a/Engine/Plugins/Runtime/MassGameplay/Config/BaseMassGameplay.ini b/Engine/Plugins/Runtime/MassGameplay/Config/BaseMassGameplay.ini
--- a/Engine/Plugins/Runtime/MassGameplay/Config/BaseMassGameplay.ini
+++ b/Engine/Plugins/Runtime/MassGameplay/Config/BaseMassGameplay.ini
@@ -16,31 +16,22 @@
 +PropertyRedirects=(OldName="/Script/MassCommon.DataFragment_Transform", NewName="/Script/MassCommon.TransformFragment")
 
 ;MassActors
-+StructRedirects=(OldName="DataFragment_Actor", NewName="MassActorFragment")
-+StructRedirects=(OldName="DataFragment_CapsuleComponentWrapper", NewName="CapsuleComponentWrapperFragment")
-+StructRedirects=(OldName="DataFragment_CharacterMovementComponentWrapper", NewName="CharacterMovementComponentWrapperFragment")
 +StructRedirects=(OldName="/Script/MassActors.DataFragment_Actor", NewName="/Script/MassActors.MassActorFragment")
 +StructRedirects=(OldName="/Script/MassActors.DataFragment_CapsuleComponentWrapper", NewName="/Script/MassActors.CapsuleComponentWrapperFragment")
 +StructRedirects=(OldName="/Script/MassActors.DataFragment_CharacterMovementComponentWrapper", NewName="/Script/MassActors.CharacterMovementComponentWrapperFragment")
 
 ;MassGameplayDebug
-+StructRedirects=(OldName="SimDebugVisComponent", NewName="SimDebugVisFragment")
 +StructRedirects=(OldName="/Script/MassGameplayDebug.SimDebugVisComponent", NewName="/Script/MassGameplayDebug.SimDebugVisFragment")
 
 ;MassSpawner
-+StructRedirects=(OldName="DataFragment_ReplicationTemplateID", NewName="ReplicationTemplateIDFragment")
 +StructRedirects=(OldName="/Script/MassSpawner.DataFragment_ReplicationTemplateID", NewName="/Script/MassSpawner.ReplicationTemplateIDFragment")
 
 ;MassRepresentation
-+StructRedirects=(OldName="MassRepresentationConfig", NewName="MassRepresentationParameters")
-+StructRedirects=(OldName="MassRepresentationSubsystemFragment", NewName="MassRepresentationSubsystemSharedFragment")
 +StructRedirects=(OldName="/Script/MassRepresentation.MassRepresentationConfig", NewName="/Script/MassRepresentation.MassRepresentationParameters")
 +StructRedirects=(OldName="/Script/MassRepresentation.MassRepresentationSubsystemFragment", NewName="/Script/MassRepresentation.MassRepresentationSubsystemSharedFragment")
 +PropertyRedirects=(OldName="MassVisualizationTrait.Config", NewName="Params")
 
 ;MassLOD
-+StructRedirects=(OldName="MassSimulationVariableTickConfig", NewName="MassSimulationVariableTickParameters")
-+StructRedirects=(OldName="MassSimulationLODConfig", NewName="MassSimulationLODParameters")
 +StructRedirects=(OldName="/Script/MassLOD.MassSimulationVariableTickConfig", NewName="/Script/MassLOD.MassSimulationVariableTickParameters")
 +StructRedirects=(OldName="/Script/MassLOD.MassSimulationLODConfig", NewName="/Script/MassLOD.MassSimulationLODParameters")
 +PropertyRedirects=(OldName="MassSimulationLODTrait.Config", NewName="Params")
diff --git a/Engine/Source/Runtime/CoreUObject/Private/UObject/CoreRedirects.cpp b/Engine/Source/Runtime/CoreUObject/Private/UObject/CoreRedirects.cpp
--- a/Engine/Source/Runtime/CoreUObject/Private/UObject/CoreRedirects.cpp
+++ b/Engine/Source/Runtime/CoreUObject/Private/UObject/CoreRedirects.cpp
@@ -162,6 +162,22 @@ int32 FCoreRedirectObjectName::MatchScore(const FCoreRedirectObjectName& Other)
 	return MatchScore;
 }
 
+void FCoreRedirectObjectName::UnionFieldsInline(const FCoreRedirectObjectName& Other)
+{
+	if (ObjectName.IsNone())
+	{
+		ObjectName = Other.ObjectName;
+	}
+	if (OuterName.IsNone())
+	{
+		OuterName = Other.OuterName;
+	}
+	if (PackageName.IsNone())
+	{
+		PackageName = Other.PackageName;
+	}
+}
+
 bool FCoreRedirectObjectName::HasValidCharacters() const
 {
 	// ObjectNames in Blueprint may contain spaces.
@@ -1104,19 +1120,42 @@ bool FCoreRedirects::AddSingleRedirect(const FCoreRedirect& NewRedirect, const F
 		if (ExistingRedirect.IdenticalMatchRules(NewRedirect))
 		{
 			bFoundDuplicate = true;
-			if (ExistingRedirect.NewName == NewRedirect.NewName && NewRedirect.HasValueChanges())
+			bool bIsSameNewName = ExistingRedirect.NewName == NewRedirect.NewName;
+			bool bOneIsPartialOtherIsFull = false;
+			if (!bIsSameNewName &&
+				ExistingRedirect.OldName.Matches(NewRedirect.OldName, FCoreRedirectObjectName::EMatchFlags::AllowPartialRHSMatch) &&
+				ExistingRedirect.NewName.Matches(NewRedirect.NewName, FCoreRedirectObjectName::EMatchFlags::AllowPartialRHSMatch))
 			{
-				// Merge value redirects as names are the same
-				ExistingRedirect.ValueChanges.Append(NewRedirect.ValueChanges);
+				bIsSameNewName = true;
+				bOneIsPartialOtherIsFull = true;
 			}
-			else if (ExistingRedirect.NewName != NewRedirect.NewName)
+
+			if (bIsSameNewName)
 			{
-				UE_LOG(LogLinker, Error, TEXT("AddRedirectList(%s) found conflicting redirectors for %s! Old: %s, New: %s"), *SourceString, *ExistingRedirect.OldName.ToString(), *ExistingRedirect.NewName.ToString(), *NewRedirect.NewName.ToString());
+				// Merge fields from the two duplicate redirects
+				bool bBothHaveValueChanges = ExistingRedirect.HasValueChanges() && NewRedirect.HasValueChanges();
+				ExistingRedirect.OldName.UnionFieldsInline(NewRedirect.OldName);
+				ExistingRedirect.NewName.UnionFieldsInline(NewRedirect.NewName);
+				ExistingRedirect.ValueChanges.Append(NewRedirect.ValueChanges);
+				if (bBothHaveValueChanges)
+				{
+					// No warning when there are values changes to merge
+				}
+				else if (bOneIsPartialOtherIsFull)
+				{
+					UE_LOG(LogLinker, Warning, TEXT("AddRedirectList(%s) found duplicate redirectors for %s to %s, one a FullPath and the other an ObjectName-only.")
+						TEXT(" This used to be required for StructRedirects but is no longer. Remove the ObjectName-only redirector and keep the FullPath."),
+						*SourceString, *ExistingRedirect.OldName.ToString(), *ExistingRedirect.NewName.ToString());
+				}
+				else
+				{
+					// ENABLE THIS ONCE INIS ARE CONVERTED
+					//UE_LOG(LogLinker, Warning, TEXT("AddRedirectList(%s) found duplicate redirectors for %s to %s!"), *SourceString, *ExistingRedirect.OldName.ToString(), *ExistingRedirect.NewName.ToString());
+				}
 			}
 			else
 			{
-				// ENABLE THIS ONCE INIS ARE CONVERTED
-				//UE_LOG(LogLinker, Error, TEXT("AddRedirectList(%s) found duplicate redirectors for %s to %s!), *SourceString, *ExistingRedirect.OldName.ToString(), *ExistingRedirect.NewName.ToString());
+				UE_LOG(LogLinker, Error, TEXT("AddRedirectList(%s) found conflicting redirectors for %s! Old: %s, New: %s"), *SourceString, *ExistingRedirect.OldName.ToString(), *ExistingRedirect.NewName.ToString(), *NewRedirect.NewName.ToString());
 			}
 			break;
 		}
diff --git a/Engine/Source/Runtime/CoreUObject/Public/UObject/CoreRedirects.h b/Engine/Source/Runtime/CoreUObject/Public/UObject/CoreRedirects.h
--- a/Engine/Source/Runtime/CoreUObject/Public/UObject/CoreRedirects.h
+++ b/Engine/Source/Runtime/CoreUObject/Public/UObject/CoreRedirects.h
@@ -134,6 +134,9 @@ struct COREUOBJECT_API FCoreRedirectObjectName
 	/** Returns integer of degree of match. 0 if doesn't match at all, higher integer for better matches */
 	int32 MatchScore(const FCoreRedirectObjectName& Other) const;
 
+	/** Fills in any empty fields on this with the corresponding fields from Other. */
+	void UnionFieldsInline(const FCoreRedirectObjectName& Other);
+
 	/** Returns the name used as the key into the acceleration map */
 	FName GetSearchKey(ECoreRedirectFlags Type) const
 	{
