The package trailer can now be polled for the size of a payload on disk based on it's FIoHash

#rnx

- Note that the size of disk we return is the payload's compressed size, which should probably be renamed in the package trailer to SizeOnDisk since we don't keep a record of it once the payload is virtualized (as we don't really know if a virtualization backend might change the compression of the payload while it is being stored etc)

diff --git a/Engine/Source/Runtime/CoreUObject/Private/UObject/PackageTrailer.cpp b/Engine/Source/Runtime/CoreUObject/Private/UObject/PackageTrailer.cpp
--- a/Engine/Source/Runtime/CoreUObject/Private/UObject/PackageTrailer.cpp
+++ b/Engine/Source/Runtime/CoreUObject/Private/UObject/PackageTrailer.cpp
@@ -35,6 +35,9 @@ static_assert(FPackageTrailer::FHeader::StaticHeaderSizeOnDisk == 28, "FPackageT
 static_assert(Private::FLookupTableEntry::SizeOnDisk == 49, "FLookupTableEntry size has been changed, if this was intentional then update this assert");
 static_assert(FPackageTrailer::FFooter::SizeOnDisk == 20, "FPackageTrailer::FFooter size has been changed, if this was intentional then update this assert");
 
+namespace
+{
+
 /** Utility for recording failed package open reasons */
 void LogPackageOpenFailureMessage(const FPackagePath& PackagePath)
 {
@@ -54,6 +57,30 @@ void LogPackageOpenFailureMessage(const FPackagePath& PackagePath)
 	}
 }
 
+/* Utility for finding a lookup table entry for reading from inside of a FHeader*/
+const Private::FLookupTableEntry* FindEntryInHeader(const FPackageTrailer::FHeader& Header, const FIoHash& Id)
+{
+	const Private::FLookupTableEntry* Entry = Header.PayloadLookupTable.FindByPredicate([&Id](const Private::FLookupTableEntry& Entry)->bool
+		{
+			return Entry.Identifier == Id;
+		});
+
+	return Entry;
+}
+
+/* Utility for finding a lookup table entry for writing from inside of a FHeader */
+Private::FLookupTableEntry* FindEntryInHeader(FPackageTrailer::FHeader& Header, const FIoHash& Id)
+{
+	Private::FLookupTableEntry* Entry = Header.PayloadLookupTable.FindByPredicate([&Id](const Private::FLookupTableEntry& Entry)->bool
+		{
+			return Entry.Identifier == Id;
+		});
+
+	return Entry;
+}
+
+} //namespace
+
 namespace Private
 {
 
@@ -428,10 +455,7 @@ FCompressedBuffer FPackageTrailer::LoadLocalPayload(const FIoHash& Id, FArchive&
 {
 	// TODO: This method should be able to load the payload in all cases, but we need a good way of passing the Archive/PackagePath 
 	// to the trailer etc. Would work if we stored the package path in the trailer.
-	const Private::FLookupTableEntry* Entry = Header.PayloadLookupTable.FindByPredicate([&Id](const Private::FLookupTableEntry& Entry)->bool
-		{
-			return Entry.Identifier == Id;
-		});
+	const Private::FLookupTableEntry* Entry = FindEntryInHeader(Header, Id);
 
 	if (Entry == nullptr || Entry->AccessMode != EPayloadAccessMode::Local)
 	{
@@ -449,10 +473,7 @@ FCompressedBuffer FPackageTrailer::LoadLocalPayload(const FIoHash& Id, FArchive&
 
 bool FPackageTrailer::UpdatePayloadAsVirtualized(const FIoHash& Identifier)
 {
-	Private::FLookupTableEntry* Entry = Header.PayloadLookupTable.FindByPredicate([&Identifier](const Private::FLookupTableEntry& Entry)->bool
-		{
-			return Entry.Identifier == Identifier;
-		});
+	Private::FLookupTableEntry* Entry = FindEntryInHeader(Header, Identifier);
 
 	if (Entry != nullptr)
 	{
@@ -470,10 +491,7 @@ bool FPackageTrailer::UpdatePayloadAsVirtualized(const FIoHash& Identifier)
 
 EPayloadStatus FPackageTrailer::FindPayloadStatus(const FIoHash& Id) const
 {
-	const Private::FLookupTableEntry* Entry = Header.PayloadLookupTable.FindByPredicate([&Id](const Private::FLookupTableEntry& Entry)->bool
-		{
-			return Entry.Identifier == Id;
-		});
+	const Private::FLookupTableEntry* Entry = FindEntryInHeader(Header, Id);
 
 	if (Entry == nullptr)
 	{
@@ -505,10 +523,7 @@ int64 FPackageTrailer::FindPayloadOffsetInFile(const FIoHash& Id) const
 {
 	if (!Id.IsZero())
 	{
-		const Private::FLookupTableEntry* Entry = Header.PayloadLookupTable.FindByPredicate([&Id](const Private::FLookupTableEntry& Entry)->bool
-			{
-				return Entry.Identifier == Id;
-			});
+		const Private::FLookupTableEntry* Entry = FindEntryInHeader(Header, Id);
 
 		//TODO Better way to return an error?
 		check(TrailerPositionInFile != INDEX_NONE);
@@ -541,6 +556,21 @@ int64 FPackageTrailer::FindPayloadOffsetInFile(const FIoHash& Id) const
 	}
 }
 
+int64 FPackageTrailer::FindPayloadSizeOnDisk(const FIoHash& Id) const
+{
+	if (!Id.IsZero())
+	{
+		const Private::FLookupTableEntry* Entry = FindEntryInHeader(Header, Id);
+
+		if (Entry != nullptr)
+		{
+			return Entry->CompressedSize;
+		}
+	}
+
+	return INDEX_NONE;
+}
+
 int64 FPackageTrailer::GetTrailerLength() const
 {
 	return Header.HeaderLength + Header.PayloadsDataLength + FFooter::SizeOnDisk;
diff --git a/Engine/Source/Runtime/CoreUObject/Public/UObject/PackageTrailer.h b/Engine/Source/Runtime/CoreUObject/Public/UObject/PackageTrailer.h
--- a/Engine/Source/Runtime/CoreUObject/Public/UObject/PackageTrailer.h
+++ b/Engine/Source/Runtime/CoreUObject/Public/UObject/PackageTrailer.h
@@ -366,6 +366,9 @@ public:
 	/** Returns the absolute offset of the payload in the package file, invalid and virtualized payloads will return INDEX_NONE */
 	[[nodiscard]] int64 FindPayloadOffsetInFile(const FIoHash& Id) const;
 
+	/** Returns the size of the payload on as stored on disk, invalid and virtualized payloads will return INDEX_NONE */
+	[[nodiscard]] int64 FindPayloadSizeOnDisk(const FIoHash& Id) const;
+
 	/** Returns the total size of the of the trailer on disk in bytes */
 	[[nodiscard]] int64 GetTrailerLength() const;
 
