- Adding safety checks to FAsynPackage2::PostLoadInstances

diff --git a/Engine/Source/Runtime/CoreUObject/Private/Serialization/AsyncLoading2.cpp b/Engine/Source/Runtime/CoreUObject/Private/Serialization/AsyncLoading2.cpp
--- a/Engine/Source/Runtime/CoreUObject/Private/Serialization/AsyncLoading2.cpp
+++ b/Engine/Source/Runtime/CoreUObject/Private/Serialization/AsyncLoading2.cpp
@@ -5934,8 +5934,14 @@ EAsyncPackageState::Type FAsyncPackage2::PostLoadInstances(FAsyncLoadingThreadSt
 
 		if (!(Export.bFiltered | Export.bExportLoadFailed))
 		{
-			UClass* ObjClass = Export.Object->GetClass();
-			ObjClass->PostLoadInstance(Export.Object);
+			if (IsValid(Export.Object))
+			{
+				UClass* ObjClass = Export.Object->GetClass();
+				if (IsValid(ObjClass))
+				{
+					ObjClass->PostLoadInstance(Export.Object);
+				}
+			}
 		}
 	}
 	return PostLoadInstanceIndex == ExportCount ? EAsyncPackageState::Complete : EAsyncPackageState::TimeOut;
