Disable fix from 19830106 while investigating why it doesn't work
#rnx

diff --git a/Engine/Source/Runtime/CoreUObject/Private/UObject/UObjectGlobals.cpp b/Engine/Source/Runtime/CoreUObject/Private/UObject/UObjectGlobals.cpp
--- a/Engine/Source/Runtime/CoreUObject/Private/UObject/UObjectGlobals.cpp
+++ b/Engine/Source/Runtime/CoreUObject/Private/UObject/UObjectGlobals.cpp
@@ -1177,7 +1177,7 @@ UPackage* LoadPackageInternal(UPackage* InOuter, const FPackagePath& PackagePath
 
 		GSyncLoadUsingAsyncLoaderCount++;
 
-#if WITH_EDITOR
+#if 0
 		// Temporarily clear the list of loaded objects before entering async loading. Otherwise the (uncooked) async loader might try to preload them which in turn can lead to
 		// recursive imports of the same cooked package.
 		FUObjectSerializeContext* LoadContext = FUObjectThreadContext::Get().GetSerializeContext();
@@ -1192,7 +1192,7 @@ UPackage* LoadPackageInternal(UPackage* InOuter, const FPackagePath& PackagePath
 			FlushAsyncLoading(RequestID);
 		}
 
-#if WITH_EDITOR
+#if 0
 		check(ThreadObjLoaded.IsEmpty());
 		Swap(SavedObjLoaded, ThreadObjLoaded);
 #endif
