Exposing FTopLevelAssetPath to Blueprints

diff --git a/Engine/Source/Runtime/CoreUObject/Private/UObject/Class.cpp b/Engine/Source/Runtime/CoreUObject/Private/UObject/Class.cpp
--- a/Engine/Source/Runtime/CoreUObject/Private/UObject/Class.cpp
+++ b/Engine/Source/Runtime/CoreUObject/Private/UObject/Class.cpp
@@ -51,6 +51,7 @@
 #include "HAL/ThreadSafeCounter.h"
 #include "Math/InterpCurvePoint.h"
 #include "UObject/UE5MainStreamObjectVersion.h"
+#include "UObject/TopLevelAssetPath.h"
 
 // This flag enables some expensive class tree validation that is meant to catch mutations of 
 // the class tree outside of SetSuperStruct. It has been disabled because loading blueprints 
@@ -6393,6 +6394,12 @@ UScriptStruct* TBaseStructure<FTestUninitializedScriptStructMembersTest>::Get()
 	return ScriptStruct;
 }
 
+UScriptStruct* TBaseStructure<FTopLevelAssetPath>::Get()
+{
+	static UScriptStruct* ScriptStruct = StaticGetBaseStructureInternal(TEXT("TopLevelAssetPath"));
+	return ScriptStruct;
+}
+
 
 
 #define UE_DEFINE_CORE_VARIANT_TYPE(VARIANT, CORE)												\
diff --git a/Engine/Source/Runtime/CoreUObject/Public/UObject/Class.h b/Engine/Source/Runtime/CoreUObject/Public/UObject/Class.h
--- a/Engine/Source/Runtime/CoreUObject/Public/UObject/Class.h
+++ b/Engine/Source/Runtime/CoreUObject/Public/UObject/Class.h
@@ -3695,7 +3695,11 @@ template<> struct TBaseStructure<FTestUninitializedScriptStructMembersTest>
 	COREUOBJECT_API static UScriptStruct* Get();
 };
 
-
+struct FTopLevelAssetPath;
+template<> struct TBaseStructure<FTopLevelAssetPath>
+{
+	COREUOBJECT_API static UScriptStruct* Get();
+};
 
 // TBaseStructure for explicit core variant types only. e.g. FVector3d returns "Vector3d" struct. 
 template< class T > struct TVariantStructure
diff --git a/Engine/Source/Runtime/CoreUObject/Public/UObject/NoExportTypes.h b/Engine/Source/Runtime/CoreUObject/Public/UObject/NoExportTypes.h
--- a/Engine/Source/Runtime/CoreUObject/Public/UObject/NoExportTypes.h
+++ b/Engine/Source/Runtime/CoreUObject/Public/UObject/NoExportTypes.h
@@ -1688,6 +1688,21 @@ struct FSoftClassPath : public FSoftObjectPath
 {
 };
 
+/**
+ * A struct that can reference a top level asset such as '/Path/To/Package.AssetName'
+ * @note The full C++ class is located here: Engine\Source\Runtime\CoreUObject\Public\UObject\TopLevelAssetPath.h
+ */
+USTRUCT(noexport, BlueprintType)
+struct FTopLevelAssetPath
+{
+	/** Name of the package containing the asset e.g. /Path/To/Package */
+	UPROPERTY(EditAnywhere, SaveGame, BlueprintReadWrite, Category = TopLevelAssetPath)
+	FName PackageName;
+	/** Name of the asset within the package e.g. 'AssetName' */
+	UPROPERTY(EditAnywhere, SaveGame, BlueprintReadWrite, Category = TopLevelAssetPath)
+	FName AssetName;
+};
+
 /** 
  * A type of primary asset, used by the Asset Manager system.
  * @note The full C++ class is located here: Engine\Source\Runtime\CoreUObject\Public\UObject\PrimaryAssetId.h
diff --git a/Engine/Source/Runtime/CoreUObject/Public/UObject/TopLevelAssetPath.h b/Engine/Source/Runtime/CoreUObject/Public/UObject/TopLevelAssetPath.h
--- a/Engine/Source/Runtime/CoreUObject/Public/UObject/TopLevelAssetPath.h
+++ b/Engine/Source/Runtime/CoreUObject/Public/UObject/TopLevelAssetPath.h
@@ -17,6 +17,7 @@ class CORE_API FString;
  *  b) storing an empty FString for a subobject path as FSoftObjectPath allows
  * Can also be used to reference the package itself in which case the second name is NAME_None
  * and the object resolves to the string `/Path/To/Package`
+ * This struct is mirrored and exposed to the UE reflection system in NoExportTypes.h
 */
 struct COREUOBJECT_API FTopLevelAssetPath
 {
