New StaticFind*Object(s) helper functions related to ANY_PACKAGE removal

diff --git a/Engine/Source/Runtime/CoreUObject/Private/UObject/UObjectGlobals.cpp b/Engine/Source/Runtime/CoreUObject/Private/UObject/UObjectGlobals.cpp
--- a/Engine/Source/Runtime/CoreUObject/Private/UObject/UObjectGlobals.cpp
+++ b/Engine/Source/Runtime/CoreUObject/Private/UObject/UObjectGlobals.cpp
@@ -346,6 +346,27 @@ UObject* StaticFindObjectFastSafe(UClass* ObjectClass, UObject* ObjectPackage, F
 	return FoundObject;
 }
 
+#if WITH_EDITOR
+static UObject* LoadObjectWhenImportingT3D(UClass* ObjectClass, const TCHAR* OrigInName)
+{
+	UObject* MatchingObject = nullptr;
+	static bool s_bCurrentlyLoading = false;
+	if (s_bCurrentlyLoading == false)
+	{
+		FString NameCheck = OrigInName;
+		if (NameCheck.Contains(TEXT("."), ESearchCase::CaseSensitive) &&
+			!NameCheck.Contains(TEXT("'"), ESearchCase::CaseSensitive) &&
+			!NameCheck.Contains(TEXT(":"), ESearchCase::CaseSensitive))
+		{
+			s_bCurrentlyLoading = true;
+			MatchingObject = StaticLoadObject(ObjectClass, nullptr, OrigInName, nullptr, LOAD_NoWarn, nullptr);
+			s_bCurrentlyLoading = false;
+		}
+	}
+	return MatchingObject;
+}
+
+#endif // WITH_EDITOR
 //
 // Find an optional object.
 //
@@ -360,25 +381,12 @@ UObject* StaticFindObject( UClass* ObjectClass, UObject* InObjectPackage, const
 	UObject* MatchingObject = nullptr;
 #if WITH_EDITOR
 	// If the editor is running, and T3D is being imported, ensure any packages referenced are fully loaded.
-	if ((GIsEditor == true) && (GIsImportingT3D == true))// && (ObjectPackage != ANY_PACKAGE) && (ObjectPackage != NULL))
+	if ((GIsEditor == true) && (GIsImportingT3D == true))
 	{
-		static bool s_bCurrentlyLoading = false;
-
-		if (s_bCurrentlyLoading == false)
+		MatchingObject = LoadObjectWhenImportingT3D(ObjectClass, OrigInName);
+		if (MatchingObject)
 		{
-			FString NameCheck = OrigInName;
-			if (NameCheck.Contains(TEXT("."), ESearchCase::CaseSensitive) && 
-				!NameCheck.Contains(TEXT("'"), ESearchCase::CaseSensitive) && 
-				!NameCheck.Contains(TEXT(":"), ESearchCase::CaseSensitive) )
-			{
-				s_bCurrentlyLoading = true;
-				MatchingObject = StaticLoadObject(ObjectClass, nullptr, OrigInName, nullptr,  LOAD_NoWarn, nullptr);
-				s_bCurrentlyLoading = false;
-				if (MatchingObject != nullptr)
-				{
-					return MatchingObject;
-				}
-			}
+			return MatchingObject;
 		}
 	}
 #endif	//#if !WITH_EDITOR
@@ -437,6 +445,128 @@ UObject* StaticFindObjectSafe( UClass* ObjectClass, UObject* ObjectParent, const
 	}
 }
 
+bool StaticFindAllObjectsFast(TArray<UObject*>& OutFoundObjects, UClass* ObjectClass, FName ObjectName, bool ExactClass, EObjectFlags ExclusiveFlags, EInternalObjectFlags ExclusiveInternalFlags)
+{
+	UE_CLOG(UE::IsSavingPackage(nullptr) || IsGarbageCollectingOnGameThread(), LogUObjectGlobals, Fatal, TEXT("Illegal call to StaticFindAllObjectsFast() while serializing object data or garbage collecting!"));
+
+	// We don't want to return any objects that are currently being background loaded unless we're using FindObject during async loading.
+	ExclusiveInternalFlags |= IsInAsyncLoadingThread() ? EInternalObjectFlags::None : EInternalObjectFlags::AsyncLoading;
+	return StaticFindAllObjectsFastInternal(OutFoundObjects, ObjectClass, ObjectName, ExactClass, ExclusiveFlags, ExclusiveInternalFlags);
+}
+
+bool StaticFindAllObjectsFastSafe(TArray<UObject*>& OutFoundObjects, UClass* ObjectClass, FName ObjectName, bool ExactClass, EObjectFlags ExclusiveFlags, EInternalObjectFlags ExclusiveInternalFlags)
+{
+	bool bFoundObjects = false;
+	if (!UE::IsSavingPackage(nullptr) && !IsGarbageCollectingOnGameThread())
+	{
+		// We don't want to return any objects that are currently being background loaded unless we're using FindObject during async loading.
+		ExclusiveInternalFlags |= IsInAsyncLoadingThread() ? EInternalObjectFlags::None : EInternalObjectFlags::AsyncLoading;
+		bFoundObjects = StaticFindAllObjectsFastInternal(OutFoundObjects, ObjectClass, ObjectName, ExactClass, ExclusiveFlags, ExclusiveInternalFlags);
+	}
+	return bFoundObjects;
+}
+
+bool StaticFindAllObjects(TArray<UObject*>& OutFoundObjects, UClass* ObjectClass, const TCHAR* OrigInName, bool ExactClass)
+{
+	INC_DWORD_STAT(STAT_FindObject);
+
+	UE_CLOG(UE::IsSavingPackage(nullptr), LogUObjectGlobals, Fatal, TEXT("Illegal call to StaticFindAllObjects() while serializing object data!"));
+	UE_CLOG(IsGarbageCollectingOnGameThread(), LogUObjectGlobals, Fatal, TEXT("Illegal call to StaticFindAllObjects() while collecting garbage!"));
+
+#if WITH_EDITOR
+	// If the editor is running, and T3D is being imported, ensure any packages referenced are fully loaded.
+	if ((GIsEditor == true) && (GIsImportingT3D == true))
+	{
+		UObject* MatchingObject = LoadObjectWhenImportingT3D(ObjectClass, OrigInName);
+		if (MatchingObject != nullptr)
+		{
+			OutFoundObjects.Add(MatchingObject);
+			return true;
+		}
+	}
+#endif	//#if !WITH_EDITOR
+
+	// Don't resolve the name since we're searching in any package
+	FString InName = OrigInName;
+	ConstructorHelpers::StripObjectClass(InName);
+	FName ObjectName(*InName, FNAME_Add);
+
+	return StaticFindAllObjectsFast(OutFoundObjects, ObjectClass, ObjectName, ExactClass);
+}
+
+bool StaticFindAllObjectsSafe(TArray<UObject*>& OutFoundObjects, UClass* ObjectClass, const TCHAR* OrigInName, bool ExactClass)
+{
+	bool bFoundObjects = false;
+	if (!UE::IsSavingPackage(nullptr) && !IsGarbageCollectingOnGameThread())
+	{
+		bFoundObjects = StaticFindAllObjects(OutFoundObjects, ObjectClass, OrigInName, ExactClass);
+	}
+	return bFoundObjects;
+}
+
+UObject* StaticFindFirstObject(UClass* Class, const TCHAR* Name, EFindFirstObjectOptions Options /*= EFindFirstObjectOptions::None*/, ELogVerbosity::Type AmbiguousMessageVerbosity /*= ELogVerbosity::Warning*/, const TCHAR* InCurrentOperation /*= nullptr*/)
+{
+	UObject* Result = nullptr;
+	// Don't resolve the name since we're searching in any package
+	FString InName = Name;
+	ConstructorHelpers::StripObjectClass(InName);
+	FName ObjectName(*InName, FNAME_Add);
+
+	if (AmbiguousMessageVerbosity == ELogVerbosity::NoLogging)
+	{
+		Result = StaticFindFirstObjectFastInternal(Class, ObjectName, !!(Options & EFindFirstObjectOptions::ExactClass));
+	}
+	else
+	{
+		TArray<UObject*> FoundObjects;
+		if (StaticFindAllObjectsFast(FoundObjects, Class, ObjectName, !!(Options & EFindFirstObjectOptions::ExactClass)))
+		{
+			Result = FoundObjects[0];
+			if (AmbiguousMessageVerbosity != ELogVerbosity::NoLogging && FoundObjects.Num() > 1)
+			{
+				TStringBuilder<256> Message;
+				Message.Append(TEXT("StaticFindFirstObject: Ambiguous object name "));
+				Message.Append(Name);
+				if (InCurrentOperation)
+				{
+					Message.Append(TEXT(" while "));
+					Message.Append(InCurrentOperation);
+				}
+				Message.Append(TEXT(", will return "));
+				Message.Append(Result->GetPathName());
+				Message.Append(TEXT(" but could also be: "));
+				for (int32 ObjectIndex = 1; ObjectIndex < FoundObjects.Num(); ++ObjectIndex)
+				{
+					Message.Append(FoundObjects[ObjectIndex]->GetPathName());
+					if (ObjectIndex < FoundObjects.Num() - 1)
+					{
+						Message.Append(TEXT(","));
+					}
+				}
+
+				if (AmbiguousMessageVerbosity == ELogVerbosity::Fatal)
+				{
+					UE_LOG(LogUObjectGlobals, Fatal, TEXT("%s"), *Message.ToString());
+				}
+				else
+				{
+					GLog->CategorizedLogf(TEXT("LogUObjectGlobals"), AmbiguousMessageVerbosity, TEXT("%s"), *Message.ToString());
+				}
+			}
+		}
+	}
+	return Result;
+}
+
+UObject* StaticFindFirstObjectSafe(UClass* Class, const TCHAR* Name, EFindFirstObjectOptions Options /*= EFindFirstObjectOptions::None*/, ELogVerbosity::Type AmbiguousMessageVerbosity /*= ELogVerbosity::Warning*/, const TCHAR* InCurrentOperation /*= nullptr*/)
+{
+	UObject* Result = nullptr;
+	if (!UE::IsSavingPackage(nullptr) && !IsGarbageCollectingOnGameThread())
+	{
+		Result = StaticFindFirstObject(Class, Name, Options, AmbiguousMessageVerbosity, InCurrentOperation);
+	}
+	return Result;
+}
 
 //
 // Global property setting.
diff --git a/Engine/Source/Runtime/CoreUObject/Private/UObject/UObjectHash.cpp b/Engine/Source/Runtime/CoreUObject/Private/UObject/UObjectHash.cpp
--- a/Engine/Source/Runtime/CoreUObject/Private/UObject/UObjectHash.cpp
+++ b/Engine/Source/Runtime/CoreUObject/Private/UObject/UObjectHash.cpp
@@ -822,6 +822,95 @@ bool DoesObjectPossiblyExist(const UObject* InOuter, FName ObjectName)
 	return ThreadHash.HashOuter.Contains(Hash);
 }
 
+bool StaticFindAllObjectsFastInternal(TArray<UObject*>& OutFoundObjects, const UClass* ObjectClass, FName ObjectName, bool bExactClass, EObjectFlags ExcludeFlags, EInternalObjectFlags ExclusiveInternalFlags)
+{
+	INC_DWORD_STAT(STAT_FindObjectFast);
+
+	ExclusiveInternalFlags |= EInternalObjectFlags::Unreachable;
+
+	FObjectSearchPath SearchPath(ObjectName);
+	const int32 Hash = GetObjectHash(SearchPath.Inner);
+	uint32 NumFoundObjects = 0; // Keeping track of the number of objects foundn. We allow OutFoundObjects to not be empty
+
+	FUObjectHashTables& ThreadHash = FUObjectHashTables::Get();
+	FHashTableLock HashLock(ThreadHash);
+
+	FHashBucket* Bucket = ThreadHash.Hash.Find(Hash);
+	if (Bucket)
+	{
+		for (FHashBucketIterator It(*Bucket); It; ++It)
+		{
+			UObject* Object = (UObject*)*It;
+
+			if
+				((Object->GetFName() == SearchPath.Inner)
+
+					/* Don't return objects that have any of the exclusive flags set */
+					&& !Object->HasAnyFlags(ExcludeFlags)
+
+					/** If a class was specified, check that the object is of the correct class */
+					&& (ObjectClass == nullptr || (bExactClass ? Object->GetClass() == ObjectClass : Object->IsA(ObjectClass)))
+
+					/** Include (or not) pending kill objects */
+					&& !Object->HasAnyInternalFlags(ExclusiveInternalFlags)
+
+					/** Ensure that the partial path provided matches the object found */
+					&& SearchPath.MatchOuterNames(Object->GetOuter()))
+			{
+				checkf(!Object->IsUnreachable(), TEXT("%s"), *Object->GetFullName());
+				OutFoundObjects.Add(Object);
+				NumFoundObjects++;
+			}
+		}
+	}
+	return !!NumFoundObjects;
+}
+
+UObject* StaticFindFirstObjectFastInternal(const UClass* ObjectClass, FName ObjectName, bool bExactClass, EObjectFlags ExcludeFlags, EInternalObjectFlags ExclusiveInternalFlags)
+{
+	INC_DWORD_STAT(STAT_FindObjectFast);
+
+	ExclusiveInternalFlags |= EInternalObjectFlags::Unreachable;
+
+	UObject* Result = nullptr;
+	FObjectSearchPath SearchPath(ObjectName);
+	const int32 Hash = GetObjectHash(SearchPath.Inner);
+	uint32 NumFoundObjects = 0; // Keeping track of the number of objects found. We allow OutFoundObjects to not be empty
+
+	FUObjectHashTables& ThreadHash = FUObjectHashTables::Get();
+	FHashTableLock HashLock(ThreadHash);
+
+	FHashBucket* Bucket = ThreadHash.Hash.Find(Hash);
+	if (Bucket)
+	{
+		for (FHashBucketIterator It(*Bucket); It; ++It)
+		{
+			UObject* Object = (UObject*)*It;
+
+			if
+				((Object->GetFName() == SearchPath.Inner)
+
+					/* Don't return objects that have any of the exclusive flags set */
+					&& !Object->HasAnyFlags(ExcludeFlags)
+
+					/** If a class was specified, check that the object is of the correct class */
+					&& (ObjectClass == nullptr || (bExactClass ? Object->GetClass() == ObjectClass : Object->IsA(ObjectClass)))
+
+					/** Include (or not) pending kill objects */
+					&& !Object->HasAnyInternalFlags(ExclusiveInternalFlags)
+
+					/** Ensure that the partial path provided matches the object found */
+					&& SearchPath.MatchOuterNames(Object->GetOuter()))
+			{
+				checkf(!Object->IsUnreachable(), TEXT("%s"), *Object->GetFullName());
+				Result = Object;
+				break;
+			}
+		}
+	}
+	return Result;
+}
+
 // Assumes that ThreadHash's critical is already locked
 FORCEINLINE static void AddToOuterMap(FUObjectHashTables& ThreadHash, UObjectBase* Object)
 {
diff --git a/Engine/Source/Runtime/CoreUObject/Public/UObject/UObjectGlobals.h b/Engine/Source/Runtime/CoreUObject/Public/UObject/UObjectGlobals.h
--- a/Engine/Source/Runtime/CoreUObject/Public/UObject/UObjectGlobals.h
+++ b/Engine/Source/Runtime/CoreUObject/Public/UObject/UObjectGlobals.h
@@ -288,6 +288,96 @@ COREUOBJECT_API UObject* StaticFindObjectChecked( UClass* Class, UObject* InOute
 /** Internal version of StaticFindObject that will not assert on GIsSavingPackage or IsGarbageCollecting() */
 COREUOBJECT_API UObject* StaticFindObjectSafe( UClass* Class, UObject* InOuter, const TCHAR* Name, bool ExactClass=false );
 
+/**
+ * Fast version of StaticFindAllObjects that relies on the passed in FName being the object name without any group/package qualifiers.
+ * This will find all objects matching the specified name and class.
+ *
+ * @param	OutFoundObjects	Array of objects matching the search parameters
+ * @param	ObjectClass		The to be found object's class
+ * @param	ObjectName		Object name to look for relative to InOuter
+ * @param	ExactClass		Whether to require an exact match with the passed in class
+ * @param	ExclusiveFlags	Ignores objects that contain any of the specified exclusive flags
+ * @param	ExclusiveInternalFlags  Ignores objects that contain any of the specified internal exclusive flags
+ *
+ * @return	Returns true if any objects were found, false otherwise
+ */
+COREUOBJECT_API bool StaticFindAllObjectsFast(TArray<UObject*>& OutFoundObjects, UClass* ObjectClass, FName ObjectName, bool ExactClass = false, EObjectFlags ExclusiveFlags = RF_NoFlags, EInternalObjectFlags ExclusiveInternalFlags = EInternalObjectFlags::None);
+
+/**
+ * Fast version of StaticFindAllObjects that relies on the passed in FName being the object name without any group/package qualifiers.
+ * This will find all objects matching the specified name and class.
+ * This version of StaticFindAllObjectsFast will not assert on GIsSavingPackage or IsGarbageCollecting()
+ * 
+ * @param	OutFoundObjects	Array of objects matching the search parameters
+ * @param	ObjectClass		The to be found object's class
+ * @param	ObjectName		Object name to look for relative to InOuter
+ * @param	ExactClass		Whether to require an exact match with the passed in class
+ * @param	ExclusiveFlags	Ignores objects that contain any of the specified exclusive flags
+ * @param	ExclusiveInternalFlags  Ignores objects that contain any of the specified internal exclusive flags
+ *
+ * @return	Returns true if any objects were found, false otherwise
+ */
+COREUOBJECT_API bool StaticFindAllObjectsFastSafe(TArray<UObject*>& OutFoundObjects, UClass* ObjectClass, FName ObjectName, bool ExactClass = false, EObjectFlags ExclusiveFlags = RF_NoFlags, EInternalObjectFlags ExclusiveInternalFlags = EInternalObjectFlags::None);
+
+/**
+ * Tries to find all objects matching the search paramters in memory. This will handle fully qualified paths of the form /path/packagename.object:subobject and resolve references for you.
+ *
+ * @param	OutFoundObjects	Array of objects matching the search parameters
+ * @param	Class			The to be found object's class
+ * @param	Name			The object path to search for an object, relative to InOuter
+ * @param	ExactClass		Whether to require an exact match with the passed in class
+ *
+ * @return	Returns true if any objects were found, false otherwise
+ */
+COREUOBJECT_API bool StaticFindAllObjects(TArray<UObject*>& OutFoundObjects, UClass* Class, const TCHAR* Name, bool ExactClass = false);
+
+/**
+ * Tries to find all objects matching the search paramters in memory. This will handle fully qualified paths of the form /path/packagename.object:subobject and resolve references for you.
+ * This version of StaticFindAllObjects will not assert on GIsSavingPackage or IsGarbageCollecting()
+ * 
+ * @param	OutFoundObjects	Array of objects matching the search parameters
+ * @param	Class			The to be found object's class
+ * @param	Name			The object path to search for an object, relative to InOuter
+ * @param	ExactClass		Whether to require an exact match with the passed in class
+ *
+ * @return	Returns true if any objects were found, false otherwise
+ */
+COREUOBJECT_API bool StaticFindAllObjectsSafe(TArray<UObject*>& OutFoundObjects, UClass* Class, const TCHAR* Name, bool ExactClass = false);
+
+
+enum class EFindFirstObjectOptions
+{
+	None = 0, // Unused / defaults to Quiet
+	ExactClass = 1 << 1 // Whether to require an exact match with the passed in class
+};
+ENUM_CLASS_FLAGS(EFindFirstObjectOptions);
+
+/**
+ * Tries to find the first objects matching the search paramters in memory. This will handle fully qualified paths of the form /path/packagename.object:subobject and resolve references for you.
+ *
+ * @param	Class						The to be found object's class
+ * @param	Name						The object path to search for an object, relative to InOuter
+ * @param	Options						Search options
+ * @param	AmbiguousMessageVerbosity	Verbosity with which to print a message if the search result is ambiguous
+ * @param	InCurrentOperation			Current operation to be logged with ambiguous search warning
+ *
+ * @return	Returns a pointer to an object if found, null otherwise
+ */
+COREUOBJECT_API UObject* StaticFindFirstObject(UClass* Class, const TCHAR* Name, EFindFirstObjectOptions Options = EFindFirstObjectOptions::None, ELogVerbosity::Type AmbiguousMessageVerbosity = ELogVerbosity::NoLogging, const TCHAR* InCurrentOperation = nullptr);
+
+/**
+ * Tries to find the first objects matching the search paramters in memory. This will handle fully qualified paths of the form /path/packagename.object:subobject and resolve references for you.
+ * This version of StaticFindFirstObject will not assert on GIsSavingPackage or IsGarbageCollecting()
+ * 
+ * @param	Class						The to be found object's class
+ * @param	Name						The object path to search for an object, relative to InOuter
+ * @param	Options						Search options
+ * @param	AmbiguousMessageVerbosity	Verbosity with which to print a message if the search result is ambiguous
+ * @param	InCurrentOperation			Current operation to be logged with ambiguous search warning
+ *
+ * @return	Returns a pointer to an object if found, null otherwise
+ */
+COREUOBJECT_API UObject* StaticFindFirstObjectSafe(UClass* Class, const TCHAR* Name, EFindFirstObjectOptions Options = EFindFirstObjectOptions::None, ELogVerbosity::Type AmbiguousMessageVerbosity = ELogVerbosity::NoLogging, const TCHAR* InCurrentOperation = nullptr);
 
 /**
  * Parse a reference to an object from a text representation
@@ -1536,6 +1626,26 @@ inline T* FindObjectSafe( UObject* Outer, const TCHAR* Name, bool ExactClass=fal
 	return (T*)StaticFindObjectSafe( T::StaticClass(), Outer, Name, ExactClass );
 }
 
+/**
+ * Find an optional object with proper handling of potential ambiguity.
+ * @see StaticFindFirstObject()
+ */
+template< class T >
+inline T* FindFirstObject(const TCHAR* Name, EFindFirstObjectOptions Options = EFindFirstObjectOptions::None, ELogVerbosity::Type AmbiguousMessageVerbosity = ELogVerbosity::NoLogging, const TCHAR* CurrentOperation = nullptr)
+{
+	return (T*)StaticFindFirstObject(T::StaticClass(), Name, Options, AmbiguousMessageVerbosity, CurrentOperation);
+}
+
+/**
+ * Find an optional object with proper handling of potential ambiguity without asserting on GIsSavingPackage or IsGarbageCollecting()
+ * @see StaticFindFirstObject()
+ */
+template< class T >
+inline T* FindFirstObjectSafe(const TCHAR* Name, EFindFirstObjectOptions Options = EFindFirstObjectOptions::None, ELogVerbosity::Type AmbiguousMessageVerbosity = ELogVerbosity::NoLogging, const TCHAR* CurrentOperation = nullptr)
+{
+	return (T*)StaticFindFirstObjectSafe(T::StaticClass(), Name, Options, AmbiguousMessageVerbosity, CurrentOperation);
+}
+
 /** 
  * Load an object. 
  * @see StaticLoadObject()
diff --git a/Engine/Source/Runtime/CoreUObject/Public/UObject/UObjectHash.h b/Engine/Source/Runtime/CoreUObject/Public/UObject/UObjectHash.h
--- a/Engine/Source/Runtime/CoreUObject/Public/UObject/UObjectHash.h
+++ b/Engine/Source/Runtime/CoreUObject/Public/UObject/UObjectHash.h
@@ -29,7 +29,32 @@ DECLARE_DWORD_COUNTER_STAT_EXTERN(TEXT("NumObjects"), STAT_Hash_NumObjects, STAT
  * @param	ExclusiveInternalFlags	Ignores objects that contain any of the specified internal exclusive flags
  * @return	Returns a pointer to the found object or NULL if none could be found
  */
-UObject* StaticFindObjectFastInternal(const UClass* Class, const UObject* InOuter, FName InName, bool ExactClass = false, bool AnyPackage = false, EObjectFlags ExclusiveFlags = RF_NoFlags, EInternalObjectFlags ExclusiveInternalFlags = EInternalObjectFlags::None);
+COREUOBJECT_API UObject* StaticFindObjectFastInternal(const UClass* Class, const UObject* InOuter, FName InName, bool ExactClass = false, bool AnyPackage = false, EObjectFlags ExclusiveFlags = RF_NoFlags, EInternalObjectFlags ExclusiveInternalFlags = EInternalObjectFlags::None);
+
+/**
+ * Private internal version of StaticFindAllObjectsFast that allows using 0 exclusion flags.
+ *
+ * @param	OutFoundObjects	Array of objects matching the specified search parameters
+ * @param	ObjectClass		The to be found object's class
+ * @param	ObjectName		The to be found object's name
+ * @param	bExactClass		Whether to require an exact match with the passed in class
+ * @param	ExcludeFlags	Ignores objects that contain any of the specified exclusive flags
+ * @param	ExclusiveInternalFlags	Ignores objects that contain any of the specified internal exclusive flags
+ * @return	Returns true if any objects were found, false otherwise
+ */
+COREUOBJECT_API bool StaticFindAllObjectsFastInternal(TArray<UObject*>& OutFoundObjects, const UClass* ObjectClass, FName ObjectName, bool bExactClass, EObjectFlags ExcludeFlags = RF_NoFlags, EInternalObjectFlags ExclusiveInternalFlags = EInternalObjectFlags::None);
+
+/**
+ * Private internal version of StaticFindFirstObject that allows using 0 exclusion flags.
+ *
+ * @param	ObjectClass		The to be found object's class
+ * @param	ObjectName		The to be found object's name
+ * @param	bExactClass		Whether to require an exact match with the passed in class
+ * @param	ExcludeFlag		Ignores objects that contain any of the specified exclusive flags
+ * @param	ExclusiveInternalFlags	Ignores objects that contain any of the specified internal exclusive flags
+ * @return	Returns first object with the specified name and class, null if not found
+ */
+COREUOBJECT_API UObject* StaticFindFirstObjectFastInternal(const UClass* ObjectClass, FName ObjectName, bool bExactClass, EObjectFlags ExcludeFlag = RF_NoFlags, EInternalObjectFlags ExclusiveInternalFlags = EInternalObjectFlags::None);
 
 /**
  * Variation of StaticFindObjectFast that uses explicit path.
