Moving utility functions related to linker tables from FLinker to FLinkerTables to have access to these functions from the AssetRegistry code (required for ANY_PACKAGE removal)

diff --git a/Engine/Source/Runtime/CoreUObject/Private/UObject/Linker.cpp b/Engine/Source/Runtime/CoreUObject/Private/UObject/Linker.cpp
--- a/Engine/Source/Runtime/CoreUObject/Private/UObject/Linker.cpp
+++ b/Engine/Source/Runtime/CoreUObject/Private/UObject/Linker.cpp
@@ -127,11 +127,11 @@ void FLinkerTables::SerializeSearchableNamesMap(FStructuredArchive::FSlot Slot)
 	Slot << SearchableNamesMap;
 }
 
-FName FLinker::GetExportClassName( int32 i )
+FName FLinkerTables::GetExportClassName(int32 ExportIndex)
 {
-	if (ExportMap.IsValidIndex(i))
+	if (ExportMap.IsValidIndex(ExportIndex))
 	{
-		FObjectExport& Export = ExportMap[i];
+		FObjectExport& Export = ExportMap[ExportIndex];
 		if( !Export.ClassIndex.IsNull() )
 		{
 			return ImpExp(Export.ClassIndex).ObjectName;
@@ -225,7 +225,6 @@ void FLinker::AddReferencedObjects(FReferenceCollector& Collector)
 #endif
 }
 
-// FLinker interface.
 /**
  * Return the path name of the UObject represented by the specified import. 
  * (can be used with StaticFindObject)
@@ -234,16 +233,16 @@ void FLinker::AddReferencedObjects(FReferenceCollector& Collector)
  *
  * @return	the path name of the UObject represented by the resource at ImportIndex
  */
-FString FLinker::GetImportPathName(int32 ImportIndex)
+FString FLinkerTables::GetImportPathName(int32 ImportIndex)
 {
 	FString Result;
 	for (FPackageIndex LinkerIndex = FPackageIndex::FromImport(ImportIndex); !LinkerIndex.IsNull();)
 	{
 		const FObjectResource& Resource = ImpExp(LinkerIndex);
 		bool bSubobjectDelimiter=false;
-
-		if (Result.Len() > 0 && GetClassName(LinkerIndex) != NAME_Package
-			&& (Resource.OuterIndex.IsNull() || GetClassName(Resource.OuterIndex) == NAME_Package) )
+		checkf(LinkerIndex.IsImport(), TEXT("Found an export while traversing the import map"));
+		if (Result.Len() > 0 && GetImportClassName(LinkerIndex) != NAME_Package
+			&& (Resource.OuterIndex.IsNull() || GetImportClassName(Resource.OuterIndex) == NAME_Package) )
 		{
 			bSubobjectDelimiter = true;
 		}
@@ -270,17 +269,17 @@ FString FLinker::GetImportPathName(int32 ImportIndex)
 /**
  * Return the path name of the UObject represented by the specified export.
  * (can be used with StaticFindObject)
- * 
+ *
+ * @param	RootPackagePath			Name of the root package for this export
  * @param	ExportIndex				index into the ExportMap for the resource to get the name for
- * @param	FakeRoot				Optional name to replace use as the root package of this object instead of the linker
  * @param	bResolveForcedExports	if true, the package name part of the return value will be the export's original package,
  *									not the name of the package it's currently contained within.
  *
  * @return	the path name of the UObject represented by the resource at ExportIndex
  */
-FString FLinker::GetExportPathName(int32 ExportIndex, const TCHAR* FakeRoot,bool bResolveForcedExports/*=false*/)
+FString FLinkerTables::GetExportPathName(const FString& RootPackagePath, int32 ExportIndex, bool bResolveForcedExports/*=false*/)
 {
-	FString Result;
+	TStringBuilder<64> Result;
 
 	bool bForcedExport = false;
 	bool bHasOuterImport = false;
@@ -296,42 +295,44 @@ FString FLinker::GetExportPathName(int32 ExportIndex, const TCHAR* FakeRoot,bool
 			if ((Resource.OuterIndex.IsNull() || GetClassName(Resource.OuterIndex) == NAME_Package)
 			  && GetClassName(LinkerIndex) != NAME_Package)
 			{
-				Result = FString(SUBOBJECT_DELIMITER) + Result;
+				Result.Prepend(SUBOBJECT_DELIMITER);
 			}
 			else
 			{
-				Result = FString(TEXT(".")) + Result;
+				Result.Prepend(TEXT("."));
 			}
 		}
-		Result = Resource.ObjectName.ToString() + Result;
+		Result.Prepend(Resource.ObjectName.ToString());
 		bForcedExport = bForcedExport || (LinkerIndex.IsExport() ? Exp(LinkerIndex).bForcedExport : false);
 	}
 
-	if ((bForcedExport && FakeRoot == nullptr && bResolveForcedExports) ||
+	if ((bForcedExport && bResolveForcedExports) ||
 		// if the export we are building the path of has an import in its outer chain, no need to append the LinkerRoot path
 		bHasOuterImport )
 	{
 		// Result already contains the correct path name for this export
-		return Result;
+		return Result.ToString();
 	}
 
-	return (FakeRoot ? FakeRoot : LinkerRoot->GetPathName()) + TEXT(".") + Result;
+	Result.Prepend(TEXT("."));
+	Result.Prepend(RootPackagePath);
+	return Result.ToString();
 }
 
-FString FLinker::GetImportFullName(int32 ImportIndex)
+FString FLinkerTables::GetImportFullName(int32 ImportIndex)
 {
 	return ImportMap[ImportIndex].ClassName.ToString() + TEXT(" ") + GetImportPathName(ImportIndex);
 }
 
-FString FLinker::GetExportFullName(int32 ExportIndex, const TCHAR* FakeRoot,bool bResolveForcedExports/*=false*/)
+FString FLinkerTables::GetExportFullName(const FString& RootPackagePath, int32 ExportIndex, bool bResolveForcedExports/*=false*/)
 {
 	FPackageIndex ClassIndex = ExportMap[ExportIndex].ClassIndex;
 	FName ClassName = ClassIndex.IsNull() ? FName(NAME_Class) : ImpExp(ClassIndex).ObjectName;
 
-	return ClassName.ToString() + TEXT(" ") + GetExportPathName(ExportIndex, FakeRoot, bResolveForcedExports);
+	return ClassName.ToString() + TEXT(" ") + GetExportPathName(RootPackagePath, ExportIndex, bResolveForcedExports);
 }
 
-FPackageIndex FLinker::ResourceGetOutermost(FPackageIndex LinkerIndex) const
+FPackageIndex FLinkerTables::ResourceGetOutermost(FPackageIndex LinkerIndex) const
 {
 	const FObjectResource* Res = &ImpExp(LinkerIndex);
 	while (!Res->OuterIndex.IsNull())
@@ -342,7 +343,7 @@ FPackageIndex FLinker::ResourceGetOutermost(FPackageIndex LinkerIndex) const
 	return LinkerIndex;
 }
 
-bool FLinker::ResourceIsIn(FPackageIndex LinkerIndex, FPackageIndex OuterIndex) const
+bool FLinkerTables::ResourceIsIn(FPackageIndex LinkerIndex, FPackageIndex OuterIndex) const
 {
 	LinkerIndex = ImpExp(LinkerIndex).OuterIndex;
 	while (!LinkerIndex.IsNull())
@@ -356,12 +357,12 @@ bool FLinker::ResourceIsIn(FPackageIndex LinkerIndex, FPackageIndex OuterIndex)
 	return false;
 }
 
-bool FLinker::DoResourcesShareOutermost(FPackageIndex LinkerIndexLHS, FPackageIndex LinkerIndexRHS) const
+bool FLinkerTables::DoResourcesShareOutermost(FPackageIndex LinkerIndexLHS, FPackageIndex LinkerIndexRHS) const
 {
 	return ResourceGetOutermost(LinkerIndexLHS) == ResourceGetOutermost(LinkerIndexRHS);
 }
 
-bool FLinker::ImportIsInAnyExport(int32 ImportIndex) const
+bool FLinkerTables::ImportIsInAnyExport(int32 ImportIndex) const
 {
 	FPackageIndex LinkerIndex = ImportMap[ImportIndex].OuterIndex;
 	while (!LinkerIndex.IsNull())
@@ -376,7 +377,7 @@ bool FLinker::ImportIsInAnyExport(int32 ImportIndex) const
 
 }
 
-bool FLinker::AnyExportIsInImport(int32 ImportIndex) const
+bool FLinkerTables::AnyExportIsInImport(int32 ImportIndex) const
 {
 	FPackageIndex OuterIndex = FPackageIndex::FromImport(ImportIndex);
 	for (int32 ExportIndex = 0; ExportIndex < ExportMap.Num(); ++ExportIndex)
@@ -389,7 +390,7 @@ bool FLinker::AnyExportIsInImport(int32 ImportIndex) const
 	return false;
 }
 
-bool FLinker::AnyExportShareOuterWithImport(int32 ImportIndex) const
+bool FLinkerTables::AnyExportShareOuterWithImport(int32 ImportIndex) const
 {
 	FPackageIndex Import = FPackageIndex::FromImport(ImportIndex);
 	for (int32 ExportIndex = 0; ExportIndex < ExportMap.Num(); ++ExportIndex)
@@ -403,6 +404,26 @@ bool FLinker::AnyExportShareOuterWithImport(int32 ImportIndex) const
 	return false;
 }
 
+FString FLinker::GetExportPathName(int32 ExportIndex, const TCHAR* FakeRoot /*= nullptr*/, bool bResolveForcedExports/*=false*/)
+{	
+	if (FakeRoot)
+	{
+		// This is to mimic the original behavior where bResolveForcedExports was only respected if FakeRoot was not specified
+		bResolveForcedExports = false;
+	}
+	return FLinkerTables::GetExportPathName(FakeRoot ? FakeRoot : LinkerRoot->GetPathName(), ExportIndex, bResolveForcedExports);
+}
+
+FString FLinker::GetExportFullName(int32 ExportIndex, const TCHAR* FakeRoot /*= nullptr*/, bool bResolveForcedExports/*=false*/)
+{
+	if (FakeRoot)
+	{
+		// This is to mimic the original behavior where bResolveForcedExports was only respected if FakeRoot was not specified
+		bResolveForcedExports = false;
+	}
+	return FLinkerTables::GetExportFullName(FakeRoot ? FakeRoot : LinkerRoot->GetPathName(), ExportIndex, bResolveForcedExports);
+}
+
 /**
  * Tell this linker to start SHA calculations
  */
diff --git a/Engine/Source/Runtime/CoreUObject/Private/UObject/LinkerLoad.cpp b/Engine/Source/Runtime/CoreUObject/Private/UObject/LinkerLoad.cpp
--- a/Engine/Source/Runtime/CoreUObject/Private/UObject/LinkerLoad.cpp
+++ b/Engine/Source/Runtime/CoreUObject/Private/UObject/LinkerLoad.cpp
@@ -2400,7 +2400,7 @@ UObject* FLinkerLoad::FindExistingExport(int32 ExportIndex)
 	if (OuterObject)
 	{
 		// find the class of this object
-		UClass* TheClass;
+		UClass* TheClass = nullptr;
 		if (Export.ClassIndex.IsNull())
 		{
 			TheClass = UClass::StaticClass();
@@ -2409,9 +2409,24 @@ UObject* FLinkerLoad::FindExistingExport(int32 ExportIndex)
 		{
 			// Check if this object export is a non-native class, non-native classes are always exports.
 			// If so, then use the outer object as a package.
-			UObject* ClassPackage = Export.ClassIndex.IsExport() ? LinkerRoot : nullptr;
-			const bool bAnyPackage = ClassPackage == nullptr;
-			TheClass = (UClass*)StaticFindObjectFast(UClass::StaticClass(), ClassPackage, ImpExp(Export.ClassIndex).ObjectName, /*bExactClass*/false, bAnyPackage);
+			UObject* ClassPackage = nullptr;
+			if (Export.ClassIndex.IsExport())
+			{
+				ClassPackage = LinkerRoot;
+			}
+			else
+			{
+				FObjectImport& ClassImport = Imp(Export.ClassIndex);
+				ClassPackage = StaticFindObjectFast(UPackage::StaticClass(), nullptr, ClassImport.ClassPackage, /*bExactClass*/true);
+			}
+			if (ClassPackage)
+			{
+				TheClass = (UClass*)StaticFindObjectFast(UClass::StaticClass(), ClassPackage, ImpExp(Export.ClassIndex).ObjectName, /*bExactClass*/false);
+			}
+			else
+			{
+				TheClass = FindFirstObject<UClass>(*ImpExp(Export.ClassIndex).ObjectName.ToString(), EFindFirstObjectOptions::None, ELogVerbosity::Fatal, TEXT("finding existing export"));
+			}
 		}
 
 		// if the class exists, try to find the object
@@ -5583,7 +5598,7 @@ bool FLinkerLoad::CreateImportClassAndPackage( FName ClassName, FName PackageNam
 	return true;
 }
 
-TArray<FName> FLinkerLoad::FindPreviousNamesForClass(FString CurrentClassPath, bool bIsInstance)
+TArray<FName> FLinkerLoad::FindPreviousNamesForClass(const FString& CurrentClassPath, bool bIsInstance)
 {
 	TArray<FName> OldNames;
 	TArray<FCoreRedirectObjectName> OldObjectNames;
@@ -5611,6 +5626,34 @@ TArray<FName> FLinkerLoad::FindPreviousNamesForClass(FString CurrentClassPath, b
 	return OldNames;
 }
 
+TArray<FString> FLinkerLoad::FindPreviousPathNamesForClass(const FString& CurrentClassPath, bool bIsInstance)
+{
+	TArray<FString> OldNames;
+	TArray<FCoreRedirectObjectName> OldObjectNames;
+
+	if (FCoreRedirects::FindPreviousNames(ECoreRedirectFlags::Type_Class, FCoreRedirectObjectName(CurrentClassPath), OldObjectNames))
+	{
+		for (FCoreRedirectObjectName& OldObjectName : OldObjectNames)
+		{
+			OldNames.AddUnique(OldObjectName.ToString());
+		}
+	}
+
+	if (bIsInstance)
+	{
+		OldObjectNames.Empty();
+		if (FCoreRedirects::FindPreviousNames(ECoreRedirectFlags::Type_Class | ECoreRedirectFlags::Category_InstanceOnly, FCoreRedirectObjectName(CurrentClassPath), OldObjectNames))
+		{
+			for (FCoreRedirectObjectName& OldObjectName : OldObjectNames)
+			{
+				OldNames.AddUnique(OldObjectName.ToString());
+			}
+		}
+	}
+
+	return OldNames;
+}
+
 FName FLinkerLoad::FindNewNameForEnum(const FName OldEnumName)
 {
 	FCoreRedirectObjectName OldName = FCoreRedirectObjectName(OldEnumName, NAME_None, NAME_None);
@@ -5660,6 +5703,42 @@ FName FLinkerLoad::FindNewNameForClass(FName OldClassName, bool bIsInstance)
 	return NAME_None;
 }
 
+FString FLinkerLoad::FindNewPathNameForClass(const FString& OldClassNameOrPathName, bool bIsInstance)
+{
+	FCoreRedirectObjectName OldName = FCoreRedirectObjectName(OldClassNameOrPathName);
+	FCoreRedirectObjectName NewName = FCoreRedirects::GetRedirectedName(ECoreRedirectFlags::Type_Class, OldName);
+	FString NewClassPathName;
+
+	if (NewName != OldName)
+	{
+		NewClassPathName = NewName.ToString();
+	}
+	else if (bIsInstance)
+	{
+		// Also check instance types
+		NewName = FCoreRedirects::GetRedirectedName(ECoreRedirectFlags::Type_Class | ECoreRedirectFlags::Category_InstanceOnly, OldName);
+
+		if (NewName != OldName)
+		{
+			NewClassPathName = NewName.ToString();
+		}
+	}
+	if (!NewClassPathName.IsEmpty() && !FPackageName::IsValidObjectPath(NewClassPathName))
+	{
+		TArray<UObject*> FoundClasses;
+		UClass* ExistingClass = FindFirstObject<UClass>(*NewClassPathName, EFindFirstObjectOptions::None, ELogVerbosity::Fatal, TEXT("FindNewPathNameForClass"));
+		if (ExistingClass)
+		{
+			NewClassPathName = ExistingClass->GetPathName();
+		}
+		else
+		{
+			UE_LOG(LogLinker, Fatal, TEXT("No classes that match \"%s\" class name found when looking for redirected class"), *NewClassPathName);
+		}
+	}
+	return NewClassPathName;
+}
+
 bool FLinkerLoad::IsKnownMissingPackage(FName PackageName)
 {
 	return FCoreRedirects::IsKnownMissing(ECoreRedirectFlags::Type_Package, FCoreRedirectObjectName(NAME_None, NAME_None, PackageName));
@@ -5864,10 +5943,11 @@ FLinkerLoad::ELinkerStatus FLinkerLoad::FixupExportMap()
 					const FName OuterClassName = GetExportClassName(Export.OuterIndex);
 					if (OuterClassName == NAME_BlueprintGeneratedClass || OuterClassName == NAME_Blueprint)
 					{
-						static const UClass* ActorComponentClass = FindObjectChecked<UClass>(ANY_PACKAGE, TEXT("ActorComponent"), true);
+						static const UClass* ActorComponentClass = FindObjectChecked<UClass>(nullptr, TEXT("/Script/Engine.ActorComponent"), true);
 						static const FString BPGeneratedClassPostfix(TEXT("_C"));
 						const FString NameClassString = NameClass.ToString();
-						UClass* Class = FindObject<UClass>(ANY_PACKAGE, *NameClassString);
+						UPackage* ClassPackage = Cast<UPackage>(StaticFindObjectFast(UPackage::StaticClass(), nullptr, NamePackage));
+						UClass* Class = Cast<UClass>(StaticFindObjectFast(UClass::StaticClass(), ClassPackage, NameClass));
 
 						// It is (obviously) a component if the class is a child of actor component
 						// and (almost certainly) a component if the class cannot be loaded but it ends in _C meaning it was generated from a blueprint
diff --git a/Engine/Source/Runtime/CoreUObject/Public/UObject/Linker.h b/Engine/Source/Runtime/CoreUObject/Public/UObject/Linker.h
--- a/Engine/Source/Runtime/CoreUObject/Public/UObject/Linker.h
+++ b/Engine/Source/Runtime/CoreUObject/Public/UObject/Linker.h
@@ -130,7 +130,7 @@
 			check(ImportMap.IsValidIndex(Index.ToImport()));
 			return &ImportMap[Index.ToImport()];
 		}
-		return NULL;
+		return nullptr;
 	}
 
 	/**
@@ -160,13 +160,192 @@
 			check(ExportMap.IsValidIndex(Index.ToExport()));
 			return &ExportMap[Index.ToExport()];
 		}
-		return NULL;
+		return nullptr;
 	}
 
 	/** Serializes the searchable name map */
 	COREUOBJECT_API void SerializeSearchableNamesMap(FArchive &Ar);
 	COREUOBJECT_API void SerializeSearchableNamesMap(FStructuredArchive::FSlot Slot);
 
+	/**
+	 * Return the path name of the UObject represented by the specified import.
+	 * (can be used with StaticFindObject)
+	 *
+	 * @param	ImportIndex	index into the ImportMap for the resource to get the name for
+	 *
+	 * @return	the path name of the UObject represented by the resource at ImportIndex
+	 */
+	COREUOBJECT_API FString GetImportPathName(int32 ImportIndex);
+	/**
+	 * Return the path name of the UObject represented by the specified import.
+	 * (can be used with StaticFindObject)
+	 *
+	 * @param	PackageIndex	package index for the resource to get the name for
+	 *
+	 * @return	the path name of the UObject represented by the resource at PackageIndex, or the empty string if this isn't an import
+	 */
+	FString GetImportPathName(FPackageIndex PackageIndex)
+	{
+		if (PackageIndex.IsImport())
+		{
+			return GetImportPathName(PackageIndex.ToImport());
+		}
+		return FString();
+	}
+	/**
+	 * Return the path name of the UObject represented by the specified export.
+	 * (can be used with StaticFindObject)
+	 *
+	 * @param	RootPackagePath			Name of the root package for this export
+	 * @param	ExportIndex				index into the ExportMap for the resource to get the name for	 
+	 * @param	bResolveForcedExports	if true, the package name part of the return value will be the export's original package,
+	 *									not the name of the package it's currently contained within.
+	 *
+	 * @return	the path name of the UObject represented by the resource at ExportIndex
+	 */
+	COREUOBJECT_API FString GetExportPathName(const FString& RootPackagePath, int32 ExportIndex, bool bResolveForcedExports = false);
+
+	/**
+	 * Return the full name of the UObject represented by the specified import.
+	 *
+	 * @param	ImportIndex	index into the ImportMap for the resource to get the name for
+	 *
+	 * @return	the full name of the UObject represented by the resource at ImportIndex
+	 */
+
+	COREUOBJECT_API FString GetImportFullName(int32 ImportIndex);
+	/**
+	 * Return the full name of the UObject represented by the specified package index
+	 *
+	 * @param	PackageIndex	package index for the resource to get the name for
+	 *
+	 * @return	the full name of the UObject represented by the resource at PackageIndex
+	 */
+	FString GetImportFullName(FPackageIndex PackageIndex)
+	{
+		if (PackageIndex.IsImport())
+		{
+			return GetImportFullName(PackageIndex.ToImport());
+		}
+		return FString();
+	}
+
+	/**
+	 * Return the full name of the UObject represented by the specified export.
+	 *
+	 * @param	RootPackagePath			Name of the root package for this export
+	 * @param	ExportIndex				index into the ExportMap for the resource to get the name for
+	 * @param	bResolveForcedExports	if true, the package name part of the return value will be the export's original package,
+	 *									not the name of the package it's currently contained within.
+	 *
+	 * @return	the full name of the UObject represented by the resource at ExportIndex
+	 */
+	COREUOBJECT_API FString GetExportFullName(const FString& RootPackagePath, int32 ExportIndex, bool bResolveForcedExports = false);
+
+	/**
+	 * Return the outermost resource package index of the resource pointed by LinkerIndex
+	 * @param LinkerIndex the resource to find the outermost of
+	 * @return the package index of the outermost
+	 */
+	FPackageIndex ResourceGetOutermost(FPackageIndex LinkerIndex) const;
+
+	/**
+	 * Return if a resource has the specified other resource in its outer chain
+	 * @param LinkerIndex the package index of the resource to verify
+	 * @param OuterIndex the package index of the possible outer for the resource
+	 * @return true if the resource pointed by LinkerIndex is in the resource pointed by OuterIndex
+	 */
+	bool ResourceIsIn(FPackageIndex LinkerIndex, FPackageIndex OuterIndex) const;
+
+	/**
+	 * Return if two resources share the same outermost
+	 * @param LinkerIndexLHS the package index of the first resource
+	 * @param LinkerIndexRHS the package index of the second resource
+	 * @return true if they share an outer
+	 */
+	bool DoResourcesShareOutermost(FPackageIndex LinkerIndexLHS, FPackageIndex LinkerIndexRHS) const;
+
+	/**
+	 * Return if the import is in any export
+	 * @param ImportIndex the import to verify
+	 * @return true if an export is in the import outer chain
+	 */
+	bool ImportIsInAnyExport(int32 ImportIndex) const;
+
+	/**
+	 * Return if any export is in the import
+	 * @param ImportIndex the import to verify
+	 * @return true if the import is in any export outer chain
+	 */
+	bool AnyExportIsInImport(int32 ImportIndex) const;
+
+	/**
+	 * Return if any export share an outer with the import
+	 * @param ImportIndex the import to verify
+	 * @reutrn true if any export share the same outer has the import
+	 */
+	bool AnyExportShareOuterWithImport(int32 ImportIndex) const;
+
+	/** 
+	 * Gets the class name for the specified index in the export map.
+	 * @param ExportIdx			Export index
+	 * @return Class name of the export at specified index
+	 */
+	COREUOBJECT_API FName GetExportClassName(int32 ExportIdx);
+	/** 
+	 * Gets the class name for the specified index in the export map.
+	 * @param PackageIndex		PackageIndex that represents the export index
+	 * @return Class name of the export at specified PackageIndex if the PackageIndex represents an export. Otherwise NAME_None
+	 */
+	FName GetExportClassName(FPackageIndex PackageIndex)
+	{
+		if (PackageIndex.IsExport())
+		{
+			return GetExportClassName(PackageIndex.ToExport());
+		}
+		return FName();
+	}
+	/** 
+	 * Gets the class name for the specified index in the import map.
+	 * @param ImportIdx			Import index
+	 * @return Class name of the import at specified index
+	 */
+	FName GetImportClassName(int32 ImportIdx)
+	{
+		return ImportMap[ImportIdx].ClassName;
+	}
+	/** 
+	 * Gets the class name for the specified index in the import map.
+	 * @param ImportIdx			Import PackageIndex
+	 * @return Class name of the import at specified PackageIndex if the PackageIndex represents an import. Otherwise NAME_None
+	 */
+	FName GetImportClassName(FPackageIndex PackageIndex)
+	{
+		if (PackageIndex.IsImport())
+		{
+			return GetImportClassName(PackageIndex.ToImport());
+		}
+		return FName();
+	}
+	/** 
+	 * Gets the class name for the specified PackageIndex.
+	 * @param RootPackagePath	Name of the root package these linker tables represent
+	 * @param PackageIndex		PackageIndex for the entry in the linker tables (import or export)
+	 * @return Class name of the export at specified PackageIndex if the PackageIndex is non-null. Otherwise NAME_None
+	 */
+	FName GetClassName(FPackageIndex PackageIndex)
+	{
+		if (PackageIndex.IsImport())
+		{
+			return GetImportClassName(PackageIndex);
+		}
+		else if (PackageIndex.IsExport())
+		{
+			return GetExportClassName(PackageIndex);
+		}
+		return FName();
+	}
+
 	/** Returns the amount of memory allocated by this container, not including sizeof(*this). */
 	COREUOBJECT_API SIZE_T GetAllocatedSize() const;
 };
@@ -262,45 +441,6 @@
 
 	virtual ~FLinker();
 
-	/** Gets the class name for the specified index in the export map. */
-	COREUOBJECT_API FName GetExportClassName(int32 ExportIdx);
-	/** Gets the class name for the specified index in the import map. */
-	FName GetExportClassName(FPackageIndex PackageIndex)
-	{
-		if (PackageIndex.IsExport())
-		{
-			return GetExportClassName(PackageIndex.ToExport());
-		}
-		return NAME_None;
-	}
-	/** Gets the class name for the specified index in the import map. */
-	FName GetImportClassName(int32 ImportIdx)
-	{
-		return ImportMap[ImportIdx].ClassName;
-	}
-	/** Gets the class name for the specified index in the import map. */
-	FName GetImportClassName(FPackageIndex PackageIndex)
-	{
-		if (PackageIndex.IsImport())
-		{
-			return GetImportClassName(PackageIndex.ToImport());
-		}
-		return NAME_None;
-	}
-	/** Gets the class name for the specified package index. */
-	FName GetClassName(FPackageIndex PackageIndex)
-	{
-		if (PackageIndex.IsImport())
-		{
-			return GetImportClassName(PackageIndex);
-		}
-		else if (PackageIndex.IsExport())
-		{
-			return GetExportClassName(PackageIndex);
-		}
-		return NAME_None;
-	}
-
 	FORCEINLINE ELinkerType::Type GetType() const
 	{
 		return LinkerType;
@@ -316,35 +456,11 @@
 	 */
 	void Serialize(FArchive& Ar);
 	void AddReferencedObjects(FReferenceCollector& Collector);
+	
 	/**
-	 * Return the path name of the UObject represented by the specified import. 
-	 * (can be used with StaticFindObject)
-	 * 
-	 * @param	ImportIndex	index into the ImportMap for the resource to get the name for
-	 *
-	 * @return	the path name of the UObject represented by the resource at ImportIndex
-	 */
-	COREUOBJECT_API FString GetImportPathName(int32 ImportIndex);
-	/**
-	 * Return the path name of the UObject represented by the specified import. 
-	 * (can be used with StaticFindObject)
-	 * 
-	 * @param	PackageIndex	package index for the resource to get the name for
-	 *
-	 * @return	the path name of the UObject represented by the resource at PackageIndex, or the empty string if this isn't an import
-	 */
-	FString GetImportPathName(FPackageIndex PackageIndex)
-	{
-		if (PackageIndex.IsImport())
-		{
-			return GetImportPathName(PackageIndex.ToImport());
-		}
-		return FString();
-	}
-	/**
 	 * Return the path name of the UObject represented by the specified export.
 	 * (can be used with StaticFindObject)
-	 * 
+	 *
 	 * @param	ExportIndex				index into the ExportMap for the resource to get the name for
 	 * @param	FakeRoot				Optional name to replace use as the root package of this object instead of the linker
 	 * @param	bResolveForcedExports	if true, the package name part of the return value will be the export's original package,
@@ -352,11 +468,12 @@
 	 *
 	 * @return	the path name of the UObject represented by the resource at ExportIndex
 	 */
-	COREUOBJECT_API FString GetExportPathName(int32 ExportIndex, const TCHAR* FakeRoot=NULL,bool bResolveForcedExports=false);
+	COREUOBJECT_API FString GetExportPathName(int32 ExportIndex, const TCHAR* FakeRoot = nullptr, bool bResolveForcedExports = false);
+
 	/**
 	 * Return the path name of the UObject represented by the specified export.
 	 * (can be used with StaticFindObject)
-	 * 
+	 *
 	 * @param	PackageIndex			package index for the resource to get the name for
 	 * @param	FakeRoot				Optional name to replace use as the root package of this object instead of the linker
 	 * @param	bResolveForcedExports	if true, the package name part of the return value will be the export's original package,
@@ -364,7 +481,7 @@
 	 *
 	 * @return	the path name of the UObject represented by the resource at PackageIndex, or the empty string if this isn't an export
 	 */
-	FString GetExportPathName(FPackageIndex PackageIndex, const TCHAR* FakeRoot=NULL,bool bResolveForcedExports=false)
+	FString GetExportPathName(FPackageIndex PackageIndex, const TCHAR* FakeRoot = nullptr, bool bResolveForcedExports = false)
 	{
 		if (PackageIndex.IsExport())
 		{
@@ -374,9 +491,9 @@
 	}
 
 	/**
-	 * Return the path name of the UObject represented by the specified import. 
+	 * Return the path name of the UObject represented by the specified import.
 	 * (can be used with StaticFindObject)
-	 * 
+	 *
 	 * @param	PackageIndex	package index
 	 *
 	 * @return	the path name of the UObject represented by the resource at PackageIndex, or the empty string if this is null
@@ -393,33 +510,10 @@
 		}
 		return FString();
 	}
-	/**
-	 * Return the full name of the UObject represented by the specified import.
-	 * 
-	 * @param	ImportIndex	index into the ImportMap for the resource to get the name for
-	 *
-	 * @return	the full name of the UObject represented by the resource at ImportIndex
-	 */
-	COREUOBJECT_API FString GetImportFullName(int32 ImportIndex);
-	/**
-	 * Return the full name of the UObject represented by the specified package index
-	 * 
-	 * @param	PackageIndex	package index for the resource to get the name for
-	 *
-	 * @return	the full name of the UObject represented by the resource at PackageIndex
-	 */
-	FString GetImportFullName(FPackageIndex PackageIndex)
-	{
-		if (PackageIndex.IsImport())
-		{
-			return GetImportFullName(PackageIndex.ToImport());
-		}
-		return FString();
-	}
 
 	/**
 	 * Return the full name of the UObject represented by the specified export.
-	 * 
+	 *
 	 * @param	ExportIndex				index into the ExportMap for the resource to get the name for
 	 * @param	FakeRoot				Optional name to replace use as the root package of this object instead of the linker
 	 * @param	bResolveForcedExports	if true, the package name part of the return value will be the export's original package,
@@ -427,10 +521,11 @@
 	 *
 	 * @return	the full name of the UObject represented by the resource at ExportIndex
 	 */
-	COREUOBJECT_API FString GetExportFullName(int32 ExportIndex, const TCHAR* FakeRoot=NULL,bool bResolveForcedExports=false);
+	COREUOBJECT_API FString GetExportFullName(int32 ExportIndex, const TCHAR* FakeRoot = nullptr, bool bResolveForcedExports = false);
+
 	/**
 	 * Return the full name of the UObject represented by the specified package index
-	 * 
+	 *
 	 * @param	PackageIndex			package index for the resource to get the name for
 	 * @param	FakeRoot				Optional name to replace use as the root package of this object instead of the linker
 	 * @param	bResolveForcedExports	if true, the package name part of the return value will be the export's original package,
@@ -438,7 +533,7 @@
 	 *
 	 * @return	the full name of the UObject represented by the resource at PackageIndex
 	 */
-	FString GetExportFullName(FPackageIndex PackageIndex, const TCHAR* FakeRoot=NULL,bool bResolveForcedExports=false)
+	FString GetExportFullName(FPackageIndex PackageIndex, const TCHAR* FakeRoot = nullptr, bool bResolveForcedExports = false)
 	{
 		if (PackageIndex.IsExport())
 		{
@@ -449,7 +544,7 @@
 
 	/**
 	 * Return the full name of the UObject represented by the specified export.
-	 * 
+	 *
 	 * @param	PackageIndex	package index
 	 *
 	 * @return	the path name of the UObject represented by the resource at PackageIndex, or the empty string if this is null
@@ -468,49 +563,5 @@
 	}
 
 	/**
-	 * Return the outermost resource package index of the resource pointed by LinkerIndex
-	 * @param LinkerIndex the resource to find the outermost of
-	 * @return the package index of the outermost
-	 */
-	FPackageIndex ResourceGetOutermost(FPackageIndex LinkerIndex) const;
-
-	/**
-	 * Return if a resource has the specified other resource in its outer chain
-	 * @param LinkerIndex the package index of the resource to verify
-	 * @param OuterIndex the package index of the possible outer for the resource
-	 * @return true if the resource pointed by LinkerIndex is in the resource pointed by OuterIndex
-	 */
-	bool ResourceIsIn(FPackageIndex LinkerIndex, FPackageIndex OuterIndex) const;
-
-	/**
-	 * Return if two resources share the same outermost
-	 * @param LinkerIndexLHS the package index of the first resource
-	 * @param LinkerIndexRHS the package index of the second resource
-	 * @return true if they share an outer
-	 */
-	bool DoResourcesShareOutermost(FPackageIndex LinkerIndexLHS, FPackageIndex LinkerIndexRHS) const;
-
-	/**
-	 * Return if the import is in any export
-	 * @param ImportIndex the import to verify
-	 * @return true if an export is in the import outer chain
-	 */
-	bool ImportIsInAnyExport(int32 ImportIndex) const;
-
-	/**
-	 * Return if any export is in the import
-	 * @param ImportIndex the import to verify
-	 * @return true if the import is in any export outer chain
-	 */
-	bool AnyExportIsInImport(int32 ImportIndex) const;
-
-	/**
-	 * Return if any export share an outer with the import
-	 * @param ImportIndex the import to verify
-	 * @reutrn true if any export share the same outer has the import
-	 */
-	bool AnyExportShareOuterWithImport(int32 ImportIndex) const;
-	
-	/**
 	 * Tell this linker to start SHA calculations
 	 */
diff --git a/Engine/Source/Runtime/CoreUObject/Public/UObject/LinkerLoad.h b/Engine/Source/Runtime/CoreUObject/Public/UObject/LinkerLoad.h
--- a/Engine/Source/Runtime/CoreUObject/Public/UObject/LinkerLoad.h
+++ b/Engine/Source/Runtime/CoreUObject/Public/UObject/LinkerLoad.h
@@ -266,15 +266,30 @@ public:
 	 * @param bIsInstance If true, we're an instance, so check instance only maps as well
 	 * @return Names without path of all classes that were redirected to this name. Empty if none found.
 	 */
-	COREUOBJECT_API static TArray<FName> FindPreviousNamesForClass(FString CurrentClassPath, bool bIsInstance);
+	COREUOBJECT_API static TArray<FName> FindPreviousNamesForClass(const FString& CurrentClassPath, bool bIsInstance);
+
+	/**
+	 * Utility functions to query the object name redirects list for previous names for a class
+	 * @param CurrentClassPath The current name of the class, with a full path
+	 * @param bIsInstance If true, we're an instance, so check instance only maps as well
+	 * @return Names without path of all classes that were redirected to this name. Empty if none found.
+	 */
+	COREUOBJECT_API static TArray<FString> FindPreviousPathNamesForClass(const FString& CurrentClassPath, bool bIsInstance);
 
 	/** 
 	 * Utility functions to query the object name redirects list for the current name for a class
 	 * @param OldClassName An old class name, without path
-	 * @return Current full path of the class. It will be None if no redirect found
+	 * @return Current name of the class. It will be None if no redirect found
 	 */
 	COREUOBJECT_API static FName FindNewNameForClass(FName OldClassName, bool bIsInstance);
 
+	/**
+	 * Utility functions to query the object name redirects list for the current name for a class
+	 * @param OldClassNameOrPathName An old class name or pathname
+	 * @return Current full path of the class. It will be empty if no redirect found
+	 */
+	COREUOBJECT_API static FString FindNewPathNameForClass(const FString& OldClassNameOrPathName, bool bIsInstance);
+
 	/** 
 	* Utility functions to query the enum name redirects list for the current name for an enum
 	* @param OldEnumName An old enum name, without path
