Reverting unnecessary changes to GetImportPathName introduced in a previous check-in.

diff --git a/Engine/Source/Runtime/CoreUObject/Private/UObject/Linker.cpp b/Engine/Source/Runtime/CoreUObject/Private/UObject/Linker.cpp
--- a/Engine/Source/Runtime/CoreUObject/Private/UObject/Linker.cpp
+++ b/Engine/Source/Runtime/CoreUObject/Private/UObject/Linker.cpp
@@ -240,9 +240,9 @@ FString FLinkerTables::GetImportPathName(int32 ImportIndex)
 	{
 		const FObjectResource& Resource = ImpExp(LinkerIndex);
 		bool bSubobjectDelimiter=false;
-		checkf(LinkerIndex.IsImport(), TEXT("Found an export while traversing the import map"));
-		if (Result.Len() > 0 && GetImportClassName(LinkerIndex) != NAME_Package
-			&& (Resource.OuterIndex.IsNull() || GetImportClassName(Resource.OuterIndex) == NAME_Package) )
+
+		if (Result.Len() > 0 && GetClassName(LinkerIndex) != NAME_Package
+			&& (Resource.OuterIndex.IsNull() || GetClassName(Resource.OuterIndex) == NAME_Package) )
 		{
 			bSubobjectDelimiter = true;
 		}
