AssetRegistry loading optimization, zero out less memory

Saves ~400ms from asset data gatherer during editor startup, which is not on current critical path, for an internal project. Probably saves bit of critical path time in runtime startup.

diff --git a/Engine/Source/Runtime/CoreUObject/Private/AssetRegistry/AssetDataTagMap.cpp b/Engine/Source/Runtime/CoreUObject/Private/AssetRegistry/AssetDataTagMap.cpp
--- a/Engine/Source/Runtime/CoreUObject/Private/AssetRegistry/AssetDataTagMap.cpp
+++ b/Engine/Source/Runtime/CoreUObject/Private/AssetRegistry/AssetDataTagMap.cpp
@@ -999,6 +999,17 @@ namespace FixedTagPrivate
 			}
 		}
 
+		template<class T>
+		void MemZeroViewDataIfNeeded(TArrayView<T> Items)
+		{
+			// Handle LoadViewData() assignment and early destruction in case of archive error
+			// for non-trivial types like FText
+			if (!std::is_trivially_move_assignable_v<T> || !std::is_trivially_destructible_v<T>)
+			{
+				FMemory::Memzero(Items.GetData(), Items.Num() * sizeof(T));
+			}
+		}
+
 		template<typename T>
 		void LoadViewData(TArrayView<T> Items)
 		{
@@ -1075,7 +1086,7 @@ namespace FixedTagPrivate
 			// Load view sizes
 			VisitViews(Store, [&] (auto& View) { SetNum(View, LoadItem<int32>()); });
 
-			// Calculate total size, allocate and zero data
+			// Calculate total size and allocate data
 			uint64 Bytes = 0;
 			VisitViews(Store, [&] (auto& View) 
 				{
@@ -1084,7 +1095,6 @@ namespace FixedTagPrivate
 				});
 
 			uint8* Ptr = reinterpret_cast<uint8*>(FMemory::Malloc(Bytes, MaxViewAlignment));
-			FMemory::Memzero(Ptr, Bytes);
 			Store.Data = Ptr;
 
 			// Set view data pointers
@@ -1094,9 +1104,11 @@ namespace FixedTagPrivate
 					SetUntypedDataPtr(View, Ptr);
 					Ptr += GetBytes(View);
 				});
-
 			check(Ptr - Bytes == Store.Data);
 
+			// Zero out required data
+			VisitViews(Store, [&](auto& View) { MemZeroViewDataIfNeeded(View); });
+
 			return Order;
 		}
 		
