BulkDataRegistry: Add an ini setting for projects to suppress the warning about duplicate ids.
#rnx

diff --git a/Engine/Source/Runtime/CoreUObject/Private/Serialization/EditorBulkData.cpp b/Engine/Source/Runtime/CoreUObject/Private/Serialization/EditorBulkData.cpp
--- a/Engine/Source/Runtime/CoreUObject/Private/Serialization/EditorBulkData.cpp
+++ b/Engine/Source/Runtime/CoreUObject/Private/Serialization/EditorBulkData.cpp
@@ -612,8 +612,18 @@ void FEditorBulkData::LogRegisterError(UE::BulkDataRegistry::ERegisterResult Val
 				{
 					SilenceWarningMessage = TEXT(" To silence this warning, run the ResavePackagesCommandlet with \"-autocheckout -resaveondemand=bulkdataduplicates\".");
 				}
-				UE_LOG(LogBulkDataRegistry, Warning, TEXT("%s updated BulkData %s on load because it collided with an ID in package %s.%s"),
+				bool bSuppressWarning = false; // Allow projects to suppress the warning because they might not be able to resave packages during an integration
+				GConfig->GetBool(TEXT("CookSettings"), TEXT("BulkDataRegistrySuppressDuplicateWarning"), bSuppressWarning, GEditorIni);
+				FString LogMessage = FString::Printf(TEXT("%s updated BulkData %s on load because it collided with an ID in package %s.%s"),
 					*OwnerPathName, *FailedBulkDataId.ToString(), *OtherOwnerPackageName, *SilenceWarningMessage);
+				if (bSuppressWarning)
+				{
+					UE_LOG(LogBulkDataRegistry, Verbose, TEXT("%s"), *LogMessage);
+				}
+				else
+				{
+					UE_LOG(LogBulkDataRegistry, Warning, TEXT("%s"), *LogMessage);
+				}
 			}
 			bMessageLogged = true;
 		}
