GC: Restrict the giant FGCCSyncObject and FUObjectHashTables locks to the reachability analysis phase.
It is possible to toggle to the previous/legacy behavior with gc.GCLockBehavior=1.

Legacy behavior:
A giant FGCCSyncObject lock for the complete duration of CollectGarbageInternal.
A giant FUObjectHashTables lock for most of CollectGarbageInternal except the execution of the GetPreGarbageCollectDelegate and GetPostGarbageCollect callbacks.
Giant FGCCSyncObject and FUObjectHashTables locks for IncrementalPurgeGarbage (inluding both UnhashUnreachableObjects and IncrementalDestroyGarbage).

Unchanged behavior:
The GIsGarbageCollecting contract is fuzzy and the behavior is not changed at this point to reduce risk.
GIsGarbageCollecting is set during CollectGarbageInternal.
GIsGarbageCollecting is also set during the IncrementalDestroyGarbage phase of IncrementalPurgeGarbage, but not during UnhashUnreachableObjects. Why???
Many use cases seems to be as simple as to avoid the StaticFindObject asserts, but there are definitely use cases that are harder to reason about like UObject::GetDefaultSubobjectByName, UObject::CanModify, UChildActorComponent::DestroyChildActor, AActor::ProcessEvent, AGroupActor::PostRemove, UEditorEngine::UndoTransaction, UEditorEngine::RedoTransaction, DatasmithRuntime::FSceneImporter, ...

Changed legacy and default behavior:
Move PostReachabilityAnalysis.Broadcast after GatherUnreachableObjects, NotifyUnreachableObjects and ClearWeakReferences.
StaticAllocateObject and the StaticFindObject functions are now checking IsGarbageCollectingAndLockingUObjectHashTables instead of IsGarbageCollecting/IsGarbageCollectingOnGameThread,
but when using the legacy behavior IsGarbageCollectingAndLockingUObjectHashTables should have the same meaning as IsGarbageCollecting.

New default behavior:
Restrict the giant scope of locking FGCCSyncObject and FUObjectHashTables during CollectGarbageInternal to DissolveClusters, VerifyGCAssumptions, PerformReachabilityAnalysis, GatherUnreachableObjects, NotifyUnreachableObjects and ClearWeakReferences.
Remove the giant FGCCSyncObject lock completely during IncrementalPurgeGarbage (including UnhashUnreachableObjects and IncrementalDestroyGarbage).
Remove the giant FUObjectHashTables lock completely during IncrementalDestroyGarbage.
This means BeginDestroy, FinishDestroy, destructors and callbacks are all executed without the giant FGCCSyncObject and FUObjectHashTables locks and they are all relying on existing smaller thread safety locks for for accessing the global uboject array and global hash tables.

#rnx

diff --git a/Engine/Source/Runtime/CoreUObject/Private/UObject/GarbageCollection.cpp b/Engine/Source/Runtime/CoreUObject/Private/UObject/GarbageCollection.cpp
--- a/Engine/Source/Runtime/CoreUObject/Private/UObject/GarbageCollection.cpp
+++ b/Engine/Source/Runtime/CoreUObject/Private/UObject/GarbageCollection.cpp
@@ -45,6 +45,8 @@ DEFINE_LOG_CATEGORY(LogGarbage);
 
 /** Object count during last mark phase																				*/
 FThreadSafeCounter		GObjectCountDuringLastMarkPhase;
+/** Whether UObject hash tables are locked by GC */
+bool GIsGarbageCollectingAndLockingUObjectHashTables = false;
 /** Whether incremental object purge is in progress										*/
 bool GObjIncrementalPurgeIsInProgress = false;
 /** Whether GC is currently routing BeginDestroy to objects										*/
@@ -77,6 +79,20 @@ static double GLastGCTime = 0;
 static const uint32 GNumTokensPerPointer = sizeof(void*) / sizeof(uint32); //-V514
 
 FThreadSafeBool GIsGarbageCollecting(false);
+class FIsGarbageCollectingScope
+{
+	bool bOldValue = false;
+public:
+	FIsGarbageCollectingScope()
+		: bOldValue(GIsGarbageCollecting)
+	{
+		GIsGarbageCollecting = true;
+	}
+	~FIsGarbageCollectingScope()
+	{
+		GIsGarbageCollecting = bOldValue;
+	}
+};
 
 int32 FGCReferenceTokenStream::MaxStackSize = 1;
 
@@ -87,27 +103,75 @@ const FGCReferenceInfo FGCReferenceInfo::EndOfStreamToken(GCRT_EndOfStream, 0);
 */
 void NotifyUnreachableObjects(const TArrayView<FUObjectItem*>& UnreachableObjects);
 
+enum FGCLockBehavior : int32
+{
+	Default = 0,
+	Legacy = 1,
+};
+static int32 GGCLockBehavior = FGCLockBehavior::Default;
+static FAutoConsoleVariable CVarGCLockBehavior(
+	TEXT("gc.LockBehavior"),
+	GGCLockBehavior,
+	TEXT("Set the GC lock behavior: 0=Default, 1=Legacy (more restrictive and less performant)."),
+	FConsoleVariableDelegate::CreateLambda([](IConsoleVariable* Variable)
+	{
+		GGCLockBehavior = FMath::Clamp(Variable->GetInt(), 0, 1);
+	}),
+	ECVF_Default);
+
+/** Locks all UObject hash tables for performing GC reachability analysis.
+ * Can be queried with IsGarbageCollectingAndLockingUObjectHashTables().
+ * */
+class FGCHashTableScopeLock
+{
+	bool bEnabled = (GGCLockBehavior == FGCLockBehavior::Default);
+public:
+
+	FORCEINLINE FGCHashTableScopeLock()
+	{
+		if (bEnabled)
+		{
+			GIsGarbageCollectingAndLockingUObjectHashTables = true;
+			LockUObjectHashTables();
+		}
+	}
+	FORCEINLINE ~FGCHashTableScopeLock()
+	{
+		if (bEnabled)
+		{
+			UnlockUObjectHashTables();
+			GIsGarbageCollectingAndLockingUObjectHashTables = false;
+		}
+	}
+};
+
 /** Locks all UObject hash tables when performing GC */
-class FGCScopeLock
+class FGCLegacyHashTableScopeLock
 {
+	bool bEnabled = (GGCLockBehavior == FGCLockBehavior::Legacy);
 	/** Previous value of the GetGarbageCollectingFlag() */
-	bool bPreviousGabageCollectingFlagValue;
+	bool bPreviousGabageCollectingFlagValue = GIsGarbageCollectingAndLockingUObjectHashTables;
 public:
 
 	/**
 	 * We're storing the value of GetGarbageCollectingFlag in the constructor, it's safe as only
 	 * one thread is ever going to be setting it and calling this code - the game thread.
 	 **/
-	FORCEINLINE FGCScopeLock()
-		: bPreviousGabageCollectingFlagValue(GIsGarbageCollecting)
+	FORCEINLINE FGCLegacyHashTableScopeLock()
 	{
-		LockUObjectHashTables();
-		GIsGarbageCollecting = true;
+		if (bEnabled)
+		{
+			LockUObjectHashTables();
+			GIsGarbageCollectingAndLockingUObjectHashTables = true;
+		}
 	}
-	FORCEINLINE ~FGCScopeLock()
+	FORCEINLINE ~FGCLegacyHashTableScopeLock()
 	{
-		GIsGarbageCollecting = bPreviousGabageCollectingFlagValue;
-		UnlockUObjectHashTables();
+		if (bEnabled)
+		{
+			GIsGarbageCollectingAndLockingUObjectHashTables = bPreviousGabageCollectingFlagValue;
+			UnlockUObjectHashTables();
+		}
 	}
 };
 
@@ -1588,6 +1652,7 @@ static bool ShouldForceSingleThreadedGC()
 
 void AcquireGCLock()
 {
+	TRACE_CPUPROFILER_EVENT_SCOPE(AcquireGCLock);
 	const double StartTime = FPlatformTime::Seconds();
 	FGCCSyncObject::Get().GCLock();
 	const double ElapsedTime = FPlatformTime::Seconds() - StartTime;
@@ -1603,19 +1668,19 @@ void ReleaseGCLock()
 }
 
 /** Locks GC within a scope but only if it hasn't been locked already */
-struct FConditionalGCLock
+struct FLegacyConditionalIncrementalGCLock
 {
 	bool bNeedsUnlock;
-	FConditionalGCLock()
+	FLegacyConditionalIncrementalGCLock()
 		: bNeedsUnlock(false)
 	{
-		if (!FGCCSyncObject::Get().IsGCLocked())
+		if ((GGCLockBehavior == FGCLockBehavior::Legacy) && !FGCCSyncObject::Get().IsGCLocked())
 		{
 			AcquireGCLock();
 			bNeedsUnlock = true;
 		}
 	}
-	~FConditionalGCLock()
+	~FLegacyConditionalIncrementalGCLock()
 	{
 		if (bNeedsUnlock)
 		{
@@ -1689,7 +1754,7 @@ void IncrementalPurgeGarbage(bool bUseTimeLimit, double TimeLimit)
 	{
 		// Lock before settting GCStartTime as it could be slow to lock if async loading is in progress
 		// but we still want to perform some GC work otherwise we'd be keeping objects in memory for a long time
-		FConditionalGCLock ScopedGCLock;
+		FLegacyConditionalIncrementalGCLock ScopedGCLock;
 
 		// Keep track of start time to enforce time limit unless bForceFullPurge is true;
 		GCStartTime = FPlatformTime::Seconds();
@@ -1767,7 +1832,8 @@ bool IncrementalDestroyGarbage(bool bUseTimeLimit, double TimeLimit)
 	const int32	TimeLimitEnforcementGranularityForDeletion = 100;
 
 	// Set 'I'm garbage collecting' flag - might be checked inside UObject::Destroy etc.
-	FGCScopeLock GCLock;
+	FGCLegacyHashTableScopeLock GCLegacyHashTableLock;
+	FIsGarbageCollectingScope IsGarbageCollectingScope;
 
 	if( !GObjFinishDestroyHasBeenRoutedToAllObjects && !bTimeLimitReached )
 	{
@@ -2068,6 +2134,10 @@ bool IsIncrementalPurgePending()
 	return GObjIncrementalPurgeIsInProgress || GObjPurgeIsRequired;
 }
 
+bool IsGarbageCollectingAndLockingUObjectHashTables()
+{
+	return GIsGarbageCollectingAndLockingUObjectHashTables;
+}
 
 // This counts how many times GC was skipped
 static int32 GNumAttemptsSinceLastGC = 0;
@@ -2214,15 +2284,12 @@ void CollectGarbageInternal(EObjectFlags KeepFlags, bool bPerformFullPurge)
 	GNumAttemptsSinceLastGC = 0;
 
 	// Flush streaming before GC if requested
-	if (GFlushStreamingOnGC)
+	if (GFlushStreamingOnGC && IsAsyncLoading())
 	{
-		if (IsAsyncLoading())
-		{
-			UE_LOG(LogGarbage, Log, TEXT("CollectGarbageInternal() is flushing async loading"));
-		}
-		FGCCSyncObject::Get().GCUnlock();
+		UE_LOG(LogGarbage, Log, TEXT("CollectGarbageInternal() is flushing async loading"));
+		ReleaseGCLock();
 		FlushAsyncLoading();
-		FGCCSyncObject::Get().GCLock();
+		AcquireGCLock();
 	}
 
 	// Route callbacks so we can ensure that we are e.g. not in the middle of loading something by flushing
@@ -2236,7 +2303,8 @@ void CollectGarbageInternal(EObjectFlags KeepFlags, bool bPerformFullPurge)
 	{
 		// Set 'I'm garbage collecting' flag - might be checked inside various functions.
 		// This has to be unlocked before we call post GC callbacks
-		FGCScopeLock GCLock;
+		FGCLegacyHashTableScopeLock GCLegacyHashTableLock;
+		FIsGarbageCollectingScope IsGarbageCollectingScope;
 
 		UE_LOG(LogGarbage, Log, TEXT("Collecting garbage%s"), IsAsyncLoading() ? TEXT(" while async loading") : TEXT(""));
 
@@ -2250,79 +2318,79 @@ void CollectGarbageInternal(EObjectFlags KeepFlags, bool bPerformFullPurge)
 				FMemory::Trim();
 			}
 		}
-		check(!GObjIncrementalPurgeIsInProgress);
-		check(!GObjPurgeIsRequired);
 
-		// This can happen if someone disables clusters from the console (gc.CreateGCClusters)
-		if (!GCreateGCClusters && GUObjectClusters.GetNumAllocatedClusters())
+		// Reachability analysis.
+		// When exiting this scope all objects to be destroyed have been marked unreachable and any weak references have been cleared.
 		{
-			GUObjectClusters.DissolveClusters(true);
-		}
+			// With the new FGCLockBehavior::Default behavior hash tables are only locked during this scope of reachability analysis.
+			FGCHashTableScopeLock GCHashTableLock;
 
-#if VERIFY_DISREGARD_GC_ASSUMPTIONS
-		// Only verify assumptions if option is enabled. This avoids false positives in the Editor or commandlets.
-		if (GShouldVerifyGCAssumptions)
-		{
-			DECLARE_SCOPE_CYCLE_COUNTER(TEXT("CollectGarbageInternal.VerifyGCAssumptions"), STAT_CollectGarbageInternal_VerifyGCAssumptions, STATGROUP_GC);
-			const double StartTime = FPlatformTime::Seconds();
-			if (GUObjectArray.DisregardForGCEnabled())
+			check(!GObjIncrementalPurgeIsInProgress);
+			check(!GObjPurgeIsRequired);
+
+			// This can happen if someone disables clusters from the console (gc.CreateGCClusters)
+			if (!GCreateGCClusters && GUObjectClusters.GetNumAllocatedClusters())
 			{
-				VerifyGCAssumptions();
+				GUObjectClusters.DissolveClusters(true);
 			}
-			if (GUObjectClusters.GetNumAllocatedClusters())
+
+#if VERIFY_DISREGARD_GC_ASSUMPTIONS
+			// Only verify assumptions if option is enabled. This avoids false positives in the Editor or commandlets.
+			if (GShouldVerifyGCAssumptions)
 			{
-				VerifyClustersAssumptions();
+				DECLARE_SCOPE_CYCLE_COUNTER(TEXT("CollectGarbageInternal.VerifyGCAssumptions"), STAT_CollectGarbageInternal_VerifyGCAssumptions, STATGROUP_GC);
+				const double StartTime = FPlatformTime::Seconds();
+				if (GUObjectArray.DisregardForGCEnabled())
+				{
+					VerifyGCAssumptions();
+				}
+				if (GUObjectClusters.GetNumAllocatedClusters())
+				{
+					VerifyClustersAssumptions();
+				}
+				VerifyObjectFlagMirroring();
+				UE_LOG(LogGarbage, Log, TEXT("%f ms for Verify GC Assumptions"), (FPlatformTime::Seconds() - StartTime) * 1000);
 			}
-			VerifyObjectFlagMirroring();
-			UE_LOG(LogGarbage, Log, TEXT("%f ms for Verify GC Assumptions"), (FPlatformTime::Seconds() - StartTime) * 1000);
-		}
 #endif
 
-		const EFastReferenceCollectorOptions Options = 
-			// Fall back to single threaded GC if processor count is 1 or parallel GC is disabled
-			// or detailed per class gc stats are enabled (not thread safe)
-			(ShouldForceSingleThreadedGC() ? EFastReferenceCollectorOptions::None : EFastReferenceCollectorOptions::Parallel) |
-			// Run with GC clustering code enabled only if clustering is enabled and there's actual allocated clusters
-			((!!GCreateGCClusters && GUObjectClusters.GetNumAllocatedClusters()) ? EFastReferenceCollectorOptions::WithClusters : EFastReferenceCollectorOptions::None) |
-			// Toggle between PendingKill enabled or disabled
-			(UObjectBaseUtility::IsPendingKillEnabled() ? EFastReferenceCollectorOptions::WithPendingKill : EFastReferenceCollectorOptions::None);
-
-		// Perform reachability analysis.
-		{
-			const double StartTime = FPlatformTime::Seconds();
-			FRealtimeGC TagUsedRealtimeGC;
+			const EFastReferenceCollectorOptions Options = 
+				// Fall back to single threaded GC if processor count is 1 or parallel GC is disabled
+				// or detailed per class gc stats are enabled (not thread safe)
+				(ShouldForceSingleThreadedGC() ? EFastReferenceCollectorOptions::None : EFastReferenceCollectorOptions::Parallel) |
+				// Run with GC clustering code enabled only if clustering is enabled and there's actual allocated clusters
+				((!!GCreateGCClusters && GUObjectClusters.GetNumAllocatedClusters()) ? EFastReferenceCollectorOptions::WithClusters : EFastReferenceCollectorOptions::None) |
+				// Toggle between PendingKill enabled or disabled
+				(UObjectBaseUtility::IsPendingKillEnabled() ? EFastReferenceCollectorOptions::WithPendingKill : EFastReferenceCollectorOptions::None);
 
-			TagUsedRealtimeGC.PerformReachabilityAnalysis(KeepFlags, Options);
-			UE_LOG(LogGarbage, Log, TEXT("%f ms for GC"), (FPlatformTime::Seconds() - StartTime) * 1000);
-		}
+			// Perform reachability analysis.
+			{
+				const double StartTime = FPlatformTime::Seconds();
+				FRealtimeGC TagUsedRealtimeGC;
 
-		FGCArrayPool& ArrayPool = FGCArrayPool::Get();
-		TArray<FGCArrayStruct*> AllArrays;
-		ArrayPool.GetAllArrayStructsFromPool(AllArrays);
-		// This needs to happen before clusters get dissolved otherwisise cluster information will be missing from history
-		ArrayPool.UpdateGCHistory(AllArrays);
+				TagUsedRealtimeGC.PerformReachabilityAnalysis(KeepFlags, Options);
+				UE_LOG(LogGarbage, Log, TEXT("%f ms for GC"), (FPlatformTime::Seconds() - StartTime) * 1000);
+			}
 
-		// Reconstruct clusters if needed
-		if (GUObjectClusters.ClustersNeedDissolving())
-		{
-			const double StartTime = FPlatformTime::Seconds();
-			GUObjectClusters.DissolveClusters();
-			UE_LOG(LogGarbage, Log, TEXT("%f ms for dissolving GC clusters"), (FPlatformTime::Seconds() - StartTime) * 1000);
-		}
+			FGCArrayPool& ArrayPool = FGCArrayPool::Get();
+			TArray<FGCArrayStruct*> AllArrays;
+			ArrayPool.GetAllArrayStructsFromPool(AllArrays);
+			// This needs to happen before clusters get dissolved otherwisise cluster information will be missing from history
+			ArrayPool.UpdateGCHistory(AllArrays);
 
-		// Fire post-reachability analysis hooks
-		{
-			TRACE_CPUPROFILER_EVENT_SCOPE(BroadcastPostReachabilityAnalysis);
-			FCoreUObjectDelegates::PostReachabilityAnalysis.Broadcast();
-		}
+			// Reconstruct clusters if needed
+			if (GUObjectClusters.ClustersNeedDissolving())
+			{
+				const double StartTime = FPlatformTime::Seconds();
+				GUObjectClusters.DissolveClusters();
+				UE_LOG(LogGarbage, Log, TEXT("%f ms for dissolving GC clusters"), (FPlatformTime::Seconds() - StartTime) * 1000);
+			}
 
-		{
 			ArrayPool.DumpGarbageReferencers(AllArrays);
-		
+
 			GatherUnreachableObjects(!(Options & EFastReferenceCollectorOptions::Parallel));
 			NotifyUnreachableObjects(GUnreachableObjects);
 
-			// This needs to happen after NotifyGarbageReferencers and GatherUnreachableObjects since both can mark more objects as unreachable
+			// This needs to happen after GatherUnreachableObjects since it can mark more objects as unreachable
 			ArrayPool.ClearWeakReferences(AllArrays);
 
 			// Now return arrays back to the pool and free some memory if requested
@@ -2342,12 +2410,27 @@ void CollectGarbageInternal(EObjectFlags KeepFlags, bool bPerformFullPurge)
 
 			// Make sure nothing will be using potentially freed arrays
 			AllArrays.Empty();
+		}
+		if (GGCLockBehavior == FGCLockBehavior::Default)
+		{
+			// With the new FGCLockBehavior::Default behavior the hash tables lock was released when exiting the reachability analysis scope above.
+			// BeginDestroy, FinishDestroy, destructors and callbacks are allowed to call functions like StaticAllocateObject and StaticFindObject.
+			// Now release the GC lock to allow async loading and other threads to perform UObject operations under the FGCScopeGuard.
+			ReleaseGCLock();
+		}
+		// The old FGCLockBehavior::Legacy behavior keeps the GC and hash table locks until GC is completed,
+		// and no threads are allowed to perform UObject operations yet.
 
-			if (bPerformFullPurge || !GIncrementalBeginDestroyEnabled)
-			{
-				UnhashUnreachableObjects(/**bUseTimeLimit = */ false);
-				FScopedCBDProfile::DumpProfile();
-			}
+		// Fire post-reachability analysis hooks
+		{
+			TRACE_CPUPROFILER_EVENT_SCOPE(BroadcastPostReachabilityAnalysis);
+			FCoreUObjectDelegates::PostReachabilityAnalysis.Broadcast();
+		}
+
+		if (bPerformFullPurge || !GIncrementalBeginDestroyEnabled)
+		{
+			UnhashUnreachableObjects(/**bUseTimeLimit = */ false);
+			FScopedCBDProfile::DumpProfile();
 		}
 
 		// Set flag to indicate that we are relying on a purge to be performed.
@@ -2539,8 +2622,12 @@ void CollectGarbage(EObjectFlags KeepFlags, bool bPerformFullPurge)
 	// Perform actual garbage collection
 	CollectGarbageInternal(KeepFlags, bPerformFullPurge);
 
-	// Other threads are free to use UObjects
-	ReleaseGCLock();
+	if (GGCLockBehavior == FGCLockBehavior::Legacy)
+	{
+		// Release the GC lock to allow async loading and other threads to perform UObject operations under the FGCScopeGuard.
+		ReleaseGCLock();
+	}
+	// With the new FGCLockBehavior::Default behavior the lock was released after reachability analysis inside CollectGarbageInternal
 }
 
 bool TryCollectGarbage(EObjectFlags KeepFlags, bool bPerformFullPurge)
@@ -2563,8 +2650,12 @@ bool TryCollectGarbage(EObjectFlags KeepFlags, bool bPerformFullPurge)
 		// Perform actual garbage collection
 		CollectGarbageInternal(KeepFlags, bPerformFullPurge);
 
-		// Other threads are free to use UObjects
-		ReleaseGCLock();
+		if (GGCLockBehavior == FGCLockBehavior::Legacy)
+		{
+			// Release the GC lock to allow async loading and other threads to perform UObject operations under the FGCScopeGuard.
+			ReleaseGCLock();
+		}
+		// With the new FGCLockBehavior::Default behavior the lock was released after reachability analysis inside CollectGarbageInternal
 	}
 	else
 	{
diff --git a/Engine/Source/Runtime/CoreUObject/Private/UObject/UObjectGlobals.cpp b/Engine/Source/Runtime/CoreUObject/Private/UObject/UObjectGlobals.cpp
--- a/Engine/Source/Runtime/CoreUObject/Private/UObject/UObjectGlobals.cpp
+++ b/Engine/Source/Runtime/CoreUObject/Private/UObject/UObjectGlobals.cpp
@@ -198,15 +198,6 @@ bool ShouldCreateThrottledSlowTask()
 	return ShouldReportProgress();
 }
 
-/**
- * Returns true if code is called form the game thread while collecting garbage.
- * We only have to guard against StaticFindObject on the game thread as other threads will be blocked anyway
- */
-static FORCEINLINE bool IsGarbageCollectingOnGameThread()
-{
-	return IsInGameThread() && IsGarbageCollecting();
-}
-
 // Anonymous namespace to not pollute global.
 namespace
 {
@@ -311,7 +302,7 @@ int32 UpdateSuffixForNextNewObject(UObject* Parent, const UClass* Class, TFuncti
 //
 UObject* StaticFindObjectFast(UClass* ObjectClass, UObject* ObjectPackage, FName ObjectName, bool bExactClass, bool bAnyPackage, EObjectFlags ExclusiveFlags, EInternalObjectFlags ExclusiveInternalFlags)
 {
-	if (UE::IsSavingPackage(nullptr) || IsGarbageCollectingOnGameThread())
+	if (UE::IsSavingPackage(nullptr) || IsGarbageCollectingAndLockingUObjectHashTables())
 	{
 		UE_LOG(LogUObjectGlobals, Fatal,TEXT("Illegal call to StaticFindObjectFast() while serializing object data or garbage collecting!"));
 	}
@@ -332,7 +323,7 @@ UObject* StaticFindObjectFastSafe(UClass* ObjectClass, UObject* ObjectPackage, F
 {
 	UObject* FoundObject = nullptr;
 	
-	if (!UE::IsSavingPackage(nullptr) && !IsGarbageCollectingOnGameThread())
+	if (!UE::IsSavingPackage(nullptr) && !IsGarbageCollectingAndLockingUObjectHashTables())
 	{
 		// We don't want to return any objects that are currently being background loaded unless we're using FindObject during async loading.
 		ExclusiveInternalFlags |= IsInAsyncLoadingThread() ? EInternalObjectFlags::None : EInternalObjectFlags::AsyncLoading;
@@ -430,11 +421,11 @@ UObject* StaticFindObjectChecked( UClass* ObjectClass, UObject* ObjectParent, co
 }
 
 //
-// Find an object; won't assert on UE::IsSavingPackage() or IsGarbageCollecting()
+// Find an object; won't assert on UE::IsSavingPackage() or IsGarbageCollectingAndLockingUObjectHashTables()
 //
 UObject* StaticFindObjectSafe( UClass* ObjectClass, UObject* ObjectParent, const TCHAR* InName, bool bExactClass )
 {
-	if (!UE::IsSavingPackage(nullptr) && !IsGarbageCollectingOnGameThread())
+	if (!UE::IsSavingPackage(nullptr) && !IsGarbageCollectingAndLockingUObjectHashTables())
 	{
 		FGCScopeGuard GCGuard;
 		return StaticFindObject( ObjectClass, ObjectParent, InName, bExactClass );
@@ -447,7 +438,7 @@ UObject* StaticFindObjectSafe( UClass* ObjectClass, UObject* ObjectParent, const
 
 bool StaticFindAllObjectsFast(TArray<UObject*>& OutFoundObjects, UClass* ObjectClass, FName ObjectName, bool ExactClass, EObjectFlags ExclusiveFlags, EInternalObjectFlags ExclusiveInternalFlags)
 {
-	UE_CLOG(UE::IsSavingPackage(nullptr) || IsGarbageCollectingOnGameThread(), LogUObjectGlobals, Fatal, TEXT("Illegal call to StaticFindAllObjectsFast() while serializing object data or garbage collecting!"));
+	UE_CLOG(UE::IsSavingPackage(nullptr) || IsGarbageCollectingAndLockingUObjectHashTables(), LogUObjectGlobals, Fatal, TEXT("Illegal call to StaticFindAllObjectsFast() while serializing object data or garbage collecting!"));
 
 	// We don't want to return any objects that are currently being background loaded unless we're using FindObject during async loading.
 	ExclusiveInternalFlags |= IsInAsyncLoadingThread() ? EInternalObjectFlags::None : EInternalObjectFlags::AsyncLoading;
@@ -457,7 +448,7 @@ bool StaticFindAllObjectsFast(TArray<UObject*>& OutFoundObjects, UClass* ObjectC
 bool StaticFindAllObjectsFastSafe(TArray<UObject*>& OutFoundObjects, UClass* ObjectClass, FName ObjectName, bool ExactClass, EObjectFlags ExclusiveFlags, EInternalObjectFlags ExclusiveInternalFlags)
 {
 	bool bFoundObjects = false;
-	if (!UE::IsSavingPackage(nullptr) && !IsGarbageCollectingOnGameThread())
+	if (!UE::IsSavingPackage(nullptr) && !IsGarbageCollectingAndLockingUObjectHashTables())
 	{
 		// We don't want to return any objects that are currently being background loaded unless we're using FindObject during async loading.
 		ExclusiveInternalFlags |= IsInAsyncLoadingThread() ? EInternalObjectFlags::None : EInternalObjectFlags::AsyncLoading;
@@ -471,7 +462,7 @@ bool StaticFindAllObjects(TArray<UObject*>& OutFoundObjects, UClass* ObjectClass
 	INC_DWORD_STAT(STAT_FindObject);
 
 	UE_CLOG(UE::IsSavingPackage(nullptr), LogUObjectGlobals, Fatal, TEXT("Illegal call to StaticFindAllObjects() while serializing object data!"));
-	UE_CLOG(IsGarbageCollectingOnGameThread(), LogUObjectGlobals, Fatal, TEXT("Illegal call to StaticFindAllObjects() while collecting garbage!"));
+	UE_CLOG(IsGarbageCollectingAndLockingUObjectHashTables(), LogUObjectGlobals, Fatal, TEXT("Illegal call to StaticFindAllObjects() while collecting garbage!"));
 
 #if WITH_EDITOR
 	// If the editor is running, and T3D is being imported, ensure any packages referenced are fully loaded.
@@ -497,7 +488,7 @@ bool StaticFindAllObjects(TArray<UObject*>& OutFoundObjects, UClass* ObjectClass
 bool StaticFindAllObjectsSafe(TArray<UObject*>& OutFoundObjects, UClass* ObjectClass, const TCHAR* OrigInName, bool ExactClass)
 {
 	bool bFoundObjects = false;
-	if (!UE::IsSavingPackage(nullptr) && !IsGarbageCollectingOnGameThread())
+	if (!UE::IsSavingPackage(nullptr) && !IsGarbageCollectingAndLockingUObjectHashTables())
 	{
 		bFoundObjects = StaticFindAllObjects(OutFoundObjects, ObjectClass, OrigInName, ExactClass);
 	}
@@ -561,7 +552,7 @@ UObject* StaticFindFirstObject(UClass* Class, const TCHAR* Name, EFindFirstObjec
 UObject* StaticFindFirstObjectSafe(UClass* Class, const TCHAR* Name, EFindFirstObjectOptions Options /*= EFindFirstObjectOptions::None*/, ELogVerbosity::Type AmbiguousMessageVerbosity /*= ELogVerbosity::Warning*/, const TCHAR* InCurrentOperation /*= nullptr*/)
 {
 	UObject* Result = nullptr;
-	if (!UE::IsSavingPackage(nullptr) && !IsGarbageCollectingOnGameThread())
+	if (!UE::IsSavingPackage(nullptr) && !IsGarbageCollectingAndLockingUObjectHashTables())
 	{
 		Result = StaticFindFirstObject(Class, Name, Options, AmbiguousMessageVerbosity, InCurrentOperation);
 	}
@@ -2931,7 +2922,7 @@ UObject* StaticAllocateObject
 		*GetNameSafe(InClass), *GetPathNameSafe(InOuter), *InName.ToString());
 	//checkf(InClass != UPackage::StaticClass() || !InOuter || bCreatingCDO, TEXT("Creating nested packages is not allowed: Outer=%s, Package=%s"), *GetNameSafe(InOuter), *InName.ToString());
 	check(bCreatingCDO || bCreatingArchetype || !InOuter || InOuter->IsA(InClass->ClassWithin));
-	checkf(!IsGarbageCollecting(), TEXT("Unable to create new object: %s %s.%s. Creating UObjects while Collecting Garbage is not allowed!"),
+	checkf(!IsGarbageCollectingAndLockingUObjectHashTables(), TEXT("Unable to create new object: %s %s.%s. Creating UObjects while Collecting Garbage is not allowed!"),
 		*GetNameSafe(InClass), *GetPathNameSafe(InOuter), *InName.ToString());
 
 	if (bCreatingCDO)
diff --git a/Engine/Source/Runtime/CoreUObject/Private/UObject/UObjectHash.cpp b/Engine/Source/Runtime/CoreUObject/Private/UObject/UObjectHash.cpp
--- a/Engine/Source/Runtime/CoreUObject/Private/UObject/UObjectHash.cpp
+++ b/Engine/Source/Runtime/CoreUObject/Private/UObject/UObjectHash.cpp
@@ -481,7 +481,7 @@ public:
 	FORCEINLINE FHashTableLock(FUObjectHashTables& InTables)
 	{
 #if THREADSAFE_UOBJECTS
-		if (!(IsGarbageCollecting() && IsInGameThread()))
+		if (!(IsGarbageCollectingAndLockingUObjectHashTables() && IsInGameThread()))
 		{
 			Tables = &InTables;
 			InTables.Lock();
diff --git a/Engine/Source/Runtime/CoreUObject/Public/UObject/GarbageCollection.h b/Engine/Source/Runtime/CoreUObject/Public/UObject/GarbageCollection.h
--- a/Engine/Source/Runtime/CoreUObject/Public/UObject/GarbageCollection.h
+++ b/Engine/Source/Runtime/CoreUObject/Public/UObject/GarbageCollection.h
@@ -702,3 +702,8 @@ FORCEINLINE bool IsGarbageCollecting()
 {
 	return GIsGarbageCollecting;
 }
+
+/**
+* Whether garbage collection is locking the global uobject hash tables
+*/
+COREUOBJECT_API bool IsGarbageCollectingAndLockingUObjectHashTables();
diff --git a/Engine/Source/Runtime/CoreUObject/Public/UObject/UObjectGlobals.h b/Engine/Source/Runtime/CoreUObject/Public/UObject/UObjectGlobals.h
--- a/Engine/Source/Runtime/CoreUObject/Public/UObject/UObjectGlobals.h
+++ b/Engine/Source/Runtime/CoreUObject/Public/UObject/UObjectGlobals.h
@@ -254,7 +254,7 @@ COREUOBJECT_API UObject* StaticFindObjectFast(UClass* Class, UObject* InOuter, F
 
 /**
  * Fast and safe version of StaticFindObject that relies on the passed in FName being the object name without any group/package qualifiers.
- * It will not assert on GIsSavingPackage or IsGarbageCollecting(). If called from within package saving code or GC, will return nullptr
+ * It will not assert on GIsSavingPackage or IsGarbageCollectingAndLockingUObjectHashTables(). If called from within package saving code or GC, will return nullptr
  * This will only find top level packages or subobjects nested directly within a passed in outer.
  *
  * @param	Class			The to be found object's class
@@ -285,7 +285,7 @@ COREUOBJECT_API UObject* StaticFindObject( UClass* Class, UObject* InOuter, cons
 /** Version of StaticFindObject() that will assert if the object is not found */
 COREUOBJECT_API UObject* StaticFindObjectChecked( UClass* Class, UObject* InOuter, const TCHAR* Name, bool ExactClass=false );
 
-/** Internal version of StaticFindObject that will not assert on GIsSavingPackage or IsGarbageCollecting() */
+/** Internal version of StaticFindObject that will not assert on GIsSavingPackage or IsGarbageCollectingAndLockingUObjectHashTables() */
 COREUOBJECT_API UObject* StaticFindObjectSafe( UClass* Class, UObject* InOuter, const TCHAR* Name, bool ExactClass=false );
 
 /**
@@ -1617,7 +1617,7 @@ inline T* FindObjectChecked( UObject* Outer, const TCHAR* Name, bool ExactClass=
 }
 
 /**
- * Find an object without asserting on GIsSavingPackage or IsGarbageCollecting()
+ * Find an object without asserting on GIsSavingPackage or IsGarbageCollectingAndLockingUObjectHashTables()
  * @see StaticFindObjectSafe()
  */
 template< class T > 
