Adding UObject related utility functions that use FTopLevelAssetPath

diff --git a/Engine/Source/Runtime/CoreUObject/Private/UObject/UObjectGlobals.cpp b/Engine/Source/Runtime/CoreUObject/Private/UObject/UObjectGlobals.cpp
--- a/Engine/Source/Runtime/CoreUObject/Private/UObject/UObjectGlobals.cpp
+++ b/Engine/Source/Runtime/CoreUObject/Private/UObject/UObjectGlobals.cpp
@@ -405,6 +405,23 @@ UObject* StaticFindObject( UClass* ObjectClass, UObject* InObjectPackage, const
 	return StaticFindObjectFast(ObjectClass, ObjectPackage, ObjectName, bExactClass, bAnyPackage);
 }
 
+//
+// Find an optional object.
+//
+UObject* StaticFindObject(UClass* Class, FTopLevelAssetPath ObjectPath, bool ExactClass /*= false*/)
+{
+	UObject* Result = nullptr;
+	if (!ObjectPath.IsNull())
+	{
+		UObject* Package = StaticFindObjectFast(UPackage::StaticClass(), nullptr, ObjectPath.GetPackageName());
+		if (Package)
+		{
+			Result = StaticFindObjectFast(Class, Package, ObjectPath.GetAssetName(), ExactClass);
+		}
+	}
+	return Result;
+}
+
 //
 // Find an object; can't fail.
 //
diff --git a/Engine/Source/Runtime/CoreUObject/Public/UObject/Class.h b/Engine/Source/Runtime/CoreUObject/Public/UObject/Class.h
--- a/Engine/Source/Runtime/CoreUObject/Public/UObject/Class.h
+++ b/Engine/Source/Runtime/CoreUObject/Public/UObject/Class.h
@@ -31,6 +31,7 @@
 #include "UObject/UObjectGlobals.h"
 #include "UObject/FieldPath.h"
 #include "UObject/PropertyTag.h"
+#include "UObject/TopLevelAssetPath.h"
 
 struct FBlake3Hash;
 struct FCustomPropertyListNode;
@@ -2663,6 +2664,16 @@ public:
 	virtual bool IsStructTrashed() const override;
 	// End of UStruct interface.
 
+	/**
+	 * Returns class path name as a package + class FName pair
+	 */
+	FORCEINLINE FTopLevelAssetPath GetClassPathName() const
+	{
+		// Some day this check may actually be relevant
+		checkf(GetOuter() == GetOutermost(), TEXT("Only top level classes are supported by FTopLevelAssetPath. This class is nested: %s"), *GetPathName());
+		return FTopLevelAssetPath(GetOuter()->GetFName(), GetFName());
+	}
+
 #if WITH_EDITOR
 	/** Provides access to C++ type info. */
 	const ICppClassTypeInfo* GetCppTypeInfo() const
diff --git a/Engine/Source/Runtime/CoreUObject/Public/UObject/UObjectGlobals.h b/Engine/Source/Runtime/CoreUObject/Public/UObject/UObjectGlobals.h
--- a/Engine/Source/Runtime/CoreUObject/Public/UObject/UObjectGlobals.h
+++ b/Engine/Source/Runtime/CoreUObject/Public/UObject/UObjectGlobals.h
@@ -18,6 +18,7 @@
 #include "Templates/Function.h"
 #include "Templates/IsArrayOrRefOfType.h"
 #include "Serialization/ArchiveUObject.h"
+#include "UObject/TopLevelAssetPath.h"
 
 struct FCustomPropertyListNode;
 struct FObjectInstancingGraph;
@@ -282,6 +283,17 @@ COREUOBJECT_API UObject* StaticFindObjectFastSafe(UClass* Class, UObject* InOute
  */
 COREUOBJECT_API UObject* StaticFindObject( UClass* Class, UObject* InOuter, const TCHAR* Name, bool ExactClass=false );
 
+/**
+ * Tries to find an object in memory. This version uses FTopLevelAssetPath to find the object.
+ *
+ * @param	Class			The to be found object's class
+ * @param	ObjectPath		FName pair representing the outer package object and the inner top level object (asset)
+ * @param	ExactClass		Whether to require an exact match with the passed in class
+ *
+ * @return	Returns a pointer to the found object or nullptr if none could be found
+ */
+UObject* StaticFindObject(UClass* Class, FTopLevelAssetPath ObjectPath, bool ExactClass /*= false*/);
+
 /** Version of StaticFindObject() that will assert if the object is not found */
 COREUOBJECT_API UObject* StaticFindObjectChecked( UClass* Class, UObject* InOuter, const TCHAR* Name, bool ExactClass=false );
 
@@ -1606,6 +1618,16 @@ inline T* FindObject( UObject* Outer, const TCHAR* Name, bool ExactClass=false )
 	return (T*)StaticFindObject( T::StaticClass(), Outer, Name, ExactClass );
 }
 
+/**
+ * Find an optional object.
+ * @see StaticFindObject()
+ */
+template< class T >
+inline T* FindObject(FTopLevelAssetPath InPath, bool ExactClass = false)
+{
+	return (T*)StaticFindObject(T::StaticClass(), InPath, ExactClass);
+}
+
 /**
  * Find an optional object, no failure allowed
  * @see StaticFindObjectChecked()
