Adding *Safe versions of FindObject functions that take FTopLevelAssetPath as parameter

diff --git a/Engine/Source/Runtime/CoreUObject/Private/UObject/UObjectGlobals.cpp b/Engine/Source/Runtime/CoreUObject/Private/UObject/UObjectGlobals.cpp
--- a/Engine/Source/Runtime/CoreUObject/Private/UObject/UObjectGlobals.cpp
+++ b/Engine/Source/Runtime/CoreUObject/Private/UObject/UObjectGlobals.cpp
@@ -422,6 +422,19 @@ UObject* StaticFindObject(UClass* Class, FTopLevelAssetPath ObjectPath, bool Exa
 	return Result;
 }
 
+UObject* StaticFindObjectSafe(UClass* Class, FTopLevelAssetPath ObjectPath, bool ExactClass /*= false*/)
+{
+	if (!UE::IsSavingPackage(nullptr) && !IsGarbageCollectingAndLockingUObjectHashTables())
+	{
+		FGCScopeGuard GCGuard;
+		return StaticFindObject(Class, ObjectPath, ExactClass);
+	}
+	else
+	{
+		return nullptr;
+	}
+}
+
 //
 // Find an object; can't fail.
 //
diff --git a/Engine/Source/Runtime/CoreUObject/Public/UObject/UObjectGlobals.h b/Engine/Source/Runtime/CoreUObject/Public/UObject/UObjectGlobals.h
--- a/Engine/Source/Runtime/CoreUObject/Public/UObject/UObjectGlobals.h
+++ b/Engine/Source/Runtime/CoreUObject/Public/UObject/UObjectGlobals.h
@@ -292,7 +292,7 @@ COREUOBJECT_API UObject* StaticFindObject( UClass* Class, UObject* InOuter, cons
  *
  * @return	Returns a pointer to the found object or nullptr if none could be found
  */
-UObject* StaticFindObject(UClass* Class, FTopLevelAssetPath ObjectPath, bool ExactClass /*= false*/);
+COREUOBJECT_API UObject* StaticFindObject(UClass* Class, FTopLevelAssetPath ObjectPath, bool ExactClass /*= false*/);
 
 /** Version of StaticFindObject() that will assert if the object is not found */
 COREUOBJECT_API UObject* StaticFindObjectChecked( UClass* Class, UObject* InOuter, const TCHAR* Name, bool ExactClass=false );
@@ -300,6 +300,18 @@ COREUOBJECT_API UObject* StaticFindObjectChecked( UClass* Class, UObject* InOute
 /** Internal version of StaticFindObject that will not assert on GIsSavingPackage or IsGarbageCollectingAndLockingUObjectHashTables() */
 COREUOBJECT_API UObject* StaticFindObjectSafe( UClass* Class, UObject* InOuter, const TCHAR* Name, bool ExactClass=false );
 
+/**
+ * Tries to find an object in memory. This version uses FTopLevelAssetPath to find the object.
+ * Version of StaticFindObject that will not assert on GIsSavingPackage or IsGarbageCollectingAndLockingUObjectHashTables()
+ *
+ * @param	Class			The to be found object's class
+ * @param	ObjectPath		FName pair representing the outer package object and the inner top level object (asset)
+ * @param	ExactClass		Whether to require an exact match with the passed in class
+ *
+ * @return	Returns a pointer to the found object or nullptr if none could be found
+ */
+COREUOBJECT_API UObject* StaticFindObjectSafe(UClass* Class, FTopLevelAssetPath ObjectPath, bool ExactClass /*= false*/);
+
 /**
  * Fast version of StaticFindAllObjects that relies on the passed in FName being the object name without any group/package qualifiers.
  * This will find all objects matching the specified name and class.
@@ -1648,6 +1660,16 @@ inline T* FindObjectSafe( UObject* Outer, const TCHAR* Name, bool ExactClass=fal
 	return (T*)StaticFindObjectSafe( T::StaticClass(), Outer, Name, ExactClass );
 }
 
+/**
+ * Find an optional object.
+ * @see StaticFindObject()
+ */
+template< class T >
+inline T* FindObjectSafe(FTopLevelAssetPath InPath, bool ExactClass = false)
+{
+	return (T*)StaticFindObjectSafe(T::StaticClass(), InPath, ExactClass);
+}
+
 /**
  * Find an optional object with proper handling of potential ambiguity.
  * @see StaticFindFirstObject()
