fix for bulkdata template link errors in modular builds.

#rnx

diff --git a/Engine/Source/Runtime/CoreUObject/Public/Serialization/BulkData2.h b/Engine/Source/Runtime/CoreUObject/Public/Serialization/BulkData2.h
--- a/Engine/Source/Runtime/CoreUObject/Public/Serialization/BulkData2.h
+++ b/Engine/Source/Runtime/CoreUObject/Public/Serialization/BulkData2.h
@@ -338,7 +338,7 @@ private:
  * @documentation @todo documentation
  */
 template<typename ElementType>
-class COREUOBJECT_API FUntypedBulkData2 : public FBulkDataBase
+class FUntypedBulkData2 : public FBulkDataBase
 {
 	// In the older Bulkdata system the data was being loaded as if it was POD with the option to opt out
 	// but nothing actually opted out. This check should help catch if any non-POD data was actually being
