EditorPackageLoader: Fixed issue with recursive imports of cooked packages when all three loaders are involved in loading a package (second attempt).
#rnx

diff --git a/Engine/Source/Runtime/CoreUObject/Private/UObject/UObjectGlobals.cpp b/Engine/Source/Runtime/CoreUObject/Private/UObject/UObjectGlobals.cpp
--- a/Engine/Source/Runtime/CoreUObject/Private/UObject/UObjectGlobals.cpp
+++ b/Engine/Source/Runtime/CoreUObject/Private/UObject/UObjectGlobals.cpp
@@ -1328,7 +1328,7 @@ UPackage* LoadPackageInternal(UPackage* InOuter, const FPackagePath& PackagePath
 
 		GSyncLoadUsingAsyncLoaderCount++;
 
-#if 0
+#if WITH_EDITOR
 		// Temporarily clear the list of loaded objects before entering async loading. Otherwise the (uncooked) async loader might try to preload them which in turn can lead to
 		// recursive imports of the same cooked package.
 		FUObjectSerializeContext* LoadContext = FUObjectThreadContext::Get().GetSerializeContext();
@@ -1343,9 +1343,8 @@ UPackage* LoadPackageInternal(UPackage* InOuter, const FPackagePath& PackagePath
 			FlushAsyncLoading(RequestID);
 		}
 
-#if 0
-		check(ThreadObjLoaded.IsEmpty());
-		Swap(SavedObjLoaded, ThreadObjLoaded);
+#if WITH_EDITOR
+		ThreadObjLoaded.Insert(SavedObjLoaded, 0);
 #endif
 
 		GSyncLoadUsingAsyncLoaderCount--;
