Fix a bug where virtualized payloads were writing incorrect flags to the editor domain version of the package.

#rnx

- If a payload was virtualized it was incorrectly being marked as requiring the legacy serialization path.
- This would technically work but was triggering a log warning about the bulkdata being stored in the wrong format (since it was in the legacy format) as we do not expect bulkdata with virtualized payloads to use the legacy path.
- We need to change the editor domain version to clear out any bad entries

diff --git a/Engine/Source/Editor/UnrealEd/Private/EditorDomain/EditorDomainUtils.cpp b/Engine/Source/Editor/UnrealEd/Private/EditorDomain/EditorDomainUtils.cpp
--- a/Engine/Source/Editor/UnrealEd/Private/EditorDomain/EditorDomainUtils.cpp
+++ b/Engine/Source/Editor/UnrealEd/Private/EditorDomain/EditorDomainUtils.cpp
@@ -151,7 +151,7 @@ bool bGGlobalConstructClassesInitialized = false;
 int64 GMaxBulkDataSize = -1;
 
 // Change to a new guid when EditorDomain needs to be invalidated
-const TCHAR* EditorDomainVersion = TEXT("A5D5DB1F5951446EA0E2F3E017314DFC");
+const TCHAR* EditorDomainVersion = TEXT("D3E2B4E479C34305952A97AA63A8E19D");
 
 // Identifier of the CacheBuckets for EditorDomain tables
 const TCHAR* EditorDomainPackageBucketName = TEXT("EditorDomainPackage");
diff --git a/Engine/Source/Runtime/CoreUObject/Private/Serialization/EditorBulkData.cpp b/Engine/Source/Runtime/CoreUObject/Private/Serialization/EditorBulkData.cpp
--- a/Engine/Source/Runtime/CoreUObject/Private/Serialization/EditorBulkData.cpp
+++ b/Engine/Source/Runtime/CoreUObject/Private/Serialization/EditorBulkData.cpp
@@ -1973,7 +1973,7 @@ FEditorBulkData::EFlags FEditorBulkData::BuildFlagsForSerialization(FArchive& Ar
 		// Currently we do not support storing local payloads to a trailer if it is being built for reference
 		// access (i.e. for the editor domain) and if this is detected we should force the legacy serialization
 		// path for this payload.
-		const bool bForceLegacyPath = bKeepFileDataByReference && bCanKeepFileDataByReference == false;
+		const bool bForceLegacyPath = bKeepFileDataByReference && bCanKeepFileDataByReference == false && !IsDataVirtualized(UpdatedFlags);
 
 		if (ShouldUseLegacySerialization(LinkerSave) == true || bForceLegacyPath == true)
 		{
