Fix a number of bugs found when saving more complex package types with virtualization enabled.

#rnx

- Note that these problems were not causing crashes, but they did cause errors to be logged about invalid formats and request the packages to be re-saved.
- Most of the problem comes from the code paths being too complicated, hopefully now things are locked down we should be able to start pruning this soon and make it easier to maintain in the future.

### Problem
- The problem comes from certain package setups that allow for the existence of duplicate payload entries where one entry is local or referenced and another version is virtualized. If a trailer is created from this data then the virtualized entry will be discarded (as the look up table is a map) and we revert to using the non virtualized payload instead. Ideally we'd want to discard the local/referenced payloads in favour of the virtualized one to save space.
-- By itself this bug would waste disk space but not cause other problems, however we were also writing out the virtualized flag so that when the bulkdata is next loaded it has the virtualized flag but then checks with the trailer and is told that it is not virtualized which triggers some asserts I added to find badly formed data.
- How the package gets duplicate entries like this is somewhat complicated involing duplicating objects from other (already virtualized) packages during load, then saving them (found in packages with Niagra assets at least)

### Fix
- When creating a trailer from a trailer build, we first find any local and/or referenced entries and check if there are duplicate entries in the virtualized payload entries. If so we remove the local/referenced version in favor of using the virtualized version.
-- Ideally we'd do this during the serialization of exports as it might let us skip the loading of some payloads for save entirely, however we don't know all the entries until after export serialization is completed so this is the only point we can search for duplicates safely.
- Added a way to get the full name of the owning asset, including package name + asset, when asking the user to re-save. This makes debugging the problems easier.
- Add a new check when loading to find payloads saved with the virtualization flag and to a trailer. If we do find one we warn the user that they should re-save the package and remove the flag so that no further problems occur.
- Demoted the old check for bad formats to a warning as we don't need to fail builds because of it.
-- Additionally we now only give the warning if we detect that we are loading from a package, as being virtualized and not in a package trailer would be valid when loading from other sources.
- Added comments to consider removing both checks asking the user to re-save the package after 5.2, as in theory we should only ever see these problems internally at Epic as it is unlikely that externals will have enabled the system.
- When saving a virtualized payload to the trailer, make sure to remove the virtualized flag, that should be picked up on load by polling the trailer for the payload status and should not be baked into the export data.
-- Hypothetically this will allow for us to re-hydrate the package without re-saving it through manipulating the trailer.

diff --git a/Engine/Source/Runtime/CoreUObject/Private/Serialization/EditorBulkData.cpp b/Engine/Source/Runtime/CoreUObject/Private/Serialization/EditorBulkData.cpp
--- a/Engine/Source/Runtime/CoreUObject/Private/Serialization/EditorBulkData.cpp
+++ b/Engine/Source/Runtime/CoreUObject/Private/Serialization/EditorBulkData.cpp
@@ -203,6 +203,19 @@ static FPackagePath GetPackagePathFromOwner(UObject* Owner, EPackageSegment& Out
 	}
 }
 
+/** Utility for finding a valid debug name to print from the owning UObject */
+FString GetDebugNameFromOwner(UObject* Owner)
+{
+	if (Owner != nullptr)
+	{
+		return Owner->GetFullName();
+	}
+	else
+	{
+		return FString(TEXT("Unknown"));
+	}
+}
+
 /** Utility for hashing a payload, will return a default FIoHash if the payload is invalid or of zero length */
 static FIoHash HashPayload(const FSharedBuffer& InPayload)
 {
@@ -960,11 +973,22 @@ void FEditorBulkData::Serialize(FArchive& Ar, UObject* Owner, bool bAllowRegiste
 				// Cache the offset from the trailer (if we move the loading of the payload to the trailer 
 				// at a later point then we can skip this)
 				OffsetInFile = Trailer->FindPayloadOffsetInFile(PayloadContentId);
+
+				// Attempt to catch data that saved with the virtualization flag when it's package has a trailers.
+				// It is unlikely this will ever trigger in the wild but keeping the code path for now to be safe.
+				// TODO: Consider removing in 5.2+
+				if (IsDataVirtualized() && Trailer->FindPayloadStatus(PayloadContentId) != EPayloadStatus::StoredVirtualized)
+				{
+					UE_LOG(LogSerialization, Warning, TEXT("Payload in '%s' is was saved with an invalid flag and required fixing. Please re-save the package!"), *GetDebugNameFromOwner(Owner));
+					EnumRemoveFlags(Flags, EFlags::IsVirtualized);
+				}
 			}
 			else
 			{
-				// TODO: This check is for older virtualized formats that might be seen in older test projects.
-				UE_CLOG(IsDataVirtualized(), LogSerialization, Error, TEXT("Payload in '%s' is virtualized in an older format and should be re-saved!"), *Owner->GetName());
+				// This check is for older virtualized formats that might be seen in older test projects.
+				// But we only care if the archive has a linker! (loading from a package)
+				// TODO: Consider removing in 5.2+
+				UE_CLOG(IsDataVirtualized() && Ar.GetLinker() != nullptr, LogSerialization, Warning, TEXT("Payload in '%s' is virtualized in an older format and should be re-saved!"), *GetDebugNameFromOwner(Owner));
 				if (!IsDataVirtualized())
 				{
 					Ar << OffsetInFile;
@@ -1590,10 +1614,25 @@ void FEditorBulkData::SerializeToPackageTrailer(FLinkerSave& LinkerSave, FCompre
 					return;
 				}
 
-				this->PackagePath = InPackagePath;
-				check(!this->PackagePath.IsEmpty()); // LinkerSave guarantees a valid PackagePath if we're updating loaded path
 				this->OffsetInFile = PayloadOffset;
 				this->Flags = (this->Flags & EFlags::TransientFlags) | (UpdatedFlags & ~EFlags::TransientFlags);
+	
+				if (PayloadOffset != INDEX_NONE)
+				{
+					// A valid payload means that it was saved locally to disk. We can set the path at this point but will only
+					// be able to load from it if the package is reattached to the file on disk, which is not a currently
+					// supported feature.
+					this->PackagePath = InPackagePath;
+					check(!this->PackagePath.IsEmpty()); // LinkerSave guarantees a valid PackagePath if we're updating loaded path
+				}
+				else
+				{
+					// If the payload offset we are given is INDEX_NONE it means that the payload was discarded as there was an identical
+					// virtualized entry and the virtualized version takes priority. Since we know that the payload is in the 
+					// virtualization system at this point we can set the virtualization flag allowing us to unload any existing payload
+					// in memory to help reduce bloat.
+					EnumAddFlags(this->Flags, EFlags::IsVirtualized);
+				}
 
 				if (CanUnloadData())
 				{
@@ -1982,6 +2021,10 @@ FEditorBulkData::EFlags FEditorBulkData::BuildFlagsForSerialization(FArchive& Ar
 		else
 		{
 			EnumAddFlags(UpdatedFlags, EFlags::StoredInPackageTrailer);
+
+			// If the payload is virtualized or not will be controlled by the package trailer so we shouldn't
+			//serialize this out to the exports
+			EnumRemoveFlags(UpdatedFlags, EFlags::IsVirtualized);
 		}
 
 		return UpdatedFlags;
diff --git a/Engine/Source/Runtime/CoreUObject/Private/UObject/PackageTrailer.cpp b/Engine/Source/Runtime/CoreUObject/Private/UObject/PackageTrailer.cpp
--- a/Engine/Source/Runtime/CoreUObject/Private/UObject/PackageTrailer.cpp
+++ b/Engine/Source/Runtime/CoreUObject/Private/UObject/PackageTrailer.cpp
@@ -3,6 +3,7 @@
 #include "UObject/PackageTrailer.h"
 
 #include "Algo/Count.h"
+#include "Algo/RemoveIf.h"
 #include "CoreGlobals.h"
 #include "Misc/ConfigCacheIni.h"
 #include "Misc/PackagePath.h"
@@ -222,7 +223,8 @@ bool FPackageTrailerBuilder::BuildAndAppendTrailer(FLinkerSave* Linker, FArchive
 	// parse and manipulate the trailer without needing to worry that we might change how
 	// our containers serialize. 
 	
-	// First we build a trailer structure
+	RemoveDuplicateEntries();
+
 	FPackageTrailer Trailer;
 	
 	Trailer.Header.Tag = FPackageTrailer::FHeader::HeaderTag;
@@ -233,7 +235,7 @@ bool FPackageTrailerBuilder::BuildAndAppendTrailer(FLinkerSave* Linker, FArchive
 	Trailer.Header.HeaderLength = FPackageTrailer::FHeader::StaticHeaderSizeOnDisk + DynamicHeaderSizeOnDisk;
 	
 	Trailer.Header.PayloadsDataLength = 0;
-	Trailer.Header.PayloadLookupTable.Reserve(LocalEntries.Num() + VirtualizedEntries.Num());
+	Trailer.Header.PayloadLookupTable.Reserve(LocalEntries.Num() + ReferencedEntries.Num() + VirtualizedEntries.Num());
 
 	for (const TPair<FIoHash, LocalEntry>& It : LocalEntries)
 	{
@@ -345,6 +347,27 @@ int32 FPackageTrailerBuilder::GetNumVirtualizedPayloads() const
 	return VirtualizedEntries.Num();
 }
 
+void FPackageTrailerBuilder::RemoveDuplicateEntries()
+{
+	for (auto Iter = LocalEntries.CreateIterator(); Iter; ++Iter)
+	{
+		if (VirtualizedEntries.Contains(Iter.Key()))
+		{	
+			UE_LOG(LogSerialization, Verbose, TEXT("Replacing localized payload '%s' with the virtualized version when building the package trailer for '%s'"), *LexToString(Iter.Key()) , *PackageName.ToString());
+			Iter.RemoveCurrent();
+		}
+	}
+
+	for (auto Iter = ReferencedEntries.CreateIterator(); Iter; ++Iter)
+	{
+		if (VirtualizedEntries.Contains(Iter.Key()))
+		{
+			UE_LOG(LogSerialization, Verbose, TEXT("Replacing localized payload '%s' with the virtualized version when building the package trailer for '%s'"), *LexToString(Iter.Key()), *PackageName.ToString());
+			Iter.RemoveCurrent();
+		}
+	}
+}
+
 bool FPackageTrailer::IsEnabled()
 {
 	static struct FUsePackageTrailer
diff --git a/Engine/Source/Runtime/CoreUObject/Public/UObject/PackageTrailer.h b/Engine/Source/Runtime/CoreUObject/Public/UObject/PackageTrailer.h
--- a/Engine/Source/Runtime/CoreUObject/Public/UObject/PackageTrailer.h
+++ b/Engine/Source/Runtime/CoreUObject/Public/UObject/PackageTrailer.h
@@ -279,6 +279,8 @@ private:
 		int64 RawSize = INDEX_NONE;
 	};
 
+	void RemoveDuplicateEntries();
+
 	// Members used when building the trailer
 
 	/** Name of the package the trailer is being built for, used to give meaningful error messages */
