Remove part of the fix from CL 19963043 as it had unintended consequences and needs rethinking

#rnx

- Although we don't really want to write out the virtualization flag, by removing it during ::BuildFlagsForSerialization we later check UpdatedFlags to see if the payload is virtualized and if we should write it out to disk or not.
- Removing the flag means we re-hydrate virtualized payloads when saved to disk, or in the case of an editordomain, try to store it as a local payload in the editor domain trailer which is not supported (if we did have a newly generated payload in an editor domain save it should go to the legacy path for now)
- The fix for this is probably quite simple, but the testing it would require means it's better to pull out the line of code and allow packages to serialize out the virtualization flag for now until a better fix is ready.

diff --git a/Engine/Source/Runtime/CoreUObject/Private/Serialization/EditorBulkData.cpp b/Engine/Source/Runtime/CoreUObject/Private/Serialization/EditorBulkData.cpp
--- a/Engine/Source/Runtime/CoreUObject/Private/Serialization/EditorBulkData.cpp
+++ b/Engine/Source/Runtime/CoreUObject/Private/Serialization/EditorBulkData.cpp
@@ -2021,10 +2021,6 @@ FEditorBulkData::EFlags FEditorBulkData::BuildFlagsForSerialization(FArchive& Ar
 		else
 		{
 			EnumAddFlags(UpdatedFlags, EFlags::StoredInPackageTrailer);
-
-			// If the payload is virtualized or not will be controlled by the package trailer so we shouldn't
-			//serialize this out to the exports
-			EnumRemoveFlags(UpdatedFlags, EFlags::IsVirtualized);
 		}
 
 		return UpdatedFlags;
