BulkDataRegistry: Add comment about silencing the duplicate guid warning.
#rnx

diff --git a/Engine/Source/Runtime/CoreUObject/Private/Serialization/EditorBulkData.cpp b/Engine/Source/Runtime/CoreUObject/Private/Serialization/EditorBulkData.cpp
--- a/Engine/Source/Runtime/CoreUObject/Private/Serialization/EditorBulkData.cpp
+++ b/Engine/Source/Runtime/CoreUObject/Private/Serialization/EditorBulkData.cpp
@@ -623,6 +623,10 @@ void FEditorBulkData::LogRegisterError(UE::BulkDataRegistry::ERegisterResult Val
 				FString SilenceWarningMessage;
 				if (OwnerPackage)
 				{
+					// More notes on silencing the warning:
+					// 1) Resaveondemand does not yet work for startup packages. If the warning is coming from startup packages, run ResavePackagesCommandlet with
+					// argument "-package=<PackageName>" for each warned package.
+					// 2) The guid will only be updated on resave if both packages contributing to the collision are loaded during the resave.
 					SilenceWarningMessage = TEXT(" To silence this warning, run the ResavePackagesCommandlet with \"-autocheckout -resaveondemand=bulkdataduplicates\".");
 				}
 				bool bSuppressWarning = false; // Allow projects to suppress the warning because they might not be able to resave packages during an integration
