Make sure one time init code in IsEventDrivenLoaderEnabled() is run in async loading thread construction

diff --git a/Engine/Source/Runtime/CoreUObject/Private/Serialization/AsyncLoading2.cpp b/Engine/Source/Runtime/CoreUObject/Private/Serialization/AsyncLoading2.cpp
--- a/Engine/Source/Runtime/CoreUObject/Private/Serialization/AsyncLoading2.cpp
+++ b/Engine/Source/Runtime/CoreUObject/Private/Serialization/AsyncLoading2.cpp
@@ -5200,6 +5200,7 @@ FAsyncLoadingThread2::FAsyncLoadingThread2(FIoDispatcher& InIoDispatcher, IAsync
 	, UncookedPackageLoader(InUncookedPackageLoader)
 {
 #if !WITH_IOSTORE_IN_EDITOR
+	IsEventDrivenLoaderEnabled(); // make sure the one time init inside runs
 	GEventDrivenLoaderEnabled = true;
 #endif
 
