Mark FPlatformInputDeviceState properties as BP read only so that you can access them when needed in blueprints, and they show up in the details panel.

#rnx

diff --git a/Engine/Source/Runtime/CoreUObject/Public/UObject/NoExportTypes.h b/Engine/Source/Runtime/CoreUObject/Public/UObject/NoExportTypes.h
--- a/Engine/Source/Runtime/CoreUObject/Public/UObject/NoExportTypes.h
+++ b/Engine/Source/Runtime/CoreUObject/Public/UObject/NoExportTypes.h
@@ -2068,11 +2068,11 @@ struct FPlatformInputDeviceState
 	GENERATED_BODY()
 
 	/** The platform user that this input device belongs to */
-	UPROPERTY(VisibleAnywhere, Category = "PlatformInputDevice")
+	UPROPERTY(BlueprintReadOnly, VisibleAnywhere, Category = "PlatformInputDevice")
 	FPlatformUserId OwningPlatformUser = PLATFORMUSERID_NONE;
 
 	/** The connection state of this input device */
-	UPROPERTY(VisibleAnywhere, Category = "PlatformInputDevice")
+	UPROPERTY(BlueprintReadOnly, VisibleAnywhere, Category = "PlatformInputDevice")
 	EInputDeviceConnectionState ConnectionState = EInputDeviceConnectionState::Invalid;
 };
 
