Shrink FObjectExport (144 -> 120): Remove deprecated PackageGUID and Dynamic Type, and properly pack members

diff --git a/Engine/Source/Editor/UnrealEd/Private/Commandlets/PackageUtilities.cpp b/Engine/Source/Editor/UnrealEd/Private/Commandlets/PackageUtilities.cpp
--- a/Engine/Source/Editor/UnrealEd/Private/Commandlets/PackageUtilities.cpp
+++ b/Engine/Source/Editor/UnrealEd/Private/Commandlets/PackageUtilities.cpp
@@ -1207,9 +1207,6 @@ void FPkgInfoReporter_Log::GeneratePackageReport( FLinkerLoad* InLinker /*=nullp
 				Out.Logf(ELogVerbosity::Display, TEXT("\t\t        Parent: '%s' (%d)"), *ParentName, Export.SuperIndex.ForDebugging());
 				Out.Logf(ELogVerbosity::Display, TEXT("\t\t      Template: '%s' (%d)"), *TemplateName, Export.TemplateIndex.ForDebugging());
 				Out.Logf(ELogVerbosity::Display, TEXT("\t\t         Outer: '%s' (%d)"), *OuterName, Export.OuterIndex.ForDebugging() );
-				PRAGMA_DISABLE_DEPRECATION_WARNINGS
-				Out.Logf(ELogVerbosity::Display, TEXT("\t\t      Pkg Guid: %s"), *Export.PackageGuid.ToString());
-				PRAGMA_ENABLE_DEPRECATION_WARNINGS
 				Out.Logf(ELogVerbosity::Display, TEXT("\t\t   ObjectFlags: 0x%08X"), (uint32)Export.ObjectFlags );
 				Out.Logf(ELogVerbosity::Display, TEXT("\t\t          Size: %d"), Export.SerialSize );
 				if ( !IsHideOffsets())
diff --git a/Engine/Source/Runtime/Core/Public/UObject/ObjectVersion.h b/Engine/Source/Runtime/Core/Public/UObject/ObjectVersion.h
--- a/Engine/Source/Runtime/Core/Public/UObject/ObjectVersion.h
+++ b/Engine/Source/Runtime/Core/Public/UObject/ObjectVersion.h
@@ -56,7 +56,10 @@ enum class EUnrealEngineObjectUE5Version : uint32
 	
 	// Large world coordinates converts a number of core types to double components by default.
 	LARGE_WORLD_COORDINATES,       
-	
+
+	// Remove package GUID
+	REMOVE_PACKAGE_GUID,
+
 	// -----<new versions can be added before this line>-------------------------------------------------
 	// - this needs to be the last line (see note below)
 	AUTOMATIC_VERSION_PLUS_ONE,
diff --git a/Engine/Source/Runtime/CoreUObject/Private/Serialization/ArchiveStackTrace.cpp b/Engine/Source/Runtime/CoreUObject/Private/Serialization/ArchiveStackTrace.cpp
--- a/Engine/Source/Runtime/CoreUObject/Private/Serialization/ArchiveStackTrace.cpp
+++ b/Engine/Source/Runtime/CoreUObject/Private/Serialization/ArchiveStackTrace.cpp
@@ -1129,9 +1129,6 @@ FString ConvertItemToText(const FObjectImport& Import, FLinkerLoad* Linker)
 bool CompareTableItem(FLinkerLoad* SourceLinker, FLinkerLoad* DestLinker, const FObjectExport& SourceExport, const FObjectExport& DestExport)
 {
 	if (SourceExport.ObjectName != DestExport.ObjectName ||
-		PRAGMA_DISABLE_DEPRECATION_WARNINGS
-		SourceExport.PackageGuid != DestExport.PackageGuid ||
-		PRAGMA_ENABLE_DEPRECATION_WARNINGS
 		SourceExport.PackageFlags != DestExport.PackageFlags ||
 		SourceExport.ObjectFlags != DestExport.ObjectFlags ||
 		SourceExport.SerialSize != DestExport.SerialSize ||
@@ -1223,15 +1220,12 @@ bool ComparePackageIndices(FLinkerLoad* SourceLinker, FLinkerLoad* DestLinker, c
 FString ConvertItemToText(const FObjectExport& Export, FLinkerLoad* Linker)
 {
 	FName ClassName = Export.ClassIndex.IsNull() ? FName(NAME_Class) : Linker->ImpExp(Export.ClassIndex).ObjectName;
-	return FString::Printf(TEXT("%s Super: %s, Template: %s, Flags: %d, Size: %lld, PackageGuid: %s, PackageFlags: %d, ForcedExport: %d, NotForClient: %d, NotForServer: %d, NotAlwaysLoadedForEditorGame: %d, IsAsset: %d"),
+	return FString::Printf(TEXT("%s Super: %s, Template: %s, Flags: %d, Size: %lld, PackageFlags: %d, ForcedExport: %d, NotForClient: %d, NotForServer: %d, NotAlwaysLoadedForEditorGame: %d, IsAsset: %d"),
 		*GetTableKey(Linker, Export),
 		*GetTableKeyForIndex(Linker, Export.SuperIndex),
 		*GetTableKeyForIndex(Linker, Export.TemplateIndex),
 		(int32)Export.ObjectFlags,
 		Export.SerialSize,
-		PRAGMA_DISABLE_DEPRECATION_WARNINGS
-		*Export.PackageGuid.ToString(),
-		PRAGMA_ENABLE_DEPRECATION_WARNINGS
 		Export.PackageFlags,
 		Export.bForcedExport,
 		Export.bNotForClient,
diff --git a/Engine/Source/Runtime/CoreUObject/Private/UObject/LinkerDiff.cpp b/Engine/Source/Runtime/CoreUObject/Private/UObject/LinkerDiff.cpp
--- a/Engine/Source/Runtime/CoreUObject/Private/UObject/LinkerDiff.cpp
+++ b/Engine/Source/Runtime/CoreUObject/Private/UObject/LinkerDiff.cpp
@@ -460,9 +460,6 @@ void FLinkerDiff::GenerateExportDiff(FLinker* LHSLinker, const FLinkerExportObje
 
 	COMPARE_MEMBER(LHSExport.Export, RHSExport.Export, ObjectFlags);
 	COMPARE_MEMBER(LHSExport.Export, RHSExport.Export, SerialSize);
-	PRAGMA_DISABLE_DEPRECATION_WARNINGS
-	COMPARE_MEMBER(LHSExport.Export, RHSExport.Export, PackageGuid);
-	PRAGMA_ENABLE_DEPRECATION_WARNINGS
 	COMPARE_MEMBER(LHSExport.Export, RHSExport.Export, PackageFlags);
 	COMPARE_MEMBER(LHSExport.Export, RHSExport.Export, SerializationBeforeSerializationDependencies);
 	COMPARE_MEMBER(LHSExport.Export, RHSExport.Export, CreateBeforeSerializationDependencies);
diff --git a/Engine/Source/Runtime/CoreUObject/Private/UObject/ObjectResource.cpp b/Engine/Source/Runtime/CoreUObject/Private/UObject/ObjectResource.cpp
--- a/Engine/Source/Runtime/CoreUObject/Private/UObject/ObjectResource.cpp
+++ b/Engine/Source/Runtime/CoreUObject/Private/UObject/ObjectResource.cpp
@@ -47,10 +47,6 @@ FObjectExport::FObjectExport()
 , bIsAsset(false)
 , bGeneratePublicHash(false)
 , bExportLoadFailed(false)
-// @todo: BP2CPP_remove
-PRAGMA_DISABLE_DEPRECATION_WARNINGS
-, DynamicType(EDynamicType::NotDynamicExport)
-PRAGMA_ENABLE_DEPRECATION_WARNINGS
 , bWasFiltered(false)
 , PackageFlags(0)
 , FirstExportDependency(-1)
@@ -77,10 +73,6 @@ FObjectExport::FObjectExport( UObject* InObject, bool bInNotAlwaysLoadedForEdito
 , bIsAsset(false)
 , bGeneratePublicHash(false)
 , bExportLoadFailed(false)
-// @todo: BP2CPP_remove
-PRAGMA_DISABLE_DEPRECATION_WARNINGS
-, DynamicType(EDynamicType::NotDynamicExport)
-PRAGMA_ENABLE_DEPRECATION_WARNINGS
 , bWasFiltered(false)
 , PackageFlags(0)
 , FirstExportDependency(-1)
@@ -150,28 +142,37 @@ void operator<<(FStructuredArchive::FSlot Slot, FObjectExport& E)
 		Record << SA_VALUE(TEXT("SerialOffset"), E.SerialOffset);
 	}
 
-	Record << SA_VALUE(TEXT("bForcedExport"), E.bForcedExport);
-	Record << SA_VALUE(TEXT("bNotForClient"), E.bNotForClient);
-	Record << SA_VALUE(TEXT("bNotForServer"), E.bNotForServer);
+	#define SERIALIZE_BIT_TO_RECORD(bValue) { \
+		bool b = E.bValue; \
+		Record << SA_VALUE(TEXT(#bValue), b); \
+		E.bValue = b; \
+	}
+
+	SERIALIZE_BIT_TO_RECORD(bForcedExport);
+	SERIALIZE_BIT_TO_RECORD(bNotForClient);
+	SERIALIZE_BIT_TO_RECORD(bNotForServer);
+
+	if (BaseArchive.UEVer() < EUnrealEngineObjectUE5Version::REMOVE_PACKAGE_GUID)
+	{
+		FGuid DummyPackageGuid;
+		Record << SA_VALUE(TEXT("PackageGuid"), DummyPackageGuid);
+	}
 
-	PRAGMA_DISABLE_DEPRECATION_WARNINGS
-	Record << SA_VALUE(TEXT("PackageGuid"), E.PackageGuid);
-	PRAGMA_ENABLE_DEPRECATION_WARNINGS
 	Record << SA_VALUE(TEXT("PackageFlags"), E.PackageFlags);
 
 	if (BaseArchive.UEVer() >= VER_UE4_LOAD_FOR_EDITOR_GAME)
 	{
-		Record << SA_VALUE(TEXT("bNotAlwaysLoadedForEditorGame"), E.bNotAlwaysLoadedForEditorGame);
+		SERIALIZE_BIT_TO_RECORD(bNotAlwaysLoadedForEditorGame);
 	}
 
 	if (BaseArchive.UEVer() >= VER_UE4_COOKED_ASSETS_IN_EDITOR_SUPPORT)
 	{
-		Record << SA_VALUE(TEXT("bIsAsset"), E.bIsAsset);
+		SERIALIZE_BIT_TO_RECORD(bIsAsset);
 	}
 
 	if (BaseArchive.UEVer() >= EUnrealEngineObjectUE5Version::OPTIONAL_RESOURCES)
 	{
-		Record << SA_VALUE(TEXT("bGeneratePublicHash"), E.bGeneratePublicHash);
+		SERIALIZE_BIT_TO_RECORD(bGeneratePublicHash);
 	}
 
 	if (BaseArchive.UEVer() >= VER_UE4_PRELOAD_DEPENDENCIES_IN_COOKED_EXPORTS)
@@ -182,6 +183,8 @@ void operator<<(FStructuredArchive::FSlot Slot, FObjectExport& E)
 		Record << SA_VALUE(TEXT("SerializationBeforeCreateDependencies"), E.SerializationBeforeCreateDependencies);
 		Record << SA_VALUE(TEXT("CreateBeforeCreateDependencies"), E.CreateBeforeCreateDependencies);
 	}	
+
+	#undef SERIALIZE_BIT_TO_RECORD
 }
 
 /*-----------------------------------------------------------------------------
@@ -236,17 +239,28 @@ void operator<<(FStructuredArchive::FSlot Slot, FObjectTextExport& E)
 		E.Export.ObjectFlags = EObjectFlags(Save & RF_Load);
 	}
 
-	Slot << SA_OPTIONAL_ATTRIBUTE(TEXT("bForcedExport"), E.Export.bForcedExport, false);
-	Slot << SA_OPTIONAL_ATTRIBUTE(TEXT("bNotForClient"), E.Export.bNotForClient, false);
-	Slot << SA_OPTIONAL_ATTRIBUTE(TEXT("bNotForServer"), E.Export.bNotForServer, false);
+	#define SERIALIZE_BIT_TO_SLOT(bValue) { \
+		bool b = E.Export.bValue; \
+		Slot << SA_OPTIONAL_ATTRIBUTE(TEXT(#bValue), b, false); \
+		E.Export.bValue = b; \
+	}
+
+	SERIALIZE_BIT_TO_SLOT(bForcedExport);
+	SERIALIZE_BIT_TO_SLOT(bNotForClient);
+	SERIALIZE_BIT_TO_SLOT(bNotForServer);
+
+	if (BaseArchive.UEVer() < EUnrealEngineObjectUE5Version::REMOVE_PACKAGE_GUID)
+	{
+		FGuid DummyPackageGuid;
+		Slot << SA_OPTIONAL_ATTRIBUTE(TEXT("PackageGuid"), DummyPackageGuid, FGuid());
+	}
 
-	PRAGMA_DISABLE_DEPRECATION_WARNINGS
-	Slot << SA_OPTIONAL_ATTRIBUTE(TEXT("PackageGuid"), E.Export.PackageGuid, FGuid());
-	PRAGMA_ENABLE_DEPRECATION_WARNINGS
 	Slot << SA_OPTIONAL_ATTRIBUTE(TEXT("PackageFlags"), E.Export.PackageFlags, 0);
 
-	Slot << SA_OPTIONAL_ATTRIBUTE(TEXT("bNotAlwaysLoadedForEditorGame"), E.Export.bNotAlwaysLoadedForEditorGame, false);
-	Slot << SA_OPTIONAL_ATTRIBUTE(TEXT("bIsAsset"), E.Export.bIsAsset, false);
+	SERIALIZE_BIT_TO_SLOT(bNotAlwaysLoadedForEditorGame);
+	SERIALIZE_BIT_TO_SLOT(bIsAsset);
+
+	#undef SERIALIZE_BIT_TO_SLOT
 }
 
 /*-----------------------------------------------------------------------------
diff --git a/Engine/Source/Runtime/CoreUObject/Private/UObject/SavePackage.cpp b/Engine/Source/Runtime/CoreUObject/Private/UObject/SavePackage.cpp
--- a/Engine/Source/Runtime/CoreUObject/Private/UObject/SavePackage.cpp
+++ b/Engine/Source/Runtime/CoreUObject/Private/UObject/SavePackage.cpp
@@ -3079,12 +3079,6 @@ FSavePackageResultStruct UPackage::Save(UPackage* InOuter, UObject* InAsset, con
 						if (Package != nullptr)
 						{
 							Linker->ExportMap[i].PackageFlags = Package->GetPackageFlags();
-							if (!Package->HasAnyPackageFlags(PKG_ServerSideOnly))
-							{
-								PRAGMA_DISABLE_DEPRECATION_WARNINGS
-								Linker->ExportMap[i].PackageGuid = Package->GetGuid();
-								PRAGMA_ENABLE_DEPRECATION_WARNINGS
-							}
 						}
 					}
 				}
diff --git a/Engine/Source/Runtime/CoreUObject/Private/UObject/SavePackage2.cpp b/Engine/Source/Runtime/CoreUObject/Private/UObject/SavePackage2.cpp
--- a/Engine/Source/Runtime/CoreUObject/Private/UObject/SavePackage2.cpp
+++ b/Engine/Source/Runtime/CoreUObject/Private/UObject/SavePackage2.cpp
@@ -1074,12 +1074,6 @@ ESavePackageResult BuildLinker(FSaveContext& SaveContext)
 			if (UPackage* Package = Cast<UPackage>(Export.Object))
 			{
 				Export.PackageFlags = Package->GetPackageFlags();
-				if (!Package->HasAnyPackageFlags(PKG_ServerSideOnly))
-				{
-					PRAGMA_DISABLE_DEPRECATION_WARNINGS
-					Export.PackageGuid = Package->GetGuid();
-					PRAGMA_ENABLE_DEPRECATION_WARNINGS
-				}
 			}
 		}
 		//Sort(&Linker->ExportMap[0], Linker->ExportMap.Num(), FObjectResourceSortHelper());
diff --git a/Engine/Source/Runtime/CoreUObject/Public/UObject/ObjectResource.h b/Engine/Source/Runtime/CoreUObject/Public/UObject/ObjectResource.h
--- a/Engine/Source/Runtime/CoreUObject/Public/UObject/ObjectResource.h
+++ b/Engine/Source/Runtime/CoreUObject/Public/UObject/ObjectResource.h
@@ -282,19 +282,19 @@ struct FObjectExport : public FObjectResource
 	 * Whether the export was forced into the export table via OBJECTMARK_ForceTagExp.
 	 * Serialized
 	 */
-	bool			bForcedExport;   
+	bool			bForcedExport:1;   
 
 	/**
 	 * Whether the export should be loaded on clients
 	 * Serialized
 	 */
-	bool			bNotForClient;   
+	bool			bNotForClient:1;   
 
 	/**
 	 * Whether the export should be loaded on servers
 	 * Serialized
 	 */
-	bool			bNotForServer;
+	bool			bNotForServer:1;
 
 	/**
 	 * Whether the export should be always loaded in editor game
@@ -302,51 +302,27 @@ struct FObjectExport : public FObjectResource
 	 * True doesn't means, that the object won't be loaded.
 	 * Serialized
 	 */
-	bool			bNotAlwaysLoadedForEditorGame;
+	bool			bNotAlwaysLoadedForEditorGame:1;
 
 	/**
 	 * True if this export is an asset object.
 	 */
-	bool			bIsAsset;
+	bool			bIsAsset:1;
 
 	/**
 	 * True if this export should have its iostore public hash generated even if not RF_Public.
 	 */
-	bool			bGeneratePublicHash;
+	bool			bGeneratePublicHash:1;
 
 	/**
 	 * Force this export to not load, it failed because the outer didn't exist.
 	 */
-	bool			bExportLoadFailed;
-
-	// @todo: BP2CPP_remove
-	/**
-	 * Export is a dynamic type.
-	 */
-	enum class UE_DEPRECATED(5.0, "This type is no longer in use and will be removed.") EDynamicType : uint8
-	{
-		NotDynamicExport,
-		DynamicType,
-		ClassDefaultObject,
-	};
-
-	// @todo: BP2CPP_remove
-PRAGMA_DISABLE_DEPRECATION_WARNINGS
-	UE_DEPRECATED(5.0, "FObjectExport::DynamicType is no longer in use and will be removed.")
-	EDynamicType	DynamicType;
-PRAGMA_ENABLE_DEPRECATION_WARNINGS
+	bool			bExportLoadFailed:1;
 
 	/**
 	 * Export was filtered out on load
 	 */
-	bool			bWasFiltered;
-
-	/** If this object is a top level package (which must have been forced into the export table via OBJECTMARK_ForceTagExp)
-	 * this is the GUID for the original package file
-	 * Serialized
-	 */
-	UE_DEPRECATED(4.27, "UPackage::Guid has not been used by the engine for a long time and FObjectExport::PackageGuid will be removed.")
-	FGuid			PackageGuid;
+	bool			bWasFiltered:1;
 
 	/** If this object is a top level package (which must have been forced into the export table via OBJECTMARK_ForceTagExp)
 	 * this is the package flags for the original package file
