Add IsInheritedInstance to exports
Add missing generate public hash entries to ArchiveStackTrace

diff --git a/Engine/Source/Runtime/Core/Public/UObject/ObjectVersion.h b/Engine/Source/Runtime/Core/Public/UObject/ObjectVersion.h
--- a/Engine/Source/Runtime/Core/Public/UObject/ObjectVersion.h
+++ b/Engine/Source/Runtime/Core/Public/UObject/ObjectVersion.h
@@ -57,8 +57,11 @@ enum class EUnrealEngineObjectUE5Version : uint32
 	// Large world coordinates converts a number of core types to double components by default.
 	LARGE_WORLD_COORDINATES,       
 
-	// Remove package GUID
-	REMOVE_PACKAGE_GUID,
+	// Remove package GUID from FObjectExport
+	REMOVE_OBJECT_EXORT_PACKAGE_GUID,
+
+	// Add IsInherited to the FObjectExport entry
+	TRACK_OBJECT_EXPORT_IS_INHERITED,
 
 	// -----<new versions can be added before this line>-------------------------------------------------
 	// - this needs to be the last line (see note below)
diff --git a/Engine/Source/Runtime/CoreUObject/Private/Serialization/ArchiveStackTrace.cpp b/Engine/Source/Runtime/CoreUObject/Private/Serialization/ArchiveStackTrace.cpp
--- a/Engine/Source/Runtime/CoreUObject/Private/Serialization/ArchiveStackTrace.cpp
+++ b/Engine/Source/Runtime/CoreUObject/Private/Serialization/ArchiveStackTrace.cpp
@@ -1137,6 +1137,8 @@ bool CompareTableItem(FLinkerLoad* SourceLinker, FLinkerLoad* DestLinker, const
 		SourceExport.bNotForServer != DestExport.bNotForServer ||
 		SourceExport.bNotAlwaysLoadedForEditorGame != DestExport.bNotAlwaysLoadedForEditorGame ||
 		SourceExport.bIsAsset != DestExport.bIsAsset ||
+		SourceExport.bIsInheritedInstance != DestExport.bIsInheritedInstance ||
+		SourceExport.bGeneratePublicHash != DestExport.bGeneratePublicHash ||
 		!ComparePackageIndices(SourceLinker, DestLinker, SourceExport.TemplateIndex, DestExport.TemplateIndex) ||
 		!ComparePackageIndices(SourceLinker, DestLinker, SourceExport.OuterIndex, DestExport.OuterIndex) ||
 		!ComparePackageIndices(SourceLinker, DestLinker, SourceExport.ClassIndex, DestExport.ClassIndex) ||
@@ -1220,7 +1222,7 @@ bool ComparePackageIndices(FLinkerLoad* SourceLinker, FLinkerLoad* DestLinker, c
 FString ConvertItemToText(const FObjectExport& Export, FLinkerLoad* Linker)
 {
 	FName ClassName = Export.ClassIndex.IsNull() ? FName(NAME_Class) : Linker->ImpExp(Export.ClassIndex).ObjectName;
-	return FString::Printf(TEXT("%s Super: %s, Template: %s, Flags: %d, Size: %lld, PackageFlags: %d, ForcedExport: %d, NotForClient: %d, NotForServer: %d, NotAlwaysLoadedForEditorGame: %d, IsAsset: %d"),
+	return FString::Printf(TEXT("%s Super: %s, Template: %s, Flags: %d, Size: %lld, PackageFlags: %d, ForcedExport: %d, NotForClient: %d, NotForServer: %d, NotAlwaysLoadedForEditorGame: %d, IsAsset: %d, IsInheritedInstance: %d, GeneratePublicHash: %d"),
 		*GetTableKey(Linker, Export),
 		*GetTableKeyForIndex(Linker, Export.SuperIndex),
 		*GetTableKeyForIndex(Linker, Export.TemplateIndex),
@@ -1231,8 +1233,9 @@ FString ConvertItemToText(const FObjectExport& Export, FLinkerLoad* Linker)
 		Export.bNotForClient,
 		Export.bNotForServer,
 		Export.bNotAlwaysLoadedForEditorGame,
-		Export.bIsAsset
-	);
+		Export.bIsAsset,
+		Export.bIsInheritedInstance,
+		Export.bGeneratePublicHash);
 }
 
 static bool IsExportMapIdentical(FLinkerLoad* SourceLinker, FLinkerLoad* DestLinker)
diff --git a/Engine/Source/Runtime/CoreUObject/Private/UObject/LinkerLoad.cpp b/Engine/Source/Runtime/CoreUObject/Private/UObject/LinkerLoad.cpp
--- a/Engine/Source/Runtime/CoreUObject/Private/UObject/LinkerLoad.cpp
+++ b/Engine/Source/Runtime/CoreUObject/Private/UObject/LinkerLoad.cpp
@@ -3530,7 +3530,7 @@ bool FLinkerLoad::VerifyImportInner(const int32 ImportIndex, FString& WarningSuf
 
 UObject* FLinkerLoad::CreateExportAndPreload(int32 ExportIndex, bool bForcePreload /* = false */)
 {
-	UObject *Object = CreateExport(ExportIndex);
+	UObject* Object = CreateExport(ExportIndex);
 	if (Object && (bForcePreload || dynamic_cast<UClass*>(Object) || Object->IsTemplate() || dynamic_cast<UObjectRedirector*>(Object)))
 	{
 		Preload(Object);
@@ -4617,10 +4617,14 @@ UObject* FLinkerLoad::CreateExport( int32 Index )
 
 		// If loading the object's Outer caused the object to be loaded or if it was a forced export package created
 		// above, return it.
-		if( Export.Object != NULL )
+		if( Export.Object != nullptr )
 		{
 			return Export.Object;
 		}
+		else if (Export.bExportLoadFailed)
+		{
+			return nullptr;
+		}
 
 		// If we should have an outer but it doesn't exist because it was filtered out, we should silently be filtered out too
 		if (Export.OuterIndex.IsExport() && ThisParent == nullptr && ExportMap[Export.OuterIndex.ToExport()].bWasFiltered)
@@ -4801,6 +4805,17 @@ UObject* FLinkerLoad::CreateExport( int32 Index )
 			}
 		}
 
+		// If this export was an instance inherited from another class, 
+		// it should not be recreated unless the archetype also has it
+		if (Export.bIsInheritedInstance)
+		{
+			if (FindObjectWithOuter(Template, LoadClass, NewName) == nullptr)
+			{
+				Export.bExportLoadFailed = true;
+				return nullptr;
+			}
+		}
+
 		LoadClass->GetDefaultObject();
 
 		FStaticConstructObjectParameters Params(LoadClass);
diff --git a/Engine/Source/Runtime/CoreUObject/Private/UObject/ObjectResource.cpp b/Engine/Source/Runtime/CoreUObject/Private/UObject/ObjectResource.cpp
--- a/Engine/Source/Runtime/CoreUObject/Private/UObject/ObjectResource.cpp
+++ b/Engine/Source/Runtime/CoreUObject/Private/UObject/ObjectResource.cpp
@@ -45,6 +45,7 @@ FObjectExport::FObjectExport()
 , bNotForServer(false)
 , bNotAlwaysLoadedForEditorGame(true)
 , bIsAsset(false)
+, bIsInheritedInstance(false)
 , bGeneratePublicHash(false)
 , bExportLoadFailed(false)
 , bWasFiltered(false)
@@ -71,6 +72,7 @@ FObjectExport::FObjectExport( UObject* InObject, bool bInNotAlwaysLoadedForEdito
 , bNotForServer(false)
 , bNotAlwaysLoadedForEditorGame(bInNotAlwaysLoadedForEditorGame)
 , bIsAsset(false)
+, bIsInheritedInstance(false)
 , bGeneratePublicHash(false)
 , bExportLoadFailed(false)
 , bWasFiltered(false)
@@ -86,6 +88,7 @@ FObjectExport::FObjectExport( UObject* InObject, bool bInNotAlwaysLoadedForEdito
 		bNotForClient = !Object->NeedsLoadForClient();
 		bNotForServer = !Object->NeedsLoadForServer();
 		bIsAsset = Object->IsAsset();
+		bIsInheritedInstance = Object->GetArchetype()->IsDefaultSubobject();
 	}
 }
 
@@ -152,12 +155,17 @@ void operator<<(FStructuredArchive::FSlot Slot, FObjectExport& E)
 	SERIALIZE_BIT_TO_RECORD(bNotForClient);
 	SERIALIZE_BIT_TO_RECORD(bNotForServer);
 
-	if (BaseArchive.UEVer() < EUnrealEngineObjectUE5Version::REMOVE_PACKAGE_GUID)
+	if (BaseArchive.UEVer() < EUnrealEngineObjectUE5Version::REMOVE_OBJECT_EXORT_PACKAGE_GUID)
 	{
 		FGuid DummyPackageGuid;
 		Record << SA_VALUE(TEXT("PackageGuid"), DummyPackageGuid);
 	}
 
+	if (BaseArchive.UEVer() >= EUnrealEngineObjectUE5Version::TRACK_OBJECT_EXPORT_IS_INHERITED)
+	{
+		SERIALIZE_BIT_TO_RECORD(bIsInheritedInstance);
+	}
+
 	Record << SA_VALUE(TEXT("PackageFlags"), E.PackageFlags);
 
 	if (BaseArchive.UEVer() >= VER_UE4_LOAD_FOR_EDITOR_GAME)
diff --git a/Engine/Source/Runtime/CoreUObject/Public/UObject/ObjectResource.h b/Engine/Source/Runtime/CoreUObject/Public/UObject/ObjectResource.h
--- a/Engine/Source/Runtime/CoreUObject/Public/UObject/ObjectResource.h
+++ b/Engine/Source/Runtime/CoreUObject/Public/UObject/ObjectResource.h
@@ -309,6 +309,12 @@ struct FObjectExport : public FObjectResource
 	 */
 	bool			bIsAsset:1;
 
+	/**
+	 * If this export is an instanced object inherited from a template and should
+	 * only be created if the template has the object
+	 */
+	bool			bIsInheritedInstance:1;
+
 	/**
 	 * True if this export should have its iostore public hash generated even if not RF_Public.
 	 */
