Fix typo and failure to copy paste something while doing last minute name tweak
#rnx

diff --git a/Engine/Source/Runtime/Core/Public/UObject/ObjectVersion.h b/Engine/Source/Runtime/Core/Public/UObject/ObjectVersion.h
--- a/Engine/Source/Runtime/Core/Public/UObject/ObjectVersion.h
+++ b/Engine/Source/Runtime/Core/Public/UObject/ObjectVersion.h
@@ -58,7 +58,7 @@ enum class EUnrealEngineObjectUE5Version : uint32
 	LARGE_WORLD_COORDINATES,       
 
 	// Remove package GUID from FObjectExport
-	REMOVE_OBJECT_EXORT_PACKAGE_GUID,
+	REMOVE_OBJECT_EXPORT_PACKAGE_GUID,
 
 	// Add IsInherited to the FObjectExport entry
 	TRACK_OBJECT_EXPORT_IS_INHERITED,
diff --git a/Engine/Source/Runtime/CoreUObject/Private/UObject/ObjectResource.cpp b/Engine/Source/Runtime/CoreUObject/Private/UObject/ObjectResource.cpp
--- a/Engine/Source/Runtime/CoreUObject/Private/UObject/ObjectResource.cpp
+++ b/Engine/Source/Runtime/CoreUObject/Private/UObject/ObjectResource.cpp
@@ -155,7 +155,7 @@ void operator<<(FStructuredArchive::FSlot Slot, FObjectExport& E)
 	SERIALIZE_BIT_TO_RECORD(bNotForClient);
 	SERIALIZE_BIT_TO_RECORD(bNotForServer);
 
-	if (BaseArchive.UEVer() < EUnrealEngineObjectUE5Version::REMOVE_OBJECT_EXORT_PACKAGE_GUID)
+	if (BaseArchive.UEVer() < EUnrealEngineObjectUE5Version::REMOVE_OBJECT_EXPORT_PACKAGE_GUID)
 	{
 		FGuid DummyPackageGuid;
 		Record << SA_VALUE(TEXT("PackageGuid"), DummyPackageGuid);
@@ -257,7 +257,7 @@ void operator<<(FStructuredArchive::FSlot Slot, FObjectTextExport& E)
 	SERIALIZE_BIT_TO_SLOT(bNotForClient);
 	SERIALIZE_BIT_TO_SLOT(bNotForServer);
 
-	if (BaseArchive.UEVer() < EUnrealEngineObjectUE5Version::REMOVE_PACKAGE_GUID)
+	if (BaseArchive.UEVer() < EUnrealEngineObjectUE5Version::REMOVE_OBJECT_EXPORT_PACKAGE_GUID)
 	{
 		FGuid DummyPackageGuid;
 		Slot << SA_OPTIONAL_ATTRIBUTE(TEXT("PackageGuid"), DummyPackageGuid, FGuid());
