added EBlueprintExceptionType::AbortException, when raised it sets a flag on the FFrame to signal that script execution should stop and return ASAP. It is currently only checked in the interpreter loop inside ProcessLocalScriptFunction.

- FBlueprintContextTracker::GetScriptStack is deprecated
- FBlueprintContextTracker::GetCurrentScriptStack() now returns a TArrayView so that the view of the stack remains const
- FBlueprintContextTracker::GetCurrentScriptStackWritable() added for uses that explicitly need writable access to the stack
- FBlueprintCoreDelegates::ThrowScriptException now needs a non-const FFrame to be able to set the abort flag

diff --git a/Engine/Plugins/Animation/LiveLink/Source/LiveLink/Private/LiveLinkBlueprintLibrary.cpp b/Engine/Plugins/Animation/LiveLink/Source/LiveLink/Private/LiveLinkBlueprintLibrary.cpp
--- a/Engine/Plugins/Animation/LiveLink/Source/LiveLink/Private/LiveLinkBlueprintLibrary.cpp
+++ b/Engine/Plugins/Animation/LiveLink/Source/LiveLink/Private/LiveLinkBlueprintLibrary.cpp
@@ -304,7 +304,7 @@ namespace EvaluateLiveLinkFrame
 		SceneTime,
 	};
 
-	bool Generic_Evaluate(const ULiveLinkBlueprintLibrary* InSelf, const struct FFrame& InStack, FStructProperty* InBlueprintDataStructProp
+	bool Generic_Evaluate(const ULiveLinkBlueprintLibrary* InSelf, struct FFrame& InStack, FStructProperty* InBlueprintDataStructProp
 		, FLiveLinkSubjectRepresentation InSubjectRepresentation, EEvaluateType InEvaluationType, double InWorldTime, const FQualifiedFrameTime& InQualifiedFrameTime, void* OutBlueprintDataPtr)
 	{
 		bool bSuccess = false;
diff --git a/Engine/Source/Editor/Kismet/Private/CallStackViewer.cpp b/Engine/Source/Editor/Kismet/Private/CallStackViewer.cpp
--- a/Engine/Source/Editor/Kismet/Private/CallStackViewer.cpp
+++ b/Engine/Source/Editor/Kismet/Private/CallStackViewer.cpp
@@ -505,7 +505,7 @@ FReply SCallStackViewer::OnKeyDown( const FGeometry& MyGeometry, const FKeyEvent
 // Proxy array of the call stack. This allows us to manually refresh UI state when changes are made:
 TArray<TSharedRef<FCallStackRow>> Private_CallStackSource;
 
-void CallStackViewer::UpdateDisplayedCallstack(const TArray<const FFrame*>& ScriptStack)
+void CallStackViewer::UpdateDisplayedCallstack(TArrayView<const FFrame* const> ScriptStack)
 {
 	Private_CallStackSource.Empty();
 	if (ScriptStack.Num() > 0)
diff --git a/Engine/Source/Editor/Kismet/Private/WatchPointViewer.cpp b/Engine/Source/Editor/Kismet/Private/WatchPointViewer.cpp
--- a/Engine/Source/Editor/Kismet/Private/WatchPointViewer.cpp
+++ b/Engine/Source/Editor/Kismet/Private/WatchPointViewer.cpp
@@ -780,7 +780,7 @@ void WatchViewer::UpdateInstancedWatchDisplay()
 #if DO_BLUEPRINT_GUARD
 	{
 		Private_InstanceWatchSource.Reset();
-		const TArray<const FFrame*>& ScriptStack = FBlueprintContextTracker::Get().GetScriptStack();
+		TArrayView<const FFrame* const> ScriptStack = FBlueprintContextTracker::Get().GetCurrentScriptStack();
 
 		TSet<const UBlueprint*> SeenBlueprints;
 
diff --git a/Engine/Source/Editor/Kismet/Public/CallStackViewer.h b/Engine/Source/Editor/Kismet/Public/CallStackViewer.h
--- a/Engine/Source/Editor/Kismet/Public/CallStackViewer.h
+++ b/Engine/Source/Editor/Kismet/Public/CallStackViewer.h
@@ -10,7 +10,7 @@ struct FFrame;
 
 namespace CallStackViewer
 {
-	void KISMET_API UpdateDisplayedCallstack(const TArray<const FFrame*>& ScriptStack);
+	void KISMET_API UpdateDisplayedCallstack(TArrayView<const FFrame* const> ScriptStack);
 	FName GetTabName();
 	void RegisterTabSpawner(FTabManager& TabManager);
 }
diff --git a/Engine/Source/Editor/UnrealEd/Private/Kismet2/KismetDebugUtilities.cpp b/Engine/Source/Editor/UnrealEd/Private/Kismet2/KismetDebugUtilities.cpp
--- a/Engine/Source/Editor/UnrealEd/Private/Kismet2/KismetDebugUtilities.cpp
+++ b/Engine/Source/Editor/UnrealEd/Private/Kismet2/KismetDebugUtilities.cpp
@@ -143,7 +143,7 @@ void FKismetDebugUtilities::RequestSingleStepIn()
 void FKismetDebugUtilities::RequestStepOver()
 {
 	FKismetDebugUtilitiesData& Data = FKismetDebugUtilitiesData::Get();
-	const TArray<const FFrame*>& ScriptStack = FBlueprintContextTracker::Get().GetScriptStack();
+	TArrayView<const FFrame* const> ScriptStack = FBlueprintContextTracker::Get().GetCurrentScriptStack();
 
 	if (ScriptStack.Num() > 0)
 	{
@@ -188,7 +188,7 @@ void FKismetDebugUtilities::RequestStepOver()
 void FKismetDebugUtilities::RequestStepOut()
 {
 	FKismetDebugUtilitiesData& Data = FKismetDebugUtilitiesData::Get();
-	const TArray<const FFrame*>& ScriptStack = FBlueprintContextTracker::Get().GetScriptStack();
+	TArrayView<const FFrame* const> ScriptStack = FBlueprintContextTracker::Get().GetCurrentScriptStack();
 
 	Data.bIsSingleStepping = false;
 	if (ScriptStack.Num() > 1)
@@ -540,7 +540,7 @@ UEdGraphNode* FKismetDebugUtilities::FindSourceNodeForCodeLocation(const UObject
 void FKismetDebugUtilities::CheckBreakConditions(UEdGraphNode* NodeStoppedAt, bool bHitBreakpoint, int32 BreakpointOffset, bool& InOutBreakExecution)
 {
 	FKismetDebugUtilitiesData& Data = FKismetDebugUtilitiesData::Get();
-	const TArray<const FFrame*>& ScriptStack = FBlueprintContextTracker::Get().GetScriptStack();
+	TArrayView<const FFrame* const> ScriptStack = FBlueprintContextTracker::Get().GetCurrentScriptStack();
 
 	if (NodeStoppedAt)
 	{
@@ -731,7 +731,7 @@ void FKismetDebugUtilities::AttemptToBreakExecution(UBlueprint* BlueprintObj, co
 	if (bShouldInStackDebug)
 	{
 		TGuardValue<int32> GuardDisablePIE(GPlayInEditorID, INDEX_NONE);
-		const TArray<const FFrame*>& ScriptStack = FBlueprintContextTracker::Get().GetScriptStack();
+		TArrayView<const FFrame* const> ScriptStack = FBlueprintContextTracker::Get().GetCurrentScriptStack();
 		Data.LastExceptionMessage = Info.GetDescription();
 		FKismetEditorUtilities::BringKismetToFocusAttentionOnObject(NodeStoppedAt);
 		CallStackViewer::UpdateDisplayedCallstack(ScriptStack);
diff --git a/Engine/Source/Runtime/CoreUObject/Private/UObject/ScriptCore.cpp b/Engine/Source/Runtime/CoreUObject/Private/UObject/ScriptCore.cpp
--- a/Engine/Source/Runtime/CoreUObject/Private/UObject/ScriptCore.cpp
+++ b/Engine/Source/Runtime/CoreUObject/Private/UObject/ScriptCore.cpp
@@ -146,7 +146,7 @@ FBlueprintCoreDelegates::FOnToggleScriptProfiler FBlueprintCoreDelegates::OnTogg
 
 PRAGMA_ENABLE_DEPRECATION_WARNINGS
 
-void FBlueprintCoreDelegates::ThrowScriptException(const UObject* ActiveObject, const FFrame& StackFrame, const FBlueprintExceptionInfo& Info)
+void FBlueprintCoreDelegates::ThrowScriptException(const UObject* ActiveObject, FFrame& StackFrame, const FBlueprintExceptionInfo& Info)
 {
 	bool bShouldLogWarning = true;
 
@@ -186,6 +186,12 @@ void FBlueprintCoreDelegates::ThrowScriptException(const UObject* ActiveObject,
 		OnScriptException.Broadcast(ActiveObject, StackFrame, Info);
 	}
 
+	if (Info.GetType() == EBlueprintExceptionType::AbortExecution)
+	{
+		// abort errors halt further execution
+		StackFrame.bAbortingExecution = true;
+	}
+
 	if (Info.GetType() == EBlueprintExceptionType::FatalError)
 	{
 		// Crash maybe?
@@ -295,7 +301,7 @@ void PrintScriptCallStackImpl()
 	const FBlueprintContextTracker* BlueprintExceptionTracker = FBlueprintContextTracker::TryGet();
 	if (BlueprintExceptionTracker)
 	{
-		const TArray<const FFrame*>& RawStack = BlueprintExceptionTracker->GetScriptStack();
+		TArrayView<const FFrame* const> RawStack = BlueprintExceptionTracker->GetCurrentScriptStack();
 		FString ScriptStack = FString::Printf(TEXT("\n\nScript Stack (%d frames):\n"), RawStack.Num());
 		for (int32 FrameIdx = RawStack.Num() - 1; FrameIdx >= 0; --FrameIdx)
 		{
@@ -931,6 +937,9 @@ void ProcessScriptFunction(UObject* Context, UFunction* Function, FFrame& Stack,
 			}
 		}
 	}
+
+	// propagate abort flag up the stack
+	Stack.bAbortingExecution |= NewStack.bAbortingExecution;
 }
 
 DEFINE_FUNCTION(UObject::execCallMathFunction)
@@ -1090,7 +1099,7 @@ void ProcessLocalScriptFunction(UObject* Context, FFrame& Stack, RESULT_DECL)
 	}
 #endif
 	// Execute the bytecode
-	while (*Stack.Code != EX_Return)
+	while (*Stack.Code != EX_Return && !Stack.bAbortingExecution)
 	{
 #if DO_BLUEPRINT_GUARD
 		if(BpET.Runaway > GMaximumScriptLoopIterations )
@@ -1120,16 +1129,25 @@ void ProcessLocalScriptFunction(UObject* Context, FFrame& Stack, RESULT_DECL)
 		Stack.Step(Stack.Object, Buffer);
 	}
 
-	// Step over the return statement and evaluate the result expression
-	Stack.Code++;
-
-	if (*Stack.Code != EX_Nothing)
+	if (!Stack.bAbortingExecution)
 	{
-		Stack.Step(Stack.Object, RESULT_PARAM);
+		// Step over the return statement and evaluate the result expression
+		Stack.Code++;
+
+		if (*Stack.Code != EX_Nothing)
+		{
+			Stack.Step(Stack.Object, RESULT_PARAM);
+		}
+		else
+		{
+			Stack.Code++;
+		}
 	}
 	else
 	{
-		Stack.Code++;
+		// If we have a return property, return a zeroed value in it
+		FProperty* ReturnProp = (Function)->GetReturnProperty();
+		ClearReturnValue(ReturnProp, RESULT_PARAM);
 	}
 
 #if DO_BLUEPRINT_GUARD
diff --git a/Engine/Source/Runtime/CoreUObject/Public/UObject/Script.h b/Engine/Source/Runtime/CoreUObject/Public/UObject/Script.h
--- a/Engine/Source/Runtime/CoreUObject/Public/UObject/Script.h
+++ b/Engine/Source/Runtime/CoreUObject/Public/UObject/Script.h
@@ -331,6 +331,7 @@ namespace EBlueprintExceptionType
 		InfiniteLoop,
 		NonFatalError,
 		FatalError,
+		AbortExecution,
 	};
 }
 
@@ -479,7 +480,7 @@ public:
 	static FOnScriptExecutionEnd OnScriptExecutionEnd;
 
 public:
-	static void ThrowScriptException(const UObject* ActiveObject, const struct FFrame& StackFrame, const FBlueprintExceptionInfo& Info);
+	static void ThrowScriptException(const UObject* ActiveObject, struct FFrame& StackFrame, const FBlueprintExceptionInfo& Info);
 	static void InstrumentScriptEvent(const FScriptInstrumentationSignal& Info);
 	static void SetScriptMaximumLoopIterations( const int32 MaximumLoopIterations );
 	static bool IsDebuggingEnabled();
@@ -528,11 +529,23 @@ struct COREUOBJECT_API FBlueprintContextTracker : TThreadSingleton<FBlueprintCon
 	{
 		return ScriptEntryTag;
 	}
-	
+
+	/** Returns current script stack frame */
+	UE_DEPRECATED(5.1, "GetScriptStack() inefficiently copies the array to return and is now deprecated. Use GetCurrentScriptStack() which returns a TArrayView instead")
+	FORCEINLINE TArray<const FFrame*> GetScriptStack() const
+	{
+		return TArray<const FFrame*>(GetCurrentScriptStack());
+	}
+
 	/** Returns current script stack frame */
-	FORCEINLINE const TArray<const FFrame*>& GetScriptStack() const
+	FORCEINLINE TArrayView<const FFrame* const> GetCurrentScriptStack() const
+	{
+		return MakeArrayView<const FFrame* const>(ScriptStack.GetData(), ScriptStack.Num());
+	}
+
+	FORCEINLINE TArrayView<FFrame* const> GetCurrentScriptStackWritable() const
 	{
-		return ScriptStack;
+		return MakeArrayView<FFrame* const>(ScriptStack.GetData(), ScriptStack.Num());
 	}
 
 	/** Delegate called from EnterScriptContext, could be called on any thread! This can be used to detect entries into script from native code */
@@ -554,7 +567,7 @@ private:
 	int32 ScriptEntryTag;
 
 	// Stack pointers from the VM to be unrolled when we assert
-	TArray<const FFrame*> ScriptStack;
+	TArray<FFrame*> ScriptStack;
 
 	// Map of reported access warnings in exception handler
 	TMap<FName, int32> DisplayedWarningsMap;
diff --git a/Engine/Source/Runtime/CoreUObject/Public/UObject/Stack.h b/Engine/Source/Runtime/CoreUObject/Public/UObject/Stack.h
--- a/Engine/Source/Runtime/CoreUObject/Public/UObject/Stack.h
+++ b/Engine/Source/Runtime/CoreUObject/Public/UObject/Stack.h
@@ -113,6 +113,9 @@ public:
 
 	bool bArrayContextFailed;
 
+	/** If this flag gets set (usually from throwing a EBlueprintExceptionType::AbortExecution exception), execution shall immediately stop and return */
+	bool bAbortingExecution;
+
 #if PER_FUNCTION_SCRIPT_STATS
 	/** Increment for each PreviousFrame on the stack (Max 255) */
 	uint8 DepthCounter;
@@ -233,6 +236,7 @@ inline FFrame::FFrame( UObject* InObject, UFunction* InNode, void* InLocals, FFr
 	, PropertyChainForCompiledIn(InPropertyChainForCompiledIn)
 	, CurrentNativeFunction(nullptr)
 	, bArrayContextFailed(false)
+	, bAbortingExecution(false)
 #if PER_FUNCTION_SCRIPT_STATS
 	, DepthCounter(0)
 #endif
