Delete SavePackage1. It has been replaced by SavePackage2 and is no longer supported.

diff --git a/Engine/Source/Editor/UnrealEd/Classes/Commandlets/SavePackageUtilitiesCommandlet.h b/Engine/Source/Editor/UnrealEd/Classes/Commandlets/SavePackageUtilitiesCommandlet.h
--- a/Engine/Source/Editor/UnrealEd/Classes/Commandlets/SavePackageUtilitiesCommandlet.h
+++ b/Engine/Source/Editor/UnrealEd/Classes/Commandlets/SavePackageUtilitiesCommandlet.h
@@ -10,7 +10,7 @@
 
 /*
  * Commandlet used to validate the package saving mechanism. 
- * It can currently compare the saving of a package (or folder of packages) between UPackage::Save and UPackage::Save2
+ * It can currently compare two back to back saves of a package (or folder of packages)
  */
 UCLASS()
 class USavePackageUtilitiesCommandlet
diff --git a/Engine/Source/Editor/UnrealEd/Private/Commandlets/SavePackageUtilitiesCommandlet.cpp b/Engine/Source/Editor/UnrealEd/Private/Commandlets/SavePackageUtilitiesCommandlet.cpp
--- a/Engine/Source/Editor/UnrealEd/Private/Commandlets/SavePackageUtilitiesCommandlet.cpp
+++ b/Engine/Source/Editor/UnrealEd/Private/Commandlets/SavePackageUtilitiesCommandlet.cpp
@@ -55,44 +55,13 @@ int32 USavePackageUtilitiesCommandlet::Main(const FString& Params)
 			SaveArgs.TopLevelFlags |= RF_Standalone;
 		}
 
-		static IConsoleVariable* EnableNewSave = IConsoleManager::Get().FindConsoleVariable(TEXT("SavePackage.EnableNewSave"));
-		int32 EnableNewSavePreviousValue = EnableNewSave->GetInt();
-
-		// Do the new save package first in case, number of serialization has a by product during saving
-		// New Save Package
-		FSavePackageResultStruct NewResult;
-		{
-			EnableNewSave->Set(3); // Enable new save cooked and uncooked data
-			NewResult = GEditor->Save(Package, Asset, *Filename, SaveArgs);
-		}
-
-		// Old Save Package
-		FSavePackageResultStruct OldResult;
-		{
-			EnableNewSave->Set(0);
-			OldResult = GEditor->Save(Package, Asset, *Filename, SaveArgs);
-		}
-
-		// Old Save Package
-		FSavePackageResultStruct OldResultCheck;
-		{
-			EnableNewSave->Set(0);
-			OldResultCheck = GEditor->Save(Package, Asset, *Filename, SaveArgs);
-		}
-		EnableNewSave->Set(EnableNewSavePreviousValue);
-
-		if (OldResult.LinkerSave && NewResult.LinkerSave)
-		{
-			// Compare Linker Save info
-			FLinkerDiff LinkerDiff = FLinkerDiff::CompareLinkers(OldResult.LinkerSave.Get(), NewResult.LinkerSave.Get());
-			LinkerDiff.PrintDiff(*GWarn);
-		}
-
-		// Add a old save against itself check to test potential byproduct, doesn't catch them all, since oftentimes byproduct are caused by the first save
-		if (OldResultCheck.LinkerSave && OldResult.LinkerSave)
+		FSavePackageResultStruct FirstResult = GEditor->Save(Package, Asset, *Filename, SaveArgs);
+		FSavePackageResultStruct SecondResult = GEditor->Save(Package, Asset, *Filename, SaveArgs);
+	
+		if (FirstResult.LinkerSave && SecondResult.LinkerSave)
 		{
 			// Compare Linker Save info
-			FLinkerDiff LinkerDiff = FLinkerDiff::CompareLinkers(OldResultCheck.LinkerSave.Get(), OldResult.LinkerSave.Get());
+			FLinkerDiff LinkerDiff = FLinkerDiff::CompareLinkers(FirstResult.LinkerSave.Get(), SecondResult.LinkerSave.Get());
 			LinkerDiff.PrintDiff(*GWarn);
 		}
 
diff --git a/Engine/Source/Editor/UnrealEd/Private/Cooker/DiffPackageWriter.cpp b/Engine/Source/Editor/UnrealEd/Private/Cooker/DiffPackageWriter.cpp
--- a/Engine/Source/Editor/UnrealEd/Private/Cooker/DiffPackageWriter.cpp
+++ b/Engine/Source/Editor/UnrealEd/Private/Cooker/DiffPackageWriter.cpp
@@ -295,18 +295,7 @@ FLinkerDiffPackageWriter::FLinkerDiffPackageWriter(TUniquePtr<ICookedPackageWrit
 {
 	FString DiffModeText;
 	FParse::Value(FCommandLine::Get(), TEXT("-LINKERDIFF="), DiffModeText);
-	DiffMode = EDiffMode::LDM_Algo;
-	if (DiffModeText == TEXT("1") || DiffModeText == TEXT("algo"))
-	{
-		DiffMode = EDiffMode::LDM_Algo;
-	}
-	else if (DiffModeText == TEXT("2") || DiffModeText == TEXT("consistent"))
-	{
-		DiffMode = EDiffMode::LDM_Consistent;
-	}
-
-	EnableNewSave = IConsoleManager::Get().FindConsoleVariable(TEXT("SavePackage.EnableNewSave"));
-	CurrentEnableNewSaveValue = EnableNewSave->GetInt();
+	DiffMode = EDiffMode::LDM_Consistent; // (DiffModeText == TEXT("2") || DiffModeText == TEXT("consistent"))
 }
 
 void FLinkerDiffPackageWriter::BeginPackage(const FBeginPackageInfo& Info)
@@ -350,10 +339,6 @@ bool FLinkerDiffPackageWriter::IsAnotherSaveNeeded(FSavePackageResultStruct& Pre
 	{
 		CompareResults(PreviousResult);
 
-		if (DiffMode == EDiffMode::LDM_Algo)
-		{
-			EnableNewSave->Set(CurrentEnableNewSaveValue);
-		}
 		OtherResult.LinkerSave.Reset();
 		PreviousResult.LinkerSave.Reset();
 
@@ -363,21 +348,10 @@ bool FLinkerDiffPackageWriter::IsAnotherSaveNeeded(FSavePackageResultStruct& Pre
 
 void FLinkerDiffPackageWriter::SetupOtherAlgorithm()
 {
-	// If mode is comparing the two save algorithms, switch the cvar to the other algo before saving again,
-	// Otherwise the linker diff mode is tracking if the save is consistent across multiple save
-	if (DiffMode == EDiffMode::LDM_Algo)
-	{
-		// see CVarEnablePackageNewSave definition in SavePackageUtilities.cpp for value meaning
-		EnableNewSave->Set(CurrentEnableNewSaveValue & 1 ? 0 : 1);
-	}
 }
 
 void FLinkerDiffPackageWriter::SetupCurrentAlgorithm()
 {
-	if (DiffMode == EDiffMode::LDM_Algo)
-	{
-		EnableNewSave->Set(CurrentEnableNewSaveValue);
-	}
 }
 
 void FLinkerDiffPackageWriter::CompareResults(FSavePackageResultStruct& CurrentResult)
diff --git a/Engine/Source/Editor/UnrealEd/Private/Cooker/DiffPackageWriter.h b/Engine/Source/Editor/UnrealEd/Private/Cooker/DiffPackageWriter.h
--- a/Engine/Source/Editor/UnrealEd/Private/Cooker/DiffPackageWriter.h
+++ b/Engine/Source/Editor/UnrealEd/Private/Cooker/DiffPackageWriter.h
@@ -224,7 +224,6 @@ public:
 private:
 	enum class EDiffMode : uint8
 	{
-		LDM_Algo,
 		LDM_Consistent,
 	};
 
@@ -235,8 +234,6 @@ private:
 	FSavePackageResultStruct OtherResult;
 	FBeginPackageInfo BeginInfo;
 	TUniquePtr<ICookedPackageWriter> Inner;
-	IConsoleVariable* EnableNewSave = nullptr;
-	int32 CurrentEnableNewSaveValue = 0;
-	EDiffMode DiffMode = EDiffMode::LDM_Algo;
+	EDiffMode DiffMode = EDiffMode::LDM_Consistent;
 	bool bHasStartedSecondSave = false;
 };
\ No newline at end of file
diff --git a/Engine/Source/Editor/UnrealEd/Private/EditorEngine.cpp b/Engine/Source/Editor/UnrealEd/Private/EditorEngine.cpp
--- a/Engine/Source/Editor/UnrealEd/Private/EditorEngine.cpp
+++ b/Engine/Source/Editor/UnrealEd/Private/EditorEngine.cpp
@@ -108,7 +108,6 @@
 #include "ScopedTransaction.h"
 
 #include "ILocalizationServiceModule.h"
-#include "PackageBackup.h"
 #include "Engine/LevelStreaming.h"
 #include "LevelUtils.h"
 #include "Layers/LayersSubsystem.h"
@@ -703,7 +702,6 @@ void UEditorEngine::InitEditor(IEngineLoop* InEngineLoop)
 
 	// Setup delegate callbacks for SavePackage()
 	FCoreUObjectDelegates::IsPackageOKToSaveDelegate.BindUObject(this, &UEditorEngine::IsPackageOKToSave);
-	FCoreUObjectDelegates::AutoPackageBackupDelegate.BindStatic(&FAutoPackageBackup::BackupPackage);
 
 	// Update recents
 	UpdateRecentlyLoadedProjectFiles();
diff --git a/Engine/Source/Editor/UnrealEd/Public/PackageBackup.h b/Engine/Source/Editor/UnrealEd/Public/PackageBackup.h
--- a/Engine/Source/Editor/UnrealEd/Public/PackageBackup.h
+++ b/Engine/Source/Editor/UnrealEd/Public/PackageBackup.h
@@ -13,7 +13,7 @@
  * to be backed-up to a special auto save/backup directory. Backups are traditionally
  * made before a save occurs to a package, to help guard against corruption/data loss.
  */
-class FAutoPackageBackup
+class UE_DEPRECATED(5.1, "Backups are no longer used in Unreal Package Saves") FAutoPackageBackup
 {
 public:
 	
diff --git a/Engine/Source/Runtime/CoreUObject/Private/UObject/SavePackage.cpp b/Engine/Source/Runtime/CoreUObject/Private/UObject/SavePackage.cpp
--- a/Engine/Source/Runtime/CoreUObject/Private/UObject/SavePackage.cpp
+++ b/Engine/Source/Runtime/CoreUObject/Private/UObject/SavePackage.cpp
@@ -3,4351 +3,38 @@
 #include "UObject/SavePackage.h"
 
 #if UE_WITH_SAVEPACKAGE
-#include "UObject/SavePackage/SavePackageUtilities.h"
-#include "CoreMinimal.h"
-#include "Misc/MessageDialog.h"
-#include "HAL/FileManager.h"
-#include "HAL/ThreadSafeCounter.h"
-#include "Misc/CommandLine.h"
-#include "Misc/FileHelper.h"
-#include "Misc/Paths.h"
-#include "Misc/Guid.h"
-#include "Stats/Stats.h"
-#include "Async/AsyncWork.h"
-#include "Serialization/LargeMemoryWriter.h"
-#include "Serialization/LargeMemoryReader.h"
-#include "Serialization/BufferArchive.h"
-#include "Misc/AssetRegistryInterface.h"
-#include "Misc/FeedbackContext.h"
-#include "Misc/RedirectCollector.h"
-#include "Misc/ScopedSlowTask.h"
-#include "Misc/ObjectThumbnail.h"
-#include "UObject/ObjectMacros.h"
+#include "Misc/Optional.h"
 #include "UObject/ArchiveCookContext.h"
-#include "UObject/UObjectGlobals.h"
-#include "UObject/UObjectHash.h"
-#include "UObject/Object.h"
-#include "UObject/ObjectSaveContext.h"
-#include "Serialization/ArchiveUObject.h"
-#include "UObject/Class.h"
-#include "UObject/UObjectIterator.h"
-#include "Serialization/PropertyLocalizationDataGathering.h"
-#include "UObject/Package.h"
-#include "Templates/Casts.h"
-#include "UObject/LazyObjectPtr.h"
-#include "UObject/SoftObjectPtr.h"
-#include "UObject/PropertyPortFlags.h"
-#include "UObject/UnrealType.h"
-#include "UObject/TextProperty.h"
-#include "UObject/ObjectRedirector.h"
-#include "Misc/PackageName.h"
-#include "Serialization/BulkData.h"
-#include "UObject/PackageFileSummary.h"
-#include "UObject/ObjectResource.h"
-#include "UObject/Linker.h"
-#include "UObject/LinkerLoad.h"
-#include "UObject/LinkerSave.h"
-#include "UObject/EditorObjectVersion.h"
-#include "Blueprint/BlueprintSupport.h"
-#include "Internationalization/TextPackageNamespaceUtil.h"
-#include "UObject/Interface.h"
-#include "Interfaces/ITargetPlatform.h"
-#include "UObject/UObjectThreadContext.h"
-#include "UObject/GCScopeLock.h"
-#include "ProfilingDebugging/CookStats.h"
-#include "UObject/DebugSerializationFlags.h"
-#include "UObject/EnumProperty.h"
-#include "HAL/IConsoleManager.h"
-#include "Serialization/ArchiveStackTrace.h"
-#include "UObject/CoreRedirects.h"
-#include "Serialization/ArchiveObjectCrc32.h"
-#include "Serialization/Formatters/BinaryArchiveFormatter.h"
-#include "Serialization/Formatters/JsonArchiveOutputFormatter.h"
-#include "Serialization/ArchiveUObjectFromStructuredArchive.h"
-#include "Serialization/PackageWriter.h"
-#include "Serialization/UnversionedPropertySerialization.h"
-#include "Serialization/EditorBulkData.h"
-#include "Misc/ScopeExit.h"
-#include "Misc/PackageAccessTracking.h"
-#include "Misc/PackageAccessTrackingOps.h"
 
-#if ENABLE_COOK_STATS
-#include "ProfilingDebugging/ScopedTimers.h"
-#endif
-
-static FCriticalSection InitializeCoreClassesCritSec;
-
-#define VALIDATE_INITIALIZECORECLASSES 0
-#define EXPORT_SORTING_DETAILED_LOGGING 0
-
-// bring the UObectGlobal declaration visible to non editor
-bool IsEditorOnlyObject(const UObject* InObject, bool bCheckRecursive, bool bCheckMarks);
-
-static bool EndSavingIfCancelled() 
-{
-	if ( GWarn->ReceivedUserCancel() )
-	{
-		return true;
-	}
-	return false;
-}
-
-class FPackageNameMapSaver
-{
-public:
-	void MarkNameAsReferenced(FName Name)
-	{
-		ReferencedNames.Add(Name.GetDisplayIndex());
-	}
-
-	void MarkNameAsReferenced(FNameEntryId Name)
-	{
-		ReferencedNames.Add(Name);
-	}
-
-	bool NameExists(FNameEntryId ComparisonId) const
-	{
-		for (FNameEntryId DisplayId : ReferencedNames)
-		{
-			if (FName::GetComparisonIdFromDisplayId(DisplayId) == ComparisonId)
-			{
-				return true;
-			}
-		}
-
-		return false;
-	}
-
-	void UpdateLinker(FLinkerSave& Linker, FArchive* BinarySaver);
-
-private:
-	TSet<FNameEntryId> ReferencedNames;
-};
-
-/**
- * Archive for tagging objects and names that must be exported
- * to the file.  It tags the objects passed to it, and recursively
- * tags all of the objects this object references.
- */
-class FArchiveSaveTagExports : public FArchiveUObject
-{
-public:
-	/**
-	 * Constructor
-	 * 
-	 * @param	InOuter		the package to save
-	 */
-	FArchiveSaveTagExports( UPackage* InOuter )
-	: Outer(InOuter)
-	{
-		this->SetIsSaving(true);
-		this->SetIsPersistent(true);
-		ArIsObjectReferenceCollector = true;
-		ArShouldSkipBulkData	= true;
-	}
-
-	void ProcessBaseObject(UObject* BaseObject);
-	virtual FArchive& operator<<(UObject*& Obj) override;
-	virtual FArchive& operator<<(FWeakObjectPtr& Value) override;
-	virtual void SetSerializeContext(FUObjectSerializeContext* InLoadContext) override
-	{
-		LoadContext = InLoadContext;
-	}
-	virtual FUObjectSerializeContext* GetSerializeContext() override
-	{
-		return LoadContext;
-	}
-	/**
-	 * Package we're currently saving.  Only objects contained
-	 * within this package will be tagged for serialization.
-	 */
-	UPackage* Outer;
-
-	virtual FString GetArchiveName() const;
-
-private:
-
-	TArray<UObject*> TaggedObjects;
-	TRefCountPtr<FUObjectSerializeContext> LoadContext;
-
-	void ProcessTaggedObjects();
-};
-
-FString FArchiveSaveTagExports::GetArchiveName() const
-{
-	return Outer != nullptr
-		? *FString::Printf(TEXT("SaveTagExports (%s)"), *Outer->GetName())
-		: TEXT("SaveTagExports");
-}
-
-FArchive& FArchiveSaveTagExports::operator<<(FWeakObjectPtr& Value)
-{
-	if (IsCooking())
-	{
-		// Always serialize weak pointers for the purposes of object tagging
-		UObject* Object = static_cast<UObject*>(Value.Get(true));
-		*this << Object;
-	}
-	else
-	{
-		FArchiveUObject::SerializeWeakObjectPtr(*this, Value);
-	}
-	return *this;
-}
-
-FArchive& FArchiveSaveTagExports::operator<<(UObject*& Obj)
-{
-	if (!Obj || Obj->HasAnyMarks(OBJECTMARK_TagExp) || Obj->HasAnyFlags(RF_Transient) || !Obj->IsInPackage(Outer))
-	{
-		return *this;
-	}
-
-	check(Outer);
-
-	// Check transient and pending kill flags for outers
-	SavePackageUtilities::CheckObjectPriorToSave(*this, Obj, Outer);
-
-	// The object may have become transient in CheckObjectPriorToSave
-	if (Obj->HasAnyFlags(RF_Transient))
-	{
-		return *this;
-	}
-
-	// Check outer chain for any exlcuded object marks
-	const EObjectMark ExcludedObjectMarks = SavePackageUtilities::GetExcludedObjectMarksForTargetPlatform(CookingTarget());
-	SavePackageUtilities::ConditionallyExcludeObjectForTarget(Obj, ExcludedObjectMarks, CookingTarget());
-
-	if (!Obj->HasAnyMarks((EObjectMark)(ExcludedObjectMarks)))
-	{
-		// It passed filtering so mark as export
-		Obj->Mark(OBJECTMARK_TagExp);
-
-		// First, serialize this object's archetype 
-		UObject* Template = Obj->GetArchetype();
-		*this << Template;
-
-		// If this is a CDO, gather it's subobjects and serialize them
-		if (Obj->HasAnyFlags(RF_ClassDefaultObject))
-		{
-			if (IsCooking())
-			{
-				// Gets all subobjects defined in a class, including the CDO, CDO components and blueprint-created components
-				TArray<UObject*> ObjectTemplates;
-				ObjectTemplates.Add(Obj);
-
-				SavePackageUtilities::GetCDOSubobjects(Obj, ObjectTemplates);
-
-				for (UObject* ObjTemplate : ObjectTemplates)
-				{
-					// Recurse into templates
-					*this << ObjTemplate;
-				}
-			}
-		}
-	
-		// NeedsLoadForEditor game is inherited to child objects, so check outer chain
-		bool bNeedsLoadForEditorGame = false;
-		for (UObject* OuterIt = Obj; OuterIt; OuterIt = OuterIt->GetOuter())
-		{
-			if (OuterIt->NeedsLoadForEditorGame())
-			{
-				bNeedsLoadForEditorGame = true;
-				break;
-			}
-		}
-
-		if(!bNeedsLoadForEditorGame && Obj->HasAnyFlags(RF_ClassDefaultObject))
-		{
-			bNeedsLoadForEditorGame = Obj->GetClass()->NeedsLoadForEditorGame();
-		}
-
-		if (!bNeedsLoadForEditorGame)
-		{
-			Obj->Mark(OBJECTMARK_NotAlwaysLoadedForEditorGame);
-		}
-
-		// Recurse with this object's class and package.
-		UObject* Class  = Obj->GetClass();
-		UObject* Parent = Obj->GetOuter();
-		*this << Class << Parent;
-
-		TaggedObjects.Add(Obj);
-	}
-	return *this;
-}
-
-/**
- * Serializes the specified object, tagging all objects it references.
- *
- * @param	BaseObject	the object that should be serialized; usually the package root or
- *						[in the case of a map package] the map's UWorld object.
- */
-void FArchiveSaveTagExports::ProcessBaseObject(UObject* BaseObject)
-{
-	(*this) << BaseObject;
-	ProcessTaggedObjects();
-}
-
-/**
- * Iterates over all objects which were encountered during serialization of the root object, serializing each one in turn.
- * Objects encountered during that serialization are then added to the array and iteration continues until no new objects are
- * added to the array.
- */
-void FArchiveSaveTagExports::ProcessTaggedObjects()
-{
-	const int32 ArrayPreSize = 1024; // Was originally total number of objects, but this was unreasonably large
-	TArray<UObject*> CurrentlyTaggedObjects;
-	CurrentlyTaggedObjects.Empty(ArrayPreSize);
-	while (TaggedObjects.Num())
-	{
-		CurrentlyTaggedObjects += TaggedObjects;
-		TaggedObjects.Empty();
-
-		for (int32 ObjIndex = 0; ObjIndex < CurrentlyTaggedObjects.Num(); ObjIndex++)
-		{
-			UObject* Obj = CurrentlyTaggedObjects[ObjIndex];
-
-			if (Obj->HasAnyFlags(RF_ClassDefaultObject))
-			{
-				Obj->GetClass()->SerializeDefaultObject(Obj, *this);
-			}
-			// In the CDO case the above would serialize most of the references, including transient properties
-			// but we still want to serialize the object using the normal path to collect all custom versions it might be using.
-			Obj->Serialize(*this);
-		}
-
-		CurrentlyTaggedObjects.Empty(ArrayPreSize);
-	}
-}
-
-/**
- * Archive for tagging objects and names that must be listed in the
- * file's imports table.
- */
-class FArchiveSaveTagImports : public FArchiveUObject
-{
-public:
-	FLinkerSave* Linker;
-	FPackageNameMapSaver& NameMapSaver;
-	TArray<UObject*> Dependencies;
-	TArray<UObject*> NativeDependencies;
-	TArray<UObject*> OtherImports;
-	bool bIgnoreDependencies;
-	bool bReferencerIsEditorOnly;
-	TSet<UObject*>& ImportsUsedInGame;
-	TSet<FName>& SoftPackagesUsedInGame;
-
-	/** Helper object to save/store state of bIgnoreDependencies */
-	class FScopeIgnoreDependencies
-	{
-		FArchiveSaveTagImports& Archive;
-		bool bScopedIgnoreDependencies;
-		
-	public:
-		FScopeIgnoreDependencies(FArchiveSaveTagImports& InArchive)
-			: Archive(InArchive)
-			, bScopedIgnoreDependencies(InArchive.bIgnoreDependencies)
-		{
-			Archive.bIgnoreDependencies = true;
-		}
-		~FScopeIgnoreDependencies()
-		{
-			Archive.bIgnoreDependencies = bScopedIgnoreDependencies;
-		}
-	};
-
-	FArchiveSaveTagImports(FLinkerSave* InLinker, FPackageNameMapSaver& InNameMapSaver, TSet<UObject*>& InImportsUsedInGame, TSet<FName>& InSoftPackagesUsedInGame, bool bInReferencerIsEditorOnly)
-		: Linker(InLinker)
-		, NameMapSaver(InNameMapSaver)
-		, bIgnoreDependencies(false)
-		, bReferencerIsEditorOnly(bInReferencerIsEditorOnly)
-		, ImportsUsedInGame(InImportsUsedInGame)
-		, SoftPackagesUsedInGame(InSoftPackagesUsedInGame)
-	{
-		check(Linker);
-
-		this->SetIsSaving(true);
-		this->SetIsPersistent(true);
-		ArIsObjectReferenceCollector = true;
-		ArShouldSkipBulkData	= true;
-
-		ArPortFlags = Linker->GetPortFlags();
-		SetCookData(Linker->GetCookData());
-	}
-
-	virtual FArchive& operator<<(UObject*& Obj) override;
-	virtual FArchive& operator<< (struct FWeakObjectPtr& Value) override;
-	virtual FArchive& operator<<(FLazyObjectPtr& LazyObjectPtr) override;
-	virtual FArchive& operator<<(FSoftObjectPath& Value) override;
-	virtual FArchive& operator<<(FName& Name) override;
-	
-	virtual void MarkSearchableName(const UObject* TypeObject, const FName& ValueName) const override;
-	virtual FString GetArchiveName() const override;
-	virtual void SetSerializeContext(FUObjectSerializeContext* InLoadContext) override
-	{
-		LoadContext = InLoadContext;
-	}
-	virtual FUObjectSerializeContext* GetSerializeContext() override
-	{
-		return LoadContext;
-	}
-
-private:
-
-	TRefCountPtr<FUObjectSerializeContext> LoadContext;
-};
-
-FString FArchiveSaveTagImports::GetArchiveName() const
-{
-	if ( Linker != nullptr && Linker->LinkerRoot )
-	{
-		return FString::Printf(TEXT("SaveTagImports (%s)"), *Linker->LinkerRoot->GetName());
-	}
-
-	return TEXT("SaveTagImports");
-}
-
-FArchive& FArchiveSaveTagImports::operator<< (struct FWeakObjectPtr& Value)
-{
-	if (IsCooking())
-	{
-		// Always serialize weak pointers for the purposes of object tagging
-		UObject* Object = static_cast<UObject*>(Value.Get(true));
-		*this << Object;
-	}
-	else
-	{
-		FArchiveUObject::SerializeWeakObjectPtr(*this, Value);
-	}
-	return *this;
-}
-
-FArchive& FArchiveSaveTagImports::operator<<( UObject*& Obj )
-{
-	// Check transient and pending kill flags for outers
-	SavePackageUtilities::CheckObjectPriorToSave(*this, Obj, nullptr);
-
-	const EObjectMark ExcludedObjectMarks = SavePackageUtilities::GetExcludedObjectMarksForTargetPlatform( CookingTarget());
-	SavePackageUtilities::ConditionallyExcludeObjectForTarget(Obj, ExcludedObjectMarks, CookingTarget());
-	bool bExcludePackageFromCook = Obj && FCoreUObjectDelegates::ShouldCookPackageForPlatform.IsBound() ? !FCoreUObjectDelegates::ShouldCookPackageForPlatform.Execute(Obj->GetOutermost(), CookingTarget()) : false;
-
-	// Skip PendingKill objects and objects that don't pass the platform mark filter
-	if (Obj && (ExcludedObjectMarks == OBJECTMARK_NOMARKS || !Obj->HasAnyMarks(ExcludedObjectMarks)) && !bExcludePackageFromCook)
-	{
-		bool bIsNative = Obj->IsNative();
-		if( !Obj->HasAnyFlags(RF_Transient) || bIsNative)
-		{
-			const bool bIsTopLevelPackage = Obj->GetOuter() == nullptr && dynamic_cast<UPackage*>(Obj);
-			UObject* Outer = Obj->GetOuter();
+#if WITH_EDITOR
 
-			// See if this is inside a native class
-			while (!bIsNative && Outer)
-			{
-				if (dynamic_cast<UClass*>(Outer) && Outer->IsNative())
-				{
-					bIsNative = true;
-				}
-				Outer = Outer->GetOuter();
-			}
+COREUOBJECT_API extern bool GOutputCookingWarnings;
 
-			bool bIsImport = !Obj->HasAnyMarks(OBJECTMARK_TagExp);
-			if (bIsImport)
-			{
-#if WITH_EDITORONLY_DATA
-				if (!bReferencerIsEditorOnly && !IsEditorOnlyPropertyOnTheStack())
 #endif
-				{
-					// Check every reference to an import for whether the reference is used-in-game and upgrade the _Package_'s reference to used-in-game. Do this before we early exit for imports we have seen before
-					ImportsUsedInGame.Add(Obj);
-				}
-			}
-
-			// We add objects as dependencies even if they're also exports
-			if (!bIsTopLevelPackage && !bIgnoreDependencies)
-			{
-				TArray<UObject*>& DependencyArray = bIsNative ? NativeDependencies : Dependencies;
-				if (DependencyArray.Contains(Obj))
-				{
-					return *this;
-				}
-				DependencyArray.Add(Obj);
-			}
-			
-			if (bIsImport)
-			{
-				// Add into other imports list unless it's already there
-				if (bIsTopLevelPackage || bIgnoreDependencies)
-				{
-					if (OtherImports.Contains(Obj))
-					{
-						return *this;
-					}
-
-					OtherImports.Add(Obj);
-				}
-
-				// Mark this object as an import
-				Obj->Mark(OBJECTMARK_TagImp);
-
-				UClass* ClassObj = Cast<UClass>(Obj);
-
-				// Don't recurse into CDOs if we're already ignoring dependencies, we only want to recurse into our outer chain in that case
-				if (IsCooking() && !bIsNative && !bIgnoreDependencies && ClassObj)
-				{
-					// We don't want to add this to Dependencies, we simply want it to be an import so that a serialization before creation dependency can be created to the CDO
-					FScopeIgnoreDependencies IgnoreDependencies(*this);
-					UObject* CDO = ClassObj->GetDefaultObject();
-
-					if (CDO)
-					{
-						// Gets all subobjects defined in a class, including the CDO, CDO components and blueprint-created components
-						TArray<UObject*> ObjectTemplates;
-						ObjectTemplates.Add(CDO);
-
-						SavePackageUtilities::GetCDOSubobjects(CDO, ObjectTemplates);
-
-						for (UObject* ObjTemplate : ObjectTemplates)
-						{
-							// Recurse into templates
-							*this << ObjTemplate;
-						}
-					}
-				}
-
-				// Recurse into parent
-				UObject* Parent = Obj->GetOuter();
-				if( Parent )
-				{
-					*this << Parent;
-				}
-
-				// if the object has a non null package set, recurse into it
-				UPackage* Package = Obj->GetExternalPackage();
-				if (Package && Package != Obj)
-				{
-					*this << Package;
-				}
-				else
-				{
-					if (!IsFilterEditorOnly())
-					{
-						// operator<<(FStructuredArchive::FSlot Slot, FObjectImport& I) will need to write NAME_None for this empty ExternalPackage pointer
-						*this << NAME_None;
-					}
-				}
-
-				// For things with a BP-created class we need to recurse into that class so the import ClassPackage will load properly
-				// We don't do this for native classes to avoid bloating the import table
-				UClass* ObjClass = Obj->GetClass();
-
-				if (!ObjClass->IsNative())
-				{
-					*this << ObjClass;
-				}
-			}
-		}
-	}
-	return *this;
-}
-
-FArchive& FArchiveSaveTagImports::operator<<( FLazyObjectPtr& LazyObjectPtr)
-{
-	FUniqueObjectGuid ID;
-	ID = LazyObjectPtr.GetUniqueID();
-	return *this << ID;
-}
-
-FArchive& FArchiveSaveTagImports::operator<<(FSoftObjectPath& Value)
-{
-	if (Value.IsValid())
-	{
-		Value.SerializePath(*this);
-
-		FSoftObjectPathThreadContext& ThreadContext = FSoftObjectPathThreadContext::Get();
-		FName ReferencingPackageName, ReferencingPropertyName;
-		ESoftObjectPathCollectType CollectType = ESoftObjectPathCollectType::AlwaysCollect;
-		ESoftObjectPathSerializeType SerializeType = ESoftObjectPathSerializeType::AlwaysSerialize;
-
-		ThreadContext.GetSerializationOptions(ReferencingPackageName, ReferencingPropertyName, CollectType, SerializeType, this);
-
-		if (CollectType != ESoftObjectPathCollectType::NeverCollect && CollectType != ESoftObjectPathCollectType::NonPackage)
-		{
-			// Don't track if this is a never collect path
-			FString Path = Value.ToString();
-			FName PackageName = FName(*FPackageName::ObjectPathToPackageName(Path));
-			NameMapSaver.MarkNameAsReferenced(PackageName);
-			Linker->SoftPackageReferenceList.AddUnique(PackageName);
-			if (CollectType != ESoftObjectPathCollectType::EditorOnlyCollect && !bReferencerIsEditorOnly)
-			{
-				SoftPackagesUsedInGame.Add(PackageName);
-			}
-		}
-	}
-	return *this;
-}
-
-FArchive& FArchiveSaveTagImports::operator<<(FName& Name)
-{
-	NameMapSaver.MarkNameAsReferenced(Name);
-	return *this;
-}
-
-void FArchiveSaveTagImports::MarkSearchableName(const UObject* TypeObject, const FName& ValueName) const
-{
-	if (!TypeObject)
-	{
-		return;
-	}
-
-	if (!Dependencies.Contains(TypeObject))
-	{
-		// Serialize object to make sure it ends up in import table
-		// This is doing a const cast to avoid backward compatibility issues
-		FArchiveSaveTagImports* MutableArchive = const_cast<FArchiveSaveTagImports*>(this);
-		UObject* TempObject = const_cast<UObject*>(TypeObject);
-		(*MutableArchive) << TempObject;
-	}
-
-	// Manually mark the name as referenced, in case it got skipped due to delta serialization
-	NameMapSaver.MarkNameAsReferenced(ValueName);
-
-	Linker->SearchableNamesObjectMap.FindOrAdd(TypeObject).AddUnique(ValueName);
-}
-
-/**
- * Helper structure to sort a linker's name map
- */
-struct FObjectNameSortHelper
-{
-	/** Comparison function used by Sort */
-	FORCEINLINE bool operator()( const FName& A, const FName& B ) const
-	{
-		return A.Compare(B) < 0;
-	}
-
-	/** Comparison function used by Sort */
-	FORCEINLINE bool operator()(FNameEntryId A, FNameEntryId B) const
-	{
-		// Could be implemented without constructing FName but would new FNameEntry comparison API
-		return A != B && operator()(FName::CreateFromDisplayId(A, 0), FName::CreateFromDisplayId(B, 0));
-	}
-};
 
-void FPackageNameMapSaver::UpdateLinker(FLinkerSave& Linker, FArchive* BinarySaver)
+FSavePackageResultStruct UPackage::Save(UPackage* InOuter, UObject* Base, EObjectFlags TopLevelFlags,
+	const TCHAR* Filename, FOutputDevice* Error, FLinkerNull* ConformNO, bool bForceByteSwapping,
+	bool bWarnOfLongFilename, uint32 SaveFlags, const ITargetPlatform* TargetPlatform,
+	const FDateTime& FinalTimeStamp, bool bSlowTask, FArchiveDiffMap* InOutDiffMap,
+	FSavePackageContext* SavePackageContext)
 {
-	// Add names
-	Linker.NameMap.Reserve(Linker.NameMap.Num() + ReferencedNames.Num());
-	for (FNameEntryId Name : ReferencedNames)
-	{
-		Linker.NameMap.Add(Name);
-	}
-
-	// Sort names
-	if (Linker.NameMap.Num())
-	{
-		Sort(&Linker.NameMap[0], Linker.NameMap.Num(), FObjectNameSortHelper());
-	}
-
-	// Serialize names and build NameIndicies
-	if (BinarySaver)
-	{
-		Linker.Summary.NameCount = Linker.NameMap.Num();
-		Linker.Summary.NamesReferencedFromExportDataCount = Linker.Summary.NameCount; // Only supported with SavePackage2
-		for (int32 i = 0; i < Linker.NameMap.Num(); i++)
-		{
-			FName::GetEntry(Linker.NameMap[i])->Write(Linker);
-			Linker.NameIndices.Add(Linker.NameMap[i], i);
-		}
+	// CookData should only be nonzero if we are cooking.
+	TOptional<FArchiveCookData> CookData;
+	FArchiveCookContext CookContext(InOuter, FArchiveCookContext::ECookTypeUnknown);
+	if (TargetPlatform != nullptr)
+	{		
+		CookData.Emplace(*TargetPlatform, CookContext);
 	}
+	
+	FSavePackageArgs SaveArgs = { nullptr /* deprecated target platform */, CookData.GetPtrOrNull(), TopLevelFlags, SaveFlags, bForceByteSwapping,
+		bWarnOfLongFilename, bSlowTask, FinalTimeStamp, Error, SavePackageContext };
+	return UPackage::Save(InOuter, Base, Filename, SaveArgs);
 }
 
-class FExportReferenceSorter : public FArchiveUObject
+FSavePackageResultStruct UPackage::Save(UPackage* InOuter, UObject* InAsset, const TCHAR* Filename,
+	const FSavePackageArgs& SaveArgs)
 {
-	/**
-	 * Verifies that all objects which will be force-loaded when the export at RelativeIndex is created and/or loaded appear in the sorted list of exports
-	 * earlier than the export at RelativeIndex.
-	 *
-	 * Used for tracking down the culprit behind dependency sorting bugs.
-	 *
-	 * @param	RelativeIndex	the index into the sorted export list to check dependencies for
-	 * @param	CheckObject		the object that will be force-loaded by the export at RelativeIndex
-	 * @param	ReferenceType	the relationship between the object at RelativeIndex and CheckObject (archetype, class, etc.)
-	 * @param	out_ErrorString	if incorrect sorting is detected, receives data containing more information about the incorrectly sorted object.
-	 *
-	 * @param	true if the export at RelativeIndex appears later than the exports associated with any objects that it will force-load; false otherwise.
-	 */
-	bool VerifyDependency( const int32 RelativeIndex, UObject* CheckObject, const FString& ReferenceType, FString& out_ErrorString )
-	{
-		bool bResult = false;
-
-		checkf(ReferencedObjects.IsValidIndex(RelativeIndex), TEXT("Invalid index specified: %i (of %i)"), RelativeIndex, ReferencedObjects.Num());
-
-		UObject* SourceObject = ReferencedObjects[RelativeIndex];
-		checkf(SourceObject, TEXT("nullptr Object at location %i in ReferencedObjects list"), RelativeIndex);
-		checkf(CheckObject, TEXT("CheckObject is nullptr for %s (%s)"), *SourceObject->GetFullName(), *ReferenceType);
-
-		if ( SourceObject->GetOutermost() != CheckObject->GetOutermost() )
-		{
-			// not in the same package; therefore we can assume that the dependent object will exist
-			bResult = true;
-		}
-		else
-		{
-			int32 OtherIndex = ReferencedObjects.Find(CheckObject);
-			if ( OtherIndex != INDEX_NONE )
-			{
-				if ( OtherIndex < RelativeIndex )
-				{
-					bResult = true;
-				}
-				else
-				{
-					out_ErrorString = FString::Printf(TEXT("Sorting error detected (%s appears later in ReferencedObjects list)!  %i) %s   =>  %i) %s"), *ReferenceType, RelativeIndex,
-						*SourceObject->GetFullName(), OtherIndex, *CheckObject->GetFullName());
-
-					bResult = false;
-				}
-			}
-			else
-			{
-				// the object isn't in the list of ReferencedObjects, which means it wasn't processed as a result of processing the source object; this
-				// might indicate a bug, but might also just mean that the CheckObject was first referenced by an earlier export
-				int32 ProcessedIndex = ProcessedObjects.Find(CheckObject);
-				if ( ProcessedIndex != INDEX_NONE )
-				{
-					OtherIndex = ProcessedIndex;
-					int32 SourceIndex = ProcessedObjects.Find(SourceObject);
-
-					if ( OtherIndex < SourceIndex )
-					{
-						bResult = true;
-					}
-					else
-					{
-						out_ErrorString = FString::Printf(TEXT("Sorting error detected (%s was processed but not added to ReferencedObjects list)!  %i/%i) %s   =>  %i) %s"),
-							*ReferenceType, RelativeIndex, SourceIndex, *SourceObject->GetFullName(), OtherIndex, *CheckObject->GetFullName());
-						bResult = false;
-					}
-				}
-				else
-				{
-					int32 SourceIndex = ProcessedObjects.Find(SourceObject);
-
-					out_ErrorString = FString::Printf(TEXT("Sorting error detected (%s has not yet been processed)!  %i/%i) %s   =>  %s"),
-						*ReferenceType, RelativeIndex, SourceIndex, *SourceObject->GetFullName(), *CheckObject->GetFullName());
-
-					bResult = false;
-				}
-			}
-		}
-
-		return bResult;
-	}
-
-	/**
-	 * Pre-initializes the list of processed objects with the boot-strap classes.
-	 */
-	void InitializeCoreClasses()
-	{
-#if 1
-		FScopeLock ScopeLock(&InitializeCoreClassesCritSec);
-		check(CoreClasses.Num() == 0);
-		check(ReferencedObjects.Num() == 0);
-		check(SerializedObjects.Num() == 0);
-		check(bIgnoreFieldReferences == false);
-
-		static bool bInitializedStaticCoreClasses = false;
-		static TArray<UClass*> StaticCoreClasses;
-		static TArray<UObject*> StaticCoreReferencedObjects;
-		static FOrderedObjectSet StaticProcessedObjects;
-		static TSet<UObject*> StaticSerializedObjects;
-		
-		
-
-		// Helper class to register FlushInitializedStaticCoreClasses callback on first SavePackage run
-		struct FAddFlushInitalizedStaticCoreClasses
-		{
-			FAddFlushInitalizedStaticCoreClasses() 
-			{
-				FCoreUObjectDelegates::GetPreGarbageCollectDelegate().AddStatic(FlushInitalizedStaticCoreClasses);
-			}
-			/** Wrapper function to handle default parameter when used as function pointer */
-			static void FlushInitalizedStaticCoreClasses()
-			{
-				bInitializedStaticCoreClasses = false;
-			}
-		};
-		static FAddFlushInitalizedStaticCoreClasses MaybeAddAddFlushInitializedStaticCoreClasses;
-
-#if VALIDATE_INITIALIZECORECLASSES
-		bool bWasValid = bInitializedStaticCoreClasses;
-		bInitializedStaticCoreClasses = false;
-#endif
-
-		if (!bInitializedStaticCoreClasses)
-		{
-			bInitializedStaticCoreClasses = true;
-
-
-			// initialize the tracking maps with the core classes
-			UClass* CoreClassList[] =
-			{
-				UObject::StaticClass(),
-				UField::StaticClass(),
-				UStruct::StaticClass(),
-				UScriptStruct::StaticClass(),
-				UFunction::StaticClass(),
-				UEnum::StaticClass(),
-				UClass::StaticClass(),
-				UInterface::StaticClass()
-			};
-
-			for (UClass* CoreClass : CoreClassList)
-			{
-				CoreClasses.AddUnique(CoreClass);
-
-				ReferencedObjects.Add(CoreClass);
-				ReferencedObjects.Add(CoreClass->GetDefaultObject());
-			}
-
-			for (UClass* CoreClass : CoreClasses)
-			{
-				ProcessStruct(CoreClass);
-			}
-
-			CoreReferencesOffset = ReferencedObjects.Num();
-
-
-#if VALIDATE_INITIALIZECORECLASSES
-			if (bWasValid)
-			{
-				// make sure everything matches up 
-				check(CoreClasses.Num() == StaticCoreClasses.Num());
-				check(ReferencedObjects.Num() == StaticCoreReferencedObjects.Num());
-				check(ProcessedObjects.Num() == StaticProcessedObjects.Num());
-				check(SerializedObjects.Num() == StaticSerializedObjects.Num());
-				
-				
-				for (int I = 0; I < CoreClasses.Num(); ++I)
-				{
-					check(CoreClasses[I] == StaticCoreClasses[I]);
-				}
-				for (int I = 0; I < ReferencedObjects.Num(); ++I)
-				{
-					check(ReferencedObjects[I] == StaticCoreReferencedObjects[I]);
-				}
-				for (const auto& ProcessedObject : ProcessedObjects.ObjectsSet)
-				{
-					check(ProcessedObject.Value == StaticProcessedObjects.Find(ProcessedObject.Key));
-				}
-				for (const auto& SerializedObject : SerializedObjects)
-				{
-					check(StaticSerializedObjects.Find(SerializedObject));
-				}
-			}
-#endif
-
-			StaticCoreClasses = CoreClasses;
-			StaticCoreReferencedObjects = ReferencedObjects;
-			StaticProcessedObjects = ProcessedObjects;
-			StaticSerializedObjects = SerializedObjects;
-
-			check(CurrentClass == nullptr);
-			check(CurrentInsertIndex == INDEX_NONE);
-		}
-		else
-		{
-			CoreClasses = StaticCoreClasses;
-			ReferencedObjects = StaticCoreReferencedObjects;
-			ProcessedObjects = StaticProcessedObjects;
-			SerializedObjects = StaticSerializedObjects;
-
-			CoreReferencesOffset = StaticCoreReferencedObjects.Num();
-		}
-
-#else
-		// initialize the tracking maps with the core classes
-		UClass* CoreClassList[] =
-		{
-			UObject::StaticClass(),
-			UField::StaticClass(),
-			UStruct::StaticClass(),
-			UScriptStruct::StaticClass(),
-			UFunction::StaticClass(),
-			UEnum::StaticClass(),
-			UClass::StaticClass(),
-			FProperty::StaticClass(),
-			FByteProperty::StaticClass(),
-			FIntProperty::StaticClass(),
-			FBoolProperty::StaticClass(),
-			FFloatProperty::StaticClass(),
-			FDoubleProperty::StaticClass(),
-			FObjectProperty::StaticClass(),
-			FClassProperty::StaticClass(),
-			FInterfaceProperty::StaticClass(),
-			FNameProperty::StaticClass(),
-			FStrProperty::StaticClass(),
-			FArrayProperty::StaticClass(),
-			FTextProperty::StaticClass(),
-			FStructProperty::StaticClass(),
-			FDelegateProperty::StaticClass(),
-			UInterface::StaticClass(),
-			FMulticastDelegateProperty::StaticClass(),
-			FWeakObjectProperty::StaticClass(),
-			FObjectPropertyBase::StaticClass(),
-			FLazyObjectProperty::StaticClass(),
-			FSoftObjectProperty::StaticClass(),
-			FSoftClassProperty::StaticClass(),
-			FMapProperty::StaticClass(),
-			FSetProperty::StaticClass(),
-			FEnumProperty::StaticClass()
-		};
-
-		for (UClass* CoreClass : CoreClassList)
-		{
-			CoreClasses.AddUnique(CoreClass);
-
-			ReferencedObjects.Add(CoreClass);
-			ReferencedObjects.Add(CoreClass->GetDefaultObject());
-		}
-
-		for (UClass* CoreClass : CoreClasses)
-		{
-			ProcessStruct(CoreClass);
-		}
-
-		CoreReferencesOffset = ReferencedObjects.Num();
-#endif
-	}
-
-	/**
-	 * Adds an object to the list of referenced objects, ensuring that the object is not added more than one.
-	 *
-	 * @param	Object			the object to add to the list
-	 * @param	InsertIndex		the index to insert the object into the export list
-	 */
-	void AddReferencedObject( UObject* Object, int32 InsertIndex )
-	{
-		if ( Object != nullptr && !ReferencedObjects.Contains(Object) )
-		{
-			ReferencedObjects.Insert(Object, InsertIndex);
-		}
-	}
-
-	/**
-	 * Handles serializing and calculating the correct insertion point for an object that will be force-loaded by another object (via an explicit call to Preload).
-	 * If the RequiredObject is a UStruct or true is specified for bProcessObject, the RequiredObject will be inserted into the list of exports just before the object
-	 * that has a dependency on this RequiredObject.
-	 *
-	 * @param	RequiredObject		the object which must be created and loaded first
-	 * @param	bProcessObject		normally, only the class and archetype for non-UStruct objects are inserted into the list;  specify true to override this behavior
-	 *								if RequiredObject is going to be force-loaded, rather than just created
-	 */
-	void HandleDependency( UObject* RequiredObject, bool bProcessObject=false )
-	{
-		if ( RequiredObject != nullptr )
-		{
-			check(CurrentInsertIndex!=INDEX_NONE);
-
-			const int32 PreviousReferencedObjectCount = ReferencedObjects.Num();
-			const int32 PreviousInsertIndex = CurrentInsertIndex;
-
-			if (!PackageToSort || RequiredObject->GetOutermost() == PackageToSort)
-			{
-				// Don't compute prerequisites for objects outside the package, this will recurse into all native properties
-				if (UStruct* RequiredObjectStruct = dynamic_cast<UStruct*>(RequiredObject))
-				{
-					// if this is a struct/class/function/state, it may have a super that needs to be processed first
-					ProcessStruct(RequiredObjectStruct);
-				}
-				else if (bProcessObject)
-				{
-					// this means that RequiredObject is being force-loaded by the referencing object, rather than simply referenced
-					ProcessObject(RequiredObject);
-				}
-				else
-				{
-					// only the object's class and archetype are force-loaded, so only those objects need to be in the list before
-					// whatever object was referencing RequiredObject
-					if (ProcessedObjects.Find(RequiredObject->GetOuter()) == INDEX_NONE)
-					{
-						HandleDependency(RequiredObject->GetOuter());
-					}
-
-					// class is needed before archetype, but we need to process these in reverse order because we are inserting into the list.
-					ProcessObject(RequiredObject->GetArchetype());
-					ProcessStruct(RequiredObject->GetClass());
-				}
-			}
-			// InsertIndexOffset is the amount the CurrentInsertIndex was incremented during the serialization of SuperField; we need to
-			// subtract out this number to get the correct location of the new insert index
-			const int32 InsertIndexOffset = CurrentInsertIndex - PreviousInsertIndex;
-			const int32 InsertIndexAdvanceCount = (ReferencedObjects.Num() - PreviousReferencedObjectCount) - InsertIndexOffset;
-			if ( InsertIndexAdvanceCount > 0 )
-			{
-				// if serializing SuperField added objects to the list of ReferencedObjects, advance the insertion point so that
-				// subsequence objects are placed into the list after the SuperField and its dependencies.
-				CurrentInsertIndex += InsertIndexAdvanceCount;
-			}
-		}
-	}
-
-public:
-	/**
-	 * Constructor
-	 */
-	FExportReferenceSorter()
-	{
-		ArIsObjectReferenceCollector = true;
-		this->SetIsPersistent(true);
-		this->SetIsSaving(true);
-
-		InitializeCoreClasses();
-	}
-
-	/**
-	 * Verifies that the sorting algorithm is working correctly by checking all objects in the ReferencedObjects array to make sure that their
-	 * required objects appear in the list first
-	 */
-	void VerifySortingAlgorithm()
-	{
-		FString ErrorString;
-		for ( int32 VerifyIndex = CoreReferencesOffset; VerifyIndex < ReferencedObjects.Num(); VerifyIndex++ )
-		{
-			UObject* Object = ReferencedObjects[VerifyIndex];
-			
-			// first, make sure that the object's class and archetype appear earlier in the list
-			UClass* ObjectClass = Object->GetClass();
-			if ( !VerifyDependency(VerifyIndex, ObjectClass, TEXT("Class"), ErrorString) )
-			{
-				UE_LOG(LogSavePackage, Log, TEXT("%s"), *ErrorString);
-			}
-
-			UObject* ObjectArchetype = Object->GetArchetype();
-			if ( ObjectArchetype != nullptr && !VerifyDependency(VerifyIndex, ObjectArchetype, TEXT("Archetype"), ErrorString) )
-			{
-				UE_LOG(LogSavePackage, Log, TEXT("%s"), *ErrorString);
-			}
-
-			// UObjectRedirectors are always force-loaded as the loading code needs immediate access to the object pointed to by the Redirector
-			UObjectRedirector* Redirector = dynamic_cast<UObjectRedirector*>(Object);
-			if ( Redirector != nullptr && Redirector->DestinationObject != nullptr )
-			{
-				// the Redirector does not force-load the destination object, so we only need its class and archetype.
-				UClass* RedirectorDestinationClass = Redirector->DestinationObject->GetClass();
-				if ( !VerifyDependency(VerifyIndex, RedirectorDestinationClass, TEXT("Redirector DestinationObject Class"), ErrorString) )
-				{
-					UE_LOG(LogSavePackage, Log, TEXT("%s"), *ErrorString);
-				}
-
-				UObject* RedirectorDestinationArchetype = Redirector->DestinationObject->GetArchetype();
-				if ( RedirectorDestinationArchetype != nullptr 
-				&& !VerifyDependency(VerifyIndex, RedirectorDestinationArchetype, TEXT("Redirector DestinationObject Archetype"), ErrorString) )
-				{
-					UE_LOG(LogSavePackage, Log, TEXT("%s"), *ErrorString);
-				}
-			}
-		}
-	}
-
-	/**
-	 * Clears the list of encountered objects; should be called if you want to re-use this archive.
-	 */
-	void Clear()
-	{
-		ReferencedObjects.RemoveAt(CoreReferencesOffset, ReferencedObjects.Num() - CoreReferencesOffset);
-	}
-
-	/**
-	 * Get the list of new objects which were encountered by this archive; excludes those objects which were passed into the constructor
-	 */
-	void GetExportList( TArray<UObject*>& out_Exports, UPackage* OuterPackage, bool bIncludeCoreClasses=false )
-	{
-		PackageToSort = OuterPackage;
-		if ( !bIncludeCoreClasses )
-		{
-			const int32 NumReferencedObjects = ReferencedObjects.Num() - CoreReferencesOffset;
-			if ( NumReferencedObjects > 0 )
-			{
-				int32 OutputIndex = out_Exports.Num();
-
-				out_Exports.AddUninitialized(NumReferencedObjects);
-				for ( int32 RefIndex = CoreReferencesOffset; RefIndex < ReferencedObjects.Num(); RefIndex++ )
-				{
-					out_Exports[OutputIndex++] = ReferencedObjects[RefIndex];
-				}
-			}
-		}
-		else
-		{
-			out_Exports += ReferencedObjects;
-		}
-	}
-
-	/** 
-	 * UObject serialization operator
-	 *
-	 * @param	Object	an object encountered during serialization of another object
-	 *
-	 * @return	reference to instance of this class
-	 */
-	FArchive& operator<<( UObject*& Object )
-	{
-		// we manually handle class default objects, so ignore those here
-		if ( Object != nullptr && !Object->HasAnyFlags(RF_ClassDefaultObject) )
-		{
-			if ( ProcessedObjects.Find(Object) == INDEX_NONE )
-			{
-				// if this object is not a UField, it is an object instance that is referenced through script or defaults (when processing classes) or
-				// through an normal object reference (when processing the non-class exports).  Since classes and class default objects
-				// are force-loaded (and thus, any objects referenced by the class's script or defaults will be created when the class
-				// is force-loaded), we'll need to be sure that the referenced object's class and archetype are inserted into the list
-				// of exports before the class, so that when CreateExport is called for this object reference we don't have to seek.
-				// Note that in the non-UField case, we don't actually need the object itself to appear before the referencing object/class because it won't
-				// be force-loaded (thus we don't need to add the referenced object to the ReferencedObject list)
-
-				if (Cast<UField>(Object))
-				{
-					// when field processing is enabled, ignore any referenced classes since a class's class and CDO are both intrinsic and
-					// attempting to deal with them here will only cause problems
-					if ( !bIgnoreFieldReferences && !dynamic_cast<UClass*>(Object) )
-					{
-						if ( CurrentClass == nullptr || Object->GetOuter() != CurrentClass )
-						{
-							if ( UStruct* StructObject = dynamic_cast<UStruct*>(Object) )
-							{
-								// if this is a struct/class/function/state, it may have a super that needs to be processed first (Preload force-loads UStruct::SuperField)
-								ProcessStruct(StructObject);
-							}
-							else
-							{
-								// properties that are enum references need their enums loaded first so that config importing works
-								if (UEnum* Enum = Cast<UEnum>(Object))
-								{
-									HandleDependency(Enum, /*bProcessObject =*/true);
-								}
-
-								// a normal field - property, enum, const; just insert it into the list and keep going
-								ProcessedObjects.Add(Object);
-								
-								AddReferencedObject(Object, CurrentInsertIndex);
-								if ( !SerializedObjects.Contains(Object) )
-								{
-									SerializedObjects.Add(Object);
-									Object->Serialize(*this);
-								}
-							}
-						}
-					}
-				}
-				else
-				{
-					HandleDependency(Object);
-				}
-			}
-		}
-
-		return *this;
-	}
-
-	/**
-	 * Adds a normal object to the list of sorted exports.  Ensures that any objects which will be force-loaded when this object is created or loaded are inserted into
-	 * the list before this object.
-	 *
-	 * @param	Object	the object to process.
-	 */
-	void ProcessObject( UObject* Object )
-	{
-		// we manually handle class default objects, so ignore those here
-		if ( Object != nullptr )
-		{
-			if ( !Object->HasAnyFlags(RF_ClassDefaultObject) )
-			{
-				if ( ProcessedObjects.Find(Object) == INDEX_NONE )
-				{
-					ProcessedObjects.Add(Object);
-
-					const bool bRecursiveCall = CurrentInsertIndex != INDEX_NONE;
-					if ( !bRecursiveCall )
-					{
-						CurrentInsertIndex = ReferencedObjects.Num();
-					}
-
-					// when an object is created (CreateExport), its class and archetype will be force-loaded, so we'll need to make sure that those objects
-					// are placed into the list before this object so that when CreateExport calls Preload on these objects, no seeks occur
-					// The object's Outer isn't force-loaded, but it will be created before the current object, so we'll need to ensure that its archetype & class
-					// are placed into the list before this object.
-					HandleDependency(Object->GetClass(), true);
-					HandleDependency(Object->GetOuter());
-					HandleDependency(Object->GetArchetype(), true);
-
-					// UObjectRedirectors are always force-loaded as the loading code needs immediate access to the object pointed to by the Redirector
-					UObjectRedirector* Redirector = dynamic_cast<UObjectRedirector*>(Object);
-					if ( Redirector != nullptr && Redirector->DestinationObject != nullptr )
-					{
-						// the Redirector does not force-load the destination object, so we only need its class and archetype.
-						HandleDependency(Redirector->DestinationObject);
-					}
-
-					// now we add this object to the list
-					AddReferencedObject(Object, CurrentInsertIndex);
-
-					// then serialize the object - any required references encountered during serialization will be inserted into the list before this object, but after this object's
-					// class and archetype
-					if ( !SerializedObjects.Contains(Object) )
-					{
-						SerializedObjects.Add(Object);
-						Object->Serialize(*this);
-					}
-
-					if ( !bRecursiveCall )
-					{
-						CurrentInsertIndex = INDEX_NONE;
-					}
-				}
-			}
-		}
-	}
-
-	/**
-	 * Adds a UStruct object to the list of sorted exports.  Handles serialization and insertion for any objects that will be force-loaded by this struct (via an explicit call to Preload).
-	 *
-	 * @param	StructObject	the struct to process
-	 */
-	void ProcessStruct( UStruct* StructObject )
-	{
-		if ( StructObject != nullptr )
-		{
-			if ( ProcessedObjects.Find(StructObject) == INDEX_NONE )
-			{
-				ProcessedObjects.Add(StructObject);
-
-				const bool bRecursiveCall = CurrentInsertIndex != INDEX_NONE;
-				if ( !bRecursiveCall )
-				{
-					CurrentInsertIndex = ReferencedObjects.Num();
-				}
-
-				// this must be done after we've established a CurrentInsertIndex
-				HandleDependency(StructObject->GetInheritanceSuper());
-
-				// insert the class/function/state/struct into the list
-				AddReferencedObject(StructObject, CurrentInsertIndex);
-				if ( !SerializedObjects.Contains(StructObject) )
-				{
-					const bool bPreviousIgnoreFieldReferences = bIgnoreFieldReferences;
-
-					// first thing to do is collect all actual objects referenced by this struct's script or defaults
-					// so we turn off field serialization so that we don't have to worry about handling this struct's fields just yet
-					bIgnoreFieldReferences = true;
-
-					bool const bIsClassObject = (dynamic_cast<UClass*>(StructObject) != nullptr);
-
-					SerializedObjects.Add(StructObject);
-					StructObject->Serialize(*this);
-
-					// at this point, any objects which were referenced through this struct's script or defaults will be in the list of exports, and 
-					// the CurrentInsertIndex will have been advanced so that the object processed will be inserted just before this struct in the array
-					// (i.e. just after class/archetypes for any objects which were referenced by this struct's script)
-
-					// now re-enable field serialization and process the struct's properties, functions, enums, structs, etc.  They will be inserted into the list
-					// just ahead of the struct itself, so that those objects are created first during seek-free loading.
-					bIgnoreFieldReferences = false;
-					
-					// invoke the serialize operator rather than calling Serialize directly so that the object is handled correctly (i.e. if it is a struct, then we should
-					// call ProcessStruct, etc. and all this logic is already contained in the serialization operator)
-					if (!bIsClassObject)
-					{
-						// before processing the Children reference, set the CurrentClass to the class which contains this StructObject so that we
-						// don't inadvertently serialize other fields of the owning class too early.
-						CurrentClass = StructObject->GetOwnerClass();
-					}					
-
-					(*this) << (UObject*&)StructObject->Children;
-					CurrentClass = nullptr; //-V519
-
-					(*this) << (UObject*&)StructObject->Next;
-
-					bIgnoreFieldReferences = bPreviousIgnoreFieldReferences;
-				}
-
-				// Preload will force-load the class default object when called on a UClass object, so make sure that the CDO is always immediately after its class
-				// in the export list; we can't resolve this circular reference, but hopefully we the CDO will fit into the same memory block as the class during 
-				// seek-free loading.
-				UClass* ClassObject = dynamic_cast<UClass*>(StructObject);
-				if ( ClassObject != nullptr )
-				{
-					UObject* CDO = ClassObject->GetDefaultObject();
-					ensureMsgf(nullptr != CDO, TEXT("Error: Invalid CDO in class %s"), *GetPathNameSafe(ClassObject));
-					if ((ProcessedObjects.Find(CDO) == INDEX_NONE) && (nullptr != CDO))
-					{
-						ProcessedObjects.Add(CDO);
-
-						if ( !SerializedObjects.Contains(CDO) )
-						{
-							SerializedObjects.Add(CDO);
-							CDO->Serialize(*this);
-						}
-
-						int32 ClassIndex = ReferencedObjects.Find(ClassObject);
-						check(ClassIndex != INDEX_NONE);
-
-						// we should be the only one adding CDO's to the list, so this assertion is to catch cases where someone else
-						// has added the CDO to the list (as it will probably be in the wrong spot).
-						check(!ReferencedObjects.Contains(CDO) || CoreClasses.Contains(ClassObject));
-						AddReferencedObject(CDO, ClassIndex + 1);
-					}
-				}
-
-				if ( !bRecursiveCall )
-				{
-					CurrentInsertIndex = INDEX_NONE;
-				}
-			}
-		}
-	}
-
-	/** Do nothing when serializing soft references, this is required because the presave on soft references can fix redirectors, which is unsafe at this point */
-	virtual FArchive& operator<<(FLazyObjectPtr& Value) override { return *this; }
-	virtual FArchive& operator<<(FSoftObjectPtr& Value) override { return *this; }
-	virtual FArchive& operator<<(FSoftObjectPath& Value) override { return *this; }
-
-	/** Object pointer references that are unresolved should not be traversed at this point as we are not allowed to load/resolve anything */
-	virtual FArchive& operator<<(FObjectPtr& Value) override
-	{
-		// Everything we want to traverse in this export reference sorter will have been resolved.  If we encounter something that is unresolved,
-		// it will be part of another package and not something we need to worry about for export sorting.
-		if (IsObjectHandleResolved(Value.GetHandle()))
-		{
-			return FArchiveUObject::operator<<(Value);
-		}
-		return *this;
-	}
-
-private:
-
-	/**
-	 * The index into the ReferencedObjects array to insert new objects
-	 */
-	int32 CurrentInsertIndex = INDEX_NONE;
-
-	/**
-	 * The index into the ReferencedObjects array for the first object not referenced by one of the core classes
-	 */
-	int32 CoreReferencesOffset = INDEX_NONE;
-
-	/**
-	 * The classes which are pre-added to the array of ReferencedObjects.  Used for resolving a number of circular dependecy issues between
-	 * the boot-strap classes.
-	 */
-	TArray<UClass*> CoreClasses;
-
-	/**
-	 * The list of objects that have been evaluated by this archive so far.
-	 */
-	struct FOrderedObjectSet
-	{
-		TMap<UObject*, int32> ObjectsMap;
-
-		int32 Add(UObject* Object)
-		{
-			const int32 Index = ObjectsMap.Num();
-			ObjectsMap.Add(Object, Index);
-			return Index;
-		}
-
-		inline int32 Find(UObject* Object) const
-		{
-			const int32 *Index = ObjectsMap.Find(Object);
-			if (Index)
-			{
-				return *Index;
-			}
-			return INDEX_NONE;
-		}
-		inline int32 Num() const
-		{
-			return ObjectsMap.Num();
-		}
-	};
-	FOrderedObjectSet ProcessedObjects;
-
-	/**
-	 * The list of objects that have been serialized; used to prevent calling Serialize on an object more than once.
-	 */
-	TSet<UObject*> SerializedObjects;
-
-	/**
-	 * The list of new objects that were encountered by this archive
-	 */
-	TArray<UObject*> ReferencedObjects;
-
-	/**
-	 * Controls whether to process UField objects encountered during serialization of an object.
-	 */
-	bool bIgnoreFieldReferences = false;
-
-	/**
-	 * The UClass currently being processed.  This is used to prevent serialization of a UStruct's Children member causing other fields of the same class to be processed too early due
-	 * to being referenced (directly or indirectly) by that field.  For example, if a class has two functions which both have a struct parameter of a struct type which is declared in the same class,
-	 * the struct would be inserted into the list immediately before the first function processed.  The second function would be inserted into the list just before the struct.  At runtime,
-	 * the "second" function would be created first, which would end up force-loading the struct.  This would cause an unacceptible seek because the struct appears later in the export list, thus
-	 * hasn't been created yet.
-	 */
-	UClass* CurrentClass = nullptr;
-
-	/** Package to constrain checks to */
-	UPackage* PackageToSort = nullptr;
-};
-
-/**
- * Helper structure encapsulating functionality to sort a linker's export map to allow seek free
- * loading by creating the exports in the order they are in the export map.
- */
-struct FObjectExportSeekFreeSorter
-{
-	/**
-	 * Sorts exports in passed in linker in order to avoid seeking when creating them in order.
-	 *
-	 * @param	Linker				LinkerSave to sort export map
-	 */
-	void SortExports(FLinkerSave* Linker)
-	{
-		SortArchive.SetCookData(Linker->GetCookData());
-
-		TMap<UObject*,int32>	OriginalExportIndexes;
-
-		// Populate object to current index map.
-		for (int32 ExportIndex=0; ExportIndex<Linker->ExportMap.Num(); ExportIndex++)
-		{
-			const FObjectExport& Export = Linker->ExportMap[ExportIndex];
-			if( Export.Object )
-			{
-				// Set the index (key) in the map to the index of this object into the export map.
-				OriginalExportIndexes.Add( Export.Object, ExportIndex );
-			}
-		}
-
-		bool bRetrieveInitialReferences = true;
-
-		// Now we need to sort the export list according to the order in which objects will be loaded.  For the sake of simplicity, 
-		// process all classes first so they appear in the list first (along with any objects those classes will force-load) 
-		for (int32 ExportIndex=0; ExportIndex<Linker->ExportMap.Num(); ExportIndex++)
-		{
-			const FObjectExport& Export = Linker->ExportMap[ExportIndex];
-			if( UClass* ExportObjectClass = dynamic_cast<UClass*>(Export.Object) )
-			{
-				SortArchive.Clear();
-				SortArchive.ProcessStruct(ExportObjectClass);
-#if EXPORT_SORTING_DETAILED_LOGGING
-				TArray<UObject*> ReferencedObjects;
-				SortArchive.GetExportList(ReferencedObjects, Linker->LinkerRoot, bRetrieveInitialReferences);
-
-				UE_LOG(LogSavePackage, Log, TEXT("Referenced objects for (%i) %s in %s"), ExportIndex, *Export.Object->GetFullName(), *Linker->LinkerRoot->GetName());
-				for ( int32 RefIndex = 0; RefIndex < ReferencedObjects.Num(); RefIndex++ )
-				{
-					UE_LOG(LogSavePackage, Log, TEXT("\t%i) %s"), RefIndex, *ReferencedObjects[RefIndex]->GetFullName());
-				}
-				if ( ReferencedObjects.Num() > 1 )
-				{
-					// insert a blank line to make the output more readable
-					UE_LOG(LogSavePackage, Log, TEXT(""));
-				}
-
-				SortedExports += ReferencedObjects;
-#else
-				SortArchive.GetExportList(SortedExports, Linker->LinkerRoot, bRetrieveInitialReferences);
-#endif
-				bRetrieveInitialReferences = false;
-			}
-
-		}
-
-#if EXPORT_SORTING_DETAILED_LOGGING
-		UE_LOG(LogSavePackage, Log, TEXT("*************   Processed %i classes out of %i possible exports for package %s.  Beginning second pass...   *************"),
-			SortedExports.Num(), Linker->ExportMap.Num(), *Linker->LinkerRoot->GetName());
-#endif
-
-		// All UClasses, CDOs, functions, properties, etc. are now in the list - process the remaining objects now
-		for (int32 ExportIndex = 0; ExportIndex < Linker->ExportMap.Num(); ExportIndex++)
-		{
-			const FObjectExport& Export = Linker->ExportMap[ExportIndex];
-			if ( Export.Object )
-			{
-				SortArchive.Clear();
-				SortArchive.ProcessObject(Export.Object);
-#if EXPORT_SORTING_DETAILED_LOGGING
-				TArray<UObject*> ReferencedObjects;
-				SortArchive.GetExportList(ReferencedObjects, Linker->LinkerRoot, bRetrieveInitialReferences);
-
-				UE_LOG(LogSavePackage, Log, TEXT("Referenced objects for (%i) %s in %s"), ExportIndex, *Export.Object->GetFullName(), *Linker->LinkerRoot->GetName());
-				for ( int32 RefIndex = 0; RefIndex < ReferencedObjects.Num(); RefIndex++ )
-				{
-					UE_LOG(LogSavePackage, Log, TEXT("\t%i) %s"), RefIndex, *ReferencedObjects[RefIndex]->GetFullName());
-				}
-				if ( ReferencedObjects.Num() > 1 )
-				{
-					// insert a blank line to make the output more readable
-					UE_LOG(LogSavePackage, Log, TEXT(""));
-				}
-
-				SortedExports += ReferencedObjects;
-#else
-				SortArchive.GetExportList(SortedExports, Linker->LinkerRoot, bRetrieveInitialReferences);
-#endif
-				bRetrieveInitialReferences = false;
-			}
-		}
-
-#if EXPORT_SORTING_DETAILED_LOGGING
-		SortArchive.VerifySortingAlgorithm();
-#endif
-		// Back up existing export map and empty it so we can repopulate it in a sorted fashion.
-		TArray<FObjectExport> OldExportMap = Linker->ExportMap;
-		Linker->ExportMap.Empty( OldExportMap.Num() );
-
-		// Create new export map from sorted exports.
-		for (int32 ObjectIndex=0; ObjectIndex<SortedExports.Num(); ObjectIndex++)
-		{
-			// See whether this object was part of the to be sortable exports map...
-			UObject* Object		= SortedExports[ObjectIndex];
-			int32* ExportIndexPtr	= OriginalExportIndexes.Find( Object );
-			if( ExportIndexPtr )
-			{
-				// And add it if it has been.
-				Linker->ExportMap.Add( OldExportMap[*ExportIndexPtr] );
-			}
-		}
-
-		// Manually add any new NULL exports last as they won't be in the SortedExportsObjects list. 
-		// A NULL Export.Object can occur if you are e.g. saving an object in the game that is 
-		// OBJECTMARK_NotForClient.
-		for (int32 ExportIndex=0; ExportIndex<OldExportMap.Num(); ExportIndex++)
-		{
-			const FObjectExport& Export = OldExportMap[ExportIndex];
-			if( Export.Object == nullptr )
-			{
-				Linker->ExportMap.Add( Export );
-			}
-		}
-	}
-
-private:
-	/**
-	 * Archive for sorting an objects references according to the order in which they'd be loaded.
-	 */
-	FExportReferenceSorter SortArchive;
-
-	/** Array of regular objects encountered by CollectExportsInOrderOfUse					*/
-	TArray<UObject*>	SortedExports;
-};
-
-// helper class for clarification, encapsulation, and elimination of duplicate code
-struct FPackageExportTagger
-{
-	UObject*		Base;
-	EObjectFlags	TopLevelFlags;
-	UPackage*		Package;
-	const class ITargetPlatform* TargetPlatform;
-
-	FPackageExportTagger(UObject* CurrentBase, EObjectFlags CurrentFlags, UPackage* InPackage, const class ITargetPlatform* InTargetPlatform)
-	:	Base(CurrentBase)
-	,	TopLevelFlags(CurrentFlags)
-	,	Package(InPackage)
-	,	TargetPlatform(InTargetPlatform)
-	{}
-
-	void TagPackageExports( FArchiveSaveTagExports& ExportTagger, bool bRoutePresave, FObjectSaveContextData& ObjectSaveContext )
-	{
-		const bool bIsCooking = !!TargetPlatform;
-
-		// Route PreSave on Base and serialize it for export tagging.
-		if( Base )
-		{
-			if ( bRoutePresave )
-			{
-				if (bIsCooking && Base->HasAnyFlags(RF_ClassDefaultObject | RF_ArchetypeObject))
-				{
-					FArchiveObjectCrc32NonEditorProperties CrcArchive;
-
-					int32 Before = CrcArchive.Crc32(Base);
-					UE::SavePackageUtilities::CallPreSave(Base, ObjectSaveContext);
-					int32 After = CrcArchive.Crc32(Base);
-
-					if (Before != After)
-					{
-						UE_ASSET_LOG(
-							LogSavePackage,
-							Warning,
-							Base,
-							TEXT("Non-deterministic cook warning - PreSave() has modified %s '%s' - a resave may be required"),
-							Base->HasAnyFlags(RF_ClassDefaultObject) ? TEXT("CDO") : TEXT("archetype"),
-							*Base->GetName()
-						);
-					}
-				}
-				else
-				{
-					UE::SavePackageUtilities::CallPreSave(Base, ObjectSaveContext);
-				}
-			}
-
-			ExportTagger.ProcessBaseObject(Base);
-		}
-		if (TopLevelFlags != RF_NoFlags)
-		{
-			TArray<UObject *> ObjectsInPackage;
-			{
-				COOK_STAT(FScopedDurationTimer SerializeTimer(FSavePackageStats::TagPackageExportsGetObjectsWithOuter));
-				GetObjectsWithPackage(Package, ObjectsInPackage);
-			}
-			// Serialize objects to tag them as OBJECTMARK_TagExp.
-			for( int32 Index = 0; Index < ObjectsInPackage.Num(); Index++ )
-			{
-				UObject* Obj = ObjectsInPackage[Index];
-				// Allowed object that have any of the top level flags or have an assigned that we are saving
-				if (Obj->HasAnyFlags(TopLevelFlags))
-				{
-					ExportTagger.ProcessBaseObject(Obj);
-				}
-			}
-		}
-		if ( bRoutePresave )
-		{
-			// Route PreSave.
-			{
-				TArray<UObject*> TagExpObjects;
-				{
-					COOK_STAT(FScopedDurationTimer SerializeTimer(FSavePackageStats::TagPackageExportsGetObjectsWithMarks));
-					GetObjectsWithAnyMarks(TagExpObjects, OBJECTMARK_TagExp);
-				}
-				for(int32 Index = 0; Index < TagExpObjects.Num(); Index++)
-				{
-					UObject* Obj = TagExpObjects[Index];
-					check(Obj->HasAnyMarks(OBJECTMARK_TagExp));
-					//@warning: Objects created from within PreSave will NOT have PreSave called on them!!!
-					if (bIsCooking && Obj->HasAnyFlags(RF_ClassDefaultObject | RF_ArchetypeObject))
-					{
-						FArchiveObjectCrc32NonEditorProperties CrcArchive;
-
-						int32 Before = CrcArchive.Crc32(Obj);
-						UE::SavePackageUtilities::CallPreSave(Obj, ObjectSaveContext);
-						int32 After = CrcArchive.Crc32(Obj);
-
-						if (Before != After)
-						{
-							UE_ASSET_LOG(
-								LogSavePackage,
-								Warning,
-								Obj,
-								TEXT("Non-deterministic cook warning - PreSave() has modified %s '%s' - a resave may be required"),
-								Obj->HasAnyFlags(RF_ClassDefaultObject) ? TEXT("CDO") : TEXT("archetype"),
-								*Obj->GetName()
-							);
-						}
-					}
-					else
-					{
-						UE::SavePackageUtilities::CallPreSave(Obj, ObjectSaveContext);
-					}
-				}
-			}
-		}
-	}
-};
-
-[[nodiscard]] ESavePackageResult BuildAndWriteTrailer(UPackage* Package, FLinkerSave* Linker, FStructuredArchive::FRecord& StructuredArchiveRoot, IPackageWriter* PackageWriter, uint32 SaveFlags, bool bTextFormat)
-{
-	check(Package != nullptr);
-	check(Linker != nullptr);
-
-	Linker->Summary.PayloadTocOffset = INDEX_NONE;
-
-	if (Linker->PackageTrailerBuilder.IsValid())
-	{
-		// At the moment we assume that we cannot have reference payloads in the trailer if SAVE_BulkDataByReference is not set and we
-		// cannot have locally stored payloads if SAVE_BulkDataByReference is set.
-		checkf((SaveFlags & SAVE_BulkDataByReference) != 0 || Linker->PackageTrailerBuilder->GetNumReferencedPayloads() == 0,
-			TEXT("Attempting to build a package trailer with referenced payloads but the SAVE_BulkDataByReference flag is not set. '%s'"), *Package->GetName());
-
-		checkf((SaveFlags & SAVE_BulkDataByReference) != 0 || Linker->PackageTrailerBuilder->GetNumLocalPayloads() == 0,
-			TEXT("Attempting to build a package trailer with local payloads but the SAVE_BulkDataByReference flag is set. '%s'"), *Package->GetName());
-
-		checkf(bTextFormat == false, TEXT("Attempting to build a package trailer for text based asset '%s', this is not supported!"), *Package->GetName());
-
-		Linker->Summary.PayloadTocOffset = Linker->Tell();
-		if (!Linker->PackageTrailerBuilder->BuildAndAppendTrailer(Linker, *Linker))
-		{
-			return ESavePackageResult::Error;
-		}
-
-		Linker->PackageTrailerBuilder.Reset();
-	}
-	
-	return ESavePackageResult::Success;
-}
-
-#if WITH_EDITOR
-/**
- * Helper function to sort export objects by fully qualified names.
- */
-bool ExportObjectSorter(const UObject& Lhs, const UObject& Rhs)
-{
-	// Check names first.
-	if (Lhs.GetFName() != Rhs.GetFName())
-	{
-		return Lhs.GetFName().LexicalLess(Rhs.GetFName());
-	}
-
-	// Names equal, compare class names.
-	if (Lhs.GetClass()->GetFName() != Rhs.GetClass()->GetFName())
-	{
-		return Lhs.GetClass()->GetFName().LexicalLess(Rhs.GetClass()->GetFName());
-	}
-
-	// Compare by outers if they exist.
-	if (Lhs.GetOuter() && Rhs.GetOuter())
-	{
-		return Lhs.GetOuter()->GetFName().LexicalLess(Rhs.GetOuter()->GetFName());
-	}
-
-	if (Lhs.GetOuter())
-	{
-		return true;
-	}
-
-	return false;
-}
-
-/**
-* Helper equality comparator for export objects. Compares by names, class names and outer names.
-*/
-bool ExportEqualityComparator(UObject* Lhs, UObject* Rhs)
-{
-	check(Lhs && Rhs);
-	return Lhs->GetOuter() == Rhs->GetOuter()
-		&& Lhs->GetClass() == Rhs->GetClass()
-		&& Lhs->GetFName() == Rhs->GetFName();
-}
-
-/**
- * Remove OBJECTMARK_TagExp from duplicated objects.
- */
-TMap<UObject*, UObject*> UnmarkExportTagFromDuplicates()
-{
-	TMap<UObject*, UObject*> RedirectDuplicatesToOriginals;
-	TArray<UObject*> Objects;
-	GetObjectsWithAnyMarks(Objects, OBJECTMARK_TagExp);
-
-	Objects.Sort(ExportObjectSorter);
-
-	int32 LastUniqueObjectIndex = 0;
-	for (int32 CurrentObjectIndex = 1; CurrentObjectIndex < Objects.Num(); ++CurrentObjectIndex)
-	{
-		UObject* LastUniqueObject = Objects[LastUniqueObjectIndex];
-		UObject* CurrentObject = Objects[CurrentObjectIndex];
-
-		// Check if duplicates with different pointers
-		if (LastUniqueObject != CurrentObject
-			// but matching names
-			&& ExportEqualityComparator(LastUniqueObject, CurrentObject))
-		{
-			// Don't export duplicates.
-			CurrentObject->UnMark(OBJECTMARK_TagExp);
-			RedirectDuplicatesToOriginals.Add(CurrentObject, LastUniqueObject);
-		}
-		else
-		{
-			LastUniqueObjectIndex = CurrentObjectIndex;
-		}
-	}
-
-	return RedirectDuplicatesToOriginals;
-}
-
-COREUOBJECT_API extern bool GOutputCookingWarnings;
-
-#endif
-
-FSavePackageResultStruct UPackage::Save(UPackage* InOuter, UObject* Base, EObjectFlags TopLevelFlags,
-	const TCHAR* Filename, FOutputDevice* Error, FLinkerNull* ConformNO, bool bForceByteSwapping,
-	bool bWarnOfLongFilename, uint32 SaveFlags, const ITargetPlatform* TargetPlatform,
-	const FDateTime& FinalTimeStamp, bool bSlowTask, FArchiveDiffMap* InOutDiffMap,
-	FSavePackageContext* SavePackageContext)
-{
-	// CookData should only be nonzero if we are cooking.
-	TOptional<FArchiveCookData> CookData;
-	FArchiveCookContext CookContext(InOuter, FArchiveCookContext::ECookTypeUnknown);
-	if (TargetPlatform != nullptr)
-	{		
-		CookData.Emplace(*TargetPlatform, CookContext);
-	}
-	
-	FSavePackageArgs SaveArgs = { nullptr /* deprecated target platform */, CookData.GetPtrOrNull(), TopLevelFlags, SaveFlags, bForceByteSwapping,
-		bWarnOfLongFilename, bSlowTask, FinalTimeStamp, Error, SavePackageContext };
-	return UPackage::Save(InOuter, Base, Filename, SaveArgs);
-}
-
-FSavePackageResultStruct UPackage::Save(UPackage* InOuter, UObject* InAsset, const TCHAR* Filename,
-	const FSavePackageArgs& SaveArgs)
-{
-	const bool bIsCooking = SaveArgs.IsCooking();
-	if (SavePackageUtilities::IsNewSaveEnabled(bIsCooking))
-	{
-		return UPackage::Save2(InOuter, InAsset, Filename, SaveArgs);
-	}
-	const ITargetPlatform* TargetPlatform = SaveArgs.GetTargetPlatform();
-	UObject* Base = InAsset;
-	EObjectFlags TopLevelFlags = SaveArgs.TopLevelFlags;
-	FOutputDevice* Error = SaveArgs.Error;
-	bool bForceByteSwapping = SaveArgs.bForceByteSwapping;
-	bool bWarnOfLongFilename = SaveArgs.bWarnOfLongFilename;
-	uint32 SaveFlags = SaveArgs.SaveFlags;
-	const FDateTime& FinalTimeStamp = SaveArgs.FinalTimeStamp;
-	bool bSlowTask = SaveArgs.bSlowTask;
-	FSavePackageContext* SavePackageContext = SaveArgs.SavePackageContext;
-
-	UE_TRACK_REFERENCING_PACKAGE_SCOPED(InOuter, PackageAccessTrackingOps::NAME_Save);
-	COOK_STAT(FScopedDurationTimer FuncSaveTimer(FSavePackageStats::SavePackageTimeSec));
-	COOK_STAT(FSavePackageStats::NumPackagesSaved++);
-	SCOPED_SAVETIMER(UPackage_Save);
-
-	// Sanity checks
-	check(InOuter);
-	check(Filename);
-
-	FPackagePath TargetPackagePath = FPackagePath::FromLocalPath(Filename);
-	if (TargetPackagePath.GetHeaderExtension() == EPackageExtension::Unspecified)
-	{
-		TargetPackagePath.SetHeaderExtension(EPackageExtension::EmptyString);
-	}
-	/** PackageWriter provides a way to save to an interface rather than necessarily saving to disk. */
-	IPackageWriter* PackageWriter = nullptr;
-	bool bIgnoreHeaderDiffs = false;
-	if (SavePackageContext)
-	{
-		PackageWriter = SavePackageContext->PackageWriter;
-		if (PackageWriter)
-		{
-			bIgnoreHeaderDiffs = SavePackageContext->PackageWriterCapabilities.bIgnoreHeaderDiffs;
-		}
-	}
-	// if we are cooking we should be doing it in the editor
-	// otherwise some other assumptions are bad
-	check(!bIsCooking || WITH_EDITOR);
-	checkf(!bIsCooking || (PackageWriter && PackageWriter->AsCookedPackageWriter()), TEXT("Cook saves require an ICookedPackageWriter"));
-
-#if WITH_EDITOR
-	TMap<UObject*, UObject*> ReplacedImportOuters;
-
-	TopLevelFlags = UE::SavePackageUtilities::NormalizeTopLevelFlags(TopLevelFlags, bIsCooking);
-
-	// if the in memory package filename is different the filename we are saving it to,
-	// regenerate a new persistent id for it.
-	if (!bIsCooking && !InOuter->GetLoadedPath().IsEmpty() && InOuter->GetLoadedPath() != TargetPackagePath && !(SaveFlags & SAVE_FromAutosave))
-	{
-		InOuter->SetPersistentGuid(FGuid::NewGuid());
-	}
-#endif //WITH_EDITOR
-	const bool bSavingConcurrent = !!(SaveFlags & ESaveFlags::SAVE_Concurrent);
-	FObjectSaveContextData ObjectSaveContext(InOuter, TargetPlatform, TargetPackagePath, SaveFlags);
-
-	if (FPlatformProperties::HasEditorOnlyData())
-	{
-		TRefCountPtr<FUObjectSerializeContext> SaveContext(FUObjectThreadContext::Get().GetSerializeContext());
-
-		const bool bCompareLinker = (SaveFlags & SAVE_CompareLinker) != 0;
-
-		if (GIsSavingPackage && !bSavingConcurrent)
-		{
-			ensureMsgf(false, TEXT("Recursive SavePackage() is not supported"));
-			return ESavePackageResult::Error;
-		}
-
-		FUObjectThreadContext& ThreadContext = FUObjectThreadContext::Get();
-		FEDLCookChecker& EDLCookChecker = FEDLCookChecker::Get();
-
-#if WITH_EDITORONLY_DATA
-		if (bIsCooking && (!(SaveFlags & ESaveFlags::SAVE_KeepEditorOnlyCookedPackages)))
-		{
-			static FCanSkipEditorReferencedPackagesWhenCooking CanSkipEditorReferencedPackagesWhenCooking;
-
-			// Don't save packages marked as editor-only.
-			if (CanSkipEditorReferencedPackagesWhenCooking && InOuter->IsLoadedByEditorPropertiesOnly())
-			{
-				UE_CLOG(!(SaveFlags & SAVE_NoError), LogSavePackage, Display, TEXT("Package loaded by editor-only properties: %s. Package will not be saved."), *InOuter->GetName());
-				return ESavePackageResult::ReferencedOnlyByEditorOnlyData;
-			}
-			else if (InOuter->HasAnyPackageFlags(PKG_EditorOnly))
-			{
-				UE_CLOG(!(SaveFlags & SAVE_NoError), LogSavePackage, Display, TEXT("Package marked as editor-only: %s. Package will not be saved."), *InOuter->GetName());
-				return ESavePackageResult::ReferencedOnlyByEditorOnlyData;
-			}
-		}
-#endif
-#if WITH_EDITOR
-		if (!bIsCooking)
-		{
-			// Attempt to create a backup of this package before it is saved, if applicable
-			if (FCoreUObjectDelegates::AutoPackageBackupDelegate.IsBound())
-			{
-				FCoreUObjectDelegates::AutoPackageBackupDelegate.Execute(*InOuter);
-			}
-		}
-
-#endif	// #if WITH_EDITOR
-
-		// do any path replacements on the source DestFile
-		const FString NewPath = FString(Filename);
-
-		// point to the new version of the path
-		Filename = *NewPath;
-
-		if (!bSavingConcurrent)
-		{
-			// We need to fulfill all pending streaming and async loading requests to then allow us to lock the global IO manager. 
-			// The latter implies flushing all file handles which is a pre-requisite of saving a package. The code basically needs 
-			// to be sure that we are not reading from a file that is about to be overwritten and that there is no way we might 
-			// start reading from the file till we are done overwriting it.
-			FlushAsyncLoading();
-		}
-
-		(*GFlushStreamingFunc)();
-
-		int64 TotalPackageSizeUncompressed = 0;
-
-		// Make sure package is fully loaded before saving. 
-		// IsFullyLoaded has important byproducts for new packages, so make sure we evaluate it rather than short-circuiting it
-		if (!InOuter->IsFullyLoaded() && !Base)
-		{
-			if (!(SaveFlags & SAVE_NoError))
-			{
-				// We cannot save packages that aren't fully loaded as it would clobber existing not loaded content.
-				FText ErrorText;
-				if (InOuter->ContainsMap())
-				{
-					FFormatNamedArguments Arguments;
-					Arguments.Add(TEXT("Name"), FText::FromString(NewPath));
-					ErrorText = FText::Format(NSLOCTEXT("SavePackage", "CannotSaveMapPartiallyLoaded", "Map '{Name}' cannot be saved as it has only been partially loaded"), Arguments);
-				}
-				else
-				{
-					FFormatNamedArguments Arguments;
-					Arguments.Add(TEXT("Name"), FText::FromString(NewPath));
-					ErrorText = FText::Format(NSLOCTEXT("SavePackage", "CannotSaveAssetPartiallyLoaded", "Asset '{Name}' cannot be saved as it has only been partially loaded"), Arguments);
-				}
-				Error->Logf(ELogVerbosity::Warning, TEXT("%s"), *ErrorText.ToString());
-			}
-			return ESavePackageResult::Error;
-		}
-
-		// Make sure package is allowed to be saved.
-		if (!TargetPlatform && FCoreUObjectDelegates::IsPackageOKToSaveDelegate.IsBound())
-		{
-			bool bIsOKToSave = FCoreUObjectDelegates::IsPackageOKToSaveDelegate.Execute(InOuter, Filename, Error);
-			if (!bIsOKToSave)
-			{
-				if (!(SaveFlags & SAVE_NoError))
-				{
-					FText ErrorText;
-					if (InOuter->ContainsMap())
-					{
-						FFormatNamedArguments Arguments;
-						Arguments.Add(TEXT("Name"), FText::FromString(NewPath));
-						ErrorText = FText::Format(NSLOCTEXT("SavePackage", "MapSaveNotAllowed", "Map '{Name}' is not allowed to save (see log for reason)"), Arguments);
-					}
-					else
-					{
-						FFormatNamedArguments Arguments;
-						Arguments.Add(TEXT("Name"), FText::FromString(NewPath));
-						ErrorText = FText::Format(NSLOCTEXT("SavePackage", "AssetSaveNotAllowed", "Asset '{Name}' is not allowed to save (see log for reason)"), Arguments);
-					}
-					Error->Logf(ELogVerbosity::Warning, TEXT("%s"), *ErrorText.ToString());
-				}
-				return ESavePackageResult::Error;
-			}
-		}
-
-		const bool FilterEditorOnly = InOuter->HasAnyPackageFlags(PKG_FilterEditorOnly);
-
-		// Route PreSaveRoot to allow e.g. the world to attach components for the persistent level.
-		// If we are saving concurrently, this should have been called before UPackage::Save was called.
-		bool bCleanupRequired = false;
-		if (Base && !bSavingConcurrent)
-		{
-			UE::SavePackageUtilities::CallPreSaveRoot(Base, ObjectSaveContext);
-			bCleanupRequired = ObjectSaveContext.bCleanupRequired;
-		}
-
-		// Init.
-		FString CleanFilename = FPaths::GetCleanFilename(Filename);
-
-		FFormatNamedArguments Args;
-		Args.Add(TEXT("CleanFilename"), FText::FromString(CleanFilename));
-
-		FText StatusMessage = FText::Format(NSLOCTEXT("Core", "SavingFile", "Saving file: {CleanFilename}..."), Args);
-
-		const int32 TotalSaveSteps = 34;
-		FScopedSlowTask SlowTask(TotalSaveSteps, StatusMessage, bSlowTask);
-		SlowTask.MakeDialogDelayed(3.0f, SaveFlags & SAVE_FromAutosave ? true : false);
-
-		SlowTask.EnterProgressFrame();
-
-		bool bSuccess = true;
-		bool bRequestStub = false;
-		{
-			// FullyLoad the package's Loader, so that anything we need to serialize (bulkdata, thumbnails) is available
-			COOK_STAT(FScopedDurationTimer SaveTimer(FSavePackageStats::FullyLoadLoadersTimeSec));
-			EnsureLoadingComplete(InOuter);
-		}
-		SlowTask.EnterProgressFrame();
-
-		// Untag all objects and names.
-		UnMarkAllObjects();
-
-		TArray<UObject*> CachedObjects;
-
-		// structure to track what every export needs to import (native only)
-		TMap<UObject*, TArray<UObject*> > NativeObjectDependencies;
-
-		// Used for the asset registry and the redirect collector
-		TSet<FName> SoftPackagesUsedInGame;
-
-		// Size of serialized out package in bytes. This is before compression.
-		int64 PackageSize = INDEX_NONE;
-		TPimplPtr<FLinkerSave> Linker = nullptr;
-		uint32 SerializedPackageFlags = 0;
-		{
-			// TODO: Require a SavePackageContext and move to EditorEngine
-			FPackageNameMapSaver NameMapSaver;
-
-			uint32 ComparisonFlags = PPF_DeepCompareInstances | PPF_DeepCompareDSOsOnly;
-
-			// Export objects (tags them as OBJECTMARK_TagExp).
-			FArchiveSaveTagExports ExportTaggerArchive( InOuter );
-			ExportTaggerArchive.SetPortFlags( ComparisonFlags );
-			ExportTaggerArchive.SetCookData(SaveArgs.ArchiveCookData);
-			ExportTaggerArchive.SetSerializeContext(SaveContext);
-
-			check( ExportTaggerArchive.IsCooking() == !!TargetPlatform );
-			check( ExportTaggerArchive.IsCooking() == bIsCooking );
-
-			// Tag exports and route presave.
-			FPackageExportTagger PackageExportTagger(Base, TopLevelFlags, InOuter, TargetPlatform);
-			{
-				SCOPED_SAVETIMER(UPackage_Save_TagExportsWithPresave);
-
-				COOK_STAT(FScopedDurationTimer SaveTimer(FSavePackageStats::TagPackageExportsPresaveTimeSec));
-				// Do not route presave if saving concurrently. This should have been done before the concurrent save started.
-				// Also suppress if the PackageWriter claims already done.
-				const bool bRoutePresave = !bSavingConcurrent && (!PackageWriter || !PackageWriter->IsPreSaveCompleted());
-				PackageExportTagger.TagPackageExports(ExportTaggerArchive, bRoutePresave, ObjectSaveContext);
-				ExportTaggerArchive.SetFilterEditorOnly(FilterEditorOnly);
-			}
-		
-#if USE_STABLE_LOCALIZATION_KEYS
-			if (GIsEditor)
-			{
-				// We need to ensure that we have a package localization namespace as the package loading will need it
-				// We need to do this before entering the GIsSavingPackage block as it may change the package meta-data
-				TextNamespaceUtil::EnsurePackageNamespace(InOuter);
-			}
-#endif // USE_STABLE_LOCALIZATION_KEYS
-
-			if (FWorldTileInfo* WorldTileInfo = InOuter->GetWorldTileInfo())
-			{
-				// collect custom version from wc tile info
-				ExportTaggerArchive << *WorldTileInfo;
-			}
-
-			{
-				// set GIsSavingPackage here and mark the package as being saved
-				// as it is now illegal to create new object inside the package, they wouldn't be saved correctly
-				// it is now also illegal to find any new object references; they potentially wouldn't be saved correctly								
-				FScopedSavingFlag IsSavingFlag(bSavingConcurrent, InOuter);
-				{
-					SCOPED_SAVETIMER(UPackage_Save_TagExports);
-					COOK_STAT(FScopedDurationTimer SaveTimer(FSavePackageStats::TagPackageExportsTimeSec));
-					// Clear all marks (OBJECTMARK_TagExp and exclusion marks) again as we need to redo tagging below.
-					UnMarkAllObjects();
-			
-					// We need to serialize objects yet again to tag objects that were created by PreSave as OBJECTMARK_TagExp.
-					PackageExportTagger.TagPackageExports( ExportTaggerArchive, false, ObjectSaveContext);
-				}
-
-				// Kick off any Precaching required for the target platform to save these objects
-				// only need to do this if we are cooking a different platform then the one which is currently running
-				// TODO: if save package is canceled then call ClearCache on each object
-
-#if WITH_EDITOR
-				if ( bIsCooking && !bSavingConcurrent )
-				{
-					TArray<UObject*> TagExpObjects;
-					GetObjectsWithAnyMarks(TagExpObjects, OBJECTMARK_TagExp);
-					for ( int Index =0; Index < TagExpObjects.Num(); ++Index)
-					{
-						UObject *ExpObject = TagExpObjects[Index];
-						if ( ExpObject->HasAnyMarks( OBJECTMARK_TagExp ) )
-						{
-							ExpObject->BeginCacheForCookedPlatformData( TargetPlatform );
-							CachedObjects.Add( ExpObject );
-						}
-					}
-				}
-#endif
-
-				SlowTask.EnterProgressFrame();
-
-				// structure to track what every export needs to import
-				TMap<UObject*, TArray<UObject*> > ObjectDependencies;
-
-				// and a structure to track non-redirector references
-				TSet<UObject*> DependenciesReferencedByNonRedirectors;
-
-				/** If true, we are going to save to save to memory rather than to a FileArchive. */
-				const bool bSaveToMemory = !!(SaveFlags & SAVE_Async) || PackageWriter;
-
-				const bool bSaveUnversionedNative = !!(SaveFlags & SAVE_Unversioned_Native);
-				const bool bSaveUnversionedProperties = !!(SaveFlags & SAVE_Unversioned_Properties) && CanUseUnversionedPropertySerialization(TargetPlatform);
-
-				FArchiveFormatterType* Formatter = nullptr;
-				FArchive* TextFormatArchive = nullptr;
-				const bool bTextFormat = FString(Filename).EndsWith(FPackageName::GetTextAssetPackageExtension()) || FString(Filename).EndsWith(FPackageName::GetTextMapPackageExtension());
-				checkf(!bTextFormat, TEXT("Text format is not supported in old Save, please enable SavePackage2 instead"));
-
-				const FString BaseFilename = FPaths::GetBaseFilename(Filename);
-				// Make temp file. CreateTempFilename guarantees unique, non-existing filename.
-				// The temp file will be saved in the game save folder to not have to deal with potentially too long paths.
-				// Since the temp filename may include a 32 character GUID as well, limit the user prefix to 32 characters.
-				TOptional<FString> TempFilename;
-				TOptional<FString> TextFormatTempFilename;
-				ON_SCOPE_EXIT
-				{
-					// free the file handle and delete the temporary file
-					Linker->CloseAndDestroySaver();
-					if (TempFilename.IsSet())
-					{
-						IFileManager::Get().Delete(*TempFilename.GetValue());
-					}
-					if (TextFormatTempFilename.IsSet())
-					{
-						IFileManager::Get().Delete(*TextFormatTempFilename.GetValue());
-					}
-				};
-	
-				{
-					SCOPED_SAVETIMER(UPackage_Save_CreateLinkerSave);
-
-					if (PackageWriter || bSaveToMemory)
-					{
-						// Allocate the linker with a memory writer, forcing byte swapping if wanted.
-						TUniquePtr<FLargeMemoryWriter> ExportsArchive;
-						if (PackageWriter)
-						{
-							ExportsArchive = PackageWriter->CreateLinkerArchive(InOuter->GetFName(), InOuter->FindAssetInPackage());
-						}
-						{
-							// The LargeMemoryWriter does not need to be persistent; the LinkerSave wraps it and reports Persistent=true
-							bool bIsPersistent = false;
-							ExportsArchive.Reset(new FLargeMemoryWriter(0, bIsPersistent, *InOuter->GetFName().ToString()));
-						}
-						Linker = MakePimpl<FLinkerSave>(InOuter, ExportsArchive.Release(), bForceByteSwapping, bSaveUnversionedNative);
-					}
-					else
-					{
-						// Allocate the linker with a tempfile, forcing byte swapping if wanted.
-						TempFilename = FPaths::CreateTempFilename(*FPaths::ProjectSavedDir(), *BaseFilename.Left(32));
-						Linker = MakePimpl<FLinkerSave>(InOuter, *TempFilename.GetValue(), bForceByteSwapping, bSaveUnversionedNative);
-						if (!Linker->Saver)
-						{
-							FFormatNamedArguments Arguments;
-							Arguments.Add(TEXT("Name"), FText::FromString(*TempFilename));
-							FText ErrorText = FText::Format(NSLOCTEXT("SavePackage", "CouldNotCreateSaveFile", "Could not create temporary save filename {Name}."), Arguments);
-							UE_LOG(LogSavePackage, Error, TEXT("%s"), *ErrorText.ToString());
-							if (!(SaveFlags & SAVE_NoError))
-							{
-								Error->Logf(ELogVerbosity::Error, TEXT("%s"), *ErrorText.ToString());
-							}
-							return ESavePackageResult::Error;
-						}
-					}
-
-					Linker->bProceduralSave = ObjectSaveContext.bProceduralSave;
-					Linker->bUpdatingLoadedPath = ObjectSaveContext.bUpdatingLoadedPath;
-
-					if (UE::FPackageTrailer::IsEnabled())
-					{
-						// The package trailer is not supported for text based assets yet
-						if (!bTextFormat && !ObjectSaveContext.bProceduralSave) 
-						{
-							Linker->PackageTrailerBuilder = MakeUnique<UE::FPackageTrailerBuilder>(InOuter->GetFName());
-						}
-						else if ((SaveFlags & SAVE_BulkDataByReference) != 0)
-						{
-							if (const FLinkerLoad* LinkerLoad = FLinkerLoad::FindExistingLinkerForPackage(InOuter))
-							{
-								if (const UE::FPackageTrailer* Trailer = LinkerLoad->GetPackageTrailer())
-								{
-									Linker->PackageTrailerBuilder = UE::FPackageTrailerBuilder::CreateReferenceToTrailer(*Trailer, InOuter->GetFName());
-								}
-							}
-						}
-					}
-
-#if WITH_TEXT_ARCHIVE_SUPPORT
-					if (bTextFormat)
-					{
-						if (TempFilename.IsSet())
-						{
-							TextFormatTempFilename = TempFilename.GetValue() + FPackageName::GetTextAssetPackageExtension();
-						}
-						else
-						{
-							TextFormatTempFilename = FPaths::CreateTempFilename(*FPaths::ProjectSavedDir(), *BaseFilename.Left(32)) + FPackageName::GetTextAssetPackageExtension();
-						}
-						TextFormatArchive = IFileManager::Get().CreateFileWriter(*TextFormatTempFilename.GetValue());
-						FJsonArchiveOutputFormatter* OutputFormatter = new FJsonArchiveOutputFormatter(*TextFormatArchive);
-						OutputFormatter->SetObjectIndicesMap(&Linker->ObjectIndicesMap);
-						Formatter = OutputFormatter;
-					}
-					else
-#endif
-					{
-						Formatter = new FBinaryArchiveFormatter(*(FArchive*)Linker.Get());
-					}
-				}
-
-				FStructuredArchive* StructuredArchive = new FStructuredArchive(*Formatter);
-				FStructuredArchive::FRecord StructuredArchiveRoot = StructuredArchive->Open().EnterRecord();
-				StructuredArchiveRoot.GetUnderlyingArchive().SetSerializeContext(SaveContext);
-#if WITH_EDITOR
-				if (!!TargetPlatform)
-				{
-					Linker->SetDebugSerializationFlags(DSF_EnableCookerWarnings | Linker->GetDebugSerializationFlags());
-				}
-#endif
-
-				if (!(Linker->Summary.GetPackageFlags() & PKG_FilterEditorOnly))
-				{
-					// The Editor version is used as part of the check to see if a package is too old to use the gather cache, so we always have to add it if we have gathered loc for this asset
-					// We need to set the editor custom version before we copy the version container to the summary, otherwise we may end up with corrupt assets
-					// because we later do it on the Linker when actually gathering loc data
-					ExportTaggerArchive.UsingCustomVersion(FEditorObjectVersion::GUID);
-				}
-
-				// Use the custom versions we had previously gleaned from the export tag pass
-				Linker->Summary.SetCustomVersionContainer(ExportTaggerArchive.GetCustomVersions());
-
-				Linker->SetPortFlags(ComparisonFlags);
-				Linker->SetFilterEditorOnly( FilterEditorOnly );
-				Linker->SetCookData(SaveArgs.ArchiveCookData);
-
-				Linker->SetUseUnversionedPropertySerialization(bSaveUnversionedProperties);
-
-				// Make sure the package has the same version as the linker
-				InOuter->SetLinkerPackageVersion(Linker->UEVer());
-				InOuter->SetLinkerLicenseeVersion(Linker->LicenseeUEVer());
-				InOuter->SetLinkerCustomVersions(Linker->GetCustomVersions());
-
-				if (EndSavingIfCancelled())
-				{ 
-					return ESavePackageResult::Canceled; 
-				}
-				SlowTask.EnterProgressFrame();
-			
-				// keep a list of objects that would normally have gone into the dependency map, but since they are from cross-level dependencies, won't be found in the import map
-				TArray<UObject*> DependenciesToIgnore;
-
-				// When cooking, strip export objects that 
-				//	are both not for client and not for server by default
-				//	are for client if target is client only
-				//	are for server if target is server only
-				if (Linker->IsCooking())
-				{
-					TArray<UObject*> TagExpObjects;
-					GetObjectsWithAnyMarks(TagExpObjects, OBJECTMARK_TagExp);
-
-					const EObjectMark ExcludedObjectMarks = SavePackageUtilities::GetExcludedObjectMarksForTargetPlatform(TargetPlatform);
-					if (Linker->IsCooking() && ExcludedObjectMarks != OBJECTMARK_NOMARKS)
-					{
-						// Make sure that nothing is in the export table that should have been filtered out
-						for (UObject* ObjExport : TagExpObjects)
-						{
-							if (!ensureMsgf(!ObjExport->HasAnyMarks(ExcludedObjectMarks), TEXT("Object %s is marked for export, but has excluded mark!"), *ObjExport->GetPathName()))
-							{
-								ObjExport->UnMark(OBJECTMARK_TagExp);
-							}
-						}
-						GetObjectsWithAnyMarks(TagExpObjects, OBJECTMARK_TagExp);
-					}
-
-					// Exports got filtered out already if they're not for this platform
-					if (TagExpObjects.Num() == 0)
-					{
-						UE_CLOG(!(SaveFlags & SAVE_NoError), LogSavePackage, Verbose, TEXT("No exports found (or all exports are editor-only) for %s. Package will not be saved."), *BaseFilename);
-						return ESavePackageResult::ContainsEditorOnlyData;
-					}
-				}
-
-				// Import objects & names.
-				TSet<UPackage*> PrestreamPackages;
-				TSet<UObject*> ImportsUsedInGame;
-				{
-					SCOPED_SAVETIMER(UPackage_Save_TagImports);
-					
-					TArray<UObject*> TagExpObjects;
-					GetObjectsWithAnyMarks(TagExpObjects, OBJECTMARK_TagExp);
-					for(int32 Index = 0; Index < TagExpObjects.Num(); Index++)
-					{
-						UObject* Obj = TagExpObjects[Index];
-						check(Obj->HasAnyMarks(OBJECTMARK_TagExp));
-
-						// Build list.
-						bool bReferencerIsEditorOnly = IsEditorOnlyObject(Obj, true /* bCheckRecursive */, true /* bCheckMarks */) && !Obj->HasNonEditorOnlyReferences();
-						FArchiveSaveTagImports ImportTagger(Linker.Get(), NameMapSaver, ImportsUsedInGame, SoftPackagesUsedInGame, bReferencerIsEditorOnly);
-						ImportTagger.SetPortFlags(ComparisonFlags);
-						ImportTagger.SetFilterEditorOnly(FilterEditorOnly);
-						ImportTagger.SetSerializeContext(SaveContext);
-
-						UClass* Class = Obj->GetClass();
-
-						if ( Obj->HasAnyFlags(RF_ClassDefaultObject) )
-						{
-							Class->SerializeDefaultObject(Obj, ImportTagger);
-						}
-						else
-						{
-							Obj->Serialize( ImportTagger );
-						}
-
-						ImportTagger << Class;
-
-						// Obj can be saved in package different than their outer, if our outer isn't the package being saved check if we need to tag it as import
-						UObject* Outer = Obj->GetOuter();
-						if (Outer->GetOutermost() != InOuter)
-						{
-							ImportTagger << Outer;
-						}
-						
-						UObject* Template = Obj->GetArchetype();
-						if (Template)
-						{
-							if (Template != Class->GetDefaultObject() || TargetPlatform)
-							{
-								ImportTagger << Template;
-							}
-
-							static struct FDumpChangesSettings
-							{
-								FString ObjectName;
-								FString ArchetypeName;
-
-								FDumpChangesSettings()
-								{
-									const TCHAR* CommandLine = FCommandLine::Get();
-
-									// Check if we want to dump objects by name
-									FString LocalObjectName;
-									if (FParse::Value(CommandLine, TEXT("dumpsavestate="), LocalObjectName))
-									{
-										UE_LOG(LogSavePackage, Warning, TEXT("The -dumpsavestate command line argument is now deprecated. It will soon be removed in a future release."));
-										ObjectName = MoveTemp(LocalObjectName);
-									}
-
-									// Check if we want to dump objects by their CDO name
-									FString LocalArchetypeName;
-									if (FParse::Value(CommandLine, TEXT("dumpsavestatebyarchetype="), LocalArchetypeName))
-									{
-										UE_LOG(LogSavePackage, Warning, TEXT("The -dumpsavestatebyarchetype command line argument is now deprecated. It will soon be removed in a future release."));
-										ArchetypeName = MoveTemp(LocalArchetypeName);
-									}
-								}
-							} DumpChangesSettings;
-
-							// Dump objects and their CDO during save to show how those objects are being delta-serialized
-							if (Obj->GetFName() == *DumpChangesSettings.ObjectName || Template->GetFName() == *DumpChangesSettings.ArchetypeName)
-							{
-								auto DumpPropertiesToText = [](UObject* Object)
-								{
-									TArray<TTuple<FProperty*, FString>> Result;
-									for (FProperty* Prop : TFieldRange<FProperty>(Object->GetClass()))
-									{
-										FString PropState;
-										Prop->ExportTextItem_InContainer(PropState, Object, nullptr, Object, PPF_None);
-
-										Result.Emplace(Prop, MoveTemp(PropState));
-									}
-									return Result;
-								};
-
-								TArray<TTuple<FProperty*, FString>> TemplateOutput = DumpPropertiesToText(Template);
-								TArray<TTuple<FProperty*, FString>> ObjOutput      = DumpPropertiesToText(Obj);
-
-								FString TemplateText = FString::JoinBy(TemplateOutput, TEXT("\n"), [](const TTuple<FProperty*, FString>& PropValue)
-								{
-									return FString::Printf(TEXT("  %s: %s"), *PropValue.Get<0>()->GetName(), *PropValue.Get<1>());
-								});
-								FString ObjText = FString::JoinBy(ObjOutput, TEXT("\n"), [](const TTuple<FProperty*, FString>& PropValue)
-								{
-									return FString::Printf(TEXT("  %s: %s"), *PropValue.Get<0>()->GetName(), *PropValue.Get<1>());
-								});
-								UE_LOG(LogSavePackage, Warning, TEXT("---\nArchetype: %s\n%s\nObject: %s\n%s\n---"), *Template->GetFullName(), *TemplateText, *Obj->GetFullName(), *ObjText);
-							}
-						}
-
-						if (TargetPlatform)
-						{
-							TArray<UObject*> Deps;
-							Obj->GetPreloadDependencies(Deps);
-							for (UObject* Dep : Deps)
-							{
-								// We assume nothing in coreuobject ever loads assets in a constructor
-								if (Dep && Dep->GetOutermost()->GetFName() != GLongCoreUObjectPackageName)
-								{
-									// We want to tag these as imports, but not as dependencies
-									FArchiveSaveTagImports::FScopeIgnoreDependencies IgnoreDependencies(ImportTagger);
-									ImportTagger << Dep;
-								}
-							}
-							static const IConsoleVariable* ProcessPrestreamingRequests = IConsoleManager::Get().FindConsoleVariable(TEXT("s.ProcessPrestreamingRequests"));
-							if (ProcessPrestreamingRequests->GetInt())
-							{
-								Deps.Reset();
-								Obj->GetPrestreamPackages(Deps);
-								for (UObject* Dep : Deps)
-								{
-									if (Dep)
-									{
-										UPackage* Pkg = Dep->GetOutermost();
-										if (!Pkg->HasAnyPackageFlags(PKG_CompiledIn) && Obj->HasAnyMarks(OBJECTMARK_TagExp))
-										{
-											PrestreamPackages.Add(Pkg);
-										}
-									}
-								}
-							}
-						}
-
-						if( Obj->IsInPackage(GetTransientPackage()) )
-						{
-							UE_LOG(LogSavePackage, Fatal, TEXT("%s"), *FString::Printf( TEXT("Transient object imported: %s"), *Obj->GetFullName() ) );
-						}
-
-						if (Obj->GetClass() != UObjectRedirector::StaticClass())
-						{
-							DependenciesReferencedByNonRedirectors.Append(ImportTagger.Dependencies);
-						}
-						ObjectDependencies.Add(Obj, MoveTemp(ImportTagger.Dependencies));
-						NativeObjectDependencies.Add(Obj, MoveTemp(ImportTagger.NativeDependencies));
-					}
-				}
-				if (PrestreamPackages.Num())
-				{
-					TSet<UPackage*> KeptPrestreamPackages;
-					for (UPackage* Pkg : PrestreamPackages)
-					{
-						if (!Pkg->HasAnyMarks(OBJECTMARK_TagImp))
-						{
-							Pkg->Mark(OBJECTMARK_TagImp);
-							ImportsUsedInGame.Add(Pkg);
-							KeptPrestreamPackages.Add(Pkg);
-						}
-					}
-					Exchange(PrestreamPackages, KeptPrestreamPackages);
-				}
-
-#if WITH_EDITOR
-				// Remove TagExp from duplicate objects.
-				TMap<UObject*, UObject*> DuplicateRedirects = UnmarkExportTagFromDuplicates();
-#endif // WITH_EDITOR
-
-				if ( EndSavingIfCancelled() )
-				{ 
-					return ESavePackageResult::Canceled;
-				}
-				SlowTask.EnterProgressFrame();
-
-				if ( EndSavingIfCancelled() )
-				{ 
-					return ESavePackageResult::Canceled;
-				}
-				SlowTask.EnterProgressFrame();
-
-				TArray<UObject*> PrivateObjects;
-				TArray<UObject*> ObjectsInOtherMaps;
-				TArray<UObject*> LevelObjects;
-
-				// Tag the names for all relevant object, classes, and packages.
-				{
-					SCOPED_SAVETIMER(UPackage_Save_TagNames);
-					// Gather the top level objects to validate references
-					TArray<UObject*> TopLevelObjects;
-					GetObjectsWithPackage(InOuter, TopLevelObjects, false);
-					auto IsInAnyTopLevelObject = [&TopLevelObjects](UObject* InObject) -> bool
-					{
-						for (UObject* TopObject : TopLevelObjects)
-						{
-							if (InObject->IsInOuter(TopObject))
-							{
-								return true;
-							}
-						}
-						return false;
-					};
-					auto AnyTopLevelObjectIsIn = [&TopLevelObjects](UObject* InObject) -> bool
-					{
-						for (UObject* TopObject : TopLevelObjects)
-						{
-							if (TopObject->IsInOuter(InObject))
-							{
-								return true;
-							}
-						}
-						return false;
-					};
-					auto AnyTopLevelObjectHasSameOutermostObject = [&TopLevelObjects](UObject* InObject) -> bool
-					{
-						UObject* Outermost = InObject->GetOutermostObject();
-						for (UObject* TopObject : TopLevelObjects)
-						{
-							if (TopObject->GetOutermostObject() == Outermost)
-							{
-								return true;
-							}
-						}
-						return false;
-
-					};
-
-					TArray<UObject*> TagExpImpObjects;
-					GetObjectsWithAnyMarks(TagExpImpObjects, EObjectMark(OBJECTMARK_TagExp|OBJECTMARK_TagImp));
-					for(int32 Index = 0; Index < TagExpImpObjects.Num(); Index++)
-					{
-						UObject* Obj = TagExpImpObjects[Index];
-						check(Obj->HasAnyMarks(EObjectMark(OBJECTMARK_TagExp|OBJECTMARK_TagImp)));
-
-						NameMapSaver.MarkNameAsReferenced(Obj->GetFName());
-						if( Obj->GetOuter() )
-						{
-							NameMapSaver.MarkNameAsReferenced(Obj->GetOuter()->GetFName());
-						}
-
-						if( Obj->HasAnyMarks(OBJECTMARK_TagImp) )
-						{
-							// Make sure the package name of an import is referenced as it might be different than its outer
-							UPackage* ObjPackage = Obj->GetPackage();
-							check(ObjPackage);
-							NameMapSaver.MarkNameAsReferenced(ObjPackage->GetFName());
-
-							NameMapSaver.MarkNameAsReferenced(Obj->GetClass()->GetFName());
-							check(Obj->GetClass()->GetOuter());
-							NameMapSaver.MarkNameAsReferenced(Obj->GetClass()->GetOuter()->GetFName());
-						
-							// if a private object was marked by the cooker, it will be in memory on load, and will be found. helps with some objects
-							// from a package in a package being moved into Startup_int.xxx, but not all
-							// Imagine this:
-							// Package P:
-							//   - A (private)
-							//   - B (public, references A)
-							//   - C (public, references A)
-							// Map M:
-							//   - MObj (references B)
-							// Startup Package S:
-							//   - SObj (references C)
-							// When Startup is cooked, it will pull in C and A. When M is cooked, it will pull in B, but not A, because
-							// A was already marked by the cooker. M.xxx now has a private import to A, which is normally illegal, hence
-							// the OBJECTMARK_MarkedByCooker check below
-							if (PrestreamPackages.Contains(ObjPackage))
-							{
-								NameMapSaver.MarkNameAsReferenced(SavePackageUtilities::NAME_PrestreamPackage);
-								// These are not errors
-								UE_LOG(LogSavePackage, Display, TEXT("Prestreaming package %s "), *ObjPackage->GetPathName()); //-V595
-								continue;
-							}
-
-							// if this import shares a outer with top level object of this package then the reference is acceptable if we aren't cooking
-							if (!bIsCooking && (IsInAnyTopLevelObject(Obj) || AnyTopLevelObjectIsIn(Obj) || AnyTopLevelObjectHasSameOutermostObject(Obj)))
-							{
-								continue;
-							}
-
-							if( !Obj->HasAnyFlags(RF_Public) && !Obj->HasAnyFlags(RF_Transient))
-							{
-								if (!TargetPlatform || !ObjPackage->HasAnyPackageFlags(PKG_CompiledIn))
-								{
-									PrivateObjects.Add(Obj);
-								}
-							}
-
-							// See whether the object we are referencing is in another map package.
-							if(ObjPackage->ContainsMap())
-							{
-								if ( ObjPackage != Obj && Obj->GetFName() != NAME_PersistentLevel && Obj->GetClass()->GetFName() != SavePackageUtilities::NAME_World )
-								{
-									ObjectsInOtherMaps.Add(Obj);
-
-									if ( DependenciesReferencedByNonRedirectors.Contains(Obj) )
-									{
-										UE_LOG(LogSavePackage, Warning, TEXT( "Obj in another map: %s"), *Obj->GetFullName() );
-									}
-								}
-								else
-								{
-									LevelObjects.Add(Obj);
-								}
-							}
-						}
-					}
-				}
-
-				if ( EndSavingIfCancelled() )
-				{ 
-					return ESavePackageResult::Canceled;
-				}
-				SlowTask.EnterProgressFrame();
-
-				if ( LevelObjects.Num() > 0 && ObjectsInOtherMaps.Num() == 0 )
-				{
-					ObjectsInOtherMaps = LevelObjects;
-				}
-
-				// It is allowed for redirectors to reference objects in other maps.
-				// Form the list of objects that erroneously reference another map.
-				TArray<UObject*> IllegalObjectsInOtherMaps;
-				for ( auto ObjIt = ObjectsInOtherMaps.CreateConstIterator(); ObjIt; ++ObjIt )
-				{
-					if ( DependenciesReferencedByNonRedirectors.Contains(*ObjIt) )
-					{
-						IllegalObjectsInOtherMaps.Add(*ObjIt);
-					}
-				}
-
-				// The graph is linked to objects in a different map package!
-				if (IllegalObjectsInOtherMaps.Num() )
-				{
-					UObject* MostLikelyCulprit = nullptr;
-
-					// construct a string containing up to the first 5 objects problem objects
-					FString ObjectNames;
-					int32 MaxNamesToDisplay = 5;
-					bool DisplayIsLimited = true;
-
-					if (IllegalObjectsInOtherMaps.Num() < MaxNamesToDisplay)
-					{
-						MaxNamesToDisplay = IllegalObjectsInOtherMaps.Num();
-						DisplayIsLimited = false;
-					}
-
-					for (int32 Idx = 0; Idx < MaxNamesToDisplay; Idx++)
-					{
-						ObjectNames += IllegalObjectsInOtherMaps[Idx]->GetName() + TEXT("\n");
-					}
-					
-					// if there are more than 5 items we indicated this by adding "..." at the end of the list
-					if (DisplayIsLimited)
-					{
-						ObjectNames += TEXT("...\n");
-					}
-
-					Args.Empty();
-					Args.Add( TEXT("FileName"), FText::FromString( Filename ) );
-					Args.Add( TEXT("ObjectNames"), FText::FromString( ObjectNames ) );
-					const FText Message = FText::Format( NSLOCTEXT("Core", "LinkedToObjectsInOtherMap_FindCulpritQ", "Can't save {FileName}: Graph is linked to object(s) in external map.\nExternal Object(s):\n{ObjectNames}  \nTry to find the chain of references to that object (may take some time)?"), Args );
-
-					FString CulpritString = TEXT( "Unknown" );
-					bool bFindCulprit = true;
-					if (bFindCulprit)
-					{
-						FString Referencer;
-						SavePackageUtilities::FindMostLikelyCulprit(IllegalObjectsInOtherMaps, MostLikelyCulprit, Referencer);
-						if (MostLikelyCulprit != nullptr)
-						{
-							CulpritString = FString::Printf(TEXT("%s (%s)"), *MostLikelyCulprit->GetFullName(), *Referencer);
-						}
-					}
-					else if (!ObjectNames.IsEmpty())
-					{
-						CulpritString = ObjectNames;
-					}
-
-					FString ErrorMessage = FString::Printf(TEXT("Can't save %s: Graph is linked to object %s in external map"), Filename, *CulpritString);
-					if (!(SaveFlags & SAVE_NoError))
-					{
-						Error->Logf(ELogVerbosity::Warning, TEXT("%s"), *ErrorMessage);
-					}
-					else
-					{
-						UE_LOG(LogSavePackage, Error, TEXT("%s"), *ErrorMessage);
-					}
-					return ESavePackageResult::Error;
-				}
-
-				// The graph is linked to private objects!
-				if (PrivateObjects.Num())
-				{
-					UObject* MostLikelyCulprit = nullptr;
-					
-					// construct a string containing up to the first 5 objects problem objects
-					FString ObjectNames;
-					int32 MaxNamesToDisplay = 5;
-					bool DisplayIsLimited = true;
-
-					if (PrivateObjects.Num() < MaxNamesToDisplay)
-					{
-						MaxNamesToDisplay = PrivateObjects.Num();
-						DisplayIsLimited = false;
-					}
-
-					for (int32 Idx = 0; Idx < MaxNamesToDisplay; Idx++)
-					{
-						ObjectNames += PrivateObjects[Idx]->GetName() + TEXT("\n");
-					}
-
-					// if there are more than 5 items we indicated this by adding "..." at the end of the list
-					if (DisplayIsLimited)
-					{
-						ObjectNames += TEXT("...\n");
-					}
-
-					Args.Empty();
-					Args.Add( TEXT("FileName"), FText::FromString( Filename ) );
-					Args.Add( TEXT("ObjectNames"), FText::FromString( ObjectNames ) );
-					const FText Message = FText::Format( NSLOCTEXT("Core", "LinkedToPrivateObjectsInOtherPackage_FindCulpritQ", "Can't save {FileName}: Graph is linked to private object(s) in an external package.\nExternal Object(s):\n{ObjectNames}  \nTry to find the chain of references to that object (may take some time)?"), Args );
-
-					FString CulpritString = TEXT( "Unknown" );
-					bool bFindCulprit = true;
-					if (bFindCulprit)
-					{
-						FString Referencer;
-						SavePackageUtilities::FindMostLikelyCulprit(PrivateObjects, MostLikelyCulprit, Referencer);
-						CulpritString = FString::Printf(TEXT("%s (%s)"),
-							(MostLikelyCulprit != nullptr) ? *MostLikelyCulprit->GetFullName() : TEXT("(unknown culprit)"),
-							*Referencer);
-					}
-
-					if (!(SaveFlags & SAVE_NoError))
-					{
-						Error->Logf(ELogVerbosity::Warning, TEXT("Can't save %s: Graph is linked to external private object %s"), Filename, *CulpritString);
-					}
-					return ESavePackageResult::Error;
-				}
-
-				// Write fixed-length file summary to overwrite later.
-				if (SaveFlags & SAVE_KeepGUID)
-				{
-#if !UE_STRIP_DEPRECATED_PROPERTIES
-					// First generation file, keep existing GUID
-					PRAGMA_DISABLE_DEPRECATION_WARNINGS
-					Linker->Summary.Guid = InOuter->Guid;
-					PRAGMA_ENABLE_DEPRECATION_WARNINGS
-#endif
-#if WITH_EDITORONLY_DATA
-					Linker->Summary.PersistentGuid = InOuter->PersistentGuid;
-#endif
-					Linker->Summary.Generations = TArray<FGenerationInfo>();
-				}
-				else
-				{
-					// First generation file.
-					PRAGMA_DISABLE_DEPRECATION_WARNINGS
-					Linker->Summary.Guid = FGuid::NewGuid();
-					PRAGMA_ENABLE_DEPRECATION_WARNINGS
-#if WITH_EDITORONLY_DATA
-					Linker->Summary.PersistentGuid = InOuter->PersistentGuid;
-#endif
-					Linker->Summary.Generations = TArray<FGenerationInfo>();
-
-#if !UE_STRIP_DEPRECATED_PROPERTIES
-					// make sure the UPackage's copy of the GUID is up to date
-					PRAGMA_DISABLE_DEPRECATION_WARNINGS
-					InOuter->Guid = Linker->Summary.Guid;
-					PRAGMA_ENABLE_DEPRECATION_WARNINGS
-#endif
-				}
-				new(Linker->Summary.Generations)FGenerationInfo(0, 0);
-
-				{
-#if WITH_EDITOR
-					FArchiveStackTraceIgnoreScope IgnoreSummaryDiffsScope(bIgnoreHeaderDiffs);
-#endif // WITH_EDITOR
-					if (!bTextFormat)
-					{
-						StructuredArchiveRoot.GetUnderlyingArchive() << Linker->Summary;
-					}
-				}
-				int32 OffsetAfterPackageFileSummary = Linker->Tell();
-		
-				if ( EndSavingIfCancelled() )
-				{ 
-					return ESavePackageResult::Canceled;
-				}
-				SlowTask.EnterProgressFrame();
-
-
-#if WITH_EDITOR
-				if ( GOutputCookingWarnings )
-				{
-					// check the name list for uniqueobjectnamefor cooking
-					static FNameEntryId NAME_UniqueObjectNameForCooking = FName("UniqueObjectNameForCooking").GetComparisonIndex();
-					if (NameMapSaver.NameExists(NAME_UniqueObjectNameForCooking))
-					{
-						UE_LOG(LogSavePackage, Warning, TEXT("Saving object into cooked package %s which was created at cook time"), Filename);
-					}
-				}
-#endif
-
-				// Build NameMap.
-				Linker->Summary.NameOffset = Linker->Tell();
-				{
-					SCOPED_SAVETIMER(UPackage_Save_BuildNameMap);
-#if WITH_EDITOR
-					FArchiveStackTraceIgnoreScope IgnoreSummaryDiffsScope(bIgnoreHeaderDiffs);
-#endif
-					NameMapSaver.UpdateLinker(*Linker.Get(), bTextFormat ? nullptr : Linker->Saver);
-				}
-				if ( EndSavingIfCancelled() )
-				{ 
-					return ESavePackageResult::Canceled;
-				}
-				SlowTask.EnterProgressFrame();
-
-				FStructuredArchive::FStream Stream = StructuredArchiveRoot.EnterStream(TEXT("GatherableTextData"));
-				Linker->Summary.GatherableTextDataOffset = 0;
-				Linker->Summary.GatherableTextDataCount = 0;
-				if (!(Linker->Summary.GetPackageFlags() & PKG_FilterEditorOnly))
-				{
-					SCOPED_SAVETIMER(UPackage_Save_WriteGatherableTextData);
-
-					// The Editor version is used as part of the check to see if a package is too old to use the gather cache, so we always have to add it if we have gathered loc for this asset
-					// Note that using custom version here only works because we already added it to the export tagger before the package summary was serialized
-					Linker->UsingCustomVersion(FEditorObjectVersion::GUID);
-
-					bool bCanCacheGatheredText = false;
-					
-					// Gathers from the given package
-					EPropertyLocalizationGathererResultFlags GatherableTextResultFlags = EPropertyLocalizationGathererResultFlags::Empty;
-					FPropertyLocalizationDataGatherer(Linker->GatherableTextDataMap, InOuter, GatherableTextResultFlags);
-
-					// We can only cache packages that:
-					//	1) Don't contain script data, as script data is very volatile and can only be safely gathered after it's been compiled (which happens automatically on asset load).
-					//	2) Don't contain text keyed with an incorrect package localization ID, as these keys will be changed later during save.
-					bCanCacheGatheredText = !EnumHasAnyFlags(GatherableTextResultFlags, EPropertyLocalizationGathererResultFlags::HasScript | EPropertyLocalizationGathererResultFlags::HasTextWithInvalidPackageLocalizationID);
-
-					if (bCanCacheGatheredText)
-					{
-						Linker->Summary.GatherableTextDataOffset = Linker->Tell();
-
-						// Save gatherable text data.
-						Linker->Summary.GatherableTextDataCount = Linker->GatherableTextDataMap.Num();
-						for (FGatherableTextData& GatherableTextData : Linker->GatherableTextDataMap)
-						{
-							Stream.EnterElement() << GatherableTextData;
-						}
-					}
-				}
-
-				if ( EndSavingIfCancelled() )
-				{ 
-					return ESavePackageResult::Canceled;
-				}
-				SlowTask.EnterProgressFrame();
-
-				// Build ImportMap.
-				{
-					SCOPED_SAVETIMER(UPackage_Save_BuildImportMap);
-
-					TArray<UObject*> TagImpObjects;
-
-					const EObjectMark ExcludedObjectMarks = SavePackageUtilities::GetExcludedObjectMarksForTargetPlatform(TargetPlatform);
-					GetObjectsWithAnyMarks(TagImpObjects, OBJECTMARK_TagImp);
-
-					if (Linker->IsCooking() && ExcludedObjectMarks != OBJECTMARK_NOMARKS)
-					{
-						// Make sure that nothing is in the import table that should have been filtered out
-						for (UObject* ObjImport : TagImpObjects)
-						{
-							if (!ensureMsgf(!ObjImport->HasAnyMarks(ExcludedObjectMarks), TEXT("Object %s is marked for import, but has excluded mark!"), *ObjImport->GetPathName()))
-							{
-								ObjImport->UnMark(OBJECTMARK_TagImp);
-							}
-						}
-						GetObjectsWithAnyMarks(TagImpObjects, OBJECTMARK_TagImp);
-					}
-
-					for(int32 Index = 0; Index < TagImpObjects.Num(); Index++)
-					{
-						UObject* Obj = TagImpObjects[Index];
-						check(Obj->HasAnyMarks(OBJECTMARK_TagImp));
-						UClass* ObjClass = Obj->GetClass();
-#if WITH_EDITOR
-						bool bExcludePackageFromCook = FCoreUObjectDelegates::ShouldCookPackageForPlatform.IsBound() ? !FCoreUObjectDelegates::ShouldCookPackageForPlatform.Execute(Obj->GetOutermost(), TargetPlatform) : false;			
-						if (bExcludePackageFromCook)
-						{
-							continue;
-						}
-#endif //WITH_EDITOR
-						FObjectImport* LocObjectImport = new(Linker->ImportMap)FObjectImport(Obj, ObjClass);
-
-						if (PrestreamPackages.Contains((UPackage*)Obj))
-						{
-							LocObjectImport->ClassName = SavePackageUtilities::NAME_PrestreamPackage;
-						}
-					}
-				}
-
-
-				if ( EndSavingIfCancelled() )
-				{ 
-					return ESavePackageResult::Canceled;
-				}
-				SlowTask.EnterProgressFrame();
-
-				// sort imports
-				FObjectImportSortHelper ImportSortHelper;
-				{
-					SCOPED_SAVETIMER(UPackage_Save_SortImports);
-					ImportSortHelper.SortImports(Linker.Get());
-					Linker->Summary.ImportCount = Linker->ImportMap.Num();
-				}
-
-				if ( EndSavingIfCancelled() )
-				{ 
-					return ESavePackageResult::Canceled;
-				}
-				SlowTask.EnterProgressFrame();
-
-				
-				// Build ExportMap.
-				{
-					SCOPED_SAVETIMER(UPackage_Save_BuildExportMap);
-
-					TArray<UObject*> TagExpObjects;
-					GetObjectsWithAnyMarks(TagExpObjects, OBJECTMARK_TagExp);
-					for(int32 Index = 0; Index < TagExpObjects.Num(); Index++)
-					{
-						UObject* Obj = TagExpObjects[Index];
-						check(Obj->HasAnyMarks(OBJECTMARK_TagExp));
-						new( Linker->ExportMap )FObjectExport( Obj, Obj->HasAnyMarks(OBJECTMARK_NotAlwaysLoadedForEditorGame));
-					}
-				}
-
-#if WITH_EDITOR
-				if (GOutputCookingWarnings)
-				{
-					// check the name list for uniqueobjectnamefor cooking
-					static FName NAME_UniqueObjectNameForCooking(TEXT("UniqueObjectNameForCooking"));
-
-					for (const auto& Export : Linker->ExportMap)
-					{
-						const auto& NameInUse = Export.ObjectName;
-						if (NameInUse.GetComparisonIndex() == NAME_UniqueObjectNameForCooking.GetComparisonIndex())
-						{
-							UObject* Outer = Export.Object->GetOuter();
-							UE_LOG(LogSavePackage, Warning, TEXT(" into cooked package %s which was created at cook time, Object Name %s, Full Path %s, Class %s, Outer %s, Outer class %s"), Filename, *NameInUse.ToString(), *Export.Object->GetFullName(), *Export.Object->GetClass()->GetName(), Outer ? *Outer->GetName() : TEXT("None"), Outer ? *Outer->GetClass()->GetName() : TEXT("None") );
-						}
-					}
-				}
-#endif
-
-				if ( EndSavingIfCancelled() )
-				{ 
-					return ESavePackageResult::Canceled;
-				}
-				SlowTask.EnterProgressFrame();
-
-				// Sort exports alphabetically
-				FObjectExportSortHelper ExportSortHelper;
-				{
-					SCOPED_SAVETIMER(UPackage_Save_SortExports);
-					ExportSortHelper.SortExports(Linker.Get());
-				}
-				
-				// Sort exports for seek-free loading.
-				if (Linker->IsCooking())
-				{
-					SCOPED_SAVETIMER(UPackage_Save_SortExportsForSeekFree);
-					COOK_STAT(FScopedDurationTimer SaveTimer(FSavePackageStats::SortExportsSeekfreeInnerTimeSec));
-					FObjectExportSeekFreeSorter SeekFreeSorter;
-					SeekFreeSorter.SortExports(Linker.Get());
-				}
-
-				Linker->Summary.ExportCount = Linker->ExportMap.Num();
-
-				if ( EndSavingIfCancelled() )
-				{ 
-					return ESavePackageResult::Canceled;
-				}
-				SlowTask.EnterProgressFrame();
-
-				// Pre-size depends map.
-				Linker->DependsMap.AddZeroed( Linker->ExportMap.Num() );
-
-				// track import and export object linker index
-				TMap<UObject*, FPackageIndex> ImportToIndexMap;
-				TMap<UObject*, FPackageIndex> ExportToIndexMap;
-				for (int32 ImpIndex = 0; ImpIndex < Linker->ImportMap.Num(); ImpIndex++)
-				{
-					ImportToIndexMap.Add(Linker->ImportMap[ImpIndex].XObject, FPackageIndex::FromImport(ImpIndex));
-				}
-
-				for (int32 ExpIndex = 0; ExpIndex < Linker->ExportMap.Num(); ExpIndex++)
-				{
-					ExportToIndexMap.Add(Linker->ExportMap[ExpIndex].Object, FPackageIndex::FromExport(ExpIndex));
-				}
-
-				// go back over the (now sorted) exports and fill out the DependsMap
-				{
-					SCOPED_SAVETIMER(UPackage_Save_BuildExportDependsMap);
-					for (int32 ExpIndex = 0; ExpIndex < Linker->ExportMap.Num(); ExpIndex++)
-					{
-						UObject* Object = Linker->ExportMap[ExpIndex].Object;
-						if (Object == nullptr)
-						{
-							UE_LOG(LogSavePackage, Warning, TEXT("Object is missing for an export, unable to save dependency map. See log for more info"));
-							if (!(SaveFlags & SAVE_NoError))
-							{
-								Error->Logf(ELogVerbosity::Warning, TEXT("%s"), *FText::Format(NSLOCTEXT("Core", "SavePackageObjectIsMissingExport", "Object is missing for an export, unable to save dependency map for asset '{0}'. See log for more info"), FText::FromString(FString(Filename))).ToString());
-							}
-							continue;
-						}
-
-						// add a dependency map entry also
-						TArray<FPackageIndex>& DependIndices = Linker->DependsMap[ExpIndex];
-						// find all the objects needed by this export
-						TArray<UObject*>* SrcDepends = ObjectDependencies.Find(Object);
-						checkf(SrcDepends, TEXT("Couldn't find dependency map for %s"), *Object->GetFullName());
-
-						// go through each object and...
-						DependIndices.Reserve(SrcDepends->Num());
-						for (int32 DependIndex = 0; DependIndex < SrcDepends->Num(); DependIndex++)
-						{
-							UObject* DependentObject = (*SrcDepends)[DependIndex];
-
-							FPackageIndex DependencyIndex;
-
-							// if the dependency is in the same package, we need to save an index into our ExportMap
-							if (DependentObject->GetOutermost() == Linker->LinkerRoot)
-							{
-								// ... find the associated ExportIndex
-								DependencyIndex = ExportToIndexMap.FindRef(DependentObject);
-							}
-							// otherwise we need to save an index into the ImportMap
-							else
-							{
-								// ... find the associated ImportIndex
-								DependencyIndex = ImportToIndexMap.FindRef(DependentObject);
-							}
-
-#if WITH_EDITOR
-							// If we still didn't find index, maybe it was a duplicate export which got removed.
-							// Check if we have a redirect to original.
-							if (DependencyIndex.IsNull() && DuplicateRedirects.Contains(DependentObject))
-							{
-								UObject** const RedirectObj = DuplicateRedirects.Find(DependentObject);
-								if (RedirectObj)
-								{
-									DependencyIndex = ExportToIndexMap.FindRef(*RedirectObj);
-								}
-							}
-#endif
-							// if we didn't find it (FindRef returns 0 on failure, which is good in this case), then we are in trouble, something went wrong somewhere
-							checkf(!DependencyIndex.IsNull(), TEXT("Failed to find dependency index for %s (%s)"), *DependentObject->GetFullName(), *Object->GetFullName());
-
-							// add the import as an import for this export
-							DependIndices.Add(DependencyIndex);
-						}
-					}
-				}
-
-
-				if ( EndSavingIfCancelled() )
-				{ 
-					return ESavePackageResult::Canceled;
-				}
-				SlowTask.EnterProgressFrame();
-
-				// Set linker reverse mappings.
-				// also set netplay required data for any UPackages in the export map
-				for( int32 i=0; i<Linker->ExportMap.Num(); i++ )
-				{
-					UObject* Object = Linker->ExportMap[i].Object;
-					if( Object )
-					{
-						Linker->ObjectIndicesMap.Add(Object, FPackageIndex::FromExport(i));
-
-						UPackage* Package = dynamic_cast<UPackage*>(Object);
-						if (Package != nullptr)
-						{
-							Linker->ExportMap[i].PackageFlags = Package->GetPackageFlags();
-						}
-					}
-				}
-
-				if ( EndSavingIfCancelled() )
-				{ 
-					return ESavePackageResult::Canceled;
-				}
-				SlowTask.EnterProgressFrame();
-
-				// If this is a map package, make sure there is a world or level in the export map.
-				if ( InOuter->ContainsMap() )
-				{
-					bool bContainsMap = false;
-					for( int32 i=0; i<Linker->ExportMap.Num(); i++ )
-					{
-						UObject* Object = Linker->ExportMap[i].Object;
-						
-						// Consider redirectors to world/levels as map packages too.
-						if (UObjectRedirector* Redirector = Cast<UObjectRedirector>(Object))
-						{
-							Object = Redirector->DestinationObject;
-						}
-							
-						if ( Object )
-						{
-							const FString ExportClassName = Object->GetClass()->GetName();
-							if( ExportClassName == TEXT("World") || ExportClassName == TEXT("Level") )
-							{
-								bContainsMap = true;
-								break;
-							}
-						}
-					}
-					if (!bContainsMap)
-					{
-						ensureMsgf(false, TEXT("Attempting to save a map package '%s' that does not contain a map object."), *InOuter->GetName());
-						UE_LOG(LogSavePackage, Error, TEXT("Attempting to save a map package '%s' that does not contain a map object."), *InOuter->GetName());
-
-						if (!(SaveFlags & SAVE_NoError))
-						{
-							Error->Logf(ELogVerbosity::Warning, TEXT("%s"), *FText::Format( NSLOCTEXT( "Core", "SavePackageNoMap", "Attempting to save a map asset '{0}' that does not contain a map object" ), FText::FromString( FString( Filename ) ) ).ToString() );
-						}
-						bSuccess = false;
-					}
-				}
-
-
-				if ( EndSavingIfCancelled() )
-				{ 
-					return ESavePackageResult::Canceled;
-				}
-				SlowTask.EnterProgressFrame();
-
-				for( int32 i=0; i<Linker->ImportMap.Num(); i++ )
-				{
-					UObject* Object = Linker->ImportMap[i].XObject;
-					if( Object != nullptr )
-					{
-						const FPackageIndex PackageIndex = FPackageIndex::FromImport(i);
-						//ensure(!Linker->ObjectIndicesMap.Contains(Object)); // this ensure will fail
-						Linker->ObjectIndicesMap.Add(Object, PackageIndex);
-					}
-					else
-					{
-						checkf(false, TEXT("NULL XObject for import %i - Object: %s Class: %s"), i, *Linker->ImportMap[i].ObjectName.ToString(), *Linker->ImportMap[i].ClassName.ToString());
-					}
-				}
-				if (TargetPlatform)
-				{
-					EDLCookChecker.AddExport(InOuter); // the package isn't actually in the export map, but that is ok, we add it as export anyway for error checking
-					for (int32 i = 0; i < Linker->ImportMap.Num(); i++)
-					{
-						UObject* Object = Linker->ImportMap[i].XObject;
-						if (Object != nullptr)
-						{
-							EDLCookChecker.AddImport(Object, InOuter);
-						}
-					}
-				}
-
-				// Convert the searchable names map from UObject to packageindex
-				for (const TPair<const UObject *, TArray<FName>>& SearchableNamePair : Linker->SearchableNamesObjectMap)
-				{
-					const FPackageIndex PackageIndex = Linker->MapObject(SearchableNamePair.Key);
-
-					// This should always be in the imports already
-					if (ensure(!PackageIndex.IsNull()))
-					{
-						Linker->SearchableNamesMap.FindOrAdd(PackageIndex) = SearchableNamePair.Value;
-					}
-				}
-				// Clear temporary map
-				Linker->SearchableNamesObjectMap.Empty();
-
-				SlowTask.EnterProgressFrame();
-
-				// Find components referenced by exports.
-
-				if ( EndSavingIfCancelled() )
-				{ 
-					return ESavePackageResult::Canceled;
-				}
-				SlowTask.EnterProgressFrame();
-
-				// Save dummy import map, overwritten later.
-				if (!bTextFormat)
-				{
-					SCOPED_SAVETIMER(UPackage_Save_WriteDummyImportMap);
-#if WITH_EDITOR
-					FArchiveStackTraceIgnoreScope IgnoreSummaryDiffsScope(bIgnoreHeaderDiffs);
-#endif // WITH_EDITOR
-					Linker->Summary.ImportOffset = Linker->Tell();
-					for (int32 i = 0; i < Linker->ImportMap.Num(); i++)
-					{
-						FObjectImport& Import = Linker->ImportMap[i];
-						StructuredArchiveRoot.GetUnderlyingArchive() << Import;
-					}
-				}
-				int32 OffsetAfterImportMap = Linker->Tell();
-
-
-				if (EndSavingIfCancelled())
-				{
-					return ESavePackageResult::Canceled;
-				}
-				SlowTask.EnterProgressFrame();
-
-				// Save dummy export map, overwritten later.
-				if (!bTextFormat)
-				{
-					SCOPED_SAVETIMER(UPackage_Save_WriteDummyExportMap);
-#if WITH_EDITOR
-					FArchiveStackTraceIgnoreScope IgnoreSummaryDiffsScope(bIgnoreHeaderDiffs);
-#endif // WITH_EDITOR
-					Linker->Summary.ExportOffset = Linker->Tell();
-					for (int32 i = 0; i < Linker->ExportMap.Num(); i++)
-					{
-						FObjectExport& Export = Linker->ExportMap[i];
-						*Linker << Export;
-					}
-				}
-				int32 OffsetAfterExportMap = Linker->Tell();
-
-
-				if (EndSavingIfCancelled())
-				{
-					return ESavePackageResult::Canceled;
-				}
-				SlowTask.EnterProgressFrame();
-
-				if (!bTextFormat)
-				{
-					SCOPED_SAVETIMER(UPackage_Save_WriteDependsMap);
-
-					FStructuredArchive::FStream DependsStream = StructuredArchiveRoot.EnterStream(TEXT("DependsMap"));
-					if (Linker->IsCooking())
-					{
-#if WITH_EDITOR
-						FArchiveStackTraceIgnoreScope IgnoreSummaryDiffsScope(bIgnoreHeaderDiffs);
-#endif // WITH_EDITOR
-						//@todo optimization, this should just be stripped entirely from cooked packages
-						TArray<FPackageIndex> Depends; // empty array
-						Linker->Summary.DependsOffset = Linker->Tell();
-						for (int32 i = 0; i < Linker->ExportMap.Num(); i++)
-						{
-							DependsStream.EnterElement() << Depends;
-						}
-					}
-					else
-					{
-						// save depends map (no need for later patching)
-						check(Linker->DependsMap.Num() == Linker->ExportMap.Num());
-						Linker->Summary.DependsOffset = Linker->Tell();
-						for (int32 i = 0; i < Linker->ExportMap.Num(); i++)
-						{
-							TArray<FPackageIndex>& Depends = Linker->DependsMap[i];
-							DependsStream.EnterElement() << Depends;
-						}
-					}
-				}
-
-
-				if (EndSavingIfCancelled())
-				{ 
-					return ESavePackageResult::Canceled;
-				}
-				SlowTask.EnterProgressFrame();
-
-				// Only save string asset and searchable name map if saving for editor
-				if (!(Linker->Summary.GetPackageFlags() & PKG_FilterEditorOnly))
-				{
-					SCOPED_SAVETIMER(UPackage_Save_SaveSoftPackagesAndSearchableNames);
-
-					// Save soft package references
-					Linker->Summary.SoftPackageReferencesOffset = Linker->Tell();
-					Linker->Summary.SoftPackageReferencesCount = Linker->SoftPackageReferenceList.Num();
-					if (!bTextFormat)
-					{
-						FStructuredArchive::FStream SoftReferenceStream = StructuredArchiveRoot.EnterStream(TEXT("SoftReferences"));
-						for (FName& SoftPackageName : Linker->SoftPackageReferenceList)
-						{
-							SoftReferenceStream.EnterElement() << SoftPackageName;
-						}
-
-						// Save searchable names map
-						Linker->Summary.SearchableNamesOffset = Linker->Tell();
-						Linker->SerializeSearchableNamesMap(StructuredArchiveRoot.EnterField(TEXT("SearchableNames")));
-					}
-				}
-				else
-				{
-					Linker->Summary.SoftPackageReferencesCount = 0;
-					Linker->Summary.SoftPackageReferencesOffset = 0;
-					Linker->Summary.SearchableNamesOffset = 0;
-				}
-
-				{
-#if WITH_EDITOR
-					FArchiveStackTraceIgnoreScope IgnoreSummaryDiffsScope(bIgnoreHeaderDiffs);
-#endif // WITH_EDITOR
-
-					// Save thumbnails
-					{
-						SCOPED_SAVETIMER(UPackage_Save_SaveThumbnails);
-						SavePackageUtilities::SaveThumbnails(InOuter, Linker.Get(), StructuredArchiveRoot.EnterField(TEXT("Thumbnails")));
-					}
-
-					if (!bTextFormat)
-					{	
-						// Save asset registry data so the editor can search for information about assets in this package
-						SCOPED_SAVETIMER(UPackage_Save_SaveAssetRegistryData);
-						UE::AssetRegistry::WritePackageData(StructuredArchiveRoot, Linker->IsCooking(), InOuter, Linker.Get(), ImportsUsedInGame, SoftPackagesUsedInGame);
-					}
-
-					// Save level information used by World browser
-					{
-						SCOPED_SAVETIMER(UPackage_Save_WorldLevelData);
-						SavePackageUtilities::SaveWorldLevelInfo(InOuter, Linker.Get(), StructuredArchiveRoot);
-					}
-				}
-
-				// Map export indices
-				{
-					SCOPED_SAVETIMER(UPackage_Save_MapExportIndices);
-
-					for (int32 i = 0; i < Linker->ExportMap.Num(); i++)
-					{
-						FObjectExport& Export = Linker->ExportMap[i];
-						if (Export.Object)
-						{
-							// Set class index.
-							// If this is *exactly* a UClass, store null instead; for anything else, including UClass-derived classes, map it
-							UClass* ObjClass = Export.Object->GetClass();
-							if (ObjClass != UClass::StaticClass())
-							{
-								Export.ClassIndex = Linker->MapObject(ObjClass);
-								checkf(!Export.ClassIndex.IsNull(), TEXT("Export %s class is not mapped when saving %s"), *Export.Object->GetFullName(), *Linker->LinkerRoot->GetName());
-							}
-							else
-							{
-								Export.ClassIndex = FPackageIndex();
-							}
-
-							if (TargetPlatform)
-							{
-								UObject* Archetype = Export.Object->GetArchetype();
-								check(Archetype);
-								check(Archetype->IsA(Export.Object->HasAnyFlags(RF_ClassDefaultObject) ? ObjClass->GetSuperClass() : ObjClass));
-								Export.TemplateIndex = Linker->MapObject(Archetype);
-								UE_CLOG(Export.TemplateIndex.IsNull(), LogSavePackage, Fatal, TEXT("%s was an archetype of %s but returned a null index mapping the object."), *Archetype->GetFullName(), *Export.Object->GetFullName());
-								check(!Export.TemplateIndex.IsNull());
-							}
-
-							// Set the parent index, if this export represents a UStruct-derived object
-							if (UStruct* Struct = dynamic_cast<UStruct*>(Export.Object))
-							{
-								if (Struct->GetSuperStruct() != nullptr)
-								{
-									Export.SuperIndex = Linker->MapObject(Struct->GetSuperStruct());
-									checkf(!Export.SuperIndex.IsNull(),
-										TEXT("Export Struct (%s) of type (%s) inheriting from (%s) of type (%s) has not mapped super struct."),
-										*GetPathNameSafe(Struct),
-										*(Struct->GetClass()->GetName()),
-										*GetPathNameSafe(Struct->GetSuperStruct()),
-										*(Struct->GetSuperStruct()->GetClass()->GetName())
-									);
-								}
-								else
-								{
-									Export.SuperIndex = FPackageIndex();
-								}
-							}
-							else
-							{
-								Export.SuperIndex = FPackageIndex();
-							}
-
-							// Set FPackageIndex for this export's Outer. If the export's Outer
-							// is the UPackage corresponding to this package's LinkerRoot, the
-							if (Export.Object->GetOuter() != InOuter)
-							{
-								check(Export.Object->GetOuter());
-								Export.OuterIndex = Linker->MapObject(Export.Object->GetOuter());
-
-								// The outer of an object is now allowed to be an export hence, the following asserts do not stand anymore:
-								//checkf(Export.Object->GetOuter()->IsInPackage(InOuter),
-								//	TEXT("Export Object (%s) Outer (%s) mismatch."),
-								//	*(Export.Object->GetPathName()),
-								//	*(Export.Object->GetOuter()->GetPathName()));
-								//checkf(!Export.OuterIndex.IsImport(),
-								//	TEXT("Export Object (%s) Outer (%s) is an Import."),
-								//	*(Export.Object->GetPathName()),
-								//	*(Export.Object->GetOuter()->GetPathName()));
-
-								if (Linker->IsCooking())
-								{
-									// Only packages are allowed to have no outer
-									ensureMsgf(Export.OuterIndex != FPackageIndex() || Export.Object->IsA(UPackage::StaticClass()), TEXT("Export %s has no valid outer when cooking!"), *Export.Object->GetPathName());
-								}
-							}
-							else
-							{
-								// this export's Outer is the LinkerRoot for this package
-								Export.OuterIndex = FPackageIndex();
-							}
-						}
-					}
-				}
-
-				Linker->Summary.PreloadDependencyOffset = Linker->Tell();
-				Linker->Summary.PreloadDependencyCount = -1;
-
-				if (Linker->IsCooking())
-				{
-#if WITH_EDITOR
-					FArchiveStackTraceIgnoreScope IgnoreSummaryDiffsScope(bIgnoreHeaderDiffs);
-#endif // WITH_EDITOR
-
-					const EObjectMark ExcludedObjectMarks = SavePackageUtilities::GetExcludedObjectMarksForTargetPlatform(Linker->CookingTarget());
-					Linker->Summary.PreloadDependencyCount = 0;
-
-					auto IncludeObjectAsDependency = [&Linker, ExcludedObjectMarks](int32 CallSite, TSet<FPackageIndex>& AddTo, UObject* ToTest, UObject* ForObj, bool bMandatory, bool bOnlyIfInLinkerTable)
-					{
-						// Skip transient, editor only, and excluded client/server objects
-						if (ToTest)
-						{
-							UPackage* Outermost = ToTest->GetOutermost();
-							check(Outermost);
-							if (Outermost->GetFName() == GLongCoreUObjectPackageName)
-							{
-								return; // We assume nothing in coreuobject ever loads assets in a constructor
-							}
-							FPackageIndex Index = Linker->MapObject(ToTest);
-							if (Index.IsNull() && bOnlyIfInLinkerTable)
-							{
-								return;
-							}
-							if (!Index.IsNull() && (ToTest->HasAllFlags(RF_Transient) && !ToTest->IsNative()))
-							{
-								UE_LOG(LogSavePackage, Warning, TEXT("A dependency '%s' of '%s' is in the linker table, but is transient. We will keep the dependency anyway (%d)."), *ToTest->GetFullName(), *ForObj->GetFullName(), CallSite);
-							}
-							if (!Index.IsNull() && !IsValid(ToTest))
-							{
-								UE_LOG(LogSavePackage, Warning, TEXT("A dependency '%s' of '%s' is in the linker table, but is pending kill or garbage. We will keep the dependency anyway (%d)."), *ToTest->GetFullName(), *ForObj->GetFullName(), CallSite);
-							}
-							bool bNotFiltered = (ExcludedObjectMarks == OBJECTMARK_NOMARKS || !ToTest->HasAnyMarks(ExcludedObjectMarks)) && (!(Linker->Summary.GetPackageFlags() & PKG_FilterEditorOnly) || !IsEditorOnlyObject(ToTest, false, true));
-							if (bMandatory && !bNotFiltered)
-							{
-								UE_LOG(LogSavePackage, Warning, TEXT("A dependency '%s' of '%s' was filtered, but is mandatory. This indicates a problem with editor only stripping. We will keep the dependency anyway (%d)."), *ToTest->GetFullName(), *ForObj->GetFullName(), CallSite);
-								bNotFiltered = true;
-							}
-							if (bNotFiltered)
-							{
-								if (!Index.IsNull())
-								{
-									AddTo.Add(Index);
-									return;
-								}
-								else if (!ToTest->HasAnyFlags(RF_Transient))
-								{
-									UE_CLOG(Outermost->HasAnyPackageFlags(PKG_CompiledIn), LogSavePackage, Verbose, TEXT("A compiled in dependency '%s' of '%s' was not actually in the linker tables and so will be ignored (%d)."), *ToTest->GetFullName(), *ForObj->GetFullName(), CallSite);
-									UE_CLOG(!Outermost->HasAnyPackageFlags(PKG_CompiledIn), LogSavePackage, Fatal, TEXT("A dependency '%s' of '%s' was not actually in the linker tables and so will be ignored (%d)."), *ToTest->GetFullName(), *ForObj->GetFullName(), CallSite);
-								}
-							}
-							check(!bMandatory);
-						}
-					};
-
-					auto IncludeIndexAsDependency = [&Linker](TSet<FPackageIndex>& AddTo, FPackageIndex Dep)
-					{
-						if (!Dep.IsNull())
-						{
-							UObject* ToTest = Dep.IsExport() ? Linker->Exp(Dep).Object : Linker->Imp(Dep).XObject;
-							if (ToTest)
-							{
-								UPackage* Outermost = ToTest->GetOutermost();
-								if (Outermost && Outermost->GetFName() != GLongCoreUObjectPackageName) // We assume nothing in coreuobject ever loads assets in a constructor
-								{
-									AddTo.Add(Dep);
-								}
-							}
-						}
-					};
-
-					FStructuredArchive::FStream DepedenciesStream = StructuredArchiveRoot.EnterStream(TEXT("PreloadDependencies"));
-					TArray<UObject*> Subobjects;
-					TArray<UObject*> Deps;
-					TSet<FPackageIndex> SerializationBeforeCreateDependencies;
-					TSet<FPackageIndex> SerializationBeforeSerializationDependencies;
-					TSet<FPackageIndex> CreateBeforeSerializationDependencies;
-					TSet<FPackageIndex> CreateBeforeCreateDependencies;
-					for (int32 i = 0; i < Linker->ExportMap.Num(); i++)
-					{
-						FObjectExport& Export = Linker->ExportMap[i];
-						if (Export.Object)
-						{
-							EDLCookChecker.AddExport(Export.Object);
-							{
-								SerializationBeforeCreateDependencies.Reset();
-								IncludeIndexAsDependency(SerializationBeforeCreateDependencies, Export.ClassIndex);
-								UObject* CDO = Export.Object->GetArchetype();
-								IncludeObjectAsDependency(1, SerializationBeforeCreateDependencies, CDO, Export.Object, true, false);
-								Subobjects.Reset();
-								GetObjectsWithOuter(CDO, Subobjects);
-								for (UObject* SubObj : Subobjects)
-								{
-									// Only include subobject archetypes
-									if (SubObj->HasAnyFlags(RF_DefaultSubObject | RF_ArchetypeObject))
-									{
-										while (SubObj->HasAnyFlags(RF_Transient)) // transient components are stripped by the ICH, so find the one it will really use at runtime
-										{
-											UObject* SubObjArch = SubObj->GetArchetype();
-											if (SubObjArch->GetClass()->HasAnyClassFlags(CLASS_Native | CLASS_Intrinsic))
-											{
-												break;
-											}
-											SubObj = SubObjArch;
-										}
-										if (IsValid(SubObj))
-										{
-											IncludeObjectAsDependency(2, SerializationBeforeCreateDependencies, SubObj, Export.Object, false, false);
-										}
-									}
-								}
-							}
-							{
-								SerializationBeforeSerializationDependencies.Reset();
-								Deps.Reset();
-								Export.Object->GetPreloadDependencies(Deps);
-
-								for (UObject* Obj : Deps)
-								{
-									IncludeObjectAsDependency(3, SerializationBeforeSerializationDependencies, Obj, Export.Object, false, true);
-								}
-								if (Export.Object->HasAnyFlags(RF_ArchetypeObject | RF_ClassDefaultObject))
-								{
-									UObject *Outer = Export.Object->GetOuter();
-									if (!Outer->IsA(UPackage::StaticClass()))
-									{
-										IncludeObjectAsDependency(4, SerializationBeforeSerializationDependencies, Outer, Export.Object, true, false);
-									}
-								}
-								if (Export.Object->IsA(UClass::StaticClass()))
-								{
-									// we need to load archetypes of our subobjects before we load the class
-									UObject* CDO = CastChecked<UClass>(Export.Object)->GetDefaultObject();
-									Subobjects.Reset();
-									GetObjectsWithOuter(CDO, Subobjects);
-									for (UObject* SubObj : Subobjects)
-									{
-										// Only include subobject archetypes
-										if (SubObj->HasAnyFlags(RF_DefaultSubObject | RF_ArchetypeObject))
-										{
-											SubObj = SubObj->GetArchetype();
-											while (SubObj->HasAnyFlags(RF_Transient)) // transient components are stripped by the ICH, so find the one it will really use at runtime
-											{
-												UObject* SubObjArch = SubObj->GetArchetype();
-												if (SubObjArch->GetClass()->HasAnyClassFlags(CLASS_Native | CLASS_Intrinsic))
-												{
-													break;
-												}
-												SubObj = SubObjArch;
-											}
-											if (IsValid(SubObj))
-											{
-												IncludeObjectAsDependency(5, SerializationBeforeSerializationDependencies, SubObj, Export.Object, false, false);
-											}
-										}
-									}
-								}
-							}
-
-							{
-								CreateBeforeSerializationDependencies.Reset();
-								UClass* Class = Cast<UClass>(Export.Object);
-								UObject* ClassCDO = Class ? Class->GetDefaultObject() : nullptr;
-								{
-									TArray<FPackageIndex>& Depends = Linker->DependsMap[i];
-									for (FPackageIndex Dep : Depends)
-									{
-										UObject* ToTest = Dep.IsExport() ? Linker->Exp(Dep).Object : Linker->Imp(Dep).XObject;
-										if (ToTest != ClassCDO)
-										{
-											IncludeIndexAsDependency(CreateBeforeSerializationDependencies, Dep);
-										}
-									}
-								}
-								{
-									TArray<UObject*>& NativeDeps = NativeObjectDependencies[Export.Object];
-									for (UObject* ToTest : NativeDeps)
-									{
-										if (ToTest != ClassCDO)
-										{
-											IncludeObjectAsDependency(6, CreateBeforeSerializationDependencies, ToTest, Export.Object, false, true);
-										}
-									}
-								}
-							}
-
-							{
-								CreateBeforeCreateDependencies.Reset();
-								IncludeIndexAsDependency(CreateBeforeCreateDependencies, Export.OuterIndex);
-								IncludeIndexAsDependency(CreateBeforeCreateDependencies, Export.SuperIndex);
-							}
-
-							auto AddArcForDepChecking = [&Linker, &Export, &EDLCookChecker](bool bExportIsSerialize, FPackageIndex Dep, bool bDepIsSerialize)
-							{
-								check(Export.Object);
-								check(!Dep.IsNull());
-								UObject* DepObject = Dep.IsExport() ? Linker->Exp(Dep).Object : Linker->Imp(Dep).XObject;
-								check(DepObject);
-
-								Linker->DepListForErrorChecking.Add(Dep);
-
-								EDLCookChecker.AddArc(DepObject, bDepIsSerialize, Export.Object, bExportIsSerialize);
-							};
-
-							for (FPackageIndex Index : SerializationBeforeSerializationDependencies)
-							{
-								if (SerializationBeforeCreateDependencies.Contains(Index))
-								{
-									continue; // if the other thing must be serialized before we create, then this is a redundant dep
-								}
-								if (Export.FirstExportDependency == -1)
-								{
-									Export.FirstExportDependency = Linker->Summary.PreloadDependencyCount;
-									check(Export.SerializationBeforeSerializationDependencies == 0 && Export.CreateBeforeSerializationDependencies == 0 && Export.SerializationBeforeCreateDependencies == 0 && Export.CreateBeforeCreateDependencies == 0);
-								}
-								Linker->Summary.PreloadDependencyCount++;
-								Export.SerializationBeforeSerializationDependencies++;
-								DepedenciesStream.EnterElement() << Index;
-								AddArcForDepChecking(true, Index, true);
-							}
-							for (FPackageIndex Index : CreateBeforeSerializationDependencies)
-							{
-								if (SerializationBeforeCreateDependencies.Contains(Index))
-								{
-									continue; // if the other thing must be serialized before we create, then this is a redundant dep
-								}
-								if (SerializationBeforeSerializationDependencies.Contains(Index))
-								{
-									continue; // if the other thing must be serialized before we serialize, then this is a redundant dep
-								}
-								if (CreateBeforeCreateDependencies.Contains(Index))
-								{
-									continue; // if the other thing must be created before we are created, then this is a redundant dep
-								}
-								if (Export.FirstExportDependency == -1)
-								{
-									Export.FirstExportDependency = Linker->Summary.PreloadDependencyCount;
-									check(Export.SerializationBeforeSerializationDependencies == 0 && Export.CreateBeforeSerializationDependencies == 0 && Export.SerializationBeforeCreateDependencies == 0 && Export.CreateBeforeCreateDependencies == 0);
-								}
-								Linker->Summary.PreloadDependencyCount++;
-								Export.CreateBeforeSerializationDependencies++;
-								DepedenciesStream.EnterElement() << Index;
-								AddArcForDepChecking(true, Index, false);
-							}
-							for (FPackageIndex Index : SerializationBeforeCreateDependencies)
-							{
-								if (Export.FirstExportDependency == -1)
-								{
-									Export.FirstExportDependency = Linker->Summary.PreloadDependencyCount;
-									check(Export.SerializationBeforeSerializationDependencies == 0 && Export.CreateBeforeSerializationDependencies == 0 && Export.SerializationBeforeCreateDependencies == 0 && Export.CreateBeforeCreateDependencies == 0);
-								}
-								Linker->Summary.PreloadDependencyCount++;
-								Export.SerializationBeforeCreateDependencies++;
-								DepedenciesStream.EnterElement() << Index;
-								AddArcForDepChecking(false, Index, true);
-							}
-							for (FPackageIndex Index : CreateBeforeCreateDependencies)
-							{
-								if (Export.FirstExportDependency == -1)
-								{
-									Export.FirstExportDependency = Linker->Summary.PreloadDependencyCount;
-									check(Export.SerializationBeforeSerializationDependencies == 0 && Export.CreateBeforeSerializationDependencies == 0 && Export.SerializationBeforeCreateDependencies == 0 && Export.CreateBeforeCreateDependencies == 0);
-								}
-								Linker->Summary.PreloadDependencyCount++;
-								Export.CreateBeforeCreateDependencies++;
-								DepedenciesStream.EnterElement() << Index;
-								AddArcForDepChecking(false, Index, false);
-							}
-						}
-					}
-					UE_LOG(LogSavePackage, Verbose, TEXT("Saved %d dependencies for %d exports."), Linker->Summary.PreloadDependencyCount, Linker->ExportMap.Num());
-				}
-
-				
-				Linker->Summary.TotalHeaderSize	= Linker->Tell();
-
-				if ( EndSavingIfCancelled() )
-				{ 
-					return ESavePackageResult::Canceled;
-				}
-				SlowTask.EnterProgressFrame(1, NSLOCTEXT("Core", "ProcessingExports", "ProcessingExports..."));
-
-				// look for this package in the list of packages to generate script SHA for 
-				TArray<uint8>* ScriptSHABytes = FLinkerSave::PackagesToScriptSHAMap.Find(*FPaths::GetBaseFilename(Filename));
-
-				// if we want to generate the SHA key, start tracking script writes
-				if (ScriptSHABytes)
-				{
-					Linker->StartScriptSHAGeneration();
-				}
-
-#if WITH_EDITOR
-				TArray<FLargeMemoryWriter, TInlineAllocator<4>> AdditionalFilesFromExports;
-#endif
-				{
-					COOK_STAT(FScopedDurationTimer SaveTimer(FSavePackageStats::SerializeExportsTimeSec));
-					SCOPED_SAVETIMER(UPackage_Save_SaveExports);
-					FScopedSlowTask ExportScope(Linker->ExportMap.Num());
-
-					FStructuredArchive::FRecord ExportsRecord = StructuredArchiveRoot.EnterRecord(TEXT("Exports"));
-
-					// Save exports.
-					int32 LastExportSaveStep = 0;
-					for( int32 i=0; i<Linker->ExportMap.Num(); i++ )
-					{
-						if ( EndSavingIfCancelled() )
-						{ 
-							return ESavePackageResult::Canceled;
-						}
-						ExportScope.EnterProgressFrame();
-
-						FObjectExport& Export = Linker->ExportMap[i];
-						if (Export.Object)
-						{
-							// Save the object data.
-							Export.SerialOffset = Linker->Tell();
-							Linker->CurrentlySavingExport = FPackageIndex::FromExport(i);
-							// UE_LOG(LogSavePackage, Log, TEXT("export %s for %s"), *Export.Object->GetFullName(), *Linker->CookingTarget()->PlatformName());
-
-							FString ObjectName = Export.Object->GetPathName(InOuter);
-							FStructuredArchive::FSlot ExportSlot = ExportsRecord.EnterField(*ObjectName);
-
-							if (bTextFormat)
-							{
-								FObjectTextExport ObjectTextExport(Export, InOuter);
-								ExportSlot << ObjectTextExport;
-							}
-
-#if WITH_EDITOR
-							bool bSupportsText = UClass::IsSafeToSerializeToStructuredArchives(Export.Object->GetClass());
-#else
-							bool bSupportsText = false;
-#endif
-
-							if ( Export.Object->HasAnyFlags(RF_ClassDefaultObject) )
-							{
-								if (bSupportsText)
-								{
-									Export.Object->GetClass()->SerializeDefaultObject(Export.Object, ExportSlot);
-								}
-								else
-								{
-									FArchiveUObjectFromStructuredArchive Adapter(ExportSlot);
-									Export.Object->GetClass()->SerializeDefaultObject(Export.Object, Adapter.GetArchive());
-									Adapter.Close();
-								}
-							}
-							else
-							{
-								TGuardValue<UObject*> GuardSerializedObject(SaveContext->SerializedObject, Export.Object);
-
-								if (bSupportsText)
-								{
-									FStructuredArchive::FRecord ExportRecord = ExportSlot.EnterRecord();
-									Export.Object->Serialize(ExportRecord);
-								}
-								else
-								{
-									FArchiveUObjectFromStructuredArchive Adapter(ExportSlot);
-									Export.Object->Serialize(Adapter.GetArchive());
-									Adapter.Close();
-								}
-
-#if WITH_EDITOR
-								if (bIsCooking)
-								{
-									Export.Object->CookAdditionalFiles(Filename, TargetPlatform,
-										[&AdditionalFilesFromExports](const TCHAR* Filename, void* Data, int64 Size)
-									{
-										FLargeMemoryWriter& Writer = AdditionalFilesFromExports.Emplace_GetRef(0, true, Filename);
-										Writer.Serialize(Data, Size);
-									});
-								}
-#endif
-							}
-							Linker->CurrentlySavingExport = FPackageIndex();
-							Export.SerialSize = Linker->Tell() - Export.SerialOffset;
-
-							// Mark object as having been saved.
-							Export.Object->Mark(OBJECTMARK_Saved);
-						}
-					}
-				}
-
-				// if we want to generate the SHA key, get it out now that the package has finished saving
-				if (ScriptSHABytes && Linker->ContainsCode())
-				{
-					// make space for the 20 byte key
-					ScriptSHABytes->Empty(20);
-					ScriptSHABytes->AddUninitialized(20);
-
-					// retrieve it
-					Linker->GetScriptSHAKey(ScriptSHABytes->GetData());
-				}
-
-				
-				if ( EndSavingIfCancelled() )
-				{ 
-					return ESavePackageResult::Canceled;
-				}
-
-				SlowTask.EnterProgressFrame(1, NSLOCTEXT("Core", "SerializingBulkData", "Serializing bulk data"));
-
-				FSavePackageOutputFileArray AdditionalOutputFiles;
-				auto WriteAdditionalFiles = [&](int64 LinkerSize) -> ESavePackageResult
-				{
-					int64 DataStartOffset = LinkerSize >= 0 ? LinkerSize : Linker->Tell();
-					ESavePackageResult SaveResult =
-						SavePackageUtilities::SaveBulkData(Linker.Get(), DataStartOffset, InOuter, Filename, TargetPlatform, SavePackageContext, SaveFlags, bTextFormat,
-							TotalPackageSizeUncompressed);
-					if (SaveResult != ESavePackageResult::Success)
-					{
-						return SaveResult;
-					}
-
-					ESavePackageResult AppendAdditionalDataResult = SavePackageUtilities::AppendAdditionalData(*Linker.Get(), DataStartOffset, SavePackageContext);
-					if (AppendAdditionalDataResult != ESavePackageResult::Success)
-					{
-						return AppendAdditionalDataResult;
-					}
-
-					ESavePackageResult CreateSidecarResult = SavePackageUtilities::CreatePayloadSidecarFile(*Linker,
-						TargetPackagePath, bSaveToMemory, AdditionalOutputFiles, SavePackageContext);
-					if (CreateSidecarResult != ESavePackageResult::Success)
-					{
-						return CreateSidecarResult;
-					}
-
-#if WITH_EDITOR
-					if (bIsCooking && AdditionalFilesFromExports.Num() > 0)
-					{
-						checkf(PackageWriter, TEXT("Cooking requires a PackageWriter"));
-						for (FLargeMemoryWriter& Writer : AdditionalFilesFromExports)
-						{
-							const int64 Size = Writer.TotalSize();
-							TotalPackageSizeUncompressed += Size;
-
-							IPackageWriter::FAdditionalFileInfo FileInfo;
-							FileInfo.PackageName = InOuter->GetFName();
-							FileInfo.Filename = *Writer.GetArchiveName();
-
-							FIoBuffer FileData(FIoBuffer::AssumeOwnership, Writer.ReleaseOwnership(), Size);
-
-							PackageWriter->WriteAdditionalFile(FileInfo, FileData);
-						}
-						AdditionalFilesFromExports.Empty();
-					}
-#endif
-					return ESavePackageResult::Success;
-				};
-			
-				if (PackageWriter)
-				{
-					const int64 ExportsSize = Linker->Tell();
-					ESavePackageResult AdditionalFilesResult = WriteAdditionalFiles(ExportsSize);
-					checkf(Linker->Tell() == ExportsSize, TEXT("The writing of additional files is not allowed to append to the LinkerSave when using a PackageWriter."));
-					if (AdditionalFilesResult != ESavePackageResult::Success)
-					{
-						return AdditionalFilesResult;
-					}
-				}
-				else
-				{
-					// AdditionalFiles are appended to the Linker's archive, and so must be appended before we can calculate the full size of the Package
-					ESavePackageResult AdditionalFilesResult = WriteAdditionalFiles(-1);
-					if (AdditionalFilesResult != ESavePackageResult::Success)
-					{
-						return AdditionalFilesResult;
-					}
-
-					// write the package post tag
-					if (!bTextFormat)
-					{
-						uint32 Tag = PACKAGE_FILE_TAG;
-						StructuredArchiveRoot.GetUnderlyingArchive() << Tag;
-					}
-
-					// Now that the package is written out we can write the package trailer that is appended
-					// to the file. This should be the last thing written to the file!
-					SlowTask.EnterProgressFrame();
-					ESavePackageResult TrailerResult = BuildAndWriteTrailer(InOuter, Linker.Get(), StructuredArchiveRoot, PackageWriter, SaveFlags, bTextFormat);
-					if (TrailerResult != ESavePackageResult::Success)
-					{
-						return TrailerResult;
-					}
-				}
-
-				int64 OffsetBeforeUpdates = Linker->Tell();
-				PackageSize = OffsetBeforeUpdates;
-				if (PackageWriter)
-				{
-					PackageWriter->AddToExportsSize(PackageSize);
-				}
-
-				// Save the import map.
-				{
-#if WITH_EDITOR
-					FArchiveStackTraceIgnoreScope IgnoreSummaryDiffsScope(bIgnoreHeaderDiffs);
-#endif // WITH_EDITOR
-
-					if (!bTextFormat)
-					{
-						Linker->Seek(Linker->Summary.ImportOffset);
-
-						int32 NumImports = Linker->ImportMap.Num();
-						FStructuredArchive::FStream ImportTableStream = StructuredArchiveRoot.EnterStream(TEXT("ImportTable"));
-
-						for (int32 i = 0; i < Linker->ImportMap.Num(); i++)
-						{
-							FObjectImport& Import = Linker->ImportMap[i];
-							if (Import.XObject)
-							{
-								// Set the package index.
-								if (Import.XObject->GetOuter())
-								{
-									// if an import outer is an export and that import doesn't have a specific package set then, there's an error
-									const bool bWrongImport = Import.XObject->GetOuter()->IsInPackage(InOuter) && Import.XObject->GetExternalPackage() == nullptr;
-									if (bWrongImport)
-									{
-										if (!Import.XObject->HasAllFlags(RF_Transient) || !Import.XObject->IsNative())
-										{
-											UE_LOG(LogSavePackage, Warning, TEXT("Bad Object=%s"), *Import.XObject->GetFullName());
-										}
-										else
-										{
-											// if an object is marked RF_Transient and native, it is either an intrinsic class or
-											// a property of an intrinsic class.  Only properties of intrinsic classes will have
-											// an Outer that passes the check for "GetOuter()->IsInPackage(InOuter)" (thus ending up in this
-											// block of code).  Just verify that the Outer for this property is also marked RF_Transient and Native
-											check(Import.XObject->GetOuter()->HasAllFlags(RF_Transient) && Import.XObject->GetOuter()->IsNative());
-										}
-									}
-									check(!bWrongImport || Import.XObject->HasAllFlags(RF_Transient) || Import.XObject->IsNative());
-#if WITH_EDITOR
-									UObject** ReplacedOuter = ReplacedImportOuters.Find(Import.XObject);
-									if (ReplacedOuter && *ReplacedOuter)
-									{
-										Import.OuterIndex = Linker->MapObject(*ReplacedOuter);
-										ensure(Import.OuterIndex != FPackageIndex());
-									}
-									else
-#endif
-									{
-										Import.OuterIndex = Linker->MapObject(Import.XObject->GetOuter());
-									}
-
-									// if the import has a package set, set it up
-									if (UPackage* ImportPackage = Import.XObject->GetExternalPackage())
-									{
-										Import.SetPackageName(ImportPackage->GetFName());
-									}
-
-									if (Linker->IsCooking())
-									{
-										// Only package imports are allowed to have no outer
-										ensureMsgf(Import.OuterIndex != FPackageIndex() || Import.ClassName == NAME_Package, TEXT("Import %s has no valid outer when cooking!"), *Import.XObject->GetPathName());
-									}
-								}
-							}
-							else
-							{
-								checkf(false, TEXT("NULL XObject for import %i - Object: %s Class: %s"), i, *Import.ObjectName.ToString(), *Import.ClassName.ToString());
-							}
-
-							// Save it.
-							ImportTableStream.EnterElement() << Import;
-						}
-					}
-				}
-
-				// Save the export map.
-				if (!bTextFormat)
-				{
-					check( Linker->Tell() == OffsetAfterImportMap );
-					Linker->Seek(Linker->Summary.ExportOffset);
-
-					int32 NumExports = Linker->ExportMap.Num();
-					FStructuredArchive::FStream ExportTableStream = StructuredArchiveRoot.EnterStream(TEXT("ExportTable"));
-					{
-#if WITH_EDITOR
-						FArchiveStackTraceIgnoreScope IgnoreSummaryDiffsScope(bIgnoreHeaderDiffs);
-#endif
-						for (int32 i = 0; i < Linker->ExportMap.Num(); i++)
-						{
-							FObjectExport& Export = Linker->ExportMap[i];
-							ExportTableStream.EnterElement() << Export;
-						}
-					}
-
-					check( Linker->Tell() == OffsetAfterExportMap );
-				}
-
-				if (EndSavingIfCancelled())
-				{ 
-					return ESavePackageResult::Canceled;
-				}
-				SlowTask.EnterProgressFrame();
-
-				FFormatNamedArguments NamedArgs;
-				NamedArgs.Add( TEXT("CleanFilename"), FText::FromString( CleanFilename ) );
-				SlowTask.DefaultMessage = FText::Format( NSLOCTEXT("Core", "Finalizing", "Finalizing: {CleanFilename}..."), NamedArgs );
-
-				//@todo: remove ExportCount and NameCount - no longer used
-				Linker->Summary.Generations.Last().ExportCount = Linker->Summary.ExportCount;
-				Linker->Summary.Generations.Last().NameCount   = Linker->Summary.NameCount;	
-
-				// create the package source (based on developer or user created)
-	#if (UE_BUILD_SHIPPING && WITH_EDITOR)
-				Linker->Summary.PackageSource = FMath::Rand() * FMath::Rand();
-	#else
-				Linker->Summary.PackageSource = FCrc::StrCrc_DEPRECATED(*FPaths::GetBaseFilename(Filename).ToUpper());
-	#endif
-
-				// Flag package as requiring localization gather if the archive requires localization gathering.
-				Linker->LinkerRoot->ThisRequiresLocalizationGather(Linker->RequiresLocalizationGather());
-				
-				// Update package flags from package, in case serialization has modified package flags.
-				Linker->Summary.SetPackageFlags(Linker->LinkerRoot->GetPackageFlags());
-				
-				{
-					// Verify that the final serialization pass hasn't added any new custom versions. Otherwise this will result in crashes when loading the package.
-					bool bNewCustomVersionsUsed = false;
-					for (const FCustomVersion& LinkerCustomVer : Linker->GetCustomVersions().GetAllVersions())
-					{
-						if (Linker->Summary.GetCustomVersionContainer().GetVersion(LinkerCustomVer.Key) == nullptr)
-						{
-							UE_LOG(LogSavePackage, Error,
-								TEXT("Unexpected custom version \"%s\" found when saving %s. This usually happens when export tagging and final serialization paths differ. Package will not be saved."),
-								*LinkerCustomVer.GetFriendlyName().ToString(), *Linker->LinkerRoot->GetName());
-							bNewCustomVersionsUsed = true;
-						}
-					}
-					if (bNewCustomVersionsUsed)
-					{
-						return ESavePackageResult::Error;
-					}
-				}
-
-				if (!bTextFormat)
-				{
-					Linker->Seek(0);
-				}
-				{
-#if WITH_EDITOR
-					FArchiveStackTraceIgnoreScope IgnoreSummaryDiffsScope(bIgnoreHeaderDiffs);
-#endif // WITH_EDITOR
-					StructuredArchiveRoot.EnterField(TEXT("Summary")) << Linker->Summary;
-					SerializedPackageFlags = Linker->Summary.GetPackageFlags();
-
-					// Currently the PKG_ContainsNoAsset flag is not serialized as part of the summary
-					bool bContainsAsset = false;
-					for (FObjectExport& Export : Linker->ExportMap)
-					{
-						if (Export.bIsAsset)
-						{
-							bContainsAsset = true;
-							break;
-						}
-					}
-
-					if (!bContainsAsset)
-					{
-						SerializedPackageFlags |= PKG_ContainsNoAsset;
-					}
-				}
-
-				if (!bTextFormat)
-				{
-					check( Linker->Tell() == OffsetAfterPackageFileSummary );
-					Linker->Seek(OffsetBeforeUpdates); // Return Linker Pos to the end; some packagewriters need it there
-				}
-
-				if ( EndSavingIfCancelled() )
-				{ 
-					return ESavePackageResult::Canceled;
-				}
-				SlowTask.EnterProgressFrame();
-
-				// Destroy archives used for saving, closing file handle.
-				if (!bSaveToMemory)
-				{
-					const bool bFileWriterSuccess = Linker->CloseAndDestroySaver();
-
-					delete StructuredArchive;
-					delete Formatter;
-					delete TextFormatArchive;
-
-					if (!bFileWriterSuccess)
-					{
-						UE_LOG(LogSavePackage, Error, TEXT("Error writing temp file '%s' for '%s'"),
-							TempFilename.IsSet() ? *TempFilename.GetValue() : TEXT("UNKNOWN"), Filename);
-						return ESavePackageResult::Error;
-					}
-				}
-		
-				if ( EndSavingIfCancelled() )
-				{ 
-					return ESavePackageResult::Canceled;
-				}
-				SlowTask.EnterProgressFrame();
-
-				if( bSuccess == true )
-				{
-					{
-						// If we're writing to the existing file call ResetLoaders on the Package so that we drop the handle to the file on disk and can write to it
-						COOK_STAT(FScopedDurationTimer SaveTimer(FSavePackageStats::ResetLoadersTimeSec));
-						ResetLoadersForSave(InOuter, Filename);
-					}
-
-					// Compress the temporarily file to destination.
-					if (bSaveToMemory)
-					{
-						UE_LOG(LogSavePackage, Verbose, TEXT("Async saving from memory to '%s'"), *NewPath);
-						FLargeMemoryWriter* Writer = static_cast<FLargeMemoryWriter*>(Linker->Saver);
-
-						if (PackageWriter)
-						{
-							IPackageWriter::FPackageInfo PackageInfo;
-							PackageInfo.PackageName = InOuter->GetFName();
-							PackageInfo.LooseFilePath = Filename;
-							PackageInfo.HeaderSize = Linker->Summary.TotalHeaderSize;
-
-							FPackageId PackageId = FPackageId::FromName(PackageInfo.PackageName);
-							PackageInfo.ChunkId = CreateIoChunkId(PackageId.Value(), 0, EIoChunkType::ExportBundleData);
-
-							PackageWriter->WritePackageData(PackageInfo, *Writer, Linker->FileRegions);
-							TotalPackageSizeUncompressed += PackageSize;
-						}
-						else
-						{
-							checkf(!bIsCooking, TEXT("Cooking requires a PackageWriter"));
-							const int64 DataSize = Writer->TotalSize();
-							// Add the uasset file to the list of output files
-							AdditionalOutputFiles.Emplace(NewPath, FLargeMemoryPtr(Writer->ReleaseOwnership()), Linker->FileRegions, DataSize);
-							
-							EAsyncWriteOptions WriteOptions(EAsyncWriteOptions::None);							
-							for (FSavePackageOutputFile& Entry : AdditionalOutputFiles)
-							{
-								TotalPackageSizeUncompressed += Entry.DataSize;
-								SavePackageUtilities::AsyncWriteFile(WriteOptions, Entry);
-							}	
-						}
-						Linker->CloseAndDestroySaver();
-
-						delete StructuredArchive;
-						delete Formatter;
-						delete TextFormatArchive;
-					}
-					// Move the temporary file.
-					else
-					{
-						checkf(!PackageWriter, TEXT("PackageWriter is not currently supported with synchronous writes. %s"), Filename);
-						checkf(TempFilename.IsSet(), TEXT("The package should've been saved to a tmp file first! %s"), Filename);
-
-						// When saving in text format we will have two temp files, so we need to manually delete the non-textbased one
-						if (bTextFormat)
-						{
-							check(TextFormatTempFilename.IsSet());
-							IFileManager::Get().Delete(*TempFilename.GetValue());
-							TempFilename = TextFormatTempFilename;
-							TextFormatTempFilename.Reset();
-						}
-
-						// Add the .uasset file to the list of output files
-						AdditionalOutputFiles.Emplace(NewPath, TempFilename.GetValue(), PackageSize);
-						
-						ESavePackageResult FinalizeResult = SavePackageUtilities::FinalizeTempOutputFiles(TargetPackagePath, AdditionalOutputFiles, FinalTimeStamp);
-						if (FinalizeResult != ESavePackageResult::Success)
-						{
-							bSuccess = false;
-						}	
-
-						// Now add up the total size of the data that was saved
-						for (FSavePackageOutputFile& Entry : AdditionalOutputFiles)
-						{
-							TotalPackageSizeUncompressed += Entry.DataSize;
-						}
-					}
-
-					AdditionalOutputFiles.Empty();
-
-					if( bSuccess == false )
-					{
-						if (SaveFlags & SAVE_NoError)
-						{
-							UE_LOG(LogSavePackage, Warning, TEXT("%s"), *FString::Printf( TEXT("Error saving '%s'"), Filename ) );
-						}
-						else
-						{
-							UE_LOG(LogSavePackage, Error, TEXT("%s"), *FString::Printf( TEXT("Error saving '%s'"), Filename ) );
-							Error->Logf(ELogVerbosity::Warning, TEXT("%s"), *FText::Format( NSLOCTEXT( "Core", "SaveWarning", "Error saving '{0}'" ), FText::FromString( FString( Filename ) ) ).ToString() );
-						}
-					}
-					else
-					{
-						// Mark exports and the package as RF_Loaded after they've been serialized
-						// This is to ensue that newly created packages are properly marked as loaded (since they now exist on disk and 
-						// in memory in the exact same state).
-						for (auto& Export : Linker->ExportMap)
-						{
-							if (Export.Object)
-							{
-								Export.Object->SetFlags(RF_WasLoaded|RF_LoadCompleted);
-							}
-						}
-						if (Linker->LinkerRoot)
-						{
-							// And finally set the flag on the package itself.
-							Linker->LinkerRoot->SetFlags(RF_WasLoaded|RF_LoadCompleted);
-						}
-
-						// Clear dirty flag if desired
-						if (!(SaveFlags & SAVE_KeepDirty))
-						{
-							InOuter->SetDirtyFlag(false);
-						}
-						
-						// Update package FileSize value
-						InOuter->SetFileSize(PackageSize);
-
-						// Warn about long package names, which may be bad for consoles with limited filename lengths.
-						if( bWarnOfLongFilename == true )
-						{
-							int32 MaxFilenameLength = FPlatformMisc::GetMaxPathLength();
-
-							// If the name is of the form "_LOC_xxx.ext", remove the loc data before the length check
-							FString CleanBaseFilename = BaseFilename;
-							if( CleanBaseFilename.Find( "_LOC_" ) == BaseFilename.Len() - 8 )
-							{
-								CleanBaseFilename = BaseFilename.LeftChop( 8 );
-							}
-
-							if( CleanBaseFilename.Len() > MaxFilenameLength )
-							{
-								if (SaveFlags & SAVE_NoError)
-								{
-									UE_LOG(LogSavePackage, Warning, TEXT("%s"), *FString::Printf( TEXT("Filename is too long (%d characters); this may interfere with cooking for consoles. Unreal filenames should be no longer than %s characters. Filename value: %s"), BaseFilename.Len(), MaxFilenameLength, *BaseFilename ) );
-								}
-								else
-								{
-									FFormatNamedArguments Arguments;
-									Arguments.Add(TEXT("FileName"), FText::FromString( BaseFilename ));
-									Arguments.Add(TEXT("MaxLength"), FText::AsNumber( MaxFilenameLength ));
-									Error->Logf(ELogVerbosity::Warning, TEXT("%s"), *FText::Format( NSLOCTEXT( "Core", "Error_FilenameIsTooLongForCooking", "Filename '{FileName}' is too long; this may interfere with cooking for consoles. Unreal filenames should be no longer than {MaxLength} characters." ), Arguments ).ToString() );
-								}
-							}
-						}
-					}
-				}
-				COOK_STAT(FSavePackageStats::MBWritten += ((double)TotalPackageSizeUncompressed) / 1024.0 / 1024.0);
-
-				SlowTask.EnterProgressFrame();
-			}
-
-			// Route PostSaveRoot to allow e.g. the world to detach components for the persistent level that were
-			// attached in PreSaveRoot. PostSaveRoot is suppressed in the same situations as PreSaveRoot.
-			if (Base && !bSavingConcurrent)
-			{
-				UE::SavePackageUtilities::CallPostSaveRoot(Base, ObjectSaveContext, bCleanupRequired);
-			}
-
-			SlowTask.EnterProgressFrame();
-			
-#if WITH_EDITOR
-			if ( !bSavingConcurrent )
-			{
-				for ( int CachedObjectIndex = 0; CachedObjectIndex < CachedObjects.Num(); ++CachedObjectIndex )
-				{
-					CachedObjects[CachedObjectIndex]->ClearCachedCookedPlatformData(TargetPlatform);
-				}
-			}
-#endif
-
-		}
-		if( bSuccess == true )
-		{
-			// Package has been save, so unmark NewlyCreated flag.
-			InOuter->ClearPackageFlags(PKG_NewlyCreated);
-
-			// Notify the soft reference collector about our harvested soft references during save. 
-			// This is currently needed only for cooking which does not require editor-only references 
-#if WITH_EDITOR
-			if (bIsCooking)
-			{
-				GRedirectCollector.CollectSavedSoftPackageReferences(InOuter->GetFName(), SoftPackagesUsedInGame, false);
-			}
-#endif
-
-			if (Linker != nullptr)
-			{
-				Linker->OnPostSave(TargetPackagePath, FObjectPostSaveContext(ObjectSaveContext));
-			}
-
-			// send a message that the package was saved
-			PRAGMA_DISABLE_DEPRECATION_WARNINGS;
-			UPackage::PackageSavedEvent.Broadcast(Filename, InOuter);
-			PRAGMA_ENABLE_DEPRECATION_WARNINGS;
-			UPackage::PackageSavedWithContextEvent.Broadcast(Filename, InOuter, FObjectPostSaveContext(ObjectSaveContext));
-		}
-
-		// We're done!
-		SlowTask.EnterProgressFrame();
-
-		if (bSuccess)
-		{
-			// if the save was successful, update the internal package filename path if we aren't currently cooking
-#if WITH_EDITOR
-			if (ObjectSaveContext.bUpdatingLoadedPath)
-			{
-				InOuter->SetLoadedPath(TargetPackagePath);
-			}
-#endif
-
-			if (bRequestStub)
-			{
-				return FSavePackageResultStruct(ESavePackageResult::GenerateStub, TotalPackageSizeUncompressed, SerializedPackageFlags, bCompareLinker ? MoveTemp(Linker) : nullptr);
-			}
-			else
-			{
-				return FSavePackageResultStruct(ESavePackageResult::Success, TotalPackageSizeUncompressed, SerializedPackageFlags, bCompareLinker ? MoveTemp(Linker) : nullptr);
-			}
-		}
-		else
-		{
-			if (bRequestStub)
-			{
-				UE_LOG(LogSavePackage, Warning, TEXT("C++ stub requested, but package failed to save, may cause compile errors: %s"), Filename);
-			}
-			return ESavePackageResult::Error;
-		}
-	}
-	else
-	{
-		return ESavePackageResult::Error;
-	}
+	return UPackage::Save2(InOuter, InAsset, Filename, SaveArgs);
 }
 
 bool UPackage::SavePackage(UPackage* InOuter, UObject* Base, EObjectFlags TopLevelFlags, const TCHAR* Filename,
diff --git a/Engine/Source/Runtime/CoreUObject/Private/UObject/SavePackage/SavePackageUtilities.cpp b/Engine/Source/Runtime/CoreUObject/Private/UObject/SavePackage/SavePackageUtilities.cpp
--- a/Engine/Source/Runtime/CoreUObject/Private/UObject/SavePackage/SavePackageUtilities.cpp
+++ b/Engine/Source/Runtime/CoreUObject/Private/UObject/SavePackage/SavePackageUtilities.cpp
@@ -194,9 +194,6 @@ const FName NAME_World("World");
 const FName NAME_Level("Level");
 const FName NAME_PrestreamPackage("PrestreamPackage");
 
-// Switch to SavePackage to Save2, 0:disabled, 1: enabled for cook, 2: enabled for uncooked, 3: enabled everywhere
-static TAutoConsoleVariable<int32> CVarEnablePackageNewSave(TEXT("SavePackage.EnableNewSave"), 3, TEXT("Enable new package save mechanism over the old one."));
-
 /**
 * A utility that records the state of a package's files before we start moving and overwriting them. 
 * This provides an easy way for us to restore the original state of the package incase of failures 
@@ -275,91 +272,6 @@ private:
 	TArray<TPair<FString, FString>> MovedOriginalFiles;	
 };
 
-bool IsNewSaveEnabled(bool bForCooking)
-{
-	int32 EnableNewSave = CVarEnablePackageNewSave.GetValueOnAnyThread();
-	return (bForCooking && (EnableNewSave & 1) != 0) || (!bForCooking && (EnableNewSave & 2) != 0);
-}
-
-bool HasUnsaveableOuter(UObject* InObj, UPackage* InSavingPackage)
-{
-	UObject* Obj = InObj;
-	while (Obj)
-	{
-		if (Obj->GetClass()->HasAnyClassFlags(CLASS_Deprecated) && !Obj->HasAnyFlags(RF_ClassDefaultObject))
-		{
-			if (IsValid(InObj) && InObj->GetOutermost() == InSavingPackage)
-			{
-				UE_LOG(LogSavePackage, Warning, TEXT("%s has a deprecated outer %s, so it will not be saved"), *InObj->GetFullName(), *Obj->GetFullName());
-			}
-			return true; 
-		}
-
-		if(!IsValid(Obj))
-		{
-			return true;
-		}
-
-		if (Obj->HasAnyFlags(RF_Transient) && !Obj->IsNative())
-		{
-			return true;
-		}
-
-		Obj = Obj->GetOuter();
-	}
-	return false;
-}
-
-void CheckObjectPriorToSave(FArchiveUObject& Ar, UObject* InObj, UPackage* InSavingPackage)
-{
-	if (!InObj)
-	{
-		return;
-	}
-	UObject* SerializedObject = nullptr;
-	FUObjectSerializeContext* SaveContext = Ar.GetSerializeContext();
-	check(SaveContext);
-	SerializedObject = SaveContext->SerializedObject;
-
-	if (!InObj->IsValidLowLevelFast() || !InObj->IsValidLowLevel())
-	{
-		UE_LOG(LogLinker, Fatal, TEXT("Attempt to save bogus object %p SaveContext.SerializedObject=%s  SerializedProperty=%s"), (void*)InObj, *GetFullNameSafe(SerializedObject), *GetFullNameSafe(Ar.GetSerializedProperty()));
-		return;
-	}
-	// if the object class is abstract or has been marked as deprecated, mark this
-	// object as transient so that it isn't serialized
-	if ( InObj->GetClass()->HasAnyClassFlags(CLASS_Abstract | CLASS_Deprecated | CLASS_NewerVersionExists) )
-	{
-		if ( !InObj->HasAnyFlags(RF_ClassDefaultObject) || InObj->GetClass()->HasAnyClassFlags(CLASS_Deprecated) )
-		{
-			InObj->SetFlags(RF_Transient);
-		}
-		if ( !InObj->HasAnyFlags(RF_ClassDefaultObject) && InObj->GetClass()->HasAnyClassFlags(CLASS_HasInstancedReference) )
-		{
-			TArray<UObject*> ComponentReferences;
-			FReferenceFinder ComponentCollector(ComponentReferences, InObj, false, true, true);
-			ComponentCollector.FindReferences(InObj, SerializedObject, Ar.GetSerializedProperty());
-
-			for ( int32 Index = 0; Index < ComponentReferences.Num(); Index++ )
-			{
-				ComponentReferences[Index]->SetFlags(RF_Transient);
-			}
-		}
-	}
-	else if ( HasUnsaveableOuter(InObj, InSavingPackage) )
-	{
-		InObj->SetFlags(RF_Transient);
-	}
-
-	if ( InObj->HasAnyFlags(RF_ClassDefaultObject) && (InObj->GetClass()->IsNative() || !InObj->GetClass()->HasAnyFlags(RF_Transient)) )
-	{
-		// if this is the class default object, make sure it's not
-		// marked transient for any reason, as we need it to be saved
-		// to disk (unless it's associated with a transient generated class)
-		InObj->ClearFlags(RF_Transient);
-	}
-}
-
 /**
  * Determines the set of object marks that should be excluded for the target platform
  *
@@ -397,120 +309,6 @@ EObjectMark GetExcludedObjectMarksForTargetPlatform(const class ITargetPlatform*
 	return ObjectMarks;
 }
 
-/**
- * Marks object as not for client, not for server, or editor only. Recurses up outer/class chain as necessary
- */
-void ConditionallyExcludeObjectForTarget(UObject* Obj, EObjectMark ExcludedObjectMarks, const ITargetPlatform* TargetPlatform)
-{
-#if WITH_EDITOR
-	if (!Obj || Obj->GetOutermost()->GetFName() == GLongCoreUObjectPackageName)
-	{
-		// No object or in CoreUObject, don't exclude
-		return;
-	}
-
-	auto InheritMarks = [](EObjectMark& MarksToModify, UObject* ObjToCheck, uint32 MarkMask)
-	{
-		EObjectMark ObjToCheckMarks = ObjToCheck->GetAllMarks();
-
-		MarksToModify = (EObjectMark)(MarksToModify | (ObjToCheckMarks & MarkMask));
-	};
-
-	// MarksToProcess is a superset of marks retrieved from UPackage::GetExcludedObjectMarksForTargetPlatform
-	const uint32 MarksToProcess = OBJECTMARK_INHERITEDMARKS;
-	check((ExcludedObjectMarks & ~MarksToProcess) == 0);
-
-	EObjectMark CurrentMarks = OBJECTMARK_NOMARKS;
-	InheritMarks(CurrentMarks, Obj, MarksToProcess);
-
-	if ((CurrentMarks & MarksToProcess) != 0)
-	{
-		// Already marked
-		return;
-	}
-
-	UObject* ObjOuter = Obj->GetOuter();
-	UClass* ObjClass = Obj->GetClass();
-
-	EObjectMark NewMarks = CurrentMarks;
-
-	// Recurse into parents, then compute inherited marks
-	ConditionallyExcludeObjectForTarget(ObjClass, ExcludedObjectMarks, TargetPlatform);
-	InheritMarks(NewMarks, ObjClass, OBJECTMARK_INHERITEDMARKS);
-
-	if (ObjOuter)
-	{
-		ConditionallyExcludeObjectForTarget(ObjOuter, ExcludedObjectMarks, TargetPlatform);
-		InheritMarks(NewMarks, ObjOuter, OBJECTMARK_INHERITEDMARKS);
-	}
-
-	// Check parent struct if we have one
-	UStruct* ThisStruct = Cast<UStruct>(Obj);
-	if (ThisStruct && ThisStruct->GetSuperStruct())
-	{
-		UObject* SuperStruct = ThisStruct->GetSuperStruct();
-		ConditionallyExcludeObjectForTarget(SuperStruct, ExcludedObjectMarks, TargetPlatform);
-		InheritMarks(NewMarks, SuperStruct, OBJECTMARK_INHERITEDMARKS);
-	}
-
-	// Check archetype, this may not have been covered in the case of components
-	UObject* Archetype = Obj->GetArchetype();
-	if (Archetype)
-	{
-		ConditionallyExcludeObjectForTarget(Archetype, ExcludedObjectMarks, TargetPlatform);
-		InheritMarks(NewMarks, Archetype, OBJECTMARK_INHERITEDMARKS);
-	}
-
-	if (!Obj->HasAnyFlags(RF_ClassDefaultObject))
-	{
-		// CDOs must be included if their class is so only inherit marks, for everything else we check the native overrides as well
-		if (!(NewMarks & OBJECTMARK_EditorOnly) && IsEditorOnlyObject(Obj, false, false))
-		{
-			NewMarks = (EObjectMark)(NewMarks | OBJECTMARK_EditorOnly);
-		}
-
-		if (!(NewMarks & OBJECTMARK_NotForClient) && !Obj->NeedsLoadForClient())
-		{
-			NewMarks = (EObjectMark)(NewMarks | OBJECTMARK_NotForClient);
-		}
-
-		if (!(NewMarks & OBJECTMARK_NotForServer) && !Obj->NeedsLoadForServer())
-		{
-			NewMarks = (EObjectMark)(NewMarks | OBJECTMARK_NotForServer);
-		}
-
-		bool bCheckTargetPlatform = false;
-		if (TargetPlatform != nullptr)
-		{
-			// NotForServer && NotForClient implies EditorOnly
-			const bool bIsEditorOnlyObject = (NewMarks & OBJECTMARK_NotForServer) && (NewMarks & OBJECTMARK_NotForClient);
-			const bool bTargetAllowsEditorObjects = TargetPlatform->AllowsEditorObjects();
-
-			// no need to query the target platform if the object is editoronly and the targetplatform doesn't allow editor objects 
-			// @note: this was done in large part fo avoid a warning in SceneComponentNeedsLoadForTarget() caused by editor only object
-			// that contains non-editoronly inners, but that warning should probably allow for this valid setup
-			bCheckTargetPlatform = !bIsEditorOnlyObject || bTargetAllowsEditorObjects;
-		}
-		if (bCheckTargetPlatform && (!Obj->NeedsLoadForTargetPlatform(TargetPlatform) || !TargetPlatform->AllowObject(Obj)))
-		{
-			NewMarks = (EObjectMark)(NewMarks | OBJECTMARK_NotForTargetPlatform);
-		}
-	}
-
-	// If NotForClient and NotForServer, it is implicitly editor only
-	if ((NewMarks & OBJECTMARK_NotForClient) && (NewMarks & OBJECTMARK_NotForServer))
-	{
-		NewMarks = (EObjectMark)(NewMarks | OBJECTMARK_EditorOnly);
-	}
-
-	// If our marks are different than original, set them on the object
-	if (CurrentMarks != NewMarks)
-	{
-		Obj->Mark(NewMarks);
-	}
-#endif
-}
-
 /**
  * Find most likely culprit that caused the objects in the passed in array to be considered for saving.
  *
diff --git a/Engine/Source/Runtime/CoreUObject/Private/UObject/SavePackage/SavePackageUtilities.h b/Engine/Source/Runtime/CoreUObject/Private/UObject/SavePackage/SavePackageUtilities.h
--- a/Engine/Source/Runtime/CoreUObject/Private/UObject/SavePackage/SavePackageUtilities.h
+++ b/Engine/Source/Runtime/CoreUObject/Private/UObject/SavePackage/SavePackageUtilities.h
@@ -351,9 +351,6 @@ namespace SavePackageUtilities
 	extern const FName NAME_Level;
 	extern const FName NAME_PrestreamPackage;
 
-	// return if the new save algorithm is enabled for cooked or uncooked
-	bool IsNewSaveEnabled(bool bForCooking = false);
-
 	void SaveThumbnails(UPackage* InOuter, FLinkerSave* Linker, FStructuredArchive::FSlot Slot);
 
 	/**
@@ -398,9 +395,6 @@ namespace SavePackageUtilities
 	
 	void SaveWorldLevelInfo(UPackage* InOuter, FLinkerSave* Linker, FStructuredArchive::FRecord Record);
 	EObjectMark GetExcludedObjectMarksForTargetPlatform(const class ITargetPlatform* TargetPlatform);
-	bool HasUnsaveableOuter(UObject* InObj, UPackage* InSavingPackage);
-	void CheckObjectPriorToSave(FArchiveUObject& Ar, UObject* InObj, UPackage* InSavingPackage);
-	void ConditionallyExcludeObjectForTarget(UObject* Obj, EObjectMark ExcludedObjectMarks, const ITargetPlatform* TargetPlatform);
 	void FindMostLikelyCulprit(const TArray<UObject*>& BadObjects, UObject*& MostLikelyCulprit, FString& OutReferencer, FSaveContext* InOptionalSaveContext = nullptr);
 	
 	/** 
diff --git a/Engine/Source/Runtime/CoreUObject/Private/UObject/SavePackage2.cpp b/Engine/Source/Runtime/CoreUObject/Private/UObject/SavePackage2.cpp
--- a/Engine/Source/Runtime/CoreUObject/Private/UObject/SavePackage2.cpp
+++ b/Engine/Source/Runtime/CoreUObject/Private/UObject/SavePackage2.cpp
@@ -1123,8 +1123,6 @@ ESavePackageResult BuildLinker(FSaveContext& SaveContext)
 			{
 				FPackageIndex DependencyIndex = Linker->ObjectIndicesMap.FindRef(DependentObject);
 
-				//@todo FH: UnmarkExportTagFromDuplicates redirect??
-
 				// if we didn't find it (FindRef returns 0 on failure, which is good in this case), then we are in trouble, something went wrong somewhere
 				checkf(!DependencyIndex.IsNull(), TEXT("Failed to find dependency index for %s (%s)"), *DependentObject->GetFullName(), *Object->GetFullName());
 
@@ -2618,47 +2616,15 @@ ESavePackageResult UPackage::SaveConcurrent(TArrayView<FPackageSaveInfo> InPacka
 	SlowTask.EnterProgressFrame();
 	{
 		SCOPED_SAVETIMER(UPackage_SaveConcurrent);
-		// Use concurrent new save only if new save is enabled, otherwise use old save
-		if (SavePackageUtilities::IsNewSaveEnabled(SaveArgs.IsCooking()))
-		{
-			// @todo: Once GIsSavingPackage is reworked we should reinstore the saving flag here for the gc lock
-			// Passing in false here so that GIsSavingPackage is set to true on top of locking the GC
-			//FScopedSavingFlag IsSavingFlag(false);
+		// @todo: Once GIsSavingPackage is reworked we should reinstore the saving flag here for the gc lock
+		// Passing in false here so that GIsSavingPackage is set to true on top of locking the GC
+		//FScopedSavingFlag IsSavingFlag(false);
 
-			// Concurrent Part
-			ParallelFor(PackageSaveContexts.Num(), [&PackageSaveContexts](int32 PackageIdx)
-				{
-					InnerSave(PackageSaveContexts[PackageIdx]);
-				});
-		}
-		else
-		{
-			FSavePackageArgs ConcurrentSaveArgs = SaveArgs;
-			ConcurrentSaveArgs.SaveFlags |= SAVE_Concurrent;
-			ConcurrentSaveArgs.bSlowTask = false;
-			// save concurrently if the SAVE_Concurrent flag is present
-			if (SaveArgs.SaveFlags & SAVE_Concurrent)
+		// Concurrent Part
+		ParallelFor(PackageSaveContexts.Num(), [&PackageSaveContexts](int32 PackageIdx)
 			{
-				GIsSavingPackage = true;
-				ParallelFor(PackageSaveContexts.Num(), [&PackageSaveContexts, &ConcurrentSaveArgs](int32 PackageIdx)
-					{
-						FSaveContext& SaveContext = PackageSaveContexts[PackageIdx];
-						Save(SaveContext.GetPackage(), SaveContext.GetAsset(), SaveContext.GetFilename(),
-							ConcurrentSaveArgs);
-					});
-				GIsSavingPackage = false;
-			}
-			// otherwise save serial
-			else
-			{
-				// still pass in the SAVE_Concurrent flag on ConcurrentSaveArgs so that we don't redo things we have handled already
-				for (FSaveContext& SaveContext : PackageSaveContexts)
-				{
-					Save(SaveContext.GetPackage(), SaveContext.GetAsset(), SaveContext.GetFilename(),
-						ConcurrentSaveArgs);
-				}
-			}
-		}
+				InnerSave(PackageSaveContexts[PackageIdx]);
+			});
 	}
 
 	// Run Post Concurrent Save
diff --git a/Engine/Source/Runtime/CoreUObject/Private/UObject/UObjectGlobals.cpp b/Engine/Source/Runtime/CoreUObject/Private/UObject/UObjectGlobals.cpp
--- a/Engine/Source/Runtime/CoreUObject/Private/UObject/UObjectGlobals.cpp
+++ b/Engine/Source/Runtime/CoreUObject/Private/UObject/UObjectGlobals.cpp
@@ -143,7 +143,9 @@ FCoreUObjectDelegates::FIsPackageOKToSaveDelegate FCoreUObjectDelegates::IsPacka
 FCoreUObjectDelegates::FOnPackageReloaded FCoreUObjectDelegates::OnPackageReloaded;
 FCoreUObjectDelegates::FNetworkFileRequestPackageReload FCoreUObjectDelegates::NetworkFileRequestPackageReload;
 #if WITH_EDITOR
+PRAGMA_DISABLE_DEPRECATION_WARNINGS;
 FCoreUObjectDelegates::FAutoPackageBackupDelegate FCoreUObjectDelegates::AutoPackageBackupDelegate;
+PRAGMA_ENABLE_DEPRECATION_WARNINGS;
 FCoreUObjectDelegates::FOnPreObjectPropertyChanged FCoreUObjectDelegates::OnPreObjectPropertyChanged;
 FCoreUObjectDelegates::FOnObjectPropertyChanged FCoreUObjectDelegates::OnObjectPropertyChanged;
 TSet<UObject*> FCoreUObjectDelegates::ObjectsModifiedThisFrame;
diff --git a/Engine/Source/Runtime/CoreUObject/Public/UObject/UObjectGlobals.h b/Engine/Source/Runtime/CoreUObject/Public/UObject/UObjectGlobals.h
--- a/Engine/Source/Runtime/CoreUObject/Public/UObject/UObjectGlobals.h
+++ b/Engine/Source/Runtime/CoreUObject/Public/UObject/UObjectGlobals.h
@@ -2611,6 +2611,7 @@ struct COREUOBJECT_API FCoreUObjectDelegates
 
 	/** Delegate used by SavePackage() to create the package backup */
 	DECLARE_DELEGATE_RetVal_OneParam(bool, FAutoPackageBackupDelegate, const UPackage&);
+	UE_DEPRECATED(5.1, "Backups are no longer used in Unreal Package Saves")
 	static FAutoPackageBackupDelegate AutoPackageBackupDelegate;
 #endif // WITH_EDITOR
 
