Fix a missing Adapter.Close() when serializing a multicast delegate property. Wouldn't have any affect unless writing to a text archive

diff --git a/Engine/Source/Runtime/CoreUObject/Private/UObject/PropertyMulticastDelegate.cpp b/Engine/Source/Runtime/CoreUObject/Private/UObject/PropertyMulticastDelegate.cpp
--- a/Engine/Source/Runtime/CoreUObject/Private/UObject/PropertyMulticastDelegate.cpp
+++ b/Engine/Source/Runtime/CoreUObject/Private/UObject/PropertyMulticastDelegate.cpp
@@ -410,6 +410,7 @@ void FMulticastInlineDelegateProperty::SerializeItem(FStructuredArchive::FSlot S
 	FArchiveUObjectFromStructuredArchive Adapter(Slot);
 	FArchive& Ar = Adapter.GetArchive();
 	Ar << *GetPropertyValuePtr(Value);
+	Adapter.Close();
 }
 
 const TCHAR* FMulticastInlineDelegateProperty::ImportText_Internal(const TCHAR* Buffer, void* ContainerOrPropertyPtr, EPropertyPointerType PropertyPointerType, UObject* Parent, int32 PortFlags, FOutputDevice* ErrorText) const
