Update editor bulkdata  unit tests to work with CL 19850627

#rnx

diff --git a/Engine/Source/Runtime/CoreUObject/Private/Tests/Serialization/EditorBulkDataTests.cpp b/Engine/Source/Runtime/CoreUObject/Private/Tests/Serialization/EditorBulkDataTests.cpp
--- a/Engine/Source/Runtime/CoreUObject/Private/Tests/Serialization/EditorBulkDataTests.cpp
+++ b/Engine/Source/Runtime/CoreUObject/Private/Tests/Serialization/EditorBulkDataTests.cpp
@@ -2,6 +2,7 @@
 
 #include "Algo/AllOf.h"
 #include "Memory/SharedBuffer.h"
+#include "Serialization/BulkDataRegistry.h"
 #include "Serialization/EditorBulkData.h"
 #include "Serialization/EditorBulkDataReader.h"
 #include "Serialization/EditorBulkDataWriter.h"
@@ -20,6 +21,15 @@ namespace UE::Serialization
 	
 constexpr const uint32 TestFlags = EAutomationTestFlags::ApplicationContextMask | EAutomationTestFlags::EngineFilter;
 
+bool IsBulkDataRegistryEnabled()
+{
+#if WITH_EDITOR
+	return IBulkDataRegistry::IsEnabled();
+#else
+	return false;
+#endif //WITH_EDITOR 
+}
+
 /** Creates a buffer full of random data to make it easy to have something to test against. */
 TUniquePtr<uint8[]> CreateRandomData(int64 BufferSize)
 {
@@ -239,35 +249,79 @@ bool FEditorBulkDataTestReaderWriter::RunTest(const FString& Parameters)
 IMPLEMENT_SIMPLE_AUTOMATION_TEST(FEditorBulkDataTestSerializationToMemory, TEXT("System.CoreUObject.Serialization.EditorBulkData.SerializationToMemory"), TestFlags)
 bool FEditorBulkDataTestSerializationToMemory::RunTest(const FString& Parameters)
 {
+	const bool bIsArPersistent = true;
 	const int64 BufferSize = 1024;
 
 	TUniquePtr<uint8[]> SourceBuffer = CreateRandomData(BufferSize);
 	
-	FEditorBulkData EmptyBulkData;
-	FEditorBulkData ValidBulkData;
-	ValidBulkData.UpdatePayload(FSharedBuffer::Clone(SourceBuffer.Get(), BufferSize));
-	
 	TArray<uint8> MemoryBuffer;
-	const bool bIsArPersistent = true;
+	FGuid ValidBulkDataId;
+	FEditorBulkData EmptyBulkData;
 
-	// Write out
 	{
-		FMemoryWriter WriterAr(MemoryBuffer, bIsArPersistent);
+		FEditorBulkData ValidBulkData;
+		ValidBulkData.UpdatePayload(FSharedBuffer::Clone(SourceBuffer.Get(), BufferSize));
+		ValidBulkDataId = ValidBulkData.GetIdentifier();
+
+		// Write out
+		{
+			FMemoryWriter WriterAr(MemoryBuffer, bIsArPersistent);
 
-		ValidBulkData.Serialize(WriterAr, nullptr);
-		EmptyBulkData.Serialize(WriterAr, nullptr);
-		ValidBulkData.Serialize(WriterAr, nullptr);
-		EmptyBulkData.Serialize(WriterAr, nullptr);
+			ValidBulkData.Serialize(WriterAr, nullptr);
+			EmptyBulkData.Serialize(WriterAr, nullptr);
+			ValidBulkData.Serialize(WriterAr, nullptr);
+			EmptyBulkData.Serialize(WriterAr, nullptr);
+		}
+
+		// First test we test reading the editor bulkdata while ValidBulkData is still registered and in scope
+		// in this case if the BulkDataRegistry is enabled then SerializedBulkData and ValidBulkData should
+		// have different identifiers, if the BulkDataRegistry is disabled then they should have the same identifiers
+		FEditorBulkData SerializedBulkData;
+
+		FMemoryReader ReaderAr(MemoryBuffer, bIsArPersistent);
+
+		SerializedBulkData.Serialize(ReaderAr, nullptr);
+		if (IsBulkDataRegistryEnabled())
+		{
+			TestNotEqual(TEXT("Bulkdata identifier should change"), SerializedBulkData.GetIdentifier(), ValidBulkData.GetIdentifier());
+		}
+		else
+		{
+			TestEqual(TEXT("Bulkdata identifier should remain the same"), SerializedBulkData.GetIdentifier(), ValidBulkData.GetIdentifier());
+		}
+		
+		TestTrue(TEXT("Bulkdata payload should remain the same"), FMemory::Memcmp(SourceBuffer.Get(), SerializedBulkData.GetPayload().Get().GetData(), BufferSize) == 0);
+
+		SerializedBulkData.Serialize(ReaderAr, nullptr);
+		TestFalse(TEXT("Bulkdata identifier should be invalid"), SerializedBulkData.GetIdentifier().IsValid());
+		TestTrue(TEXT("Bulkdata should not have a payload"), SerializedBulkData.GetPayload().Get().IsNull());
+
+		SerializedBulkData.Serialize(ReaderAr, nullptr);
+		if (IsBulkDataRegistryEnabled())
+		{
+			TestNotEqual(TEXT("Bulkdata identifier should change"), SerializedBulkData.GetIdentifier(), ValidBulkData.GetIdentifier());
+		}
+		else
+		{
+			TestEqual(TEXT("Bulkdata identifier should remain the same"), SerializedBulkData.GetIdentifier(), ValidBulkData.GetIdentifier());
+		}
+		TestTrue(TEXT("Bulkdata payload should remain the same"), FMemory::Memcmp(SourceBuffer.Get(), SerializedBulkData.GetPayload().Get().GetData(), BufferSize) == 0);
+
+		SerializedBulkData.Serialize(ReaderAr, nullptr);
+		TestFalse(TEXT("Bulkdata identifier should be invalid"), SerializedBulkData.GetIdentifier().IsValid());
+		TestTrue(TEXT("Bulkdata should not have a payload"), SerializedBulkData.GetPayload().Get().IsNull());
 	}
 
-	// Read in
+	// Now test the serialization when ValidBulkData is no longer in scope and has unregistered itself. In this
+	// case it doesn't matter if the BulkDataRegistry is enabled or not. In both cases we should get the original
+	// bulkdata identifier.
 	{
 		FEditorBulkData SerializedBulkData;
 
 		FMemoryReader ReaderAr(MemoryBuffer, bIsArPersistent);
 
 		SerializedBulkData.Serialize(ReaderAr, nullptr);
-		TestEqual(TEXT("Bulkdata identifier should remain the same"), SerializedBulkData.GetIdentifier(), ValidBulkData.GetIdentifier());
+		TestEqual(TEXT("Bulkdata identifier should remain the same"), SerializedBulkData.GetIdentifier(), ValidBulkDataId);
 		TestTrue(TEXT("Bulkdata payload should remain the same"), FMemory::Memcmp(SourceBuffer.Get(), SerializedBulkData.GetPayload().Get().GetData(), BufferSize) == 0);
 
 		SerializedBulkData.Serialize(ReaderAr, nullptr);
@@ -275,13 +329,14 @@ bool FEditorBulkDataTestSerializationToMemory::RunTest(const FString& Parameters
 		TestTrue(TEXT("Bulkdata should not have a payload"), SerializedBulkData.GetPayload().Get().IsNull());
 
 		SerializedBulkData.Serialize(ReaderAr, nullptr);
-		TestEqual(TEXT("Bulkdata identifier should remain the same"), SerializedBulkData.GetIdentifier(), ValidBulkData.GetIdentifier());
+		TestEqual(TEXT("Bulkdata identifier should remain the same"), SerializedBulkData.GetIdentifier(), ValidBulkDataId);
 		TestTrue(TEXT("Bulkdata payload should remain the same"), FMemory::Memcmp(SourceBuffer.Get(), SerializedBulkData.GetPayload().Get().GetData(), BufferSize) == 0);
 
 		SerializedBulkData.Serialize(ReaderAr, nullptr);
 		TestEqual(TEXT("Bulkdata identifier should remain the same"), SerializedBulkData.GetIdentifier(), EmptyBulkData.GetIdentifier());
 		TestTrue(TEXT("Bulkdata should not have a payload"), SerializedBulkData.GetPayload().Get().IsNull());
 	}
+
 	return true;
 }
 
@@ -439,7 +494,14 @@ bool FEditorBulkDataTestIdentifiers::RunTest(const FString& Parameters)
 			DstData.Serialize(ReaderAr, nullptr);
 		}
 
-		TestEqual(TEXT("Serialization should preserve the identifier"), SrcData.GetIdentifier(), DstData.GetIdentifier());
+		if (IsBulkDataRegistryEnabled())
+		{
+			TestNotEqual(TEXT("Serialization should not preserve the identifier"), SrcData.GetIdentifier(), DstData.GetIdentifier());
+		}
+		else
+		{
+			TestEqual(TEXT("Serialization should preserve the identifier"), SrcData.GetIdentifier(), DstData.GetIdentifier());
+		}	
 	}
 
 	// Test that serialization does not change the identifier (in this case serializing to and from a memory buffer) 
