LinkerLoad: Fixed issue with imported package linkers sometimes being incomplete and unusable for packages that are loaded using LoadPackageAsync and timing out during their async loading tick.

diff --git a/Engine/Source/Runtime/CoreUObject/Private/UObject/LinkerLoad.cpp b/Engine/Source/Runtime/CoreUObject/Private/UObject/LinkerLoad.cpp
--- a/Engine/Source/Runtime/CoreUObject/Private/UObject/LinkerLoad.cpp
+++ b/Engine/Source/Runtime/CoreUObject/Private/UObject/LinkerLoad.cpp
@@ -3266,6 +3266,19 @@ bool FLinkerLoad::VerifyImportInner(const int32 ImportIndex, FString& WarningSuf
 		// if we do not have a linker, then this import is native/in memory only
 		if( Import.SourceLinker )
 		{
+			if (!Import.SourceLinker->bHasFoundExistingExports)
+			{
+				UE_LOG(LogLinker, Log, TEXT("Source linker '%s' imported from '%s' has not processed all header information, ticking it now"), *GetNameSafe(Import.SourceLinker->LinkerRoot), *GetNameSafe(LinkerRoot));
+				// This means that the source linker timed out during it's async loading tick and that the header information hasn't been fully processed yet.
+				// Make sure that the header information is available but don't process any imports or exports (LOAD_NoVerify).
+				TGuardValue<uint32> LinkerLoadFlagGuard(Import.SourceLinker->LoadFlags, Import.SourceLinker->LoadFlags | LOAD_NoVerify);
+				if (Import.SourceLinker->Tick(0.f, false, false, nullptr) == LINKER_Failed)
+				{
+					UE_LOG(LogLinker, Warning, TEXT("Failed ticking import source linker '%s'"), *GetNameSafe(Import.SourceLinker->LinkerRoot));
+					return false;
+				}
+			}
+
 			// Assign the linker root of the source linker as the package we are looking for.
 			Pkg = Import.SourceLinker->LinkerRoot;
 
