Add a utility to load a package trailer directly from a file using a string based file path rather than FPackagePath

#rnx

- This can be used to load a trailer from a package file that is not stored in a mount point (and so would fail if FPackagePath is used)

diff --git a/Engine/Source/Runtime/CoreUObject/Private/UObject/PackageTrailer.cpp b/Engine/Source/Runtime/CoreUObject/Private/UObject/PackageTrailer.cpp
--- a/Engine/Source/Runtime/CoreUObject/Private/UObject/PackageTrailer.cpp
+++ b/Engine/Source/Runtime/CoreUObject/Private/UObject/PackageTrailer.cpp
@@ -5,6 +5,7 @@
 #include "Algo/Count.h"
 #include "Algo/RemoveIf.h"
 #include "CoreGlobals.h"
+#include "HAL/FileManager.h"
 #include "Misc/ConfigCacheIni.h"
 #include "Misc/PackagePath.h"
 #include "Misc/PackageSegment.h"
@@ -40,7 +41,7 @@ namespace
 {
 
 /** Utility for recording failed package open reasons */
-void LogPackageOpenFailureMessage(const FPackagePath& PackagePath)
+void LogPackageOpenFailureMessage(const FString& DebugName)
 {
 	// TODO: Check the various error paths here again!
 	const uint32 SystemError = FPlatformMisc::GetLastError();
@@ -50,11 +51,11 @@ void LogPackageOpenFailureMessage(const FPackagePath& PackagePath)
 	{
 		TCHAR SystemErrorMsg[2048] = { 0 };
 		FPlatformMisc::GetSystemErrorMessage(SystemErrorMsg, sizeof(SystemErrorMsg), SystemError);
-		UE_LOG(LogSerialization, Error, TEXT("Could not open the file '%s' for reading due to system error: '%s' (%d))"), *PackagePath.GetDebugName(), SystemErrorMsg, SystemError);
+		UE_LOG(LogSerialization, Error, TEXT("Could not open the file '%s' for reading due to system error: '%s' (%d))"), *DebugName, SystemErrorMsg, SystemError);
 	}
 	else
 	{
-		UE_LOG(LogSerialization, Error, TEXT("Could not open (%s) to read FPackageTrailer with an unknown error"), *PackagePath.GetDebugName());
+		UE_LOG(LogSerialization, Error, TEXT("Could not open (%s) to read FPackageTrailer with an unknown error"), *DebugName);
 	}
 }
 
@@ -396,11 +397,27 @@ bool FPackageTrailer::TryLoadFromPackage(const FPackagePath& PackagePath, FPacka
 	}
 	else
 	{
-		LogPackageOpenFailureMessage(PackagePath);
+		LogPackageOpenFailureMessage(PackagePath.GetDebugName());
 		return false;
 	}	
 }
 
+bool FPackageTrailer::TryLoadFromFile(const FString& Path, FPackageTrailer& OutTrailer)
+{
+	TUniquePtr<FArchive> PackageAr(IFileManager::Get().CreateFileReader(*Path));
+
+	if (PackageAr.IsValid())
+	{
+		PackageAr->Seek(PackageAr->TotalSize());
+		return OutTrailer.TryLoadBackwards(*PackageAr);
+	}
+	else
+	{
+		LogPackageOpenFailureMessage(Path);
+		return false;
+	}
+}
+
 bool FPackageTrailer::TryLoad(FArchive& Ar)
 {
 	check(Ar.IsLoading());
diff --git a/Engine/Source/Runtime/CoreUObject/Public/UObject/PackageTrailer.h b/Engine/Source/Runtime/CoreUObject/Public/UObject/PackageTrailer.h
--- a/Engine/Source/Runtime/CoreUObject/Public/UObject/PackageTrailer.h
+++ b/Engine/Source/Runtime/CoreUObject/Public/UObject/PackageTrailer.h
@@ -315,6 +315,9 @@ public:
 	/** Try to load a trailer from a given package path. Note that it will always try to load the trailer from the workspace domain */
 	[[nodiscard]] static bool TryLoadFromPackage(const FPackagePath& PackagePath, FPackageTrailer& OutTrailer);
 
+	/** Try to load a trailer from a given file path. Note that it will always try to load the trailer from the workspace domain */
+	[[nodiscard]] static bool TryLoadFromFile(const FString& Path, FPackageTrailer& OutTrailer);
+
 	FPackageTrailer() = default;
 	~FPackageTrailer() = default;
 
