GC: Fix incorrect assert in GCObjectReferencer (fallout from 19898875).
Assertion failed: GObjUnhashUnreachableIsInProgress || GObjIncrementalPurgeIsInProgress || !IsGarbageCollecting() [File:.\Runtime/CoreUObject/Private/Misc/GCObjectReferencer.cpp]
With the new default lock behavior it should be allowed to add and remove references outside of the reachability analysis scope.
With the legacy lock behavior checking IsGarbageCollectingAndLockingUObjectHashTables() means the same thing as IsGarbageCollecting().
When the legacy lock behavior is deleted this check can be simplified to check(!IsGarbageCollectingAndLockingUObjectHashTables()).

#rnx

diff --git a/Engine/Source/Runtime/CoreUObject/Private/Misc/GCObjectReferencer.cpp b/Engine/Source/Runtime/CoreUObject/Private/Misc/GCObjectReferencer.cpp
--- a/Engine/Source/Runtime/CoreUObject/Private/Misc/GCObjectReferencer.cpp
+++ b/Engine/Source/Runtime/CoreUObject/Private/Misc/GCObjectReferencer.cpp
@@ -59,7 +59,8 @@ void UGCObjectReferencer::AddReferencedObjects(UObject* InThis, FReferenceCollec
 void UGCObjectReferencer::AddObject(FGCObject* Object)
 {
 	check(Object);
-	check(GObjUnhashUnreachableIsInProgress || GObjIncrementalPurgeIsInProgress || !IsGarbageCollecting());	FScopeLock ReferencedObjectsLock(&ReferencedObjectsCritical);
+	check(GObjUnhashUnreachableIsInProgress || GObjIncrementalPurgeIsInProgress || !IsGarbageCollectingAndLockingUObjectHashTables());
+	FScopeLock ReferencedObjectsLock(&ReferencedObjectsCritical);
 	// Make sure there are no duplicates. Should be impossible...
 	checkSlow(!ReferencedObjects.Contains(Object));
 	ReferencedObjects.Add(Object);
@@ -75,7 +76,8 @@ void UGCObjectReferencer::AddObject(FGCObject* Object)
 void UGCObjectReferencer::RemoveObject(FGCObject* Object)
 {
 	check(Object);
-	check(GObjUnhashUnreachableIsInProgress || GObjIncrementalPurgeIsInProgress || !IsGarbageCollecting());	FScopeLock ReferencedObjectsLock(&ReferencedObjectsCritical);
+	check(GObjUnhashUnreachableIsInProgress || GObjIncrementalPurgeIsInProgress || !IsGarbageCollectingAndLockingUObjectHashTables());
+	FScopeLock ReferencedObjectsLock(&ReferencedObjectsCritical);
 	int32 NumRemoved = ReferencedObjects.RemoveSingleSwap(Object);
 	check(NumRemoved == 1);
 #if VERIFY_DISREGARD_GC_ASSUMPTIONS
