Add a utility to load a package trailer directly from an archive

#rnx

diff --git a/Engine/Source/Runtime/CoreUObject/Private/UObject/PackageTrailer.cpp b/Engine/Source/Runtime/CoreUObject/Private/UObject/PackageTrailer.cpp
--- a/Engine/Source/Runtime/CoreUObject/Private/UObject/PackageTrailer.cpp
+++ b/Engine/Source/Runtime/CoreUObject/Private/UObject/PackageTrailer.cpp
@@ -418,6 +418,12 @@ bool FPackageTrailer::TryLoadFromFile(const FString& Path, FPackageTrailer& OutT
 	}
 }
 
+bool FPackageTrailer::TryLoadFromArchive(FArchive& Ar, FPackageTrailer& OutTrailer)
+{
+	Ar.Seek(Ar.TotalSize());
+	return OutTrailer.TryLoadBackwards(Ar);
+}
+
 bool FPackageTrailer::TryLoad(FArchive& Ar)
 {
 	check(Ar.IsLoading());
diff --git a/Engine/Source/Runtime/CoreUObject/Public/UObject/PackageTrailer.h b/Engine/Source/Runtime/CoreUObject/Public/UObject/PackageTrailer.h
--- a/Engine/Source/Runtime/CoreUObject/Public/UObject/PackageTrailer.h
+++ b/Engine/Source/Runtime/CoreUObject/Public/UObject/PackageTrailer.h
@@ -315,9 +315,12 @@ public:
 	/** Try to load a trailer from a given package path. Note that it will always try to load the trailer from the workspace domain */
 	[[nodiscard]] static bool TryLoadFromPackage(const FPackagePath& PackagePath, FPackageTrailer& OutTrailer);
 
-	/** Try to load a trailer from a given file path. Note that it will always try to load the trailer from the workspace domain */
+	/** Try to load a trailer from a given file path. */
 	[[nodiscard]] static bool TryLoadFromFile(const FString& Path, FPackageTrailer& OutTrailer);
 
+	/** Try to load a trailer from a given archive. Assumes that the trailer is at the end of the archive */
+	[[nodiscard]] static bool TryLoadFromArchive(FArchive& Ar, FPackageTrailer& OutTrailer);
+
 	FPackageTrailer() = default;
 	~FPackageTrailer() = default;
 
