add missing include

#rnx

diff --git a/Engine/Source/Runtime/CoreUObject/Tests/ObjectRefTrackingTestBase.h b/Engine/Source/Runtime/CoreUObject/Tests/ObjectRefTrackingTestBase.h
--- a/Engine/Source/Runtime/CoreUObject/Tests/ObjectRefTrackingTestBase.h
+++ b/Engine/Source/Runtime/CoreUObject/Tests/ObjectRefTrackingTestBase.h
@@ -8,6 +8,7 @@
 
 #include "CoreMinimal.h"
 #include "UObject/ObjectHandle.h"
+#include "Misc/ScopeLock.h"
 
 
 class FObjectRefTrackingTestBase
