Fix a crash when trying to load a package trailer (backwards) from an archive that is smaller than the package footer size

#rnx

- This is most likely to occur when trying to load a package trailer from a corrupted 0 byte file.
- Currently we load the footer then check that the package tag exists to confirm that the package file is malformed
- We need to check that there is actually room for us to seek backwards to the potential start point of the footer otherwise FArchive code will assert on a negative seek position.

diff --git a/Engine/Source/Runtime/CoreUObject/Private/UObject/PackageTrailer.cpp b/Engine/Source/Runtime/CoreUObject/Private/UObject/PackageTrailer.cpp
--- a/Engine/Source/Runtime/CoreUObject/Private/UObject/PackageTrailer.cpp
+++ b/Engine/Source/Runtime/CoreUObject/Private/UObject/PackageTrailer.cpp
@@ -472,7 +472,13 @@ bool FPackageTrailer::TryLoadBackwards(FArchive& Ar)
 {
 	check(Ar.IsLoading());
 
-	Ar.Seek(Ar.Tell() - FFooter::SizeOnDisk);
+	const int64 FooterPos = Ar.Tell() - (int64)FFooter::SizeOnDisk;
+	if (FooterPos <= 0)
+	{
+		return false;
+	}
+
+	Ar.Seek(FooterPos);
 
 	FFooter Footer;
 
