Added checks for unsupported nested NetDeltaSerialize structs.

diff --git a/Engine/Source/Runtime/CoreUObject/Private/UObject/Class.cpp b/Engine/Source/Runtime/CoreUObject/Private/UObject/Class.cpp
--- a/Engine/Source/Runtime/CoreUObject/Private/UObject/Class.cpp
+++ b/Engine/Source/Runtime/CoreUObject/Private/UObject/Class.cpp
@@ -4410,63 +4410,169 @@ void UClass::ValidateRuntimeReplicationData()
 {
 	DECLARE_SCOPE_CYCLE_COUNTER(TEXT("Class ValidateRuntimeReplicationData"), STAT_Class_ValidateRuntimeReplicationData, STATGROUP_Game);
 
+	bool bCheckLifetimeProps = true;
+
 	if (HasAnyClassFlags(CLASS_CompiledFromBlueprint) || bLayoutChanging)
 	{
 		// Blueprint classes don't always generate a GetLifetimeReplicatedProps function. 
 		// Assume the Blueprint compiler was ok to do this.
-		return;
+		bCheckLifetimeProps = false;
 	}
 
-	if (HasAnyClassFlags(CLASS_ReplicationDataIsSetUp) == false)
+	if (bCheckLifetimeProps && HasAnyClassFlags(CLASS_ReplicationDataIsSetUp) == false)
 	{
 		UE_LOG(LogClass, Warning, TEXT("ValidateRuntimeReplicationData for class %s called before ReplicationData was setup."), *GetName());
-		return;
+		bCheckLifetimeProps = false;
 	}
 
 	// Our replication data was set up, but there are no class reps, so there's nothing to do.
 	if (ClassReps.Num() == 0)
 	{
-		return;
+		bCheckLifetimeProps = false;
 	}
 
 	// Let's compare the CDO's registered lifetime properties with the Class's net properties
-	TArray<FLifetimeProperty> LifetimeProps;
-	LifetimeProps.Reserve(ClassReps.Num());
+	if (bCheckLifetimeProps)
+	{
+		TArray<FLifetimeProperty> LifetimeProps;
+		LifetimeProps.Reserve(ClassReps.Num());
 
-	const UObject* Object = GetDefaultObject();
-	Object->GetLifetimeReplicatedProps(LifetimeProps);
+		const UObject* Object = GetDefaultObject();
+		Object->GetLifetimeReplicatedProps(LifetimeProps);
 
-	if (LifetimeProps.Num() == ClassReps.Num())
-	{
-		// All replicated properties were registered for this class
-		return;
+		// Find which properties where not registered by the user code
+		if (LifetimeProps.Num() != ClassReps.Num())
+		{
+			for (int32 RepIndex=0; RepIndex<ClassReps.Num(); ++RepIndex)
+			{
+				const FProperty* RepProp = ClassReps[RepIndex].Property;
+
+				const FLifetimeProperty* LifetimeProp = LifetimeProps.FindByPredicate(
+					[&RepIndex](const FLifetimeProperty& Var) { return Var.RepIndex == RepIndex; });
+
+				if (LifetimeProp == nullptr)
+				{
+					// Check if this unregistered property type uses a custom delta serializer
+					if (const FStructProperty* StructProperty = CastField<FStructProperty>(RepProp))
+					{
+						const UScriptStruct* Struct = StructProperty->Struct;
+
+						if (EnumHasAnyFlags(Struct->StructFlags, STRUCT_NetDeltaSerializeNative))
+						{
+							UE_LOG(LogClass, Warning,
+									TEXT("Property %s::%s (SourceClass: %s) with custom net delta serializer was not registered in ")
+									TEXT("GetLifetimeReplicatedProps. This property will replicate but you should still register it."),
+									ToCStr(GetName()), ToCStr(RepProp->GetName()), ToCStr(RepProp->GetOwnerClass()->GetName()));
+
+							continue;
+						}
+					}
+
+					UE_LOG(LogClass, Warning,
+							TEXT("Property %s::%s (SourceClass: %s) was not registered in GetLifetimeReplicatedProps. ")
+							TEXT("This property will not be replicated. Use DISABLE_REPLICATED_PROPERTY if not replicating was intentional."),
+							ToCStr(GetName()), ToCStr(RepProp->GetName()), ToCStr(RepProp->GetOwnerClass()->GetName()));
+				}
+			}
+		}
 	}
 
-	// Find which properties where not registered by the user code
-	for (int32 RepIndex = 0; RepIndex < ClassReps.Num(); ++RepIndex)
+	// Find NetDeltaSerialize properties which are not top-level/parent properties (e.g. nested in another struct)
+	for (const FRepRecord& RepRecord : ClassReps)
 	{
-		const FProperty* RepProp = ClassReps[RepIndex].Property;
+		const FProperty* BaseProp = RepRecord.Property;
+		TArray<const UStruct*, TInlineAllocator<16>> RecursiveStructList;
 
-		const FLifetimeProperty* LifetimeProp = LifetimeProps.FindByPredicate([&RepIndex](const FLifetimeProperty& Var) { return Var.RepIndex == RepIndex; });
-
-		if (LifetimeProp == nullptr)
+		enum class ECheckStructType : uint8
 		{
-			// Check if this unregistered property type uses a custom delta serializer
-			if (const FStructProperty* StructProperty = CastField<FStructProperty>(RepProp))
+			TopLevel,	// Parent/top-level - NetDeltaSerialize is valid
+			Recursive	// Recursive (within other structs/arrays) - NetDeltaSerialize is not valid
+		};
+
+		auto CheckStructRecursive = [&BaseProp, &RecursiveStructList]
+									(auto CheckStructRecursive, const FProperty* InProp, ECheckStructType CheckType)
+									-> void
 			{
-				const UScriptStruct* Struct = StructProperty->Struct;
+				auto BadStruct = [&BaseProp](const FStructProperty* StructProp)
+					{
+						UE_LOG(LogClass, Warning, TEXT("Property %s contained nested NetDeltaSerialize struct '%s', ")
+								TEXT("when this is not supported. Only use NetDeltaSerialize structs at class level. ")
+								TEXT("Struct will replicate using non-delta serialization."),
+								ToCStr(BaseProp->GetPathName()), ToCStr(StructProp->GetPathName()));
+					};
 
-				if (EnumHasAnyFlags(Struct->StructFlags, STRUCT_NetDeltaSerializeNative))
+				const FStructProperty* ParamStructProp = nullptr;
+				bool bNonDeltaArrayStruct = false;
+
+				if (!EnumHasAnyFlags(InProp->PropertyFlags, CPF_RepSkip))
 				{
-					UE_LOG(LogClass, Warning, TEXT("Property %s::%s (SourceClass: %s) with custom net delta serializer was not registered in GetLifetimeReplicatedProps. This property will replicate but you should still register it."),
-						*GetName(), *RepProp->GetName(), *RepProp->GetOwnerClass()->GetName());
-					continue;
+					if (const FStructProperty* StructProp = CastField<FStructProperty>(InProp))
+					{
+						ParamStructProp = StructProp;
+					}
+					else if (const FArrayProperty* ArrayProp = CastField<FArrayProperty>(InProp))
+					{
+						if (const FStructProperty* ArrayStructProp = CastField<FStructProperty>(ArrayProp->Inner))
+						{
+							if (EnumHasAnyFlags(ArrayStructProp->Struct->StructFlags, STRUCT_NetDeltaSerializeNative))
+							{
+								BadStruct(ArrayStructProp);
+							}
+							else
+							{
+								ParamStructProp = ArrayStructProp;
+								bNonDeltaArrayStruct = true;
+							}
+						}
+					}
 				}
-			}
 
-			UE_LOG(LogClass, Warning, TEXT("Property %s::%s (SourceClass: %s) was not registered in GetLifetimeReplicatedProps. This property will not be replicated. Use DISABLE_REPLICATED_PROPERTY if not replicating was intentional."),
-				*GetName(), *RepProp->GetName(), *RepProp->GetOwnerClass()->GetName());
-		}
+
+				const bool bTopLevelNonDeltaStruct = ParamStructProp != nullptr && CheckType == ECheckStructType::TopLevel &&
+					(bNonDeltaArrayStruct || !EnumHasAnyFlags(ParamStructProp->Struct->StructFlags, STRUCT_NetDeltaSerializeNative));
+
+				const bool bFreshRecursiveStruct = ParamStructProp != nullptr && CheckType == ECheckStructType::Recursive &&
+													!RecursiveStructList.Contains(ParamStructProp->Struct);
+
+				if (bTopLevelNonDeltaStruct || bFreshRecursiveStruct || bNonDeltaArrayStruct)
+				{
+					RecursiveStructList.Add(ParamStructProp->Struct);
+
+					for (TFieldIterator<FProperty> It(ParamStructProp->Struct); It; ++It)
+					{
+						if (!EnumHasAnyFlags(It->PropertyFlags, CPF_RepSkip))
+						{
+							const FStructProperty* CurLevelStructProp = nullptr;
+
+							if (const FStructProperty* StructProp = CastField<FStructProperty>(*It))
+							{
+								CurLevelStructProp = StructProp;
+							}
+							else if (const FArrayProperty* ArrayProp = CastField<FArrayProperty>(*It))
+							{
+								if (const FStructProperty* ArrayStructProp = CastField<FStructProperty>(ArrayProp->Inner))
+								{
+									CurLevelStructProp = ArrayStructProp;
+								}
+							}
+
+							if (CurLevelStructProp != nullptr)
+							{
+								if (EnumHasAnyFlags(CurLevelStructProp->Struct->StructFlags, STRUCT_NetDeltaSerializeNative))
+								{
+									BadStruct(CurLevelStructProp);
+								}
+								else
+								{
+									CheckStructRecursive(CheckStructRecursive, CurLevelStructProp, ECheckStructType::Recursive);
+								}
+							}
+						}
+					}
+				}
+			};
+
+		CheckStructRecursive(CheckStructRecursive, BaseProp, ECheckStructType::TopLevel);
 	}
 }
 
