Always include referencing property info in HandleGarbageReference since it's guarded by a cvar anyway.
Allow ENABLE_GC_OBJECT_CHECKS to be overriden outside of GarbageCollection.h

diff --git a/Engine/Source/Runtime/CoreUObject/Private/UObject/GarbageCollection.cpp b/Engine/Source/Runtime/CoreUObject/Private/UObject/GarbageCollection.cpp
--- a/Engine/Source/Runtime/CoreUObject/Private/UObject/GarbageCollection.cpp
+++ b/Engine/Source/Runtime/CoreUObject/Private/UObject/GarbageCollection.cpp
@@ -922,14 +922,12 @@ public:
 		{
 			if (GarbageReferencer != FGCObject::GGCObjectReferencer)
 			{
-#if ENABLE_GC_OBJECT_CHECKS
 				if (TokenIndex >= 0)
 				{
 					UClass* Class = GarbageReferencer->GetClass();
 					FTokenInfo TokenInfo = Class->ReferenceTokenStream.GetTokenInfo(TokenIndex);
 					PropertyName = TokenInfo.Name;
 				}
-#endif // ENABLE_GC_OBJECT_CHECKS
 				ObjectsToSerializeStruct.GarbageReferences.Add(FGarbageReferenceInfo(GarbageReferencer, Object, Object, PropertyName));
 			}
 			else
diff --git a/Engine/Source/Runtime/CoreUObject/Public/UObject/GarbageCollection.h b/Engine/Source/Runtime/CoreUObject/Public/UObject/GarbageCollection.h
--- a/Engine/Source/Runtime/CoreUObject/Public/UObject/GarbageCollection.h
+++ b/Engine/Source/Runtime/CoreUObject/Public/UObject/GarbageCollection.h
@@ -24,7 +24,9 @@
 #define PERF_DETAILED_PER_CLASS_GC_STATS				(LOOKING_FOR_PERF_ISSUES || 0) 
 
 /** UObject pointer checks are disabled by default in shipping and test builds as they add roughly 20% overhead to GC times */
-#define ENABLE_GC_OBJECT_CHECKS (!(UE_BUILD_TEST || UE_BUILD_SHIPPING) || 0)
+#ifndef ENABLE_GC_OBJECT_CHECKS
+	#define ENABLE_GC_OBJECT_CHECKS (!(UE_BUILD_TEST || UE_BUILD_SHIPPING) || 0)
+#endif
 
 /** Token debug info (token names) enabled in non-shipping builds */
 #define ENABLE_GC_TOKEN_DEBUG_INFO (!UE_BUILD_SHIPPING)
