Add simple command-line options for enabling and disabling pending kill.

diff --git a/Engine/Source/Runtime/CoreUObject/Private/UObject/ObjectBaseUtility.cpp b/Engine/Source/Runtime/CoreUObject/Private/UObject/ObjectBaseUtility.cpp
--- a/Engine/Source/Runtime/CoreUObject/Private/UObject/ObjectBaseUtility.cpp
+++ b/Engine/Source/Runtime/CoreUObject/Private/UObject/ObjectBaseUtility.cpp
@@ -194,6 +194,15 @@ void InitNoPendingKill()
 	check(GConfig);
 	bool bPendingKillEnabled = false;
 	GConfig->GetBool(TEXT("/Script/Engine.GarbageCollectionSettings"), TEXT("gc.PendingKillEnabled"), bPendingKillEnabled, GEngineIni);
+	// Allow command-line overrides for easy testing without unique builds.
+	if (FParse::Param(FCommandLine::Get(), TEXT("DisablePendingKill")))
+	{
+		bPendingKillEnabled = false;
+	}
+	else if (FParse::Param(FCommandLine::Get(), TEXT("EnablePendingKill")))
+	{
+		bPendingKillEnabled = true;
+	}
 	// Try to sync even though we're not gonna use the console var
 	UObjectBaseUtility::bPendingKillDisabled = !bPendingKillEnabled;
 	GPendingKillEnabled = bPendingKillEnabled;
