Fix non-unity error.

diff --git a/Engine/Source/Runtime/CoreUObject/Private/UObject/UObjectBaseUtility.cpp b/Engine/Source/Runtime/CoreUObject/Private/UObject/UObjectBaseUtility.cpp
--- a/Engine/Source/Runtime/CoreUObject/Private/UObject/UObjectBaseUtility.cpp
+++ b/Engine/Source/Runtime/CoreUObject/Private/UObject/UObjectBaseUtility.cpp
@@ -15,6 +15,7 @@
 #include "ProfilingDebugging/MallocProfiler.h"
 #include "HAL/IConsoleManager.h"
 #include "Misc/ConfigCacheIni.h"
+#include "Misc/CommandLine.h"
 
 /***********************/
 /******** Names ********/
