Fix non-unity error in correct file this time.

diff --git a/Engine/Source/Runtime/CoreUObject/Private/UObject/ObjectBaseUtility.cpp b/Engine/Source/Runtime/CoreUObject/Private/UObject/ObjectBaseUtility.cpp
--- a/Engine/Source/Runtime/CoreUObject/Private/UObject/ObjectBaseUtility.cpp
+++ b/Engine/Source/Runtime/CoreUObject/Private/UObject/ObjectBaseUtility.cpp
@@ -9,6 +9,7 @@
 #include "UObject/LinkerLoad.h"
 #include "HAL/IConsoleManager.h"
 #include "Misc/ConfigCacheIni.h"
+#include "Misc/CommandLine.h"
 
 /**
  * Returns the UE version of the linker for this object.
diff --git a/Engine/Source/Runtime/CoreUObject/Private/UObject/UObjectBaseUtility.cpp b/Engine/Source/Runtime/CoreUObject/Private/UObject/UObjectBaseUtility.cpp
--- a/Engine/Source/Runtime/CoreUObject/Private/UObject/UObjectBaseUtility.cpp
+++ b/Engine/Source/Runtime/CoreUObject/Private/UObject/UObjectBaseUtility.cpp
@@ -15,7 +15,6 @@
 #include "ProfilingDebugging/MallocProfiler.h"
 #include "HAL/IConsoleManager.h"
 #include "Misc/ConfigCacheIni.h"
-#include "Misc/CommandLine.h"
 
 /***********************/
 /******** Names ********/
