UPackage::FindAssetInPackage: Prioritize Valid Assets over deleted ones if multiple are found in same package

This fixes replaced actors being considered as the package asset when replacement actor is the better choice. (AActor::PreSaveRoot was being called on wrong actor)

diff --git a/Engine/Source/Runtime/CoreUObject/Private/UObject/Package.cpp b/Engine/Source/Runtime/CoreUObject/Private/UObject/Package.cpp
--- a/Engine/Source/Runtime/CoreUObject/Private/UObject/Package.cpp
+++ b/Engine/Source/Runtime/CoreUObject/Private/UObject/Package.cpp
@@ -126,20 +126,33 @@ void UPackage::Serialize( FArchive& Ar )
 UObject* UPackage::FindAssetInPackage(EObjectFlags RequiredTopLevelFlags) const
 {
 	UObject* Asset = nullptr;
+	bool bAssetValid = false;
 
-	ForEachObjectWithPackage(this, [&Asset, RequiredTopLevelFlags](UObject* Object)
+	ForEachObjectWithPackage(this, [&Asset, &bAssetValid, RequiredTopLevelFlags](UObject* Object)
 		{
 			if (Object->IsAsset() && !UE::AssetRegistry::FFiltering::ShouldSkipAsset(Object) &&
 				(RequiredTopLevelFlags == RF_NoFlags || Object->HasAnyFlags(RequiredTopLevelFlags)))
 			{
-				if (FAssetData::IsUAsset(Object))
+				const bool bIsValid = IsValid(Object);
+				const bool bIsUAsset = FAssetData::IsUAsset(Object);
+
+				if (!Asset)
 				{
 					Asset = Object;
-					return false;
+					bAssetValid = bIsValid;
+					// stop iterating if Asset is valid and also a UAsset
+					return !(bIsValid && bIsUAsset);
 				}
-				if (!Asset)
+				else if(bIsValid)
 				{
-					Asset = Object;
+					// Overwrite found asset if previous was invalid or new one is a UAsset
+					if (!bAssetValid || bIsUAsset)
+					{
+						Asset = Object;
+						bAssetValid = true;
+					}
+					// stop iterating if found asset is a UAsset
+					return !bIsUAsset;
 				}
 			}
 			return true;
