Saving External Objects and AssetRegistry dependencies: mark objects in external packages as used-in-game (not editor-only), so that their imports are marked as used-in-game, and are followed when looking for manage references.
#rnx

diff --git a/Engine/Source/Runtime/CoreUObject/Private/UObject/SavePackage/SavePackageUtilities.cpp b/Engine/Source/Runtime/CoreUObject/Private/UObject/SavePackage/SavePackageUtilities.cpp
--- a/Engine/Source/Runtime/CoreUObject/Private/UObject/SavePackage/SavePackageUtilities.cpp
+++ b/Engine/Source/Runtime/CoreUObject/Private/UObject/SavePackage/SavePackageUtilities.cpp
@@ -776,7 +776,12 @@ bool IsEditorOnlyObject(const UObject* InObject, bool bCheckRecursive, bool bChe
 	}
 	else
 	{
-		Package = InObject->GetOutermost();
+		// In the case that the object is an external object, we want to use its host package, rather than its
+		// external package when testing editor-only. All external packages are editor-only, but the objects in
+		// the external package logically belong to the host package, and are editoronly if and only if the host is.
+		// So use GetOutermostObject()->GetPackage(). This will be the same as GetPackage for non-external objects.
+		UObject* HostObject = InObject->GetOutermostObject();
+		Package = HostObject->GetPackage();
 	}
 	
 	if (Package && Package->HasAnyPackageFlags(PKG_EditorOnly))
