SavePackage: Fix argument error in log statement - filename is a TCHAR* rather than an FString.
#rnx

diff --git a/Engine/Source/Runtime/CoreUObject/Private/UObject/SavePackage2.cpp b/Engine/Source/Runtime/CoreUObject/Private/UObject/SavePackage2.cpp
--- a/Engine/Source/Runtime/CoreUObject/Private/UObject/SavePackage2.cpp
+++ b/Engine/Source/Runtime/CoreUObject/Private/UObject/SavePackage2.cpp
@@ -1977,7 +1977,7 @@ ESavePackageResult FinalizeFile(FStructuredArchive::FRecord& StructuredArchiveRo
 	{
 		bool bIsOptionalRealm = SaveContext.GetCurrentHarvestingRealm() == ESaveRealm::Optional;
 		FLinkerSave* Linker = SaveContext.GetLinker();
-		UE_LOG(LogSavePackage, Verbose, TEXT("Async saving from memory to '%s'"), *SaveContext.GetFilename());
+		UE_LOG(LogSavePackage, Verbose, TEXT("Async saving from memory to '%s'"), SaveContext.GetFilename());
 		FLargeMemoryWriter* Writer = static_cast<FLargeMemoryWriter*>(Linker->Saver);
 
 		if (PackageWriter)
