Adding a thread_local global FFrame* tracking pointer that mimics the ScriptStack in the FBlueprintContextTracker, but remains enabled when DO_BLUEPRINT_GUARD=0

diff --git a/Engine/Source/Runtime/CoreUObject/Private/UObject/ScriptCore.cpp b/Engine/Source/Runtime/CoreUObject/Private/UObject/ScriptCore.cpp
--- a/Engine/Source/Runtime/CoreUObject/Private/UObject/ScriptCore.cpp
+++ b/Engine/Source/Runtime/CoreUObject/Private/UObject/ScriptCore.cpp
@@ -90,6 +90,8 @@ COREUOBJECT_API int32 GNativeDuplicate=0;
 
 COREUOBJECT_API int32 GMaximumScriptLoopIterations = 1000000;
 
+thread_local static FFrame* GTopTrackingStackFrame = nullptr;
+
 #if DO_BLUEPRINT_GUARD
 
 #define CHECK_RUNAWAY { FBlueprintContextTracker::Get().AddRunaway(); }
@@ -523,6 +525,22 @@ void FFrame::InitPrintScriptCallstack()
 }
 #endif
 
+COREUOBJECT_API FFrame* FFrame::PushThreadLocalTopStackFrame(FFrame* NewTopStackFrame)
+{
+	FFrame* Result = GTopTrackingStackFrame;
+	GTopTrackingStackFrame = NewTopStackFrame;
+	return Result;
+}
+
+COREUOBJECT_API void FFrame::PopThreadLocalTopStackFrame(FFrame* NewTopStackFrame)
+{
+	GTopTrackingStackFrame = NewTopStackFrame;
+}
+
+COREUOBJECT_API FFrame* FFrame::GetThreadLocalTopStackFrame()
+{
+	return GTopTrackingStackFrame;
+}
 //
 // Error or warning handler.
 //
diff --git a/Engine/Source/Runtime/CoreUObject/Public/UObject/Stack.h b/Engine/Source/Runtime/CoreUObject/Public/UObject/Stack.h
--- a/Engine/Source/Runtime/CoreUObject/Public/UObject/Stack.h
+++ b/Engine/Source/Runtime/CoreUObject/Public/UObject/Stack.h
@@ -111,6 +111,9 @@ public:
 	/** Currently executed native function */
 	UFunction* CurrentNativeFunction;
 
+	/** Previous tracking frame */
+	FFrame* PreviousTrackingFrame;
+
 	bool bArrayContextFailed;
 
 	/** If this flag gets set (usually from throwing a EBlueprintExceptionType::AbortExecution exception), execution shall immediately stop and return */
@@ -133,7 +136,18 @@ public:
 		{
 			BlueprintExceptionTracker.ScriptStack.Pop(false);
 		}
+
+		// ensure that GTopTrackingStackFrame is accurate
+		if (BlueprintExceptionTracker.ScriptStack.Num() == 0)
+		{
+			ensure(PreviousTrackingFrame == nullptr);
+		}
+		else
+		{
+			ensure(BlueprintExceptionTracker.ScriptStack.Last() == PreviousTrackingFrame);
+		}
 #endif
+		PopThreadLocalTopStackFrame(PreviousTrackingFrame);
 	}
 
 	// Functions.
@@ -216,6 +230,10 @@ public:
 #if DO_BLUEPRINT_GUARD
 	static void InitPrintScriptCallstack();
 #endif
+
+	COREUOBJECT_API static FFrame* PushThreadLocalTopStackFrame(FFrame* NewTopStackFrame);
+	COREUOBJECT_API static void PopThreadLocalTopStackFrame(FFrame* NewTopStackFrame);
+	COREUOBJECT_API static FFrame* GetThreadLocalTopStackFrame();
 };
 
 
@@ -244,6 +262,8 @@ inline FFrame::FFrame( UObject* InObject, UFunction* InNode, void* InLocals, FFr
 #if DO_BLUEPRINT_GUARD
 	FBlueprintContextTracker::Get().ScriptStack.Push(this);
 #endif
+	PreviousTrackingFrame = PushThreadLocalTopStackFrame(this);
+
 #if PER_FUNCTION_SCRIPT_STATS
 	if (InPreviousFrame)
 	{
