Look in the correct object when determining if the inherited instance should be instantiated

diff --git a/Engine/Source/Runtime/CoreUObject/Private/UObject/LinkerLoad.cpp b/Engine/Source/Runtime/CoreUObject/Private/UObject/LinkerLoad.cpp
--- a/Engine/Source/Runtime/CoreUObject/Private/UObject/LinkerLoad.cpp
+++ b/Engine/Source/Runtime/CoreUObject/Private/UObject/LinkerLoad.cpp
@@ -4822,7 +4822,7 @@ UObject* FLinkerLoad::CreateExport( int32 Index )
 		// it should not be recreated unless the archetype also has it
 		if (Export.bIsInheritedInstance)
 		{
-			if (FindObjectWithOuter(Template, LoadClass, NewName) == nullptr)
+			if (FindObjectWithOuter(ThisParent->GetArchetype(), LoadClass, NewName) == nullptr)
 			{
 				Export.bExportLoadFailed = true;
 				return nullptr;
