GC locks are no longer acquired when GIsInitialLoad is true.

diff --git a/Engine/Source/Runtime/CoreUObject/Private/UObject/GarbageCollection.cpp b/Engine/Source/Runtime/CoreUObject/Private/UObject/GarbageCollection.cpp
--- a/Engine/Source/Runtime/CoreUObject/Private/UObject/GarbageCollection.cpp
+++ b/Engine/Source/Runtime/CoreUObject/Private/UObject/GarbageCollection.cpp
@@ -2250,15 +2250,7 @@ void GatherUnreachableObjects(bool bForceSingleThreaded)
  */
 void CollectGarbageInternal(EObjectFlags KeepFlags, bool bPerformFullPurge)
 {
-#if !UE_WITH_GC
-	return;
-#else
-	if (GIsInitialLoad)
-	{
-		// During initial load classes may not yet have their GC token streams assembled
-		UE_LOG(LogGarbage, Log, TEXT("Skipping CollectGarbage() call during initial load. It's not safe."));
-		return;
-	}
+#if UE_WITH_GC
 	SCOPE_TIME_GUARD(TEXT("Collect Garbage"));
 	SCOPED_NAMED_EVENT(CollectGarbageInternal, FColor::Red);
 	CSV_EVENT_GLOBAL(TEXT("GC"));
@@ -2614,6 +2606,16 @@ bool UnhashUnreachableObjects(bool bUseTimeLimit, double TimeLimit)
 
 void CollectGarbage(EObjectFlags KeepFlags, bool bPerformFullPurge)
 {
+#if !UE_WITH_GC
+	return;
+#else
+	if (GIsInitialLoad)
+	{
+		// During initial load classes may not yet have their GC token streams assembled
+		UE_LOG(LogGarbage, Log, TEXT("Skipping CollectGarbage() call during initial load. It's not safe."));
+		return;
+	}
+
 	// No other thread may be performing UObject operations while we're running
 	AcquireGCLock();
 
@@ -2626,10 +2628,21 @@ void CollectGarbage(EObjectFlags KeepFlags, bool bPerformFullPurge)
 		ReleaseGCLock();
 	}
 	// With the new FGCLockBehavior::Default behavior the lock was released after reachability analysis inside CollectGarbageInternal
+#endif
 }
 
 bool TryCollectGarbage(EObjectFlags KeepFlags, bool bPerformFullPurge)
 {
+#if !UE_WITH_GC
+	return false;
+#else
+	if (GIsInitialLoad)
+	{
+		// During initial load classes may not yet have their GC token streams assembled
+		UE_LOG(LogGarbage, Log, TEXT("Skipping TryCollectGarbage() call during initial load. It's not safe."));
+		return false;
+	}
+
 	// No other thread may be performing UObject operations while we're running
 	bool bCanRunGC = FGCCSyncObject::Get().TryGCLock();
 	if (!bCanRunGC)
@@ -2661,6 +2674,7 @@ bool TryCollectGarbage(EObjectFlags KeepFlags, bool bPerformFullPurge)
 	}
 
 	return bCanRunGC;
+#endif
 }
 
 void UObject::CallAddReferencedObjects(FReferenceCollector& Collector)
