SavePackage: Fix fatal error when cooking a Blueprint package with a class that has a subobject that is excluded from the targetplatform.

diff --git a/Engine/Source/Runtime/CoreUObject/Private/UObject/SavePackage2.cpp b/Engine/Source/Runtime/CoreUObject/Private/UObject/SavePackage2.cpp
--- a/Engine/Source/Runtime/CoreUObject/Private/UObject/SavePackage2.cpp
+++ b/Engine/Source/Runtime/CoreUObject/Private/UObject/SavePackage2.cpp
@@ -1377,6 +1377,12 @@ void SavePreloadDependencies(FStructuredArchive::FRecord& StructuredArchiveRoot,
 						// Only include subobject archetypes
 						if (SubObj->HasAnyFlags(RF_DefaultSubObject | RF_ArchetypeObject))
 						{
+							// Don't include the archetype of SubObjects that we excluded in the harvesting phase; we didn't add their archetypes
+							if (SaveContext.IsExcluded(SubObj))
+							{
+								continue;
+							}
+
 							SubObj = SubObj->GetArchetype();
 							while (SubObj->HasAnyFlags(RF_Transient)) // transient components are stripped by the ICH, so find the one it will really use at runtime
 							{
