Rename `AreEqual` to `AreIdentical` in dynamic value representation.

diff --git a/Engine/Source/Runtime/CoreUObject/Private/UObject/DynamicallyTypedValue.cpp b/Engine/Source/Runtime/CoreUObject/Private/UObject/DynamicallyTypedValue.cpp
--- a/Engine/Source/Runtime/CoreUObject/Private/UObject/DynamicallyTypedValue.cpp
+++ b/Engine/Source/Runtime/CoreUObject/Private/UObject/DynamicallyTypedValue.cpp
@@ -19,7 +19,7 @@ UE::FDynamicallyTypedValueType& UE::FDynamicallyTypedValue::NullType()
 		virtual void SerializeValue(FStructuredArchive::FSlot Slot, void* Data) const {}
 
 		virtual uint32 GetValueHash(const void* Data) const override { return 0; }
-		virtual bool AreEqual(const void* DataA, const void* DataB) const override { return true; }
+		virtual bool AreIdentical(const void* DataA, const void* DataB) const override { return true; }
 	};
 
 	static FNullType Singleton;
diff --git a/Engine/Source/Runtime/CoreUObject/Public/UObject/DynamicallyTypedValue.h b/Engine/Source/Runtime/CoreUObject/Public/UObject/DynamicallyTypedValue.h
--- a/Engine/Source/Runtime/CoreUObject/Public/UObject/DynamicallyTypedValue.h
+++ b/Engine/Source/Runtime/CoreUObject/Public/UObject/DynamicallyTypedValue.h
@@ -38,7 +38,7 @@ namespace UE
 		virtual void SerializeValue(FStructuredArchive::FSlot Slot, void* Data) const = 0;
 
 		virtual uint32 GetValueHash(const void* Data) const = 0;
-		virtual bool AreEqual(const void* DataA, const void* DataB) const = 0;
+		virtual bool AreIdentical(const void* DataA, const void* DataB) const = 0;
 
 		SIZE_T GetNumBytes() const { return NumBytes; }
 		uint8 GetMinAlignmentLogTwo() const { return MinAlignmentLogTwo; }
