- Waorkaround to allow an optional object pointing to a non-optional in the same package (I think it's okay to allow this, but not sure - it will only affect -editoroptional mode, so should be okay for now)

diff --git a/Engine/Source/Runtime/CoreUObject/Private/UObject/SavePackage2.cpp b/Engine/Source/Runtime/CoreUObject/Private/UObject/SavePackage2.cpp
--- a/Engine/Source/Runtime/CoreUObject/Private/UObject/SavePackage2.cpp
+++ b/Engine/Source/Runtime/CoreUObject/Private/UObject/SavePackage2.cpp
@@ -359,6 +359,12 @@ ESavePackageResult HarvestPackage(FSaveContext& SaveContext)
 					// Flag the package itself to be an import
 					bHasNonOptionalSelfReference = true;
 				}
+				// if an optional object points to a non-optional object in this package, we assume it will be being saved, so we allow it
+				else if (It->Obj->IsIn(SaveContext.GetPackage()))
+				{
+					UE_LOG(LogSavePackage, Warning, TEXT("An optional object in '%s' is pointing to non-optional object %s, but since it's in the same package, we are allowing it to not be found (usually we look in non-optional packages for references from optional to non-optional, which must already be cooked)"),
+						*SaveContext.GetPackage()->GetFullName(), *It->Obj->GetFullName());
+				}
 				// if not found in the game context, record an illegal reference
 				else
 				{
