UClass: Move the UClass pointers to the cppclass's AddReferencedObject function onto a new type FUObjectCppClassStaticFunctions that allows us to easily add new static functions that are reflected in the same way.

Move the existing DeclareCustomVersions function which was declared as a virtual UObject function into a static UObject function using this method.

#rnx

diff --git a/Engine/Source/Editor/Kismet/Private/BlueprintCompilationManager.cpp b/Engine/Source/Editor/Kismet/Private/BlueprintCompilationManager.cpp
--- a/Engine/Source/Editor/Kismet/Private/BlueprintCompilationManager.cpp
+++ b/Engine/Source/Editor/Kismet/Private/BlueprintCompilationManager.cpp
@@ -1002,7 +1002,7 @@ void FBlueprintCompilationManagerImpl::FlushCompilationQueueImpl(bool bSuppressB
 					
 					SkeletonToRelink->ClassConstructor = nullptr;
 					SkeletonToRelink->ClassVTableHelperCtorCaller = nullptr;
-					SkeletonToRelink->ClassAddReferencedObjects = nullptr;
+					SkeletonToRelink->CppClassStaticFunctions.Reset();
 					SkeletonToRelink->Bind();
 					SkeletonToRelink->ClearFunctionMapsCaches();
 					SkeletonToRelink->StaticLink(true);
@@ -1935,7 +1935,7 @@ void FBlueprintCompilationManagerImpl::ReparentHierarchies(const TMap<UClass*, U
 		{
 			ClassToReinstance->ClassConstructor = nullptr;
 			ClassToReinstance->ClassVTableHelperCtorCaller = nullptr;
-			ClassToReinstance->ClassAddReferencedObjects = nullptr;
+			ClassToReinstance->CppClassStaticFunctions.Reset();
 		
 			// check to see if we're a direct parent of one of the new classes:
 			UClass* const* NewParent = OldToNewClasses.Find(ClassToReinstance->GetSuperClass());
diff --git a/Engine/Source/Editor/KismetCompiler/Private/KismetCompiler.cpp b/Engine/Source/Editor/KismetCompiler/Private/KismetCompiler.cpp
--- a/Engine/Source/Editor/KismetCompiler/Private/KismetCompiler.cpp
+++ b/Engine/Source/Editor/KismetCompiler/Private/KismetCompiler.cpp
@@ -582,7 +582,7 @@ void FKismetCompilerContext::CleanAndSanitizeClass(UBlueprintGeneratedClass* Cla
 	{
 		ParentClass = UObject::StaticClass();
 	}
-	TransientClass->ClassAddReferencedObjects = ParentClass->AddReferencedObjects;
+	TransientClass->CppClassStaticFunctions = ParentClass->CppClassStaticFunctions;
 	TransientClass->ClassGeneratedBy = Blueprint;
 	TransientClass->ClassFlags |= CLASS_CompiledFromBlueprint|CLASS_NewerVersionExists;
 
diff --git a/Engine/Source/Editor/UnrealEd/Private/EditorDomain/EditorDomainUtils.cpp b/Engine/Source/Editor/UnrealEd/Private/EditorDomain/EditorDomainUtils.cpp
--- a/Engine/Source/Editor/UnrealEd/Private/EditorDomain/EditorDomainUtils.cpp
+++ b/Engine/Source/Editor/UnrealEd/Private/EditorDomain/EditorDomainUtils.cpp
@@ -778,7 +778,7 @@ private:
 TArray<FGuid> GetCustomVersions(UClass& Class)
 {
 	FCustomVersionCollectorArchive Ar;
-	Class.GetDefaultObject()->DeclareCustomVersions(Ar);
+	Class.CallDeclareCustomVersions(Ar);
 	// Default objects of blueprint classes are serialized during SavePackage with a special call to
 	// UBlueprintGeneratedClass::SerializeDefaultObject
 	// All packages that include a BlueprintGeneratedClass import the UClass BlueprintGeneratedClass
diff --git a/Engine/Source/Editor/UnrealEd/Private/Kismet2/BlueprintEditorUtils.cpp b/Engine/Source/Editor/UnrealEd/Private/Kismet2/BlueprintEditorUtils.cpp
--- a/Engine/Source/Editor/UnrealEd/Private/Kismet2/BlueprintEditorUtils.cpp
+++ b/Engine/Source/Editor/UnrealEd/Private/Kismet2/BlueprintEditorUtils.cpp
@@ -1276,7 +1276,7 @@ void FBlueprintEditorUtils::RemoveStaleFunctions(UBlueprintGeneratedClass* Class
 		FString OrphanedClassString = FString::Printf(TEXT("ORPHANED_DATA_ONLY_%s"), *Class->GetName());
 		FName OrphanedClassName = MakeUniqueObjectName(GetTransientPackage(), UBlueprintGeneratedClass::StaticClass(), FName(*OrphanedClassString));
 		UClass* OrphanedClass = NewObject<UBlueprintGeneratedClass>(GetTransientPackage(), OrphanedClassName, RF_Public | RF_Transient);
-		OrphanedClass->ClassAddReferencedObjects = Class->AddReferencedObjects;
+		OrphanedClass->CppClassStaticFunctions = Class->CppClassStaticFunctions;
 		OrphanedClass->ClassFlags |= CLASS_CompiledFromBlueprint;
 		OrphanedClass->ClassGeneratedBy = Class->ClassGeneratedBy;
 
diff --git a/Engine/Source/Runtime/CoreUObject/Private/Blueprint/BlueprintSupport.cpp b/Engine/Source/Runtime/CoreUObject/Private/Blueprint/BlueprintSupport.cpp
--- a/Engine/Source/Runtime/CoreUObject/Private/Blueprint/BlueprintSupport.cpp
+++ b/Engine/Source/Runtime/CoreUObject/Private/Blueprint/BlueprintSupport.cpp
@@ -839,7 +839,7 @@ static PlaceholderType* MakeImportPlaceholder(UObject* Outer, const TCHAR* Targe
 	// in the FLinkerLoad's ImportPlaceholders map
 
 	// make sure the class is fully formed (has its 
-	// ClassAddReferencedObjects/ClassConstructor members set)
+	// CppClassStaticFunctions/ClassConstructor members set)
 	PlaceholderObj->Bind();
 	PlaceholderObj->StaticLink(/*bRelinkExistingProperties =*/true);
 
diff --git a/Engine/Source/Runtime/CoreUObject/Private/Misc/GCObjectReferencer.cpp b/Engine/Source/Runtime/CoreUObject/Private/Misc/GCObjectReferencer.cpp
--- a/Engine/Source/Runtime/CoreUObject/Private/Misc/GCObjectReferencer.cpp
+++ b/Engine/Source/Runtime/CoreUObject/Private/Misc/GCObjectReferencer.cpp
@@ -209,7 +209,7 @@ void UGCObjectReferencer::VerifyGCObjectNames()
 
 IMPLEMENT_CORE_INTRINSIC_CLASS(UGCObjectReferencer, UObject, 
 	{
-		Class->ClassAddReferencedObjects = &UGCObjectReferencer::AddReferencedObjects;
+		Class->CppClassStaticFunctions = UOBJECT_CPPCLASS_STATICFUNCTIONS_FORCLASS(UGCObjectReferencer);
 	}
 );
 
diff --git a/Engine/Source/Runtime/CoreUObject/Private/UObject/Class.cpp b/Engine/Source/Runtime/CoreUObject/Private/UObject/Class.cpp
--- a/Engine/Source/Runtime/CoreUObject/Private/UObject/Class.cpp
+++ b/Engine/Source/Runtime/CoreUObject/Private/UObject/Class.cpp
@@ -2159,7 +2159,7 @@ void UStruct::InstanceSubobjectTemplates( void* Data, void const* DefaultData, U
 
 IMPLEMENT_CORE_INTRINSIC_CLASS(UStruct, UField,
 	{
-		Class->ClassAddReferencedObjects = &UStruct::AddReferencedObjects;
+		Class->CppClassStaticFunctions = UOBJECT_CPPCLASS_STATICFUNCTIONS_FORCLASS(UStruct);
 		Class->EmitObjectReference(STRUCT_OFFSET(UStruct, SuperStruct), TEXT("SuperStruct"));
 		Class->EmitObjectReference(STRUCT_OFFSET(UStruct, Children), TEXT("Children"));
 
@@ -3665,10 +3665,9 @@ void UClass::PostInitProperties()
 	Super::PostInitProperties();
 	if (HasAnyFlags(RF_ClassDefaultObject))
 	{
-		if (ClassAddReferencedObjects == NULL)
+		if (!CppClassStaticFunctions.IsInitialized())
 		{
-			// Default__Class uses its own AddReferencedObjects function.
-			ClassAddReferencedObjects = &UClass::AddReferencedObjects;
+			CppClassStaticFunctions = UOBJECT_CPPCLASS_STATICFUNCTIONS_FORCLASS(UClass);
 		}
 	}
 }
@@ -3774,8 +3773,8 @@ class FRestoreClassInfo: public FRestoreForUObjectOverwrite
 	UClass::ClassConstructorType Constructor;
 	/** Saved ClassVTableHelperCtorCaller **/
 	UClass::ClassVTableHelperCtorCallerType ClassVTableHelperCtorCaller;
-	/** Saved ClassConstructor **/
-	UClass::ClassAddReferencedObjectsType AddReferencedObjects;
+	/** Saved CppClassStaticFunctions **/
+	FUObjectCppClassStaticFunctions CppClassStaticFunctions;
 	/** Saved NativeFunctionLookupTable. */
 	TArray<FNativeFunctionLookup> NativeFunctionLookupTable;
 public:
@@ -3796,7 +3795,7 @@ public:
 		CastFlags(Save->ClassCastFlags),
 		Constructor(Save->ClassConstructor),
 		ClassVTableHelperCtorCaller(Save->ClassVTableHelperCtorCaller),
-		AddReferencedObjects(Save->ClassAddReferencedObjects),
+		CppClassStaticFunctions(Save->CppClassStaticFunctions),
 		NativeFunctionLookupTable(Save->NativeFunctionLookupTable)
 	{
 	}
@@ -3813,7 +3812,7 @@ public:
 		Target->ClassCastFlags |= CastFlags;
 		Target->ClassConstructor = Constructor;
 		Target->ClassVTableHelperCtorCaller = ClassVTableHelperCtorCaller;
-		Target->ClassAddReferencedObjects = AddReferencedObjects;
+		Target->CppClassStaticFunctions = CppClassStaticFunctions;
 		Target->NativeFunctionLookupTable = NativeFunctionLookupTable;
 	}
 };
@@ -4067,7 +4066,7 @@ void UClass::Bind()
 	}
 
 	UClass* SuperClass = GetSuperClass();
-	if (SuperClass && (ClassConstructor == nullptr || ClassAddReferencedObjects == nullptr
+	if (SuperClass && (ClassConstructor == nullptr || !CppClassStaticFunctions.IsInitialized()
 		|| ClassVTableHelperCtorCaller == nullptr
 		))
 	{
@@ -4081,9 +4080,9 @@ void UClass::Bind()
 		{
 			ClassVTableHelperCtorCaller = SuperClass->ClassVTableHelperCtorCaller;
 		}
-		if (!ClassAddReferencedObjects)
+		if (!CppClassStaticFunctions.IsInitialized())
 		{
-			ClassAddReferencedObjects = SuperClass->ClassAddReferencedObjects;
+			CppClassStaticFunctions = SuperClass->CppClassStaticFunctions;
 		}
 
 		// propagate flags.
@@ -5073,7 +5072,7 @@ void UClass::PurgeClass(bool bRecompilingOnLoad)
 	RefLink = nullptr;
 	PropertyLink = nullptr;
 	DestructorLink = nullptr;
-	ClassAddReferencedObjects = nullptr;
+	CppClassStaticFunctions.Reset();
 
 	ScriptAndPropertyObjectReferences.Empty();
 	DeleteUnresolvedScriptProperties();
@@ -5261,12 +5260,12 @@ UClass::UClass
 	EObjectFlags	InFlags,
 	ClassConstructorType InClassConstructor,
 	ClassVTableHelperCtorCallerType InClassVTableHelperCtorCaller,
-	ClassAddReferencedObjectsType InClassAddReferencedObjects
+	FUObjectCppClassStaticFunctions&& InCppClassStaticFunctions
 )
 :	UStruct					( EC_StaticConstructor, InSize, InAlignment, InFlags )
 ,	ClassConstructor		( InClassConstructor )
 ,	ClassVTableHelperCtorCaller(InClassVTableHelperCtorCaller)
-,	ClassAddReferencedObjects( InClassAddReferencedObjects )
+,	CppClassStaticFunctions(MoveTemp(InCppClassStaticFunctions))
 ,	ClassUnique				( 0 )
 ,	bCooked					( false )
 ,	bLayoutChanging			( false )
@@ -5439,7 +5438,7 @@ bool UClass::HotReloadPrivateStaticClass(
 	const TCHAR*    InConfigName,
 	ClassConstructorType InClassConstructor,
 	ClassVTableHelperCtorCallerType InClassVTableHelperCtorCaller,
-	ClassAddReferencedObjectsType InClassAddReferencedObjects,
+	FUObjectCppClassStaticFunctions&& InCppClassStaticFunctions,
 	class UClass* TClass_Super_StaticClass,
 	class UClass* TClass_WithinClass_StaticClass
 	)
@@ -5460,7 +5459,7 @@ bool UClass::HotReloadPrivateStaticClass(
 	ClassConstructorType OldClassConstructor = ClassConstructor;
 	ClassConstructor = InClassConstructor;
 	ClassVTableHelperCtorCaller = InClassVTableHelperCtorCaller;
-	ClassAddReferencedObjects = InClassAddReferencedObjects;
+	CppClassStaticFunctions = InCppClassStaticFunctions; // Not MoveTemp; it is used again below
 	/* No recursive ::StaticClass calls allowed. Setup extras. */
 	/* @todo safe? 
 	if (TClass_Super_StaticClass != this)
@@ -5527,7 +5526,7 @@ bool UClass::HotReloadPrivateStaticClass(
 				{
 					Class->ClassConstructor = ClassConstructor;
 					Class->ClassVTableHelperCtorCaller = ClassVTableHelperCtorCaller;
-					Class->ClassAddReferencedObjects = ClassAddReferencedObjects;
+					Class->CppClassStaticFunctions = InCppClassStaticFunctions; // Not MoveTemp; it is used in later loop iterations
 					CountClass++;
 				}
 			}
@@ -5852,7 +5851,7 @@ bool UClass::IsClassGroupName(const TCHAR* InGroupName) const
 #if WITH_EDITORONLY_DATA
 IMPLEMENT_CORE_INTRINSIC_CLASS(UClass, UStruct,
 	{
-		Class->ClassAddReferencedObjects = &UClass::AddReferencedObjects;
+		Class->CppClassStaticFunctions = UOBJECT_CPPCLASS_STATICFUNCTIONS_FORCLASS(UClass);
 
 		Class->EmitObjectReference(STRUCT_OFFSET(UClass, ClassDefaultObject), TEXT("ClassDefaultObject"));
 		Class->EmitObjectReference(STRUCT_OFFSET(UClass, ClassWithin), TEXT("ClassWithin"));
@@ -5863,7 +5862,7 @@ IMPLEMENT_CORE_INTRINSIC_CLASS(UClass, UStruct,
 #else
 IMPLEMENT_CORE_INTRINSIC_CLASS(UClass, UStruct,
 	{
-		Class->ClassAddReferencedObjects = &UClass::AddReferencedObjects;
+		Class->CppClassStaticFunctions = UOBJECT_CPPCLASS_STATICFUNCTIONS_FORCLASS(UClass);
 
 		Class->EmitObjectReference(STRUCT_OFFSET(UClass, ClassDefaultObject), TEXT("ClassDefaultObject"));
 		Class->EmitObjectReference(STRUCT_OFFSET(UClass, ClassWithin), TEXT("ClassWithin"));
@@ -5884,7 +5883,7 @@ void GetPrivateStaticClassBody(
 	const TCHAR* InConfigName,
 	UClass::ClassConstructorType InClassConstructor,
 	UClass::ClassVTableHelperCtorCallerType InClassVTableHelperCtorCaller,
-	UClass::ClassAddReferencedObjectsType InClassAddReferencedObjects,
+	FUObjectCppClassStaticFunctions&& InCppClassStaticFunctions,
 	UClass::StaticClassFunctionType InSuperClassFn,
 	UClass::StaticClassFunctionType InWithinClassFn
 	)
@@ -5905,7 +5904,7 @@ void GetPrivateStaticClassBody(
 					InConfigName,
 					InClassConstructor,
 					InClassVTableHelperCtorCaller,
-					InClassAddReferencedObjects,
+					FUObjectCppClassStaticFunctions(InCppClassStaticFunctions),
 					InSuperClassFn(),
 					InWithinClassFn()
 					))
@@ -5941,7 +5940,7 @@ void GetPrivateStaticClassBody(
 		EObjectFlags(RF_Public | RF_Standalone | RF_Transient | RF_MarkAsNative | RF_MarkAsRootSet),
 		InClassConstructor,
 		InClassVTableHelperCtorCaller,
-		InClassAddReferencedObjects
+		MoveTemp(InCppClassStaticFunctions)
 		);
 	check(ReturnClass);
 	
@@ -6591,7 +6590,7 @@ PRAGMA_ENABLE_DEPRECATION_WARNINGS
 PRAGMA_DISABLE_DEPRECATION_WARNINGS
 IMPLEMENT_CORE_INTRINSIC_CLASS(UDynamicClass, UClass,
 {
-	Class->ClassAddReferencedObjects = &UDynamicClass::AddReferencedObjects;
+	Class->CppClassStaticFunctions = UOBJECT_CPPCLASS_STATICFUNCTIONS_FORCLASS(UDynamicClass);
 }
 );
 PRAGMA_ENABLE_DEPRECATION_WARNINGS
diff --git a/Engine/Source/Runtime/CoreUObject/Private/UObject/GarbageCollection.cpp b/Engine/Source/Runtime/CoreUObject/Private/UObject/GarbageCollection.cpp
--- a/Engine/Source/Runtime/CoreUObject/Private/UObject/GarbageCollection.cpp
+++ b/Engine/Source/Runtime/CoreUObject/Private/UObject/GarbageCollection.cpp
@@ -3244,20 +3244,20 @@ void UClass::AssembleReferenceTokenStream(bool bForce)
 		}
 
 		{
-			check(ClassAddReferencedObjects != NULL);
+			check(CppClassStaticFunctions.GetAddReferencedObjects() != nullptr);
 			const bool bKeepOuter = true;//GetFName() != NAME_Package;
 			const bool bKeepClass = true;//!HasAnyInternalFlags(EInternalObjectFlags::Native);
 
-			ClassAddReferencedObjectsType AddReferencedObjectsFn = nullptr;
+			FUObjectCppClassStaticFunctions::AddReferencedObjectsType AddReferencedObjectsFn = nullptr;
 #if !WITH_EDITOR
 			// In no-editor builds UObject::ARO is empty, thus only classes
 			// which implement their own ARO function need to have the ARO token generated.
-			if (ClassAddReferencedObjects != &UObject::AddReferencedObjects)
+			if (CppClassStaticFunctions.GetAddReferencedObjects() != &UObject::AddReferencedObjects)
 			{
-				AddReferencedObjectsFn = ClassAddReferencedObjects;
+				AddReferencedObjectsFn = CppClassStaticFunctions.GetAddReferencedObjects();
 			}
 #else
-			AddReferencedObjectsFn = ClassAddReferencedObjects;
+			AddReferencedObjectsFn = CppClassStaticFunctions.GetAddReferencedObjects();
 #endif
 			ReferenceTokenStream.Fixup(AddReferencedObjectsFn, bKeepOuter, bKeepClass);
 
diff --git a/Engine/Source/Runtime/CoreUObject/Private/UObject/Interface.cpp b/Engine/Source/Runtime/CoreUObject/Private/UObject/Interface.cpp
--- a/Engine/Source/Runtime/CoreUObject/Private/UObject/Interface.cpp
+++ b/Engine/Source/Runtime/CoreUObject/Private/UObject/Interface.cpp
@@ -6,7 +6,7 @@
 #if WITH_EDITOR
 IMPLEMENT_CORE_INTRINSIC_CLASS(UInterface, UObject,
 	{
-		Class->ClassAddReferencedObjects = &UInterface::AddReferencedObjects;
+		Class->CppClassStaticFunctions = UOBJECT_CPPCLASS_STATICFUNCTIONS_FORCLASS(UInterface);
 		Class->SetMetaData(TEXT("IsBlueprintBase"), TEXT("true"));
 		Class->SetMetaData(TEXT("CannotImplementInterfaceInBlueprint"), TEXT(""));
 	}
@@ -16,7 +16,7 @@ IMPLEMENT_CORE_INTRINSIC_CLASS(UInterface, UObject,
 
 IMPLEMENT_CORE_INTRINSIC_CLASS(UInterface, UObject,
 {
-	Class->ClassAddReferencedObjects = &UInterface::AddReferencedObjects;
+	Class->CppClassStaticFunctions = UOBJECT_CPPCLASS_STATICFUNCTIONS_FORCLASS(UInterface);
 }
 );
 
diff --git a/Engine/Source/Runtime/CoreUObject/Private/UObject/LinkerPlaceholderClass.cpp b/Engine/Source/Runtime/CoreUObject/Private/UObject/LinkerPlaceholderClass.cpp
--- a/Engine/Source/Runtime/CoreUObject/Private/UObject/LinkerPlaceholderClass.cpp
+++ b/Engine/Source/Runtime/CoreUObject/Private/UObject/LinkerPlaceholderClass.cpp
@@ -41,7 +41,7 @@ void ULinkerPlaceholderClass::Bind()
 	ClassVTableHelperCtorCaller = InternalVTableHelperCtorCaller<ULinkerPlaceholderClass>;
 	Super::Bind();
 
-	ClassAddReferencedObjects = &ULinkerPlaceholderClass::AddReferencedObjects;
+	CppClassStaticFunctions = UOBJECT_CPPCLASS_STATICFUNCTIONS_FORCLASS(ULinkerPlaceholderClass);
 }
 
 
diff --git a/Engine/Source/Runtime/CoreUObject/Private/UObject/Obj.cpp b/Engine/Source/Runtime/CoreUObject/Private/UObject/Obj.cpp
--- a/Engine/Source/Runtime/CoreUObject/Private/UObject/Obj.cpp
+++ b/Engine/Source/Runtime/CoreUObject/Private/UObject/Obj.cpp
@@ -1499,12 +1499,13 @@ void UObject::Serialize(FStructuredArchive::FRecord Record)
 #endif
 }
 
-void UObject::DeclareCustomVersions(FArchive& Ar)
+#if WITH_EDITORONLY_DATA
+void UObject::DeclareCustomVersions(FArchive& Ar, const UClass* SpecificSubclass)
 {
 	// DeclareCustomVersions is called on the default object for each class
 	// We first Serialize the object, which catches all the UsingCustomVersion statements
 	// class authors have added unconditionally in their Serialize function
-	Serialize(Ar);
+	SpecificSubclass->GetDefaultObject()->Serialize(Ar);
 
 	// To further catch CustomVersions used by non-native structs that are in an array or don't
 	// otherwise exist on the default object, Construct an instance of the struct and serialize
@@ -1512,12 +1513,11 @@ void UObject::DeclareCustomVersions(FArchive& Ar)
 	// Since structs can contain other structs, we do a tree search of the fields.
 	struct FStackData
 	{
-		UStruct* Struct;
+		const UStruct* Struct;
 		FProperty* NextProperty;
 	};
 	TArray<FStackData> StructStack;
-	UClass* Class = GetClass();
-	StructStack.Add(FStackData{ Class, Class->PropertyLink });
+	StructStack.Add(FStackData{ SpecificSubclass, SpecificSubclass->PropertyLink });
 	TArray<uint8> AllocationBuffer;
 	while (!StructStack.IsEmpty())
 	{
@@ -1572,6 +1572,7 @@ void UObject::DeclareCustomVersions(FArchive& Ar)
 		}
 	}
 }
+#endif
 
 void UObject::SerializeScriptProperties(FArchive& Ar) const
 {
diff --git a/Engine/Source/Runtime/CoreUObject/Public/UObject/Class.h b/Engine/Source/Runtime/CoreUObject/Public/UObject/Class.h
--- a/Engine/Source/Runtime/CoreUObject/Public/UObject/Class.h
+++ b/Engine/Source/Runtime/CoreUObject/Public/UObject/Class.h
@@ -2382,13 +2382,11 @@ public:
 
 	typedef void		(*ClassConstructorType)				(const FObjectInitializer&);
 	typedef UObject*	(*ClassVTableHelperCtorCallerType)	(FVTableHelper& Helper);
-	typedef void		(*ClassAddReferencedObjectsType)	(UObject*, class FReferenceCollector&);
 	typedef UClass* (*StaticClassFunctionType)();
 
 	ClassConstructorType ClassConstructor;
 	ClassVTableHelperCtorCallerType ClassVTableHelperCtorCaller;
-	/** Pointer to a static AddReferencedObjects method. */
-	ClassAddReferencedObjectsType ClassAddReferencedObjects;
+	FUObjectCppClassStaticFunctions CppClassStaticFunctions;
 
 	/** Class pseudo-unique counter; used to accelerate unique instance name generation */
 	mutable int32 ClassUnique;
@@ -2463,11 +2461,20 @@ public:
 	{
 		// The object must of this class type.
 		check(This->IsA(this)); 
-		// This is should always be set to something, at the very least to UObject::ARO
-		check(ClassAddReferencedObjects != nullptr);
-		ClassAddReferencedObjects(This, Collector);
+		// This should always be set to something, at the very least to UObject::ARO
+		check(CppClassStaticFunctions.GetAddReferencedObjects() != nullptr);
+		CppClassStaticFunctions.GetAddReferencedObjects()(This, Collector);
 	}
 
+#if WITH_EDITORONLY_DATA
+	/** Calls the c++ class's DeclareCustomVersions static function (from the nearest native parent if this is not native) */
+	void CallDeclareCustomVersions(FArchive& Ar) const
+	{
+		check(CppClassStaticFunctions.GetDeclareCustomVersions());
+		CppClassStaticFunctions.GetDeclareCustomVersions()(Ar, this);
+	}
+#endif
+
 	/** The class default object; used for delta serialization and object initialization */
 	UObject* ClassDefaultObject;
 
@@ -2555,7 +2562,7 @@ public:
 	UClass( EStaticConstructor, FName InName, uint32 InSize, uint32 InAlignment, EClassFlags InClassFlags, EClassCastFlags InClassCastFlags,
 		const TCHAR* InClassConfigName, EObjectFlags InFlags, ClassConstructorType InClassConstructor,
 		ClassVTableHelperCtorCallerType InClassVTableHelperCtorCaller,
-		ClassAddReferencedObjectsType InClassAddReferencedObjects);
+		FUObjectCppClassStaticFunctions&& InCppClassStaticFunctions);
 
 #if WITH_RELOAD
 	/**
@@ -2575,7 +2582,7 @@ public:
 		const TCHAR*    InConfigName,
 		ClassConstructorType InClassConstructor,
 		ClassVTableHelperCtorCallerType InClassVTableHelperCtorCaller,
-		ClassAddReferencedObjectsType InClassAddReferencedObjects,
+		FUObjectCppClassStaticFunctions&& InCppClassStaticFunctions,
 		class UClass* TClass_Super_StaticClass,
 		class UClass* TClass_WithinClass_StaticClass
 		);
@@ -3130,7 +3137,7 @@ public:
 	UDynamicClass(EStaticConstructor, FName InName, uint32 InSize, uint32 InAlignment, EClassFlags InClassFlags, EClassCastFlags InClassCastFlags,
 		const TCHAR* InClassConfigName, EObjectFlags InFlags, ClassConstructorType InClassConstructor,
 		ClassVTableHelperCtorCallerType InClassVTableHelperCtorCaller,
-		ClassAddReferencedObjectsType InClassAddReferencedObjects,
+		FUObjectCppClassStaticFunctions&& InCppClassStaticFunctions,
 		DynamicClassInitializerType InDynamicClassInitializer)
 	{}
 
@@ -3204,7 +3211,7 @@ COREUOBJECT_API void InitializePrivateStaticClass(
  * @param InConfigName Class config name
  * @param InClassConstructor Class constructor function pointer
  * @param InClassVTableHelperCtorCaller Class constructor function for vtable pointer
- * @param InClassAddReferencedObjects Class AddReferencedObjects function pointer
+ * @param InCppClassStaticFunctions Function pointers for the class's version of Unreal's reflected static functions
  * @param InSuperClassFn Super class function pointer
  * @param WithinClass Within class
  */
@@ -3220,7 +3227,7 @@ COREUOBJECT_API void GetPrivateStaticClassBody(
 	const TCHAR* InConfigName,
 	UClass::ClassConstructorType InClassConstructor,
 	UClass::ClassVTableHelperCtorCallerType InClassVTableHelperCtorCaller,
-	UClass::ClassAddReferencedObjectsType InClassAddReferencedObjects,
+	FUObjectCppClassStaticFunctions&& InCppClassStaticFunctions,
 	UClass::StaticClassFunctionType InSuperClassFn,
 	UClass::StaticClassFunctionType InWithinClassFn);
 
diff --git a/Engine/Source/Runtime/CoreUObject/Public/UObject/Object.h b/Engine/Source/Runtime/CoreUObject/Public/UObject/Object.h
--- a/Engine/Source/Runtime/CoreUObject/Public/UObject/Object.h
+++ b/Engine/Source/Runtime/CoreUObject/Public/UObject/Object.h
@@ -365,12 +365,14 @@ public:
 	 */
 	virtual void Serialize(FArchive& Ar);
 	virtual void Serialize(FStructuredArchive::FRecord Record);
+#if WITH_EDITORONLY_DATA
 	/**
 	 * Call Ar.UsingCustomVersion for every CustomVersion that might be serialized by this class when saving.
 	 * This duplicates CustomVersions declared in Serialize; Serialize still needs to declare them.
 	 * Used to track which customversions will be used by a package when it is resaved.
 	 * Not yet exhaustive; add CustomVersions as necessary to remove EditorDomain warnings about missing versions. */
-	virtual void DeclareCustomVersions(FArchive& Ar);
+	static void DeclareCustomVersions(FArchive& Ar, const UClass* SpecificSubclass);
+#endif
 
 	/** After a critical error, perform any mission-critical cleanup, such as restoring the video mode orreleasing hardware resources. */
 	virtual void ShutdownAfterError() {}
diff --git a/Engine/Source/Runtime/CoreUObject/Public/UObject/ObjectMacros.h b/Engine/Source/Runtime/CoreUObject/Public/UObject/ObjectMacros.h
--- a/Engine/Source/Runtime/CoreUObject/Public/UObject/ObjectMacros.h
+++ b/Engine/Source/Runtime/CoreUObject/Public/UObject/ObjectMacros.h
@@ -11,6 +11,7 @@
 #include "UObject/Script.h"
 
 class FObjectInitializer;
+class FReferenceCollector;
 struct FFrame;
 struct FClassReloadVersionInfo;
 
@@ -1838,6 +1839,109 @@ public: \
 #define DECLARE_WITHIN_UPACKAGE() \
 	DECLARE_WITHIN_INTERNAL( UPackage, true )
 
+#define UOBJECT_CPPCLASS_STATICFUNCTIONS_ALLCONFIGS(TClass) \
+	&TClass::AddReferencedObjects
+	/* UObjectCppClassStaticFunctions: Extend this macro with the address of your new static function, if it applies to all configs. */ \
+	/* Order must match the order in the FUObjectCppClassStaticFunctions constructor. */ \
+
+#if WITH_EDITORONLY_DATA
+	#define UOBJECT_CPPCLASS_STATICFUNCTIONS_WITHEDITORONLYDATA(TClass) \
+		, &TClass::DeclareCustomVersions
+		/* UObjectCppClassStaticFunctions: Extend this macro with the address of your new static function, if it is editor-only. */ \
+		/* Order must match the order in the FUObjectCppClassStaticFunctions constructor. */ \
+#else
+	#define UOBJECT_CPPCLASS_STATICFUNCTIONS_WITHEDITORONLYDATA(TClass)
+#endif
+
+/**
+ * A macro called from the IMPLEMENT_CLASS macro that allows the compiler to report to the UClass constructor
+ * the class-specific overrides of UnrealEngine's list of reflected UObject static functions.
+ */
+#define UOBJECT_CPPCLASS_STATICFUNCTIONS_FORCLASS(TClass) \
+	FUObjectCppClassStaticFunctions \
+	( \
+		UOBJECT_CPPCLASS_STATICFUNCTIONS_ALLCONFIGS(TClass) \
+		UOBJECT_CPPCLASS_STATICFUNCTIONS_WITHEDITORONLYDATA(TClass) \
+	)
+
+/**
+ * Collection of the pointers to our specified list of static functions that are defined on a specific c++ class,
+ * for reference by the corresponding UClass. The pointers in this structure point to the class's version of that
+ * function, if it exists, or to the version on the nearest parent class where it does exist, similar to a virtual
+ * function table in c++.
+ */
+struct FUObjectCppClassStaticFunctions
+{
+public:
+	typedef void (*AddReferencedObjectsType)	(UObject* ThisObject, FReferenceCollector& Ar);
+#if WITH_EDITORONLY_DATA
+	typedef void (*DeclareCustomVersionsType)   (FArchive& Ar, const UClass* SpecificSubclass);
+#endif
+	// UObjectCppClassStaticFunctions: Extend this list of types with the type of your new static function.
+
+	FUObjectCppClassStaticFunctions(AddReferencedObjectsType InAddReferencedObjects
+#if WITH_EDITORONLY_DATA
+		, DeclareCustomVersionsType InDeclareCustomVersions
+#endif
+	)
+		: AddReferencedObjects(InAddReferencedObjects)
+#if WITH_EDITORONLY_DATA
+		, DeclareCustomVersions(InDeclareCustomVersions)
+#endif
+	{
+		// Null elements are not valid in this constructor
+		check(InAddReferencedObjects);
+#if WITH_EDITORONLY_DATA
+		check(InDeclareCustomVersions);
+#endif
+		// UObjectCppClassStaticFunctions: Extend the constructor with initializers for your new static function member.
+		// Order must match the order in UOBJECT_CPPCLASS_STATICFUNCTIONS_FORCLASS.
+	}
+public:
+	FUObjectCppClassStaticFunctions() = default;
+	bool IsInitialized() const
+	{
+		// All methods that set the functions in this class guarantee that if any function is non-null, they all are.
+		// So we only need to check the first one to know whether we are completely initialized.
+		return AddReferencedObjects != nullptr;
+	}
+	void Reset()
+	{
+		*this = FUObjectCppClassStaticFunctions();
+	}
+
+public:
+	AddReferencedObjectsType GetAddReferencedObjects() const
+	{
+		return AddReferencedObjects;
+	}
+	void SetAddReferencedObjects(AddReferencedObjectsType InAddReferencedObjects)
+	{
+		check(InAddReferencedObjects != nullptr); // It is not valid to clear single elements (see IsInitialized). Call Reset to clear all elements.
+		AddReferencedObjects = InAddReferencedObjects;
+	}
+#if WITH_EDITORONLY_DATA
+	DeclareCustomVersionsType GetDeclareCustomVersions() const
+	{
+		return DeclareCustomVersions;
+	}
+	void SetDeclareCustomVersions(DeclareCustomVersionsType InDeclareCustomVersions)
+	{
+		check(InDeclareCustomVersions != nullptr); // It is not valid to clear single elements (see IsInitialized). Call Reset to clear all elements.
+		DeclareCustomVersions = InDeclareCustomVersions;
+	}
+#endif
+	// UObjectCppClassStaticFunctions: Extend the list of accessors for your new static function.
+
+private:
+	AddReferencedObjectsType AddReferencedObjects = nullptr;
+#if WITH_EDITORONLY_DATA
+	DeclareCustomVersionsType DeclareCustomVersions = nullptr;
+#endif
+	// UObjectCppClassStaticFunctions: Extend this list of members with the member for your new static function.
+
+};
+
 // Implement the GetPrivateStaticClass and the registration info but do not auto register the class.  
 // This is primarily used by UnrealHeaderTool
 #define IMPLEMENT_CLASS_NO_AUTO_REGISTRATION(TClass) \
@@ -1859,7 +1963,7 @@ public: \
 				TClass::StaticConfigName(), \
 				(UClass::ClassConstructorType)InternalConstructor<TClass>, \
 				(UClass::ClassVTableHelperCtorCallerType)InternalVTableHelperCtorCaller<TClass>, \
-				&TClass::AddReferencedObjects, \
+				UOBJECT_CPPCLASS_STATICFUNCTIONS_FORCLASS(TClass), \
 				&TClass::Super::StaticClass, \
 				&TClass::WithinClass::StaticClass \
 			); \
diff --git a/Engine/Source/Runtime/Engine/Classes/Animation/MorphTarget.h b/Engine/Source/Runtime/Engine/Classes/Animation/MorphTarget.h
--- a/Engine/Source/Runtime/Engine/Classes/Animation/MorphTarget.h
+++ b/Engine/Source/Runtime/Engine/Classes/Animation/MorphTarget.h
@@ -175,7 +175,9 @@ public:
 	//~ UObject interface
 
 	ENGINE_API virtual void Serialize(FArchive& Ar) override;
-	ENGINE_API virtual void DeclareCustomVersions(FArchive& Ar) override;
+#if WITH_EDITORONLY_DATA
+	ENGINE_API static void DeclareCustomVersions(FArchive& Ar, const UClass* SpecificSubclass);
+#endif
 	ENGINE_API virtual void PostLoad() override;
 
 	/** UObject does not support serialization via FMemoryArchive, so manually handle separately */
diff --git a/Engine/Source/Runtime/Engine/Classes/EdGraph/EdGraphNode.h b/Engine/Source/Runtime/Engine/Classes/EdGraph/EdGraphNode.h
--- a/Engine/Source/Runtime/Engine/Classes/EdGraph/EdGraphNode.h
+++ b/Engine/Source/Runtime/Engine/Classes/EdGraph/EdGraphNode.h
@@ -442,7 +442,9 @@ public:
 
 	// UObject interface
 	virtual void Serialize(FArchive& Ar) override;
-	virtual void DeclareCustomVersions(FArchive& Ar) override;
+#if WITH_EDITORONLY_DATA
+	static void DeclareCustomVersions(FArchive& Ar, const UClass* SpecificSubclass);
+#endif
 	// End of UObject interface
 
 #if WITH_EDITOR
diff --git a/Engine/Source/Runtime/Engine/Classes/EdGraph/EdGraphPin.h b/Engine/Source/Runtime/Engine/Classes/EdGraph/EdGraphPin.h
--- a/Engine/Source/Runtime/Engine/Classes/EdGraph/EdGraphPin.h
+++ b/Engine/Source/Runtime/Engine/Classes/EdGraph/EdGraphPin.h
@@ -207,7 +207,9 @@ public:
 
 	ENGINE_API bool Serialize(FArchive& Ar);
 	ENGINE_API void PostSerialize(const FArchive& Ar);
+#if WITH_EDITORONLY_DATA
 	ENGINE_API static void DeclareCustomVersions(FArchive& Ar);
+#endif
 
 	static ENGINE_API FEdGraphPinType GetPinTypeForTerminalType( const FEdGraphTerminalType& TerminalType );
 	static ENGINE_API FEdGraphPinType GetTerminalTypeForContainer( const FEdGraphPinType& ContainerType );
@@ -526,7 +528,9 @@ public:
 
 	/** Serializes an array of pins as the owner. Only the OwningNode should call this function. */
 	static void SerializeAsOwningNode(FArchive& Ar, TArray<UEdGraphPin*>& ArrayRef);
-	static void DeclareCustomVersions(FArchive& Ar);
+#if WITH_EDITORONLY_DATA
+	static void DeclarePinCustomVersions(FArchive& Ar);
+#endif
 
 	/** Marks the owning node as modified. */
 	ENGINE_API bool Modify(bool bAlwaysMarkDirty = true);
diff --git a/Engine/Source/Runtime/Engine/Classes/Engine/SkeletalMesh.h b/Engine/Source/Runtime/Engine/Classes/Engine/SkeletalMesh.h
--- a/Engine/Source/Runtime/Engine/Classes/Engine/SkeletalMesh.h
+++ b/Engine/Source/Runtime/Engine/Classes/Engine/SkeletalMesh.h
@@ -691,7 +691,9 @@ struct FSkeletalMaterial
 	}
 
 	friend FArchive& operator<<( FArchive& Ar, FSkeletalMaterial& Elem );
+#if WITH_EDITORONLY_DATA
 	static void DeclareCustomVersions(FArchive& Ar);
+#endif
 
 	ENGINE_API friend bool operator==( const FSkeletalMaterial& LHS, const FSkeletalMaterial& RHS );
 	ENGINE_API friend bool operator==( const FSkeletalMaterial& LHS, const UMaterialInterface& RHS );
@@ -2757,7 +2759,9 @@ public:
 	PRAGMA_ENABLE_DEPRECATION_WARNINGS
 	virtual void PreSave(FObjectPreSaveContext ObjectSaveContext) override;
 	virtual void Serialize(FArchive& Ar) override;
-	virtual void DeclareCustomVersions(FArchive& Ar) override;
+#if WITH_EDITORONLY_DATA
+	static void DeclareCustomVersions(FArchive& Ar, const UClass* SpecificSubclass);
+#endif
 	virtual void PostInitProperties() override;
 	virtual void PostLoad() override;
 	virtual bool IsPostLoadThreadSafe() const override;
diff --git a/Engine/Source/Runtime/Engine/Private/BlueprintGeneratedClass.cpp b/Engine/Source/Runtime/Engine/Private/BlueprintGeneratedClass.cpp
--- a/Engine/Source/Runtime/Engine/Private/BlueprintGeneratedClass.cpp
+++ b/Engine/Source/Runtime/Engine/Private/BlueprintGeneratedClass.cpp
@@ -130,7 +130,7 @@ void UBlueprintGeneratedClass::PostInitProperties()
 	if (HasAnyFlags(RF_ClassDefaultObject))
 	{
 		// Default__BlueprintGeneratedClass uses its own AddReferencedObjects function.
-		ClassAddReferencedObjects = &UBlueprintGeneratedClass::AddReferencedObjects;
+		CppClassStaticFunctions = UOBJECT_CPPCLASS_STATICFUNCTIONS_FORCLASS(UBlueprintGeneratedClass);
 	}
 }
 
@@ -1742,7 +1742,7 @@ void UBlueprintGeneratedClass::Bind()
 
 	if (UsePersistentUberGraphFrame() && UberGraphFunction)
 	{
-		ClassAddReferencedObjects = &UBlueprintGeneratedClass::AddReferencedObjectsInUbergraphFrame;
+		CppClassStaticFunctions.SetAddReferencedObjects(&UBlueprintGeneratedClass::AddReferencedObjectsInUbergraphFrame);
 	}
 }
 
diff --git a/Engine/Source/Runtime/Engine/Private/EdGraph/EdGraphNode.cpp b/Engine/Source/Runtime/Engine/Private/EdGraph/EdGraphNode.cpp
--- a/Engine/Source/Runtime/Engine/Private/EdGraph/EdGraphNode.cpp
+++ b/Engine/Source/Runtime/Engine/Private/EdGraph/EdGraphNode.cpp
@@ -242,11 +242,13 @@ void UEdGraphNode::Serialize(FArchive& Ar)
 #endif
 }
 
-void UEdGraphNode::DeclareCustomVersions(FArchive& Ar)
+#if WITH_EDITORONLY_DATA
+void UEdGraphNode::DeclareCustomVersions(FArchive& Ar, const UClass* SpecificSubclass)
 {
-	Super::DeclareCustomVersions(Ar);
-	UEdGraphPin::DeclareCustomVersions(Ar);
+	Super::DeclareCustomVersions(Ar, SpecificSubclass);
+	UEdGraphPin::DeclarePinCustomVersions(Ar);
 }
+#endif
 
 
 bool UEdGraphNode::GetCanRenameNode() const
diff --git a/Engine/Source/Runtime/Engine/Private/EdGraph/EdGraphPin.cpp b/Engine/Source/Runtime/Engine/Private/EdGraph/EdGraphPin.cpp
--- a/Engine/Source/Runtime/Engine/Private/EdGraph/EdGraphPin.cpp
+++ b/Engine/Source/Runtime/Engine/Private/EdGraph/EdGraphPin.cpp
@@ -337,12 +337,13 @@ bool FEdGraphPinType::Serialize(FArchive& Ar)
 	return true;
 }
 
+#if WITH_EDITORONLY_DATA
 void FEdGraphPinType::DeclareCustomVersions(FArchive& Ar)
 {
 	Ar.UsingCustomVersion(FFrameworkObjectVersion::GUID);
 	Ar.UsingCustomVersion(FReleaseObjectVersion::GUID);
 }
-
+#endif
 
 FEdGraphPinType FEdGraphPinType::GetPinTypeForTerminalType( const FEdGraphTerminalType& TerminalType )
 {
@@ -1782,13 +1783,15 @@ bool UEdGraphPin::Serialize(FArchive& Ar)
 	return true;
 }
 
-void UEdGraphPin::DeclareCustomVersions(FArchive& Ar)
+#if WITH_EDITORONLY_DATA
+void UEdGraphPin::DeclarePinCustomVersions(FArchive& Ar)
 {
 	Ar.UsingCustomVersion(FFrameworkObjectVersion::GUID);
 	Ar.UsingCustomVersion(FUE5MainStreamObjectVersion::GUID);
 	Ar.UsingCustomVersion(FUE5ReleaseStreamObjectVersion::GUID);
 	FEdGraphPinType::DeclareCustomVersions(Ar);
 }
+#endif
 
 void UEdGraphPin::ConvertConnectedGhostNodesToRealNodes(UEdGraphNode* InNode)
 {
diff --git a/Engine/Source/Runtime/Engine/Private/Model.cpp b/Engine/Source/Runtime/Engine/Private/Model.cpp
--- a/Engine/Source/Runtime/Engine/Private/Model.cpp
+++ b/Engine/Source/Runtime/Engine/Private/Model.cpp
@@ -542,7 +542,7 @@ void UModel::GetResourceSizeEx(FResourceSizeEx& CumulativeResourceSize)
 
 IMPLEMENT_INTRINSIC_CLASS(UModel, ENGINE_API, UObject, CORE_API, "/Script/Engine",
 	{
-		Class->ClassAddReferencedObjects = &UModel::AddReferencedObjects;
+		Class->CppClassStaticFunctions = UOBJECT_CPPCLASS_STATICFUNCTIONS_FORCLASS(UModel);
 		Class->EmitObjectReference(STRUCT_OFFSET(UModel, Polys), TEXT("Polys"));
 		const uint32 SkipIndexIndex = Class->EmitStructArrayBegin(STRUCT_OFFSET(UModel, Surfs), TEXT("Surfs"), sizeof(FBspSurf));
 		Class->EmitObjectReference(STRUCT_OFFSET(FBspSurf, Material), TEXT("Material"));
@@ -555,7 +555,7 @@ IMPLEMENT_INTRINSIC_CLASS(UModel, ENGINE_API, UObject, CORE_API, "/Script/Engine
 
 IMPLEMENT_INTRINSIC_CLASS(UModel, ENGINE_API, UObject, CORE_API, "/Script/Engine",
 	{
-		Class->ClassAddReferencedObjects = &UModel::AddReferencedObjects;
+		Class->CppClassStaticFunctions = UOBJECT_CPPCLASS_STATICFUNCTIONS_FORCLASS(UModel);
 		const uint32 SkipIndexIndex = Class->EmitStructArrayBegin(STRUCT_OFFSET(UModel, Surfs), TEXT("Surfs"), sizeof(FBspSurf));
 		Class->EmitObjectReference(STRUCT_OFFSET(FBspSurf, Material), TEXT("Material"));
 		Class->EmitObjectReference(STRUCT_OFFSET(FBspSurf, Actor), TEXT("Actor"));
diff --git a/Engine/Source/Runtime/Engine/Private/MorphMesh.cpp b/Engine/Source/Runtime/Engine/Private/MorphMesh.cpp
--- a/Engine/Source/Runtime/Engine/Private/MorphMesh.cpp
+++ b/Engine/Source/Runtime/Engine/Private/MorphMesh.cpp
@@ -36,12 +36,14 @@ void UMorphTarget::Serialize( FArchive& Ar )
 	}
 }
 
-void UMorphTarget::DeclareCustomVersions(FArchive& Ar)
+#if WITH_EDITORONLY_DATA
+void UMorphTarget::DeclareCustomVersions(FArchive& Ar, const UClass* SpecificSubclass)
 {
-	Super::DeclareCustomVersions(Ar);
+	Super::DeclareCustomVersions(Ar, SpecificSubclass);
 	FMorphTargetLODModel MorphLODModel;
 	Ar << MorphLODModel;
 }
+#endif
 
 namespace
 {
diff --git a/Engine/Source/Runtime/Engine/Private/SkeletalMesh.cpp b/Engine/Source/Runtime/Engine/Private/SkeletalMesh.cpp
--- a/Engine/Source/Runtime/Engine/Private/SkeletalMesh.cpp
+++ b/Engine/Source/Runtime/Engine/Private/SkeletalMesh.cpp
@@ -1962,14 +1962,14 @@ void USkeletalMesh::Serialize( FArchive& Ar )
 	}
 }
 
-void USkeletalMesh::DeclareCustomVersions(FArchive& Ar)
+#if WITH_EDITORONLY_DATA
+void USkeletalMesh::DeclareCustomVersions(FArchive& Ar, const UClass* SpecificSubclass)
 {
-	Super::DeclareCustomVersions(Ar);
+	Super::DeclareCustomVersions(Ar, SpecificSubclass);
 	FSkeletalMaterial::DeclareCustomVersions(Ar);
-#if WITH_EDITORONLY_DATA
 	FSkeletalMeshLODModel::DeclareCustomVersions(Ar);
-#endif
 }
+#endif
 
 void USkeletalMesh::GetPreloadDependencies(TArray<UObject*>& OutDeps)
 {
@@ -4460,11 +4460,13 @@ FArchive& operator<<(FArchive& Ar, FMeshUVChannelInfo& ChannelData)
 	return Ar;
 }
 
+#if WITH_EDITORONLY_DATA
 void FSkeletalMaterial::DeclareCustomVersions(FArchive& Ar)
 {
 	Ar.UsingCustomVersion(FEditorObjectVersion::GUID);
 	Ar.UsingCustomVersion(FCoreObjectVersion::GUID);
 }
+#endif
 
 FArchive& operator<<(FArchive& Ar, FSkeletalMaterial& Elem)
 {
