It is now possible to disable virtualization for all payloads owned by a specific asset type by adding the name of the asset to [Core.ContentVirtualization]DisabledAsset string array in the engine ini file.

#rnx

### Virtualization
- [Core.ContentVirtualization] in the engine ini file now supports an array called 'DisabledAsset' which can be used to name asset types that should not virtualize their payloads.
-- By default (in BaseEngine.ini) we have disabled the StaticMesh asset as we know it will crash if a payload is missing and the SoundWave asset as it still is pending testing.
- This new way to disable virtualization is data driven. The older hard coded method has not been removed but will likely be reworked in a future submit.
- Now when an editor bulkdata is adding it's payload to the package trailer builder during package save it will poll the virtualization system with a call to the new method ::IsDisabledForObject by passing in it's owner.
-- If the owner is valid and was present in the 'DisabledAsset' array then the method will return true and the EPayloadFlags::DisableVirtualization flag will be applied.

### Package Trailer
- The pre-existing functionality of enum EPayloadFilter has been moved to a new enum EPayloadStorageType as will only filter payloads based on their storage type.
- EPayloadFilter has been modified to filter payloads based on functionality although at the moment the only thing it can filter for is to return payloads that can be virtualized, it is left for future expansion.
- EPayloadFlags has been reduced to a uint16 with the remaining 2bytes being turned into a new member EPayloadFilterReason.
- This new member allows us to record the exact reason why a payload is excluded from virtualization. If it is zero then the payload can virtualize, otherwise it will contain one or more reasons as to why it is being excluded. For this reason the enum is a bitfield.
- Added overloads of ::GetPayloads and ::GetNumPayloads that take EPayloadFilter rather than a EPayloadStorageType
- Added wrappers around all AccessMode types for FLookupTableEntry.
- FPackageTrailerBuilder has been extended to take a EPayloadFilterReason so that the caller can already provide a reason why the payload cannot be virtualized.
-- As a future peace of work this will probably be changed and we will ask the caller to pass in the owner UObject pointer instead and then we will process the filtering when building the package trailer to a) keep all of the filtering code in one place b) keep the filtering consistent

### PackageSubmissionChecks
- The virtualization process in  will now request the payloads that can be virtualized from the package trailer, which respects the new payload flag, rather than requesting all locally stored payloads.

### UObjects
- There is no need for the SoundWave or MeshDescription classes to opt out of virtualization on construction. This will be done when the package is saved and is now data driven rather than being hardcoded.

### DumpPackagePayloadInfo
- The command has been updated to also display the filter reasons applied to each payload

diff --git a/Engine/Config/BaseEngine.ini b/Engine/Config/BaseEngine.ini
--- a/Engine/Config/BaseEngine.ini
+++ b/Engine/Config/BaseEngine.ini
@@ -1203,6 +1203,9 @@ BackendGraph=ContentVirtualizationBackendGraph_None
 FilterMode=OptOut
 FilterEngineContent=True
 FilterEnginePluginContent=True
+; The following assets should not virtualize until it is confirmed that they do not crash when virtualized payloads are missing
++DisabledAsset=StaticMesh
++DisabledAsset=SoundWave
 
 [/Script/Virtualization.VirtualizationFilterSettings]
 !ExcludePackagePaths=ClearArray
diff --git a/Engine/Source/Developer/Virtualization/Private/DumpPackagePayloadInfoCommand.cpp b/Engine/Source/Developer/Virtualization/Private/DumpPackagePayloadInfoCommand.cpp
--- a/Engine/Source/Developer/Virtualization/Private/DumpPackagePayloadInfoCommand.cpp
+++ b/Engine/Source/Developer/Virtualization/Private/DumpPackagePayloadInfoCommand.cpp
@@ -31,38 +31,38 @@ void DumpPackagePayloadInfo(const TArray<FString>& Args)
 
 		if (FPackagePath::TryFromMountedName(Arg, Path))
 		{
-			TArray<FIoHash> LocalPayloadIds;
-			TArray<FIoHash> VirtualizedPayloadIds;
-
-			if (!UE::FindPayloadsInPackageFile(Path, UE::EPayloadFilter::Local, LocalPayloadIds))
-			{
-				UE_LOG(LogVirtualization, Error, TEXT("Failed to find local payload information from package: '%s'"), *Path.GetDebugName());
-				continue;
-			}
-
-			if (!UE::FindPayloadsInPackageFile(Path, UE::EPayloadFilter::Virtualized, VirtualizedPayloadIds))
+			FPackageTrailer Trailer;
+			if (!FPackageTrailer::TryLoadFromPackage(Path, Trailer))
 			{
-				UE_LOG(LogVirtualization, Error, TEXT("Failed to find virtualized payload information from package: '%s'"), *Path.GetDebugName());
+				UE_LOG(LogVirtualization, Error, TEXT("Failed to find load the package trailer from: '%s'"), *Path.GetDebugName());
 				continue;
 			}
 
+			TArray<FIoHash> LocalPayloadIds = Trailer.GetPayloads(UE::EPayloadStorageType::Local);
+			TArray<FIoHash> VirtualizedPayloadIds = Trailer.GetPayloads(UE::EPayloadStorageType::Virtualized);
+			
+			UE_LOG(LogVirtualization, Display, TEXT("")); // Blank line to make the output easier to read
 			UE_LOG(LogVirtualization, Display, TEXT("Package: '%s' has %d local and %d virtualized payloads"), *Path.GetDebugName(), LocalPayloadIds.Num(), VirtualizedPayloadIds.Num());
 			
 			if (LocalPayloadIds.Num() > 0)
 			{
 				UE_LOG(LogVirtualization, Display, TEXT("LocalPayloads:"));
+				UE_LOG(LogVirtualization, Display, TEXT("Index|\t%-40s|\tFilterReason"), TEXT("PayloadIdentifier"));
 				for (int32 Index = 0; Index < LocalPayloadIds.Num(); ++Index)
 				{
-					UE_LOG(LogVirtualization, Display, TEXT("%02d: '%s'"), Index, *LexToString(LocalPayloadIds[Index]));
+					FPayloadInfo Info = Trailer.GetPayloadInfo(LocalPayloadIds[Index]);
+					UE_LOG(LogVirtualization, Display, TEXT("%02d:  |\t%s|\t%s"), Index, *LexToString(LocalPayloadIds[Index]), *LexToString(Info.FilterFlags));
 				}
 			}
 
 			if (VirtualizedPayloadIds.Num() > 0)
 			{
 				UE_LOG(LogVirtualization, Display, TEXT("VirtualizedPayloads:"));
+				UE_LOG(LogVirtualization, Display, TEXT("Index|\t%-40s|\tFilterReason"), TEXT("PayloadIdentifier"));
 				for (int32 Index = 0; Index < VirtualizedPayloadIds.Num(); ++Index)
 				{
-					UE_LOG(LogVirtualization, Display, TEXT("%02d: '%s'"), Index, *LexToString(VirtualizedPayloadIds[Index]));
+					FPayloadInfo Info = Trailer.GetPayloadInfo(VirtualizedPayloadIds[Index]);
+					UE_LOG(LogVirtualization, Display, TEXT("%02d:  |\t%s|\t%s"), Index, *LexToString(VirtualizedPayloadIds[Index]), *LexToString(Info.FilterFlags));
 				}
 			}
 		}
diff --git a/Engine/Source/Developer/Virtualization/Private/PackageSubmissionChecks.cpp b/Engine/Source/Developer/Virtualization/Private/PackageSubmissionChecks.cpp
--- a/Engine/Source/Developer/Virtualization/Private/PackageSubmissionChecks.cpp
+++ b/Engine/Source/Developer/Virtualization/Private/PackageSubmissionChecks.cpp
@@ -320,8 +320,8 @@ void VirtualizePackages(const TArray<FString>& FilesToSubmit, TArray<FText>& Out
 			if (FPackageTrailer::TryLoadFromPackage(PackagePath, Trailer))
 			{
 				// The following is not expected to ever happen, currently we give a user facing error but it generally means that the asset is broken somehow.
-				ensureMsgf(Trailer.GetNumPayloads(EPayloadFilter::Referenced) == 0, TEXT("Trying to virtualize a package that already contains payload references which the workspace file should not ever contain!"));
-				if (Trailer.GetNumPayloads(EPayloadFilter::Referenced) > 0)
+				ensureMsgf(Trailer.GetNumPayloads(EPayloadStorageType::Referenced) == 0, TEXT("Trying to virtualize a package that already contains payload references which the workspace file should not ever contain!"));
+				if (Trailer.GetNumPayloads(EPayloadStorageType::Referenced) > 0)
 				{
 					FText Message = FText::Format(	LOCTEXT("Virtualization_PkgHasReferences", "Cannot virtualize the package '{1}' as it has referenced payloads in the trailer"),
 													FText::FromString(PackagePath.GetDebugName()));
@@ -333,7 +333,7 @@ void VirtualizePackages(const TArray<FString>& FilesToSubmit, TArray<FText>& Out
 
 				PkgInfo.Path = MoveTemp(PackagePath);
 				PkgInfo.Trailer = MoveTemp(Trailer);
-				PkgInfo.LocalPayloads = PkgInfo.Trailer.GetPayloads(EPayloadFilter::Local);
+				PkgInfo.LocalPayloads = PkgInfo.Trailer.GetPayloads(EPayloadFilter::CanVirtualize);
 
 				if (!PkgInfo.LocalPayloads.IsEmpty())
 				{	
@@ -390,7 +390,7 @@ void VirtualizePackages(const TArray<FString>& FilesToSubmit, TArray<FText>& Out
 		// If we made changes we should recalculate the local payloads left
 		if (PackageInfo.bWasTrailerUpdated)
 		{
-			PackageInfo.LocalPayloads = PackageInfo.Trailer.GetPayloads(EPayloadFilter::Local);
+			PackageInfo.LocalPayloads = PackageInfo.Trailer.GetPayloads(EPayloadStorageType::Local);
 		}
 #endif
 
diff --git a/Engine/Source/Developer/Virtualization/Private/VirtualizationManager.cpp b/Engine/Source/Developer/Virtualization/Private/VirtualizationManager.cpp
--- a/Engine/Source/Developer/Virtualization/Private/VirtualizationManager.cpp
+++ b/Engine/Source/Developer/Virtualization/Private/VirtualizationManager.cpp
@@ -336,6 +336,24 @@ bool FVirtualizationManager::IsPushingEnabled(EStorageType StorageType) const
 	}
 }
 
+bool FVirtualizationManager::IsDisabledForObject(const UObject* OwnerObject) const
+{
+	if (OwnerObject == nullptr)
+	{
+		return false;
+	}
+
+	const UClass* OwnerClass = OwnerObject->GetClass();
+	if (OwnerClass == nullptr)
+	{
+		// TODO: Not actually sure if the class being nullptr is reasonable or if we should warn/error here?
+		return false;
+	}
+
+	FName ClassName = OwnerClass->GetFName();
+	return DisabledAssetTypes.Find(ClassName) != nullptr;
+}
+
 bool FVirtualizationManager::PushData(const FIoHash& Id, const FCompressedBuffer& Payload, EStorageType StorageType, const FString& Context)
 {
 	FPushRequest Request(Id, Payload, Context);
@@ -733,7 +751,19 @@ void FVirtualizationManager::ApplySettingsFromConfigFiles(const FConfigFile& Con
 	else
 	{
 		UE_LOG(LogVirtualization, Error, TEXT("Failed to load [Core.ContentVirtualization].FilterEnginePluginContent from config file!"));
-	}	
+	}
+
+	TArray<FString> DisabledAssetTypesFromIni;
+	if (ConfigFile.GetArray(TEXT("Core.ContentVirtualization"), TEXT("DisabledAsset"), DisabledAssetTypesFromIni) > 0)
+	{
+		UE_LOG(LogVirtualization, Display, TEXT("\tVirtualization is disabled for payloads of the following assets:"));
+		DisabledAssetTypes.Reserve(DisabledAssetTypesFromIni.Num());
+		for(const FString& AssetType : DisabledAssetTypesFromIni)
+		{ 
+			UE_LOG(LogVirtualization, Display, TEXT("\t\t%s"), *AssetType);
+			DisabledAssetTypes.Add(FName(AssetType));
+		}	
+	}
 }
 
 void FVirtualizationManager::ApplyDebugSettingsFromFromCmdline()
diff --git a/Engine/Source/Developer/Virtualization/Private/VirtualizationManager.h b/Engine/Source/Developer/Virtualization/Private/VirtualizationManager.h
--- a/Engine/Source/Developer/Virtualization/Private/VirtualizationManager.h
+++ b/Engine/Source/Developer/Virtualization/Private/VirtualizationManager.h
@@ -112,6 +112,8 @@ private:
 
 	virtual bool IsEnabled() const override;
 	virtual bool IsPushingEnabled(EStorageType StorageType) const override;
+
+	virtual bool IsDisabledForObject(const UObject* Owner) const override;
 	
 	virtual bool PushData(const FIoHash& Id, const FCompressedBuffer& Payload, EStorageType StorageType, const FString& Context) override;
 	virtual bool PushData(TArrayView<FPushRequest> Requests, EStorageType StorageType) override;
@@ -212,6 +214,9 @@ private:
 	/** The name of the current project */
 	FString ProjectName;
 
+	/** The names of all asset types that should not virtualize. See @IsDisabledForObject */
+	TSet<FName> DisabledAssetTypes;
+
 	/** All of the backends that were mounted during graph creation */
 	TArray<TUniquePtr<IVirtualizationBackend>> AllBackends;
 
diff --git a/Engine/Source/Editor/UnrealEd/Private/Commandlets/ContentCommandlets.cpp b/Engine/Source/Editor/UnrealEd/Private/Commandlets/ContentCommandlets.cpp
--- a/Engine/Source/Editor/UnrealEd/Private/Commandlets/ContentCommandlets.cpp
+++ b/Engine/Source/Editor/UnrealEd/Private/Commandlets/ContentCommandlets.cpp
@@ -1509,7 +1509,7 @@ void UResavePackagesCommandlet::PerformPreloadOperations( FLinkerLoad* PackageLi
 	if (bOnlyVirtualized)
 	{
 		const UE::FPackageTrailer* Trailer = PackageLinker->GetPackageTrailer();
-		if (Trailer == nullptr || Trailer->GetNumPayloads(UE::EPayloadFilter::Virtualized) == 0)
+		if (Trailer == nullptr || Trailer->GetNumPayloads(UE::EPayloadStorageType::Virtualized) == 0)
 		{
 			bSavePackage = false;
 			return;
diff --git a/Engine/Source/Runtime/AssetRegistry/Private/PackageReader.cpp b/Engine/Source/Runtime/AssetRegistry/Private/PackageReader.cpp
--- a/Engine/Source/Runtime/AssetRegistry/Private/PackageReader.cpp
+++ b/Engine/Source/Runtime/AssetRegistry/Private/PackageReader.cpp
@@ -296,7 +296,7 @@ bool FPackageReader::SerializePackageTrailer(FAssetPackageData& PackageData)
 		return true;
 	}
 
-	PackageData.SetHasVirtualizedPayloads(Trailer.GetNumPayloads(UE::EPayloadFilter::Virtualized) > 0);
+	PackageData.SetHasVirtualizedPayloads(Trailer.GetNumPayloads(UE::EPayloadStorageType::Virtualized) > 0);
 	return true;
 }
 
diff --git a/Engine/Source/Runtime/Core/Private/Virtualization/VirtualizationSystem.cpp b/Engine/Source/Runtime/Core/Private/Virtualization/VirtualizationSystem.cpp
--- a/Engine/Source/Runtime/Core/Private/Virtualization/VirtualizationSystem.cpp
+++ b/Engine/Source/Runtime/Core/Private/Virtualization/VirtualizationSystem.cpp
@@ -35,6 +35,11 @@ public:
 		return false;
 	}
 
+	virtual bool IsDisabledForObject(const UObject* Owner) const override
+	{
+		return false;
+	}
+
 	virtual bool PushData(const FIoHash& Id, const FCompressedBuffer& Payload, EStorageType StorageType, const FString& Context) override
 	{
 		return false;
diff --git a/Engine/Source/Runtime/Core/Public/Virtualization/VirtualizationSystem.h b/Engine/Source/Runtime/Core/Public/Virtualization/VirtualizationSystem.h
--- a/Engine/Source/Runtime/Core/Public/Virtualization/VirtualizationSystem.h
+++ b/Engine/Source/Runtime/Core/Public/Virtualization/VirtualizationSystem.h
@@ -319,6 +319,16 @@ public:
 	/** Poll to see if pushing virtualized content to the given backend storage type is enabled or not. */
 	virtual bool IsPushingEnabled(EStorageType StorageType) const = 0;
 
+	/** 
+	 * Poll to see if virtualization is disabled for the given asset type.
+	 * 
+	 * @param	Owner	The object to be tested, assumed to be an asset that 
+	 *					can own virtualized payloads.
+	 * 
+	 * @return	True if payloads owned by this object should never virtualized
+	 */
+	virtual bool IsDisabledForObject(const UObject* Owner) const = 0;
+
 	/**
 	 * Push a payload to the virtualization backends.
 	 *
diff --git a/Engine/Source/Runtime/CoreUObject/Private/Serialization/EditorBulkData.cpp b/Engine/Source/Runtime/CoreUObject/Private/Serialization/EditorBulkData.cpp
--- a/Engine/Source/Runtime/CoreUObject/Private/Serialization/EditorBulkData.cpp
+++ b/Engine/Source/Runtime/CoreUObject/Private/Serialization/EditorBulkData.cpp
@@ -1652,7 +1652,13 @@ void FEditorBulkData::SerializeToPackageTrailer(FLinkerSave& LinkerSave, FCompre
 		}
 	};
 
-	LinkerSave.PackageTrailerBuilder->AddPayload(PayloadContentId, MoveTemp(PayloadToSerialize), MoveTemp(OnPayloadWritten));
+	EPayloadFilterReason PayloadFilter = EPayloadFilterReason::None;
+	if (bSkipVirtualization || UE::Virtualization::IVirtualizationSystem::Get().IsDisabledForObject(Owner))
+	{
+		PayloadFilter = EPayloadFilterReason::Asset;
+	}
+
+	LinkerSave.PackageTrailerBuilder->AddPayload(PayloadContentId, MoveTemp(PayloadToSerialize), PayloadFilter, MoveTemp(OnPayloadWritten));
 }
 
 void FEditorBulkData::UpdatePayloadImpl(FSharedBuffer&& InPayload, FIoHash&& InPayloadID, UObject* Owner)
@@ -2090,13 +2096,6 @@ FText FEditorBulkData::GetCorruptedPayloadErrorMsgForSave(FLinkerSave* Linker) c
 
 bool FEditorBulkData::ShouldUseLegacySerialization(const FLinkerSave* LinkerSave) const
 {
-#if UE_ENABLE_VIRTUALIZATION_TOGGLE
-	if (bSkipVirtualization == true)
-	{
-		return true;
-	}
-#endif // UE_ENABLE_VIRTUALIZATION_TOGGLE 
-
 	if (LinkerSave == nullptr)
 	{
 		return true;
diff --git a/Engine/Source/Runtime/CoreUObject/Private/UObject/PackageTrailer.cpp b/Engine/Source/Runtime/CoreUObject/Private/UObject/PackageTrailer.cpp
--- a/Engine/Source/Runtime/CoreUObject/Private/UObject/PackageTrailer.cpp
+++ b/Engine/Source/Runtime/CoreUObject/Private/UObject/PackageTrailer.cpp
@@ -103,6 +103,7 @@ void FLookupTableEntry::Serialize(FArchive& Ar, EPackageTrailerVersion PackageTr
 	if (Ar.IsSaving() || PackageTrailerVersion >= EPackageTrailerVersion::PAYLOAD_FLAGS)
 	{
 		Ar << Flags;
+		Ar << FilterFlags;
 	}
 
 	if (Ar.IsSaving() || PackageTrailerVersion >= EPackageTrailerVersion::ACCESS_PER_PAYLOAD)
@@ -126,7 +127,7 @@ FPackageTrailerBuilder FPackageTrailerBuilder::CreateFromTrailer(const FPackageT
 			case EPayloadAccessMode::Local:
 			{
 				FCompressedBuffer Payload = Trailer.LoadLocalPayload(Entry.Identifier, Ar);
-				Builder.LocalEntries.Add(Entry.Identifier, LocalEntry(MoveTemp(Payload)));
+				Builder.LocalEntries.Add(Entry.Identifier, LocalEntry(MoveTemp(Payload), Entry.FilterFlags));
 			}
 			break;
 
@@ -199,13 +200,25 @@ FPackageTrailerBuilder::FPackageTrailerBuilder(const FName& InPackageName)
 {
 }
 
-void FPackageTrailerBuilder::AddPayload(const FIoHash& Identifier, FCompressedBuffer Payload, AdditionalDataCallback&& Callback)
+void FPackageTrailerBuilder::AddPayload(const FIoHash& Identifier, FCompressedBuffer Payload, EPayloadFilterReason FilterFlags, AdditionalDataCallback&& Callback)
 {
 	Callbacks.Emplace(MoveTemp(Callback));
 
 	if (!Identifier.IsZero())
 	{
-		LocalEntries.FindOrAdd(Identifier, LocalEntry(MoveTemp(Payload)));
+		// If the payload already exists and the DisableVirtualization flag has been passed in
+		// then we need to make sure that it is applied.
+		// TODO: This will have to be done here for every new flag added as we don't know what
+		// future flags will want to do in the case of a duplicate entry but we probably need
+		// a nicer way to handle this longer term.
+		if (LocalEntry* Entry = LocalEntries.Find(Identifier))
+		{
+			Entry->FilterFlags |= FilterFlags;
+		}
+		else
+		{
+			LocalEntries.Add(Identifier, LocalEntry(MoveTemp(Payload), FilterFlags));
+		}
 	}
 }
 
@@ -248,6 +261,8 @@ bool FPackageTrailerBuilder::BuildAndAppendTrailer(FLinkerSave* Linker, FArchive
 		Entry.CompressedSize = It.Value.Payload.GetCompressedSize();
 		Entry.RawSize = It.Value.Payload.GetRawSize();
 		Entry.AccessMode = EPayloadAccessMode::Local;
+		Entry.Flags = EPayloadFlags::None;
+		Entry.FilterFlags = It.Value.FilterFlags;
 
 		Trailer.Header.PayloadsDataLength += It.Value.Payload.GetCompressedSize();
 	}
@@ -628,60 +643,83 @@ int64 FPackageTrailer::GetTrailerLength() const
 	return Header.HeaderLength + Header.PayloadsDataLength + FFooter::SizeOnDisk;
 }
 
-TArray<FIoHash> FPackageTrailer::GetPayloads(EPayloadFilter Type) const
+FPayloadInfo FPackageTrailer::GetPayloadInfo(const FIoHash& Id) const
+{
+	const Private::FLookupTableEntry* Entry = FindEntryInHeader(Header, Id);
+
+	if (Entry != nullptr)
+	{
+		FPayloadInfo Info;
+
+		Info.OffsetInFile = Entry->OffsetInFile;
+		Info.CompressedSize = Entry->CompressedSize;
+		Info.RawSize = Entry->RawSize;
+		Info.AccessMode = Entry->AccessMode;
+		Info.Flags = Entry->Flags;
+		Info.FilterFlags = Entry->FilterFlags;
+
+		return Info;
+	}
+	else
+	{
+		return FPayloadInfo();
+	}
+}
+
+TArray<FIoHash> FPackageTrailer::GetPayloads(EPayloadStorageType StorageType) const
 {
 	TArray<FIoHash> Identifiers;
 	Identifiers.Reserve(Header.PayloadLookupTable.Num());
 
 	for (const Private::FLookupTableEntry& Entry : Header.PayloadLookupTable)
 	{
-		switch (Type)
+		switch (StorageType)
 		{
-		case EPayloadFilter::All:
-			Identifiers.Add(Entry.Identifier);
-			break;
-
-		case EPayloadFilter::Local:
-			if (!Entry.IsVirtualized())
-			{
+			case EPayloadStorageType::Any:
 				Identifiers.Add(Entry.Identifier);
-			}
-			break;
+				break;
 
-		case EPayloadFilter::Virtualized:
-			if (Entry.IsVirtualized())
-			{
-				Identifiers.Add(Entry.Identifier);
-			}
-			break;
+			case EPayloadStorageType::Local:
+				if (Entry.IsLocal())
+				{
+					Identifiers.Add(Entry.Identifier);
+				}
+				break;
 
-		default:
-			checkNoEntry();
+			case EPayloadStorageType::Virtualized:
+				if (Entry.IsVirtualized())
+				{
+					Identifiers.Add(Entry.Identifier);
+				}
+				break;
+
+			default:
+				checkNoEntry();
 		}	
 	}
 
 	return Identifiers;
 }
 
-int32 FPackageTrailer::GetNumPayloads(EPayloadFilter Type) const
+int32 FPackageTrailer::GetNumPayloads(EPayloadStorageType Type) const
 {
 	int32 Count = 0;
 
 	switch (Type)
 	{
-		case EPayloadFilter::All:
+		case EPayloadStorageType::Any:
 			Count = Header.PayloadLookupTable.Num();
 			break;
 
-		case EPayloadFilter::Local:
+		case EPayloadStorageType::Local:
 			Count = (int32)Algo::CountIf(Header.PayloadLookupTable, [](const Private::FLookupTableEntry& Entry) { return Entry.AccessMode == EPayloadAccessMode::Local; });
 			break;
 
-		case EPayloadFilter::Referenced:
+		case EPayloadStorageType::Referenced:
 			Count = (int32)Algo::CountIf(Header.PayloadLookupTable, [](const Private::FLookupTableEntry& Entry) { return Entry.AccessMode == EPayloadAccessMode::Referenced; });
 			break;
 
-		case EPayloadFilter::Virtualized:
+		case EPayloadStorageType::Virtualized:
 			Count = (int32)Algo::CountIf(Header.PayloadLookupTable, [](const Private::FLookupTableEntry& Entry) { return Entry.AccessMode == EPayloadAccessMode::Virtualized; });
 			break;
 
@@ -692,6 +730,49 @@ int32 FPackageTrailer::GetNumPayloads(EPayloadFilter Type) const
 	return Count;
 }
 
+TArray<FIoHash> FPackageTrailer::GetPayloads(EPayloadFilter Filter) const
+{
+	TArray<FIoHash> Identifiers;
+	Identifiers.Reserve(Header.PayloadLookupTable.Num());
+
+	for (const Private::FLookupTableEntry& Entry : Header.PayloadLookupTable)
+	{
+		switch (Filter)
+		{
+			case EPayloadFilter::CanVirtualize:
+				if (Entry.IsLocal() && Entry.FilterFlags == EPayloadFilterReason::None)
+				{
+					Identifiers.Add(Entry.Identifier);
+				}
+				break;
+
+			default:
+				checkNoEntry();
+		}
+	}
+
+	return Identifiers;
+}
+
+int32 FPackageTrailer::GetNumPayloads(EPayloadFilter Filter) const
+{
+	int32 Count = 0;
+
+	switch (Filter)
+	{
+		case EPayloadFilter::CanVirtualize:
+			Count = (int32)Algo::CountIf(Header.PayloadLookupTable, [](const Private::FLookupTableEntry& Entry)
+				{
+					return Entry.AccessMode == EPayloadAccessMode::Local && Entry.FilterFlags == EPayloadFilterReason::None;
+				});
+			break;
+
+		default:
+			checkNoEntry();
+	}
+
+	return Count;
+}
 FPackageTrailer::FFooter FPackageTrailer::CreateFooter() const
 {
 	FFooter Footer;
@@ -733,7 +814,7 @@ FArchive& operator<<(FArchive& Ar, FPackageTrailer::FFooter& Footer)
 	return Ar;
 }
 
-bool FindPayloadsInPackageFile(const FPackagePath& PackagePath, EPayloadFilter Filter, TArray<FIoHash>& OutPayloadIds)
+bool FindPayloadsInPackageFile(const FPackagePath& PackagePath, EPayloadStorageType Filter, TArray<FIoHash>& OutPayloadIds)
 {
 	if (FPackageName::IsTextPackageExtension(PackagePath.GetHeaderExtension()))
 	{
@@ -769,3 +850,64 @@ bool FindPayloadsInPackageFile(const FPackagePath& PackagePath, EPayloadFilter F
 
 } //namespace UE
 
+FString LexToString(UE::EPayloadFilterReason FilterFlags)
+{
+	using namespace UE;
+
+	if (FilterFlags == EPayloadFilterReason::None)
+	{
+		return TEXT("None");
+	}
+	else
+	{
+		TStringBuilder<512> Builder;
+		auto AddSeparatorIfNeeded = [&Builder]()
+		{
+			if (Builder.Len() != 0)
+			{
+				Builder << TEXT("|");
+			}
+		};
+
+		if (EnumHasAllFlags(FilterFlags, EPayloadFilterReason::Asset))
+		{
+			Builder << TEXT("Asset");
+		}
+
+		if (EnumHasAllFlags(FilterFlags, EPayloadFilterReason::Path))
+		{
+			AddSeparatorIfNeeded();
+			Builder << TEXT("Path");
+		}
+
+		if (EnumHasAllFlags(FilterFlags, EPayloadFilterReason::MinSize))
+		{
+			AddSeparatorIfNeeded();
+			Builder << TEXT("MinSize");
+		}
+
+		if (EnumHasAllFlags(FilterFlags, EPayloadFilterReason::EditorBulkDataCode))
+		{
+			AddSeparatorIfNeeded();
+			Builder << TEXT("EditorBulkDataCode");
+		}
+
+		// In case a new entry was added to EPayloadFilterReason without this function being
+		// updated we need to check if any of the other bits are set and if they are print 
+		// and unknown entry.
+		// We don't want a Max or Count entry being added to EPayloadFilterReason as we'd then
+		// need to validate all places taking the enum to make sure that nobody is using it.
+		const uint16 Mask = ~uint16(EPayloadFilterReason::Asset | EPayloadFilterReason::Path |
+			EPayloadFilterReason::MinSize | EPayloadFilterReason::EditorBulkDataCode);
+
+		const bool bHasUnknownBits = ((uint16)FilterFlags & Mask) != 0;
+		if (bHasUnknownBits)
+		{
+			AddSeparatorIfNeeded();
+			Builder << TEXT("UnknownBits");
+		}
+
+		return Builder.ToString();
+	}
+}
+
diff --git a/Engine/Source/Runtime/CoreUObject/Public/UObject/PackageTrailer.h b/Engine/Source/Runtime/CoreUObject/Public/UObject/PackageTrailer.h
--- a/Engine/Source/Runtime/CoreUObject/Public/UObject/PackageTrailer.h
+++ b/Engine/Source/Runtime/CoreUObject/Public/UObject/PackageTrailer.h
@@ -62,11 +62,11 @@ namespace UE
  * |____________________________________________________________________________________________________________________________________________|
  */
 
- /** Used to filter requests to a specific type of payload */
-enum class EPayloadFilter
+ /** Used to filter requests based on how a payload is stored*/
+enum class EPayloadStorageType : uint8
 {
-	/** All payload types. */
-	All,
+	/** All payload regardless of type. */
+	Any,
 	/** All payloads stored locally in the package trailer. */
 	Local,
 	/** All payloads that are a reference to payloads stored in the workspace domain trailer*/
@@ -75,6 +75,15 @@ enum class EPayloadFilter
 	Virtualized
 };
 
+/** Used to filter requests based on how a payload can function */
+enum class EPayloadFilter
+{
+	/** All payloads that are stored locally and do not have virtualization disabled */
+	CanVirtualize
+};
+
+// TODO: Could consider merging EPayloadStorageType and EPayloadStatus, only difference is Any vs NotFound
+
 /** Used to show the status of a payload */
 enum class EPayloadStatus
 {
@@ -100,12 +109,30 @@ enum class EPayloadAccessMode : uint8
 };
 
 /** Flags that can be set on payloads in a payload trailer */
-enum class EPayloadFlags : uint32
+enum class EPayloadFlags : uint16
 {
 	/** No flags are set */
-	None
+	None = 0,
+};
+
+ENUM_CLASS_FLAGS(EPayloadFlags);
+
+enum class EPayloadFilterReason : uint16
+{
+	/** Not filtered */
+	None				= 0,
+	/** Filtered due to the asset type of the owning UObject */
+	Asset				= 1 << 0,
+	/** Filtered due to the path of the owning UPackage */
+	Path				= 1 << 1,
+	/** Filtered because the payload size is below the minimum size for virtualization */
+	MinSize				= 1 << 2,
+	/** Filtered because the owning editor bulkdata had virtualization disabled programmatically */
+	EditorBulkDataCode	= 1 << 3
 };
 
+ENUM_CLASS_FLAGS(EPayloadFilterReason);
+
 enum class EPackageTrailerVersion : uint32;
 
 namespace Private
@@ -126,6 +153,16 @@ struct FLookupTableEntry
 
 	void Serialize(FArchive& Ar, EPackageTrailerVersion PackageTrailerVersion);
 	
+	[[nodiscard]] bool IsLocal() const
+	{
+		return AccessMode == EPayloadAccessMode::Local;
+	}
+
+	[[nodiscard]] bool IsReferenced() const
+	{
+		return AccessMode == EPayloadAccessMode::Referenced;
+	}
+
 	[[nodiscard]] bool IsVirtualized() const
 	{
 		return AccessMode == EPayloadAccessMode::Virtualized;
@@ -141,12 +178,27 @@ struct FLookupTableEntry
 	uint64 RawSize = INDEX_NONE;
 	/** Bitfield of flags, see @UE::EPayloadFlags */
 	EPayloadFlags Flags = EPayloadFlags::None;
+	/** Bitfield of flags showing if the payload allowed to be virtualized or the reason why it cannot be virtualized, see @UE::EPayloadFilterReason */
+	EPayloadFilterReason FilterFlags = EPayloadFilterReason::None;
 
 	EPayloadAccessMode AccessMode = EPayloadAccessMode::Local;
 };
 
 } // namespace Private
 
+//** Info about a payload stored in the trailer
+struct FPayloadInfo
+{
+	int64 OffsetInFile = INDEX_NONE;
+
+	uint64 CompressedSize = INDEX_NONE;
+	uint64 RawSize = INDEX_NONE;
+
+	EPayloadAccessMode AccessMode = EPayloadAccessMode::Local;
+	EPayloadFlags Flags = EPayloadFlags::None;
+	EPayloadFilterReason FilterFlags = EPayloadFilterReason::None;
+};
+
 /** 
  * This class is used to build a FPackageTrailer and write it disk.
  * 
@@ -191,9 +243,10 @@ public:
 	 * 
 	 * @param Identifier	The identifier of the payload
 	 * @param Payload		The payload data
+	 * @param Flags			The custom flags to be applied to the payload
 	 * @param Callback		This callback will be invoked once the FPackageTrailer has been built and appended to disk.
 	 */
-	void AddPayload(const FIoHash& Identifier, FCompressedBuffer Payload, AdditionalDataCallback&& Callback);
+	void AddPayload(const FIoHash& Identifier, FCompressedBuffer Payload, EPayloadFilterReason Filter, AdditionalDataCallback&& Callback);
 
 	/**
 	 * Adds an already virtualized payload to the builder to be written to the trailer. When the trailer is written
@@ -236,14 +289,16 @@ private:
 	struct LocalEntry
 	{
 		LocalEntry() = default;
-		LocalEntry(FCompressedBuffer&& InPayload)
+		LocalEntry(FCompressedBuffer&& InPayload, EPayloadFilterReason InFilterFlags)
 			: Payload(InPayload)
+			, FilterFlags(InFilterFlags)
 		{
 
 		}
 		~LocalEntry() = default;
 
 		FCompressedBuffer Payload;
+		EPayloadFilterReason FilterFlags = EPayloadFilterReason::None;
 	};
 
 	/** All of the data required to add a reference to a payload stored in another trailer */
@@ -380,11 +435,19 @@ public:
 	/** Returns the total size of the of the trailer on disk in bytes */
 	[[nodiscard]] int64 GetTrailerLength() const;
 
-	/** Returns an array of the payloads that match the given filter type. @See EPayloadType */
-	[[nodiscard]] TArray<FIoHash> GetPayloads(EPayloadFilter Type) const;
+	[[nodiscard]] FPayloadInfo GetPayloadInfo(const FIoHash& Id) const;
+
+	/** Returns an array of the payloads with the given storage type. @See EPayloadStoragetype */
+	[[nodiscard]] TArray<FIoHash> GetPayloads(EPayloadStorageType StorageType) const;
 
-	/** Returns the number of payloads that the trailer owns that match the given filter type. @See EPayloadType */
-	[[nodiscard]] int32 GetNumPayloads(EPayloadFilter Type) const;
+	/** Returns the number of payloads that the trailer owns with the given storage type. @See EPayloadStoragetype */
+	[[nodiscard]] int32 GetNumPayloads(EPayloadStorageType Type) const;
+
+	/** Returns an array of the payloads that match the given filter type. @See EPayloadFilter */
+	[[nodiscard]] TArray<FIoHash> GetPayloads(EPayloadFilter Filter) const;
+
+	/** Returns the number of payloads that the trailer owns that match the given filter type. @See EPayloadFilter */
+	[[nodiscard]] int32 GetNumPayloads(EPayloadFilter Filter) const;
 
 	struct FHeader
 	{
@@ -463,6 +526,8 @@ private:
  *
  * @return 				True if the package was parsed successfully (although it might not have contained any payloads) and false if opening or parsing the package file failed.
  */
-[[nodiscard]] COREUOBJECT_API bool FindPayloadsInPackageFile(const FPackagePath& PackagePath, EPayloadFilter Filter, TArray<FIoHash>& OutPayloadIds);
+[[nodiscard]] COREUOBJECT_API bool FindPayloadsInPackageFile(const FPackagePath& PackagePath, EPayloadStorageType Filter, TArray<FIoHash>& OutPayloadIds);
+
+} //namespace UE
 
-} //namespace UE
\ No newline at end of file
+[[nodiscard]] COREUOBJECT_API FString LexToString(UE::EPayloadFilterReason FilterFlags);
diff --git a/Engine/Source/Runtime/Engine/Private/SoundWave.cpp b/Engine/Source/Runtime/Engine/Private/SoundWave.cpp
--- a/Engine/Source/Runtime/Engine/Private/SoundWave.cpp
+++ b/Engine/Source/Runtime/Engine/Private/SoundWave.cpp
@@ -663,15 +663,6 @@ USoundWave::USoundWave(const FObjectInitializer& ObjectInitializer)
 	EnvelopeFollowerFrameSize = 1024;
 	EnvelopeFollowerAttackTime = 10;
 	EnvelopeFollowerReleaseTime = 100;
-
-#if UE_ENABLE_VIRTUALIZATION_TOGGLE
-	// Calling ::SetVirtualizationOptOut allows us to prevent the bulkdata from virtualizing if a project
-	// has that option enabled. This is not intended to be a shipping feature.
-	PRAGMA_DISABLE_DEPRECATION_WARNINGS
-		RawData.SetVirtualizationOptOut(true);
-	PRAGMA_ENABLE_DEPRECATION_WARNINGS
-#endif //UE_ENABLE_VIRTUALIZATION_TOGGLE
-
 #endif
 
 	bCachedSampleRateFromPlatformSettings = false;
diff --git a/Engine/Source/Runtime/MeshDescription/Public/MeshDescription.h b/Engine/Source/Runtime/MeshDescription/Public/MeshDescription.h
--- a/Engine/Source/Runtime/MeshDescription/Public/MeshDescription.h
+++ b/Engine/Source/Runtime/MeshDescription/Public/MeshDescription.h
@@ -1297,14 +1297,6 @@ public:
 		: bBulkDataUpdated(false)
 		, bGuidIsHash(false)
 	{
-#if UE_ENABLE_VIRTUALIZATION_TOGGLE
-		// Currently FMeshDescriptionBulkData will cause a hard crash if the virtualized data cannot be found.
-		// Calling ::SetVirtualizationOptOut allows us to prevent the bulkdata from virtualizing if a project
-		// has that option enabled. This is not intended to be a shipping feature.
-		PRAGMA_DISABLE_DEPRECATION_WARNINGS	
-		BulkData.SetVirtualizationOptOut(true);
-		PRAGMA_ENABLE_DEPRECATION_WARNINGS
-#endif //UE_ENABLE_VIRTUALIZATION_TOGGLE
 	}
 
 	/** Serialization */
