Add an error when saving a package if the package header gets too large which would cause one or more of the offsets stored in the FPackageFileSummary to truncate.

#rnx

- The offsets and TotalHeaderSize value in FPackageFileSummary are all stored as int32 and although it is very unlikely that the package header would ever be large enough that we would need to store these values as 64 bit, we do support 64bit length files and if the header was large enough we could end up with truncated values in the summary which would cause the package to probably crash the editor when loaded.
- Since all of these values are assigned during ::WritePackageHeader we don't need to add a check every time one is assigned, we can have a single check at the end of writing out the header. Technically it might be faster to fail earlier when writing out the header, but it is unlikely that this condition will ever trigger and adding bespoke checks on every assignment would either bloat the code and make it harder to read or require a macro.
- I opted for a log error and not a user facing error because realistically if this occurs there is nothing an end user can do to fix it other than by making a bug report.

diff --git a/Engine/Source/Runtime/CoreUObject/Private/UObject/SavePackage/SavePackageUtilities.cpp b/Engine/Source/Runtime/CoreUObject/Private/UObject/SavePackage/SavePackageUtilities.cpp
--- a/Engine/Source/Runtime/CoreUObject/Private/UObject/SavePackage/SavePackageUtilities.cpp
+++ b/Engine/Source/Runtime/CoreUObject/Private/UObject/SavePackage/SavePackageUtilities.cpp
@@ -1601,7 +1601,7 @@ void SaveThumbnails(UPackage* InOuter, FLinkerSave* Linker, FStructuredArchive::
 				FObjectFullNameAndThumbnail& CurObjectThumb = ObjectsWithThumbnails[ CurObjectIndex ];
 
 				// Store the file offset to this thumbnail
-				CurObjectThumb.FileOffset = Linker->Tell();
+				CurObjectThumb.FileOffset = (int32)Linker->Tell();
 
 				// Serialize the thumbnail!
 				FObjectThumbnail* SerializableThumbnail = const_cast< FObjectThumbnail* >( CurObjectThumb.ObjectThumbnail );
@@ -1611,7 +1611,7 @@ void SaveThumbnails(UPackage* InOuter, FLinkerSave* Linker, FStructuredArchive::
 
 			// Store the thumbnail table of contents
 			{
-				Linker->Summary.ThumbnailTableOffset = Linker->Tell();
+				Linker->Summary.ThumbnailTableOffset = (int32)Linker->Tell();
 
 				// Save number of thumbnails
 				int32 ThumbnailCount = ObjectsWithThumbnails.Num();
@@ -2166,7 +2166,7 @@ void SaveWorldLevelInfo(UPackage* InOuter, FLinkerSave* Linker, FStructuredArchi
 	
 	if(FWorldTileInfo* WorldTileInfo = InOuter->GetWorldTileInfo())
 	{
-		Linker->Summary.WorldTileInfoDataOffset = Linker->Tell();
+		Linker->Summary.WorldTileInfoDataOffset = (int32)Linker->Tell();
 		Record << SA_VALUE(TEXT("WorldLevelInfo"), *WorldTileInfo);
 	}
 }
@@ -2230,7 +2230,7 @@ namespace AssetRegistry
 		check(!BinaryArchive.IsTextFormat());
 
 		// Store the asset registry offset in the file and enter a record for the asset registry data
-		Linker->Summary.AssetRegistryDataOffset = BinaryArchive.Tell();
+		Linker->Summary.AssetRegistryDataOffset = (int32)BinaryArchive.Tell();
 		FStructuredArchiveRecord AssetRegistryRecord = ParentRecord.EnterField(TEXT("AssetRegistry")).EnterRecord();
 
 		// Offset to Dependencies
diff --git a/Engine/Source/Runtime/CoreUObject/Private/UObject/SavePackage2.cpp b/Engine/Source/Runtime/CoreUObject/Private/UObject/SavePackage2.cpp
--- a/Engine/Source/Runtime/CoreUObject/Private/UObject/SavePackage2.cpp
+++ b/Engine/Source/Runtime/CoreUObject/Private/UObject/SavePackage2.cpp
@@ -1301,7 +1301,7 @@ void SavePreloadDependencies(FStructuredArchive::FRecord& StructuredArchiveRoot,
 		}
 	};
 
-	Linker->Summary.PreloadDependencyOffset = Linker->Tell();
+	Linker->Summary.PreloadDependencyOffset = (int32)Linker->Tell();
 	Linker->Summary.PreloadDependencyCount = -1;
 
 	if (SaveContext.IsCooking())
@@ -1532,7 +1532,7 @@ void WriteGatherableText(FStructuredArchive::FRecord& StructuredArchiveRoot, FSa
 		// Note that using custom version here only works because we already added it to the export tagger before the package summary was serialized
 		Linker->UsingCustomVersion(FEditorObjectVersion::GUID);
 
-		Linker->Summary.GatherableTextDataOffset = Linker->Tell();
+		Linker->Summary.GatherableTextDataOffset = (int32)Linker->Tell();
 		Linker->Summary.GatherableTextDataCount = Linker->GatherableTextDataMap.Num();
 		for (FGatherableTextData& GatherableTextData : Linker->GatherableTextDataMap)
 		{
@@ -1574,31 +1574,31 @@ ESavePackageResult WritePackageHeader(FStructuredArchive::FRecord& StructuredArc
 	// Save Dummy Import Map, overwritten later.
 	{
 		SCOPED_SAVETIMER(UPackage_Save_WriteDummyImportMap);
-		Linker->Summary.ImportOffset = Linker->Tell();
+		Linker->Summary.ImportOffset = (int32)Linker->Tell();
 		for (FObjectImport& Import : Linker->ImportMap)
 		{
 			StructuredArchiveRoot.GetUnderlyingArchive() << Import;
 		}
 	}
-	SaveContext.OffsetAfterImportMap = Linker->Tell();
+	SaveContext.OffsetAfterImportMap = (int32)Linker->Tell();
 
 	// Save Dummy Export Map, overwritten later.
 	{
 		SCOPED_SAVETIMER(UPackage_Save_WriteDummyExportMap);
-		Linker->Summary.ExportOffset = Linker->Tell();
+		Linker->Summary.ExportOffset = (int32)Linker->Tell();
 		for (FObjectExport& Export : Linker->ExportMap)
 		{
 			*Linker << Export;
 		}
 	}
-	SaveContext.OffsetAfterExportMap = Linker->Tell();
+	SaveContext.OffsetAfterExportMap = (int32)Linker->Tell();
 
 	// Save Depend Map
 	{
 		SCOPED_SAVETIMER(UPackage_Save_WriteDependsMap);
 
 		FStructuredArchive::FStream DependsStream = StructuredArchiveRoot.EnterStream(TEXT("DependsMap"));
-		Linker->Summary.DependsOffset = Linker->Tell();
+		Linker->Summary.DependsOffset = (int32)Linker->Tell();
 		if (SaveContext.IsCooking())
 		{
 			//@todo optimization, this should just be stripped entirely from cooked packages
@@ -1625,7 +1625,7 @@ ESavePackageResult WritePackageHeader(FStructuredArchive::FRecord& StructuredArc
 		SCOPED_SAVETIMER(UPackage_Save_SaveSoftPackagesAndSearchableNames);
 
 		// Save soft package references
-		Linker->Summary.SoftPackageReferencesOffset = Linker->Tell();
+		Linker->Summary.SoftPackageReferencesOffset = (int32)Linker->Tell();
 		Linker->Summary.SoftPackageReferencesCount = Linker->SoftPackageReferenceList.Num();
 		{
 			FStructuredArchive::FStream SoftReferenceStream = StructuredArchiveRoot.EnterStream(TEXT("SoftReferences"));
@@ -1635,7 +1635,7 @@ ESavePackageResult WritePackageHeader(FStructuredArchive::FRecord& StructuredArc
 			}
 
 			// Save searchable names map
-			Linker->Summary.SearchableNamesOffset = Linker->Tell();
+			Linker->Summary.SearchableNamesOffset = (int32)Linker->Tell();
 			Linker->SerializeSearchableNamesMap(StructuredArchiveRoot.EnterField(TEXT("SearchableNames")));
 		}
 	}
@@ -1667,7 +1667,18 @@ ESavePackageResult WritePackageHeader(FStructuredArchive::FRecord& StructuredArc
 		SCOPED_SAVETIMER(UPackage_Save_PreloadDependencies);
 		SavePreloadDependencies(StructuredArchiveRoot, SaveContext);
 	}
-	Linker->Summary.TotalHeaderSize = Linker->Tell();
+	Linker->Summary.TotalHeaderSize = (int32)Linker->Tell();
+
+	// Rather than check if an offset is truncated every time we assign one, we can just check the final TotalHeaderSize to see if it is truncated.
+	// Checking every time an offset is assigned would let us fail quicker but a) relies that new code follows the convention b) bloats the code a fair bit.
+	if (Linker->Tell() > MAX_int32)
+	{
+		UE_LOG(LogSavePackage, Error, TEXT("Package header for '%s' is too large (%" UINT64_FMT " bytes), some package file summary offsets will be truncated when stored as a int32"),
+			*SaveContext.GetPackage()->GetName(), Linker->Tell());
+
+		return ESavePackageResult::Error;
+	}
+
 	return ReturnSuccessOrCancel();
 }
 
