Fix a bug where saving a static mesh multiple times would cause it to emit serialization warnings about invalid flags when loaded.

#rnx

### Problem
- When saving an editor bulkdata with an invalid (empty) payload, it was being incorrectly identified has sharing a payload with another bulkdata instance that had already been virtualized and so was applying the virtualization flag to itself on post save.
- When saved a second time this incorrect flag would be saved to the bulkdata object itself.
- When loading we would detect that the bulkdata believed itself to be virtualized even though virtualizing an empty payload makes no sense and would emit the log warning.

### Fix
- When fixing up an editor bulkdata after saving we should not check to see if the payload is a virtualized duplicate if the payload is invalid.
- Additionally we should not record the package path if the payload is invalid as there would be nothing to load from.

diff --git a/Engine/Source/Runtime/CoreUObject/Private/Serialization/EditorBulkData.cpp b/Engine/Source/Runtime/CoreUObject/Private/Serialization/EditorBulkData.cpp
--- a/Engine/Source/Runtime/CoreUObject/Private/Serialization/EditorBulkData.cpp
+++ b/Engine/Source/Runtime/CoreUObject/Private/Serialization/EditorBulkData.cpp
@@ -1621,21 +1621,26 @@ void FEditorBulkData::SerializeToPackageTrailer(FLinkerSave& LinkerSave, FCompre
 				this->OffsetInFile = PayloadOffset;
 				this->Flags = (this->Flags & EFlags::TransientFlags) | (UpdatedFlags & ~EFlags::TransientFlags);
 	
-				if (PayloadOffset != INDEX_NONE)
+				// If the payload is valid we might want to fix up the package path or virtualization flags now
+				// that the package has saved.
+				if (!this->PayloadContentId.IsZero())
 				{
-					// A valid payload means that it was saved locally to disk. We can set the path at this point but will only
-					// be able to load from it if the package is reattached to the file on disk, which is not a currently
-					// supported feature.
-					this->PackagePath = InPackagePath;
-					check(!this->PackagePath.IsEmpty()); // LinkerSave guarantees a valid PackagePath if we're updating loaded path
-				}
-				else
-				{
-					// If the payload offset we are given is INDEX_NONE it means that the payload was discarded as there was an identical
-					// virtualized entry and the virtualized version takes priority. Since we know that the payload is in the 
-					// virtualization system at this point we can set the virtualization flag allowing us to unload any existing payload
-					// in memory to help reduce bloat.
-					EnumAddFlags(this->Flags, EFlags::IsVirtualized);
+					if (PayloadOffset != INDEX_NONE)
+					{
+						// A valid payload offset means that it was saved locally to disk. We can set the path at this point
+						// but will only be able to load from it if the package is reattached to the file on disk, which is
+						// not a currently supported feature.
+						this->PackagePath = InPackagePath;
+						check(!this->PackagePath.IsEmpty()); // LinkerSave guarantees a valid PackagePath if we're updating loaded path
+					}
+					else
+					{
+						// If the payload offset we are given is INDEX_NONE it means that the payload was discarded as there was an identical
+						// virtualized entry and the virtualized version takes priority. Since we know that the payload is in the 
+						// virtualization system at this point we can set the virtualization flag allowing us to unload any existing payload
+						// in memory to help reduce bloat.
+						EnumAddFlags(this->Flags, EFlags::IsVirtualized);
+					}
 				}
 
 				if (CanUnloadData())
