EditorDomain: Add UObject::AppendToClassSchema so that textures can declare EnableLegacyAlphaCoverageThresholdScaling and use different EditorDomain packages when that variable is changed. This will also add to the new hybrid iterative cook system to allow detection of when iteratively-cooked packages need to be invalidated due to changes in Serialize and PostLoad.
#rnx

diff --git a/Engine/Source/Runtime/CoreUObject/Private/Serialization/UnversionedPropertySerialization.cpp b/Engine/Source/Runtime/CoreUObject/Private/Serialization/UnversionedPropertySerialization.cpp
--- a/Engine/Source/Runtime/CoreUObject/Private/Serialization/UnversionedPropertySerialization.cpp
+++ b/Engine/Source/Runtime/CoreUObject/Private/Serialization/UnversionedPropertySerialization.cpp
@@ -316,6 +316,11 @@ struct FUnversionedStructSchema
 		FUnversionedStructSchema* Schema = reinterpret_cast<FUnversionedStructSchema*>(FMemory::Malloc(Bytes, alignof(FUnversionedPropertySerializer)));
 		
 #if WITH_EDITORONLY_DATA
+		const UClass* StructAsClass = Cast<const UClass>(Struct);
+		if (StructAsClass)
+		{
+			StructAsClass->CallAppendToClassSchema(HashBuilder);
+		}
 		Schema->SchemaHash = HashBuilder.Finalize();
 #endif
 		Schema->Num = Serializers.Num();
diff --git a/Engine/Source/Runtime/CoreUObject/Private/UObject/Obj.cpp b/Engine/Source/Runtime/CoreUObject/Private/UObject/Obj.cpp
--- a/Engine/Source/Runtime/CoreUObject/Private/UObject/Obj.cpp
+++ b/Engine/Source/Runtime/CoreUObject/Private/UObject/Obj.cpp
@@ -1572,6 +1572,10 @@ void UObject::DeclareCustomVersions(FArchive& Ar, const UClass* SpecificSubclass
 		}
 	}
 }
+
+void UObject::AppendToClassSchema(FBlake3& Hasher)
+{
+}
 #endif
 
 void UObject::SerializeScriptProperties(FArchive& Ar) const
diff --git a/Engine/Source/Runtime/CoreUObject/Public/UObject/Class.h b/Engine/Source/Runtime/CoreUObject/Public/UObject/Class.h
--- a/Engine/Source/Runtime/CoreUObject/Public/UObject/Class.h
+++ b/Engine/Source/Runtime/CoreUObject/Public/UObject/Class.h
@@ -2473,6 +2473,13 @@ public:
 		check(CppClassStaticFunctions.GetDeclareCustomVersions());
 		CppClassStaticFunctions.GetDeclareCustomVersions()(Ar, this);
 	}
+
+	/** Calls the c++ class's AppendToClassSchema static function */
+	void CallAppendToClassSchema(FBlake3& Hasher) const
+	{
+		check(CppClassStaticFunctions.GetAppendToClassSchema());
+		CppClassStaticFunctions.GetAppendToClassSchema()(Hasher);
+	}
 #endif
 
 	/** The class default object; used for delta serialization and object initialization */
diff --git a/Engine/Source/Runtime/CoreUObject/Public/UObject/Object.h b/Engine/Source/Runtime/CoreUObject/Public/UObject/Object.h
--- a/Engine/Source/Runtime/CoreUObject/Public/UObject/Object.h
+++ b/Engine/Source/Runtime/CoreUObject/Public/UObject/Object.h
@@ -15,6 +15,7 @@
 #include "UObject/PrimaryAssetId.h"
 
 struct FAssetData;
+class FBlake3;
 class FConfigCacheIni;
 class FCustomPropertyConditionState;
 class FEditPropertyChain;
@@ -372,6 +373,12 @@ public:
 	 * Used to track which customversions will be used by a package when it is resaved.
 	 * Not yet exhaustive; add CustomVersions as necessary to remove EditorDomain warnings about missing versions. */
 	static void DeclareCustomVersions(FArchive& Ar, const UClass* SpecificSubclass);
+	/**
+	 * Append config values or settings that can change how instances of the class are cooked, including especially
+	 * values that determine how version upgraded are conducted. Can also append a unique guid when necessary to
+	 * invalidate previous results because serialization changed and no custom version was udpated.
+	 */
+	static void AppendToClassSchema(FBlake3& Hasher);
 #endif
 
 	/** After a critical error, perform any mission-critical cleanup, such as restoring the video mode orreleasing hardware resources. */
diff --git a/Engine/Source/Runtime/CoreUObject/Public/UObject/ObjectMacros.h b/Engine/Source/Runtime/CoreUObject/Public/UObject/ObjectMacros.h
--- a/Engine/Source/Runtime/CoreUObject/Public/UObject/ObjectMacros.h
+++ b/Engine/Source/Runtime/CoreUObject/Public/UObject/ObjectMacros.h
@@ -10,6 +10,7 @@
 #include "Misc/EnumClassFlags.h"
 #include "UObject/Script.h"
 
+class FBlake3;
 class FObjectInitializer;
 class FReferenceCollector;
 struct FFrame;
@@ -1846,7 +1847,8 @@ public: \
 
 #if WITH_EDITORONLY_DATA
 	#define UOBJECT_CPPCLASS_STATICFUNCTIONS_WITHEDITORONLYDATA(TClass) \
-		, &TClass::DeclareCustomVersions
+		, &TClass::DeclareCustomVersions \
+		, &TClass::AppendToClassSchema
 		/* UObjectCppClassStaticFunctions: Extend this macro with the address of your new static function, if it is editor-only. */ \
 		/* Order must match the order in the FUObjectCppClassStaticFunctions constructor. */ \
 #else
@@ -1876,23 +1878,27 @@ public:
 	typedef void (*AddReferencedObjectsType)	(UObject* ThisObject, FReferenceCollector& Ar);
 #if WITH_EDITORONLY_DATA
 	typedef void (*DeclareCustomVersionsType)   (FArchive& Ar, const UClass* SpecificSubclass);
+	typedef void (*AppendToClassSchemaType)   (FBlake3& Hasher);
 #endif
 	// UObjectCppClassStaticFunctions: Extend this list of types with the type of your new static function.
 
 	FUObjectCppClassStaticFunctions(AddReferencedObjectsType InAddReferencedObjects
 #if WITH_EDITORONLY_DATA
 		, DeclareCustomVersionsType InDeclareCustomVersions
+		, AppendToClassSchemaType InAppendToClassSchema
 #endif
 	)
 		: AddReferencedObjects(InAddReferencedObjects)
 #if WITH_EDITORONLY_DATA
 		, DeclareCustomVersions(InDeclareCustomVersions)
+		, AppendToClassSchema(InAppendToClassSchema)
 #endif
 	{
 		// Null elements are not valid in this constructor
 		check(InAddReferencedObjects);
 #if WITH_EDITORONLY_DATA
 		check(InDeclareCustomVersions);
+		check(InAppendToClassSchema);
 #endif
 		// UObjectCppClassStaticFunctions: Extend the constructor with initializers for your new static function member.
 		// Order must match the order in UOBJECT_CPPCLASS_STATICFUNCTIONS_FORCLASS.
@@ -1930,6 +1936,15 @@ public:
 		check(InDeclareCustomVersions != nullptr); // It is not valid to clear single elements (see IsInitialized). Call Reset to clear all elements.
 		DeclareCustomVersions = InDeclareCustomVersions;
 	}
+	AppendToClassSchemaType GetAppendToClassSchema() const
+	{
+		return AppendToClassSchema;
+	}
+	void SetAppendToClassSchema(AppendToClassSchemaType InAppendToClassSchema)
+	{
+		check(InAppendToClassSchema != nullptr); // It is not valid to clear single elements (see IsInitialized). Call Reset to clear all elements.
+		AppendToClassSchema = InAppendToClassSchema;
+	}
 #endif
 	// UObjectCppClassStaticFunctions: Extend the list of accessors for your new static function.
 
@@ -1937,6 +1952,7 @@ private:
 	AddReferencedObjectsType AddReferencedObjects = nullptr;
 #if WITH_EDITORONLY_DATA
 	DeclareCustomVersionsType DeclareCustomVersions = nullptr;
+	AppendToClassSchemaType AppendToClassSchema = nullptr;
 #endif
 	// UObjectCppClassStaticFunctions: Extend this list of members with the member for your new static function.
 
diff --git a/Engine/Source/Runtime/Engine/Classes/Engine/Texture.h b/Engine/Source/Runtime/Engine/Classes/Engine/Texture.h
--- a/Engine/Source/Runtime/Engine/Classes/Engine/Texture.h
+++ b/Engine/Source/Runtime/Engine/Classes/Engine/Texture.h
@@ -1632,6 +1632,7 @@ public:
 	ENGINE_API virtual void FinishDestroy() override;
 	ENGINE_API virtual void PostCDOContruct() override;
 #if WITH_EDITORONLY_DATA
+	ENGINE_API static void AppendToClassSchema(FBlake3& Hasher);
 	ENGINE_API virtual void GetAssetRegistryTags(TArray<FAssetRegistryTag>& OutTags) const override;
 #endif
 	ENGINE_API virtual bool IsPostLoadThreadSafe() const override;
diff --git a/Engine/Source/Runtime/Engine/Private/Texture.cpp b/Engine/Source/Runtime/Engine/Private/Texture.cpp
--- a/Engine/Source/Runtime/Engine/Private/Texture.cpp
+++ b/Engine/Source/Runtime/Engine/Private/Texture.cpp
@@ -40,6 +40,7 @@
 #include "ImageCoreUtils.h"
 #include "ImageUtils.h"
 #include "Algo/Unique.h"
+#include "Hash/Blake3.h"
 
 #if WITH_EDITOR
 #include "DerivedDataBuildVersion.h"
@@ -711,6 +712,22 @@ void UTexture::PostEditChangeProperty(FPropertyChangedEvent& PropertyChangedEven
 }
 #endif // WITH_EDITOR
 
+#if WITH_EDITORONLY_DATA
+static bool IsEnableLegacyAlphaCoverageThresholdScaling()
+{
+	static struct ReadConfigOnce
+	{
+		bool bBool;
+		ReadConfigOnce()
+		{
+			bBool = false;
+			GConfig->GetBool(TEXT("Texture"), TEXT("EnableLegacyAlphaCoverageThresholdScaling"), bBool, GEditorIni);
+		}
+	} ConfigValue;
+	return ConfigValue.bBool;
+}
+#endif
+
 void UTexture::Serialize(FArchive& Ar)
 {
 	Ar.UsingCustomVersion(FUE5MainStreamObjectVersion::GUID);
@@ -757,19 +774,8 @@ void UTexture::Serialize(FArchive& Ar)
 		{
 			// if value is (0,0,0,1)
 			//	that was previously incorrectly being set by default and enabling alpha coverage processing
-			// we don't want that, but to optionally preserve old behavior you can set a config option :
-
-			static struct ReadConfigOnce
-			{
-				bool bBool;
-				ReadConfigOnce()
-				{
-					bBool = false;
-					GConfig->GetBool(TEXT("Texture"), TEXT("EnableLegacyAlphaCoverageThresholdScaling"), bBool, GEditorIni);
-				}
-			} ConfigValue;
-
-			bDoScaleMipsForAlphaCoverage = ConfigValue.bBool;
+			// we don't want that, but to optionally preserve old behavior you can set a config option
+			bDoScaleMipsForAlphaCoverage = IsEnableLegacyAlphaCoverageThresholdScaling();
 		}
 	}
 
@@ -887,6 +893,18 @@ void UTexture::Serialize(FArchive& Ar)
 #endif // #if WITH_EDITORONLY_DATA
 }
 
+#if WITH_EDITORONLY_DATA
+void UTexture::AppendToClassSchema(FBlake3& Hasher)
+{
+	Super::AppendToClassSchema(Hasher);
+
+	// IsEnableLegacyAlphaCoverageThresholdScaling affects upgrades executed in Serialize, so include it in the ClassSchema
+	uint8 LegacyScalingBool = IsEnableLegacyAlphaCoverageThresholdScaling();
+	Hasher.Update(&LegacyScalingBool, sizeof(LegacyScalingBool));
+}
+#endif
+
+
 void UTexture::PostInitProperties()
 {
 #if WITH_EDITORONLY_DATA
