SavePackage: Call PostSaveRoot when early-exiting, even if bCleanupRequired was set to false by PreSaveRoot. bCleanupRequired refers to specific internal details of PostSaveRoot; PostSaveRoot needs to be called in all cases where PreSaveRoot was called.
#rnx

diff --git a/Engine/Source/Runtime/CoreUObject/Private/UObject/SavePackage/SaveContext.h b/Engine/Source/Runtime/CoreUObject/Private/UObject/SavePackage/SaveContext.h
--- a/Engine/Source/Runtime/CoreUObject/Private/UObject/SavePackage/SaveContext.h
+++ b/Engine/Source/Runtime/CoreUObject/Private/UObject/SavePackage/SaveContext.h
@@ -424,7 +424,7 @@ public:
 
 	~FSaveContext()
 	{
-		if (bNeedPreSaveCleanup && Asset)
+		if (bPostSaveRootRequired && Asset)
 		{
 			UE::SavePackageUtilities::CallPostSaveRoot(Asset, ObjectSaveContext, bNeedPreSaveCleanup);
 		}
@@ -630,6 +630,16 @@ public:
 		return PPF_DeepCompareInstances | PPF_DeepCompareDSOsOnly;
 	}
 
+	bool GetPostSaveRootRequired() const
+	{
+		return bPostSaveRootRequired;
+	}
+
+	void SetPostSaveRootRequired(bool bInPostSaveRootRequired)
+	{
+		bPostSaveRootRequired = bInPostSaveRootRequired;
+	}
+
 	bool GetPreSaveCleanup() const
 	{
 		return bNeedPreSaveCleanup;
@@ -991,6 +1001,7 @@ private:
 	bool bTextFormat = false;
 	bool bIsProcessingPrestreamPackages = false;
 	bool bIsFixupStandaloneFlags = false;
+	bool bPostSaveRootRequired = false;
 	bool bNeedPreSaveCleanup = false;
 	bool bGenerateFileStub = false;
 	bool bIgnoreHeaderDiffs = false;
diff --git a/Engine/Source/Runtime/CoreUObject/Private/UObject/SavePackage2.cpp b/Engine/Source/Runtime/CoreUObject/Private/UObject/SavePackage2.cpp
--- a/Engine/Source/Runtime/CoreUObject/Private/UObject/SavePackage2.cpp
+++ b/Engine/Source/Runtime/CoreUObject/Private/UObject/SavePackage2.cpp
@@ -2567,6 +2567,7 @@ FSavePackageResultStruct UPackage::Save2(UPackage* InPackage, UObject* InAsset,
 	{
 		FObjectSaveContextData& ObjectSaveContext = SaveContext.GetObjectSaveContext();
 		UE::SavePackageUtilities::CallPreSaveRoot(SaveContext.GetAsset(), ObjectSaveContext);
+		SaveContext.SetPostSaveRootRequired(true);
 		SaveContext.SetPreSaveCleanup(ObjectSaveContext.bCleanupRequired);
 	}
 
@@ -2601,10 +2602,10 @@ FSavePackageResultStruct UPackage::Save2(UPackage* InPackage, UObject* InAsset,
 
 	// PostSave Asset
 	SlowTask.EnterProgressFrame();
-	if (SaveContext.GetAsset() && !SaveContext.IsConcurrent())
+	if (SaveContext.GetPostSaveRootRequired() && SaveContext.GetAsset())
 	{
 		UE::SavePackageUtilities::CallPostSaveRoot(SaveContext.GetAsset(), SaveContext.GetObjectSaveContext(), SaveContext.GetPreSaveCleanup());
-		SaveContext.SetPreSaveCleanup(false);
+		SaveContext.SetPostSaveRootRequired(false);
 	}
 
 	ClearCachedPlatformCookedData(SaveContext);
