Optimize resolving FGCObjectInfo by searching using FNames only.

diff --git a/Engine/Source/Runtime/CoreUObject/Private/UObject/GCObjectInfo.cpp b/Engine/Source/Runtime/CoreUObject/Private/UObject/GCObjectInfo.cpp
--- a/Engine/Source/Runtime/CoreUObject/Private/UObject/GCObjectInfo.cpp
+++ b/Engine/Source/Runtime/CoreUObject/Private/UObject/GCObjectInfo.cpp
@@ -8,7 +8,25 @@
 
 UObject* FGCObjectInfo::TryResolveObject()
 {
-	return StaticFindObject(UObject::StaticClass(), nullptr, *GetPathName());
+	UClass* ResolvedClass = Class ? Cast<UClass>(Class->TryResolveObject()) : UClass::StaticClass();
+	if (!ResolvedClass)
+	{
+		return nullptr;
+	}
+
+	UObject* Resolved = nullptr;
+	if (Outer)
+	{
+		if (UObject* OuterObj = Outer->TryResolveObject())
+		{
+			return StaticFindObjectFast(ResolvedClass, OuterObj, Name, true, false);
+		}
+		return nullptr;
+	}
+	else
+	{
+		return StaticFindObjectFast(ResolvedClass, nullptr, Name, true, false);
+	}
 }
 
 void FGCObjectInfo::GetPathName(FStringBuilderBase& ResultString) const
@@ -30,6 +48,7 @@ void FGCObjectInfo::GetPathName(FStringBuilderBase& ResultString) const
 			}
 		}
 		Name.AppendString(ResultString);
+
 	}
 	else
 	{
