Fix non-unity error.

diff --git a/Engine/Source/Runtime/CoreUObject/Private/UObject/GCObjectInfo.cpp b/Engine/Source/Runtime/CoreUObject/Private/UObject/GCObjectInfo.cpp
--- a/Engine/Source/Runtime/CoreUObject/Private/UObject/GCObjectInfo.cpp
+++ b/Engine/Source/Runtime/CoreUObject/Private/UObject/GCObjectInfo.cpp
@@ -5,6 +5,7 @@
 =============================================================================*/
 
 #include "UObject/GCObjectInfo.h"
+#include "UObject/Class.h"
 
 UObject* FGCObjectInfo::TryResolveObject()
 {
