Revert optimizaiton to fix infinite recursion

diff --git a/Engine/Source/Runtime/CoreUObject/Private/UObject/GCObjectInfo.cpp b/Engine/Source/Runtime/CoreUObject/Private/UObject/GCObjectInfo.cpp
--- a/Engine/Source/Runtime/CoreUObject/Private/UObject/GCObjectInfo.cpp
+++ b/Engine/Source/Runtime/CoreUObject/Private/UObject/GCObjectInfo.cpp
@@ -5,29 +5,10 @@
 =============================================================================*/
 
 #include "UObject/GCObjectInfo.h"
-#include "UObject/Class.h"
 
 UObject* FGCObjectInfo::TryResolveObject()
 {
-	UClass* ResolvedClass = Class ? Cast<UClass>(Class->TryResolveObject()) : UClass::StaticClass();
-	if (!ResolvedClass)
-	{
-		return nullptr;
-	}
-
-	UObject* Resolved = nullptr;
-	if (Outer)
-	{
-		if (UObject* OuterObj = Outer->TryResolveObject())
-		{
-			return StaticFindObjectFast(ResolvedClass, OuterObj, Name, true, false);
-		}
-		return nullptr;
-	}
-	else
-	{
-		return StaticFindObjectFast(ResolvedClass, nullptr, Name, true, false);
-	}
+	return StaticFindObject(UObject::StaticClass(), nullptr, *GetPathName());
 }
 
 void FGCObjectInfo::GetPathName(FStringBuilderBase& ResultString) const
@@ -49,7 +30,6 @@ void FGCObjectInfo::GetPathName(FStringBuilderBase& ResultString) const
 			}
 		}
 		Name.AppendString(ResultString);
-
 	}
 	else
 	{
