[Insights] Added more "UObject" mem scopes (LLM_SCOPE) in an effort to reduce the Untagged allocations:
  - Capture untagged memory allocated from FRegisterCompiledInInfo (i.e. used by IMPLEMENT_CLASS).
  - Capture untagged memory allocated from the lambda used to load an import package.

diff --git a/Engine/Source/Runtime/CoreUObject/Private/UObject/LinkerLoad.cpp b/Engine/Source/Runtime/CoreUObject/Private/UObject/LinkerLoad.cpp
--- a/Engine/Source/Runtime/CoreUObject/Private/UObject/LinkerLoad.cpp
+++ b/Engine/Source/Runtime/CoreUObject/Private/UObject/LinkerLoad.cpp
@@ -3015,6 +3015,8 @@ bool FLinkerLoad::VerifyImportInner(const int32 ImportIndex, FString& WarningSuf
 	// Lambda used to load an import package
 	auto LoadImportPackage = [this, ImportIndex](FObjectImport& Import, TOptional<FScopedSlowTask>& SlowTask) -> UPackage*
 	{
+		LLM_SCOPE(ELLMTag::UObject);
+
 		// Either this import is a UPackage or it has PackageName set.
 		check(Import.ClassName == NAME_Package || Import.HasPackageName());
 
diff --git a/Engine/Source/Runtime/CoreUObject/Private/UObject/UObjectBase.cpp b/Engine/Source/Runtime/CoreUObject/Private/UObject/UObjectBase.cpp
--- a/Engine/Source/Runtime/CoreUObject/Private/UObject/UObjectBase.cpp
+++ b/Engine/Source/Runtime/CoreUObject/Private/UObject/UObjectBase.cpp
@@ -697,6 +697,8 @@ void RegisterCompiledInInfo(UPackage* (*InOuterRegister)(), const TCHAR* InPacka
 // Multiple registrations
 void RegisterCompiledInInfo(const TCHAR* PackageName, const FClassRegisterCompiledInInfo* ClassInfo, size_t NumClassInfo, const FStructRegisterCompiledInInfo* StructInfo, size_t NumStructInfo, const FEnumRegisterCompiledInInfo* EnumInfo, size_t NumEnumInfo)
 {
+	LLM_SCOPE(ELLMTag::UObject);
+
 	for (size_t Index = 0; Index < NumClassInfo; ++Index)
 	{
 		const FClassRegisterCompiledInInfo& Info = ClassInfo[Index];
