SavePackage: Undo the call to IsCachedCookedPlatformDataLoaded, since it has caused multiple fallout issues.
Cooker: Fix the issue that we were trying to fix with that previous change; it is specific to WorldPartition and CookPackageSplitter; we need to requery the list of objects in the package that need BeginCacheForCookedPlatformData after we call on WorldPartition to move all of its OFPA objects into the Level object for cooking.
#rnx

diff --git a/Engine/Source/Editor/UnrealEd/Classes/CookOnTheSide/CookOnTheFlyServer.h b/Engine/Source/Editor/UnrealEd/Classes/CookOnTheSide/CookOnTheFlyServer.h
--- a/Engine/Source/Editor/UnrealEd/Classes/CookOnTheSide/CookOnTheFlyServer.h
+++ b/Engine/Source/Editor/UnrealEd/Classes/CookOnTheSide/CookOnTheFlyServer.h
@@ -1139,7 +1139,7 @@ private:
 	* @param bOutCompleted			Return true is function is done splitting
 	* @param bOutError				Return true if function fails
 	*/
-	void SplitPackage(UE::Cook::FGeneratorPackage* Generator, bool& bCompleted, bool& bOutError);
+	void SplitPackage(UE::Cook::FPackageData& PackageData, UE::Cook::FGeneratorPackage* Generator, bool& bCompleted, bool& bOutError);
 
 	/** Try calling the splitter's populate to create the package */
 	UPackage* TryPopulateGeneratedPackage(UE::Cook::FPopulatePackageContext& Context);
diff --git a/Engine/Source/Editor/UnrealEd/Private/CookOnTheFlyServer.cpp b/Engine/Source/Editor/UnrealEd/Private/CookOnTheFlyServer.cpp
--- a/Engine/Source/Editor/UnrealEd/Private/CookOnTheFlyServer.cpp
+++ b/Engine/Source/Editor/UnrealEd/Private/CookOnTheFlyServer.cpp
@@ -2453,7 +2453,7 @@ UE::Cook::FGeneratorPackage* UCookOnTheFlyServer::CreateGeneratorPackage(UE::Coo
 	return GeneratorPackage;
 }
 
-void UCookOnTheFlyServer::SplitPackage(UE::Cook::FGeneratorPackage* GeneratorStruct, bool& bOutCompleted, bool& bOutError)
+void UCookOnTheFlyServer::SplitPackage(UE::Cook::FPackageData& PackageData, UE::Cook::FGeneratorPackage* GeneratorStruct, bool& bOutCompleted, bool& bOutError)
 {
 	using namespace UE::Cook;
 
@@ -2535,9 +2535,12 @@ void UCookOnTheFlyServer::SplitPackage(UE::Cook::FGeneratorPackage* GeneratorStr
 		}
 	}
 
-	UPackage* OwnerPackage = GeneratorStruct->GetOwner().GetPackage();
+	UPackage* OwnerPackage = PackageData.GetPackage();
 	check(OwnerPackage);
 	GeneratorStruct->GetCookPackageSplitterInstance()->PreSaveGeneratorPackage(OwnerPackage, SplitObject, SplitterDatas);
+	// PreSaveGenerator package can add new objects to the package due to moving them in from OneFilePerActor packages
+	// We want to support calling BeginCacheForCookedPlatformData on those new objects, so recreate the ObjectCache now.
+	PackageData.RecreateObjectCache();
 
 	bOutCompleted = true;
 }
@@ -2704,7 +2707,7 @@ bool UCookOnTheFlyServer::BeginPrepareSave(UE::Cook::FPackageData& PackageData,
 				return false;
 			}
 			bool bCompleted = false;
-			SplitPackage(Generator, bCompleted, bError);
+			SplitPackage(PackageData, Generator, bCompleted, bError);
 			if (bError)
 			{
 				// SplitPackage failure marks the package data for PumpSaves to handle this as an error and to put it in idle state
diff --git a/Engine/Source/Editor/UnrealEd/Private/Cooker/CookPackageData.cpp b/Engine/Source/Editor/UnrealEd/Private/Cooker/CookPackageData.cpp
--- a/Engine/Source/Editor/UnrealEd/Private/Cooker/CookPackageData.cpp
+++ b/Engine/Source/Editor/UnrealEd/Private/Cooker/CookPackageData.cpp
@@ -983,6 +983,19 @@ void FPackageData::CreateObjectCache()
 	}
 }
 
+void FPackageData::RecreateObjectCache()
+{
+	check(GetPackage());
+	if (GetHasSaveCache())
+	{
+		// It is not valid to recreate the ObjectCache when we have already started calling BeginCacheForCookedPlatformData on the objects in it.
+		// ReleaseCookedPlatformData must be called first to tear down all of the calls and reset GetCookedPlatformDataNextIndex to 0
+		check(GetCookedPlatformDataNextIndex() == 0);
+		ClearObjectCache();
+	}
+	CreateObjectCache();
+}
+
 void FPackageData::ClearObjectCache()
 {
 	CachedObjectsInOuter.Empty();
diff --git a/Engine/Source/Editor/UnrealEd/Private/Cooker/CookPackageData.h b/Engine/Source/Editor/UnrealEd/Private/Cooker/CookPackageData.h
--- a/Engine/Source/Editor/UnrealEd/Private/Cooker/CookPackageData.h
+++ b/Engine/Source/Editor/UnrealEd/Private/Cooker/CookPackageData.h
@@ -357,6 +357,8 @@ public:
 	void CheckObjectCacheEmpty() const;
 	/** Populate CachedObjectsInOuter if not already populated. Invalid to call except when in the save state. */
 	void CreateObjectCache();
+	/** Clear CachedObjectsInOuter and recompute it. */
+	void RecreateObjectCache();
 	/** Clear the CachedObjectsInOuter list, when e.g. leaving the save state. */
 	void ClearObjectCache();
 
diff --git a/Engine/Source/Runtime/CoreUObject/Private/UObject/SavePackage2.cpp b/Engine/Source/Runtime/CoreUObject/Private/UObject/SavePackage2.cpp
--- a/Engine/Source/Runtime/CoreUObject/Private/UObject/SavePackage2.cpp
+++ b/Engine/Source/Runtime/CoreUObject/Private/UObject/SavePackage2.cpp
@@ -2087,6 +2087,12 @@ ESavePackageResult FinalizeFile(FStructuredArchive::FRecord& StructuredArchiveRo
 ESavePackageResult BeginCachePlatformCookedData(FSaveContext& SaveContext)
 {
 #if WITH_EDITOR
+	// TODO: Call IsCachedCookedPlatformDataLoaded on all the objects until it returns true. This is required for the contract.
+	// We tried enabling this once, but it created knockon bugs: Textures created by landscape were not handling it correctly
+	// (which we have subsequently fixed) and MaterialInstanceConstants created by landscape were not handling it correctly (which 
+	// we have not yet diagnosed)
+#define SAVEPACKAGE_CALL_ISCACHEDCOOKEDPLATFORMDATALOADED 0
+
 	// TODO: Remove BeginCacheForCookedPlatformData from SavePackage; it is already called by the cooker.
 	// Cache platform cooked data
 	if (SaveContext.IsCooking() && !SaveContext.IsConcurrent())
@@ -2103,12 +2109,15 @@ ESavePackageResult BeginCachePlatformCookedData(FSaveContext& SaveContext)
 				continue;
 			}
 			Object->BeginCacheForCookedPlatformData(TargetPlatform);
+#if SAVEPACKAGE_CALL_ISCACHEDCOOKEDPLATFORMDATALOADED
 			if (Object->IsCachedCookedPlatformDataLoaded(TargetPlatform))
 			{
 				Iter.RemoveCurrentSwap();
 				continue;
 			}
+#endif
 		}
+#if SAVEPACKAGE_CALL_ISCACHEDCOOKEDPLATFORMDATALOADED
 		if (ObjectsInPackage.Num())
 		{
 			if (SaveContext.GetSaveArgs().SaveFlags & SAVE_AllowTimeout)
@@ -2135,6 +2144,7 @@ ESavePackageResult BeginCachePlatformCookedData(FSaveContext& SaveContext)
 				}
 			}
 		}
+#endif
 	}
 #endif
 	return ESavePackageResult::Success;
