AsyncLoading2: Fix nullptr crash inside check format statement.

#rnx

diff --git a/Engine/Source/Runtime/CoreUObject/Private/Serialization/AsyncLoading2.cpp b/Engine/Source/Runtime/CoreUObject/Private/Serialization/AsyncLoading2.cpp
--- a/Engine/Source/Runtime/CoreUObject/Private/Serialization/AsyncLoading2.cpp
+++ b/Engine/Source/Runtime/CoreUObject/Private/Serialization/AsyncLoading2.cpp
@@ -6007,7 +6007,7 @@ void FAsyncPackage2::CreateUPackage(const FZenPackageSummary* PackageSummary, co
 				UPackage* FoundPackage = FindObjectFast<UPackage>(nullptr, Desc.UPackageName);
 				checkf(LinkerRoot == FoundPackage,
 					TEXT("LinkerRoot '%s' (%p) is different from FoundPackage '%s' (%p)"),
-					*LinkerRoot->GetName(), LinkerRoot, *FoundPackage->GetName(), FoundPackage);
+					*LinkerRoot->GetName(), LinkerRoot, FoundPackage ? *FoundPackage->GetName() : TEXT("null"), FoundPackage);
 			}
 #endif
 		}
