FMapProperty and FSetProperty fixes for IsPermutation.

diff --git a/Engine/Source/Runtime/CoreUObject/Private/UObject/PropertyMap.cpp b/Engine/Source/Runtime/CoreUObject/Private/UObject/PropertyMap.cpp
--- a/Engine/Source/Runtime/CoreUObject/Private/UObject/PropertyMap.cpp
+++ b/Engine/Source/Runtime/CoreUObject/Private/UObject/PropertyMap.cpp
@@ -154,7 +154,7 @@ namespace UEMapProperty_Private
 		for (;;)
 		{
 			const uint8* PairA = MapHelperA.GetPairPtr(IndexA);
-			if (!AnyEqual(MapHelperA, FirstIndexA, FirstNum - Num, PairA, PortFlags) && !RangesContainSameAmountsOfVal(MapHelperA, IndexA, MapHelperB, IndexB, Num, PairA, PortFlags))
+			if (!AnyEqual(MapHelperA, FirstIndexA, FirstNum - Num, PairA, PortFlags) && !RangesContainSameAmountsOfVal(MapHelperA, FirstIndexA, MapHelperB, FirstIndexB, FirstNum, PairA, PortFlags))
 			{
 				return false;
 			}
@@ -165,15 +165,11 @@ namespace UEMapProperty_Private
 				return true;
 			}
 
+			++IndexA;
 			while (!MapHelperA.IsValidIndex(IndexA))
 			{
 				++IndexA;
 			}
-
-			while (!MapHelperB.IsValidIndex(IndexB))
-			{
-				++IndexB;
-			}
 		}
 	}
 }
diff --git a/Engine/Source/Runtime/CoreUObject/Private/UObject/PropertySet.cpp b/Engine/Source/Runtime/CoreUObject/Private/UObject/PropertySet.cpp
--- a/Engine/Source/Runtime/CoreUObject/Private/UObject/PropertySet.cpp
+++ b/Engine/Source/Runtime/CoreUObject/Private/UObject/PropertySet.cpp
@@ -136,7 +136,7 @@ namespace UESetProperty_Private
 		for (;;)
 		{
 			const uint8* ElementA = SetHelperA.GetElementPtr(IndexA);
-			if (!AnyEqual(SetHelperA, FirstIndexA, FirstNum - Num, ElementA, PortFlags) && !RangesContainSameAmountsOfVal(SetHelperA, IndexA, SetHelperB, IndexB, Num, ElementA, PortFlags))
+			if (!AnyEqual(SetHelperA, FirstIndexA, FirstNum - Num, ElementA, PortFlags) && !RangesContainSameAmountsOfVal(SetHelperA, FirstIndexA, SetHelperB, FirstIndexB, FirstNum, ElementA, PortFlags))
 			{
 				return false;
 			}
@@ -147,15 +147,12 @@ namespace UESetProperty_Private
 				return true;
 			}
 
+			++IndexA;
 			while (!SetHelperA.IsValidIndex(IndexA))
 			{
 				++IndexA;
 			}
 
-			while (!SetHelperB.IsValidIndex(IndexB))
-			{
-				++IndexB;
-			}
 		}
 	}
 }
