Introduce new FVirtualStackAllocator which provides alloca() type functionality backed by virtual memory. Only supported on platforms with support for manual virtual memory management.

diff --git a/Engine/Source/Runtime/Core/Private/Memory/VirtualStackAllocator.cpp b/Engine/Source/Runtime/Core/Private/Memory/VirtualStackAllocator.cpp
new file mode 100644
--- /dev/null
+++ b/Engine/Source/Runtime/Core/Private/Memory/VirtualStackAllocator.cpp
@@ -0,0 +1,111 @@
+// Copyright Epic Games, Inc. All Rights Reserved.
+
+#include "Memory/VirtualStackAllocator.h"
+#include "Templates/AlignmentTemplates.h"
+#include "HAL/LowLevelMemTracker.h"
+
+LLM_DEFINE_TAG(VirtualStackAllocator);
+
+template<typename T> static inline T* OffsetPointer(T* Start, size_t Offset)
+{
+    return reinterpret_cast<T*>(reinterpret_cast<intptr_t>(Start) + Offset);
+}
+
+template<typename T, typename U> static inline ptrdiff_t PointerDifference(T* End, U* Start)
+{
+    return static_cast<ptrdiff_t>(reinterpret_cast<intptr_t>(End) - reinterpret_cast<intptr_t>(Start));
+}
+
+FScopedStackAllocatorBookmark::~FScopedStackAllocatorBookmark()
+{
+    if (RestorePointer != nullptr)
+    {
+        check(Owner != nullptr);
+        Owner->Free(RestorePointer);
+    }
+}
+
+FVirtualStackAllocator::FVirtualStackAllocator(size_t RequestedStackSize, EVirtualStackAllocatorDecommitMode Mode)
+    : PageSize(FPlatformMemory::FPlatformVirtualMemoryBlock::GetCommitAlignment()),
+	DecommitMode(Mode)
+{
+    TotalReservationSize = Align(RequestedStackSize, PageSize);
+    
+	if (TotalReservationSize > 0)
+	{
+		VirtualMemory = FPlatformMemory::FPlatformVirtualMemoryBlock::AllocateVirtual(TotalReservationSize);
+		NextUncommittedPage = VirtualMemory.GetVirtualPointer();
+		NextAllocationStart = NextUncommittedPage;
+		RecentHighWaterMark = NextUncommittedPage;
+	}
+}
+
+FVirtualStackAllocator::~FVirtualStackAllocator()
+{
+    check(GetAllocatedBytes() == 0);
+    VirtualMemory.FreeVirtual();
+}
+
+void* FVirtualStackAllocator::Allocate(size_t Size, size_t Alignment)
+{
+    void* const AllocationStart = Align(NextAllocationStart, Alignment);
+    void* const AllocationEnd = OffsetPointer(AllocationStart, Size);
+
+    if (Size > (TotalReservationSize - PageSize))
+    {
+        FPlatformMemory::OnOutOfMemory(Size, Alignment);
+    }
+
+    void* const UsableMemoryEnd = OffsetPointer(VirtualMemory.GetVirtualPointer(), TotalReservationSize - PageSize);
+
+	// After the high water mark is established, needing to commit pages should be rare
+    if (UNLIKELY(AllocationEnd > NextUncommittedPage))
+    {
+        // We need to commit some more pages. Let's see how many
+        uintptr_t RequiredAdditionalCommit = PointerDifference(AllocationEnd, NextUncommittedPage);
+        // CommitByPtr doesn't round up the size for you
+        size_t SizeToCommit = Align(RequiredAdditionalCommit, PageSize);
+        VirtualMemory.CommitByPtr(NextUncommittedPage, SizeToCommit);
+
+		LLM_IF_ENABLED(FLowLevelMemTracker::Get().OnLowLevelAlloc(ELLMTracker::Default, NextUncommittedPage, SizeToCommit, LLM_TAG_NAME(VirtualStackAllocator)));
+
+        NextUncommittedPage = Align(AllocationEnd, PageSize);
+    }
+
+	if ((char*)AllocationEnd > (char*)RecentHighWaterMark)
+	{
+		RecentHighWaterMark = Align(AllocationEnd, PageSize);
+	}
+
+    NextAllocationStart = AllocationEnd;
+
+    return AllocationStart;
+}
+
+void FVirtualStackAllocator::DecommitUnusedPages()
+{
+	// This should only be called when the allocator is empty
+	check(NextAllocationStart == VirtualMemory.GetVirtualPointer());
+
+	if (DecommitMode == EVirtualStackAllocatorDecommitMode::AllOnStackEmpty)
+	{
+		VirtualMemory.Decommit();
+		NextUncommittedPage = VirtualMemory.GetVirtualPointer();
+	}
+	else if (DecommitMode == EVirtualStackAllocatorDecommitMode::ExcessOnStackEmpty)
+	{
+		// In this mode, each time we get down to zero memory in use we consider decommitting some of the memory above the most recent high water mark
+		ptrdiff_t AmountToFree = (intptr_t)NextUncommittedPage - (intptr_t)RecentHighWaterMark;
+
+		// We will only decommit memory if it would free up at least 25% of the current commit. This helps prevent us from thrashing pages if our
+		// memory usage is consistant but not exactly constant and ensures we only pay to decommit if it will actually result in a significant savings
+		ptrdiff_t MinimumToDecommit = PointerDifference(NextUncommittedPage, VirtualMemory.GetVirtualPointer()) / 4;
+		if (AmountToFree > MinimumToDecommit)
+		{
+			// We have used less memory this time than the last time, decommit the excess
+			VirtualMemory.DecommitByPtr(RecentHighWaterMark, AmountToFree);
+			NextUncommittedPage = RecentHighWaterMark;
+		}
+	}
+	RecentHighWaterMark = VirtualMemory.GetVirtualPointer();
+}
diff --git a/Engine/Source/Runtime/Core/Public/HAL/LowLevelMemTracker.h b/Engine/Source/Runtime/Core/Public/HAL/LowLevelMemTracker.h
--- a/Engine/Source/Runtime/Core/Public/HAL/LowLevelMemTracker.h
+++ b/Engine/Source/Runtime/Core/Public/HAL/LowLevelMemTracker.h
@@ -119,6 +119,7 @@ enum class ELLMTagSet : uint8
 	Max,	// note: check out FLowLevelMemTracker::ShouldReduceThreads and IsAssetTagForAssets if you add any asset-style tagsets
 };
 
+// Do not add to these macros. Please use the LLM_DECLARE_TAG family of macros below to create new tags.
 #define LLM_ENUM_GENERIC_TAGS(macro) \
 	macro(Untagged,								"Untagged",						NAME_None,													NAME_None,										-1)\
 	macro(Paused,								"Paused",						NAME_None,													NAME_None,										-1)\
@@ -259,7 +260,7 @@ enum class ELLMTag : LLM_TAG_TYPE
 
 	// anything above this value is treated as an FName for a stat section
 };
-static_assert( ELLMTag::GenericTagCount <= ELLMTag::PlatformTagStart, "too many LLM tags defined"); 
+static_assert( ELLMTag::GenericTagCount <= ELLMTag::PlatformTagStart, "too many LLM tags defined -- Instead of adding a new tag and updating the limits, please use the LLM_DECLARE_TAG macros below"); 
 
 static constexpr uint32 LLM_TAG_COUNT = 256;
 static constexpr uint32 LLM_CUSTOM_TAG_START = (int32)ELLMTag::PlatformTagStart;
diff --git a/Engine/Source/Runtime/Core/Public/Memory/VirtualStackAllocator.h b/Engine/Source/Runtime/Core/Public/Memory/VirtualStackAllocator.h
new file mode 100644
--- /dev/null
+++ b/Engine/Source/Runtime/Core/Public/Memory/VirtualStackAllocator.h
@@ -0,0 +1,128 @@
+// Copyright Epic Games, Inc. All Rights Reserved.
+
+#pragma once
+
+#include "CoreMinimal.h"
+#include "HAL/PlatformMemory.h"
+
+class FVirtualStackAllocator;
+
+/**
+ * Implements a stack-style allocator backed by a dedicated block of virtual memory
+ *
+ * FVirtualStackAllocator provides the ability to reserve a block of virtual memory and handles the allocation process
+ *
+ * Freeing memory is handled by FScopedStackAllocatorBookmark which will, upon destruction, free all memory that was allocated after its creation.
+ *
+ * To use the system, first create a FVirtualStackAllocator with a lifetime longer than any allocations you will make from it.
+ * In each scope that you would like be able to bulk free allocations, create an FScopedStackAllocatorBookmark by calling Allocator->CreateScopedBookmark()
+ * Make one or more allocations in any nested scope by calling Allocator->Allocate(Size, Alignment)
+ * When the FScopedStackAllocatorBookmark is destructed, all memory allocated after that will be freed back to the allocator
+ * 
+ * FVirtualStackAllocator provides a one page guard at the end of its reservation to protect against overruns. The total usable memory is therefore
+ * NextMultipleOf(ReqestedStackSize, SystemPageSize) - SystemPageSize
+ */
+
+enum class EVirtualStackAllocatorDecommitMode : uint8
+{
+	// Default mode, does not decommit pages until the allocator is destroyed
+	AllOnDestruction = 0,
+	// Decommits all pages once none are in use
+	AllOnStackEmpty = 1,
+	// Tracks the high water mark and uses it to free excess memory that is not expected to be used again soon
+	// This enables us to quickly release memory consumed by a one-off spikey usage pattern while avoiding frequent page management in the steady state
+	// See DecommitUnusedPages() for details of the heuristic used
+	ExcessOnStackEmpty = 2,
+	NumModes
+};
+
+struct FScopedStackAllocatorBookmark
+{
+public:
+
+    CORE_API ~FScopedStackAllocatorBookmark();
+
+private:
+    friend FVirtualStackAllocator;
+
+    FScopedStackAllocatorBookmark(void* InRestorePointer, FVirtualStackAllocator* Owner)
+		: RestorePointer(InRestorePointer)
+		, Owner(Owner) 
+	{
+	}
+
+    FScopedStackAllocatorBookmark() = delete;
+    FScopedStackAllocatorBookmark(const FScopedStackAllocatorBookmark& Other) = delete;
+    FScopedStackAllocatorBookmark(FScopedStackAllocatorBookmark&& Other) = delete;
+    void operator=(const FScopedStackAllocatorBookmark& Other) = delete;
+    void operator=(FScopedStackAllocatorBookmark&& Other) = delete;
+
+    void* RestorePointer;
+    FVirtualStackAllocator* Owner;
+};
+
+class FVirtualStackAllocator
+{
+public:
+
+    // RequestedStackSize will be rounded up to a whole number of pages and reserved
+    // One page will be earmarked as a guard, so the available memory will be
+    // NextMultipleOf(ReqestedStackSize, SystemPageSize) - SystemPageSize
+    CORE_API FVirtualStackAllocator(size_t RequestedStackSize, EVirtualStackAllocatorDecommitMode Mode);
+    CORE_API ~FVirtualStackAllocator();
+
+	FORCEINLINE FScopedStackAllocatorBookmark CreateScopedBookmark()
+    {
+        // Note that this depends on mandatory Return Value Optimization to ensure the memory is only freed once
+        return FScopedStackAllocatorBookmark(NextAllocationStart, this);
+    }
+
+    CORE_API void* Allocate(size_t Size, size_t Alignment);
+
+    size_t GetAllocatedBytes() const
+    {
+        return (char*)NextAllocationStart - (char*)VirtualMemory.GetVirtualPointer();
+    }
+
+	size_t GetCommittedBytes() const
+	{
+		return (char*)NextUncommittedPage - (char*)VirtualMemory.GetVirtualPointer();
+	}
+
+private:
+	FVirtualStackAllocator(const FVirtualStackAllocator& Other) = delete;
+	FVirtualStackAllocator(FVirtualStackAllocator&& Other) = delete;
+	void operator=(const FVirtualStackAllocator& Other) = delete;
+	void operator=(FVirtualStackAllocator&& Other) = delete;
+
+    friend FScopedStackAllocatorBookmark;
+
+	FVirtualStackAllocator();
+
+	FORCEINLINE void Free(void* RestorePointer)
+	{
+		check(RestorePointer <= NextAllocationStart);
+
+		NextAllocationStart = RestorePointer;
+		if (UNLIKELY(NextAllocationStart == VirtualMemory.GetVirtualPointer() && DecommitMode != EVirtualStackAllocatorDecommitMode::AllOnDestruction))
+		{
+			DecommitUnusedPages();
+		}
+	}
+
+	// Frees unused pages according to the current decommit mode
+	void DecommitUnusedPages();
+
+    FPlatformMemory::FPlatformVirtualMemoryBlock VirtualMemory;
+
+    void* NextUncommittedPage = nullptr;
+    void* NextAllocationStart = nullptr;
+
+    size_t TotalReservationSize;
+    const size_t PageSize;
+
+	EVirtualStackAllocatorDecommitMode DecommitMode = EVirtualStackAllocatorDecommitMode::AllOnDestruction;
+	void* RecentHighWaterMark = nullptr;
+};
+
+
diff --git a/Engine/Source/Runtime/CoreUObject/Private/UObject/ScriptCore.cpp b/Engine/Source/Runtime/CoreUObject/Private/UObject/ScriptCore.cpp
--- a/Engine/Source/Runtime/CoreUObject/Private/UObject/ScriptCore.cpp
+++ b/Engine/Source/Runtime/CoreUObject/Private/UObject/ScriptCore.cpp
@@ -148,6 +148,40 @@ FBlueprintCoreDelegates::FOnToggleScriptProfiler FBlueprintCoreDelegates::OnTogg
 
 PRAGMA_ENABLE_DEPRECATION_WARNINGS
 
+static int32 BlueprintContextVirtualStackAllocatorSize = 8 * 1024 * 1024;
+FAutoConsoleVariableRef CVarBlueprintContextVirtualStackAllocatorStackSize(
+	TEXT("r.FBlueprintContext.VirtualStackAllocatorStackSize"),
+	BlueprintContextVirtualStackAllocatorSize,
+	TEXT("Default size for FBlueprintContext's FVirtualStackAllocator"),
+	ECVF_ReadOnly
+);
+
+static int BlueprintContextVirtualStackAllocatorDecommitMode = (int)EVirtualStackAllocatorDecommitMode::AllOnDestruction;
+FAutoConsoleVariableRef CVarBlueprintContextVirtualStackAllocatorDecommitMode(
+	TEXT("r.FBlueprintContext.VirtualStackAllocator.DecommitMode"),
+	BlueprintContextVirtualStackAllocatorDecommitMode,
+	TEXT("Specifies DecommitMode for FVirtualStackAllocator when used through its ThreadSingleton. Values are from EVirtualStackAllocatorDecommitMode."),
+	ECVF_ReadOnly
+);
+
+FBlueprintContext::FBlueprintContext()
+#if UE_USE_VIRTUAL_STACK_ALLOCATOR_FOR_SCRIPT_VM
+	: VirtualStackAllocator(
+		BlueprintContextVirtualStackAllocatorSize,
+		static_cast<EVirtualStackAllocatorDecommitMode>(BlueprintContextVirtualStackAllocatorDecommitMode))
+#else
+	: VirtualStackAllocator(0, EVirtualStackAllocatorDecommitMode::AllOnDestruction)
+#endif
+{
+	ensure(BlueprintContextVirtualStackAllocatorDecommitMode >= 0 && BlueprintContextVirtualStackAllocatorDecommitMode < (int)EVirtualStackAllocatorDecommitMode::NumModes);
+}
+
+FBlueprintContext* FBlueprintContext::GetThreadSingleton()
+{
+	static thread_local FBlueprintContext ThreadLocalContext;
+	return &ThreadLocalContext;
+}
+
 void FBlueprintCoreDelegates::ThrowScriptException(const UObject* ActiveObject, FFrame& StackFrame, const FBlueprintExceptionInfo& Info)
 {
 	bool bShouldLogWarning = true;
@@ -785,7 +819,9 @@ static uint8 GRegisterCast(ECastToken CastCode, const FNativeFuncPtr& Func)
 void UObject::SkipFunction(FFrame& Stack, RESULT_DECL, UFunction* Function)
 {
 	// allocate temporary memory on the stack for evaluating parameters
-	uint8* Frame = (uint8*)FMemory_Alloca_Aligned(Function->PropertiesSize, Function->GetMinAlignment());
+	UE_VSTACK_MAKE_FRAME(SkipFunctionBookmark, Stack.CachedThreadVirtualStackAllocator);
+
+	uint8* Frame = (uint8*)UE_VSTACK_ALLOC_ALIGNED(Stack.CachedThreadVirtualStackAllocator, Function->PropertiesSize, Function->GetMinAlignment());
 	FMemory::Memzero(Frame, Function->PropertiesSize);
 	for (FProperty* Property = (FProperty*)(Function->ChildProperties); *Stack.Code != EX_EndFunctionParms; Property = (FProperty*)(Property->Next))
 	{
@@ -838,13 +874,14 @@ void ProcessScriptFunction(UObject* Context, UFunction* Function, FFrame& Stack,
 	// the desire to inline calls to our Execution function are the reason for this template function:
 	uint8* FrameMemory = nullptr;
 	FFrame NewStack(Context, Function, nullptr, &Stack, Function->ChildProperties);
+	UE_VSTACK_MAKE_FRAME(ProcessScriptFunctionBookmark, NewStack.CachedThreadVirtualStackAllocator);
 #if USE_UBER_GRAPH_PERSISTENT_FRAME
 	FrameMemory = Function->GetOuterUClassUnchecked()->GetPersistentUberGraphFrame(Context, Function);
 #endif
 	bool bUsePersistentFrame = (nullptr != FrameMemory);
 	if (!bUsePersistentFrame)
 	{
-		FrameMemory = (uint8*)FMemory_Alloca_Aligned(Function->PropertiesSize, Function->GetMinAlignment());
+		FrameMemory = (uint8*)UE_VSTACK_ALLOC_ALIGNED(NewStack.CachedThreadVirtualStackAllocator, Function->PropertiesSize, Function->GetMinAlignment());
 		FMemory::Memzero(FrameMemory, Function->PropertiesSize);
 	}
 
@@ -858,7 +895,7 @@ void ProcessScriptFunction(UObject* Context, UFunction* Function, FFrame& Stack,
 		FProperty* ReturnProperty = Function->GetReturnProperty();
 		if(ensure(ReturnProperty))
 		{
- 			FOutParmRec* RetVal = (FOutParmRec*)FMemory_Alloca(sizeof(FOutParmRec));
+ 			FOutParmRec* RetVal = (FOutParmRec*)UE_VSTACK_ALLOC(NewStack.CachedThreadVirtualStackAllocator, sizeof(FOutParmRec));
 
  			/* Our context should be that we're in a variable assignment to the return value, so ensure that we have a valid property to return to */
  			check(RESULT_PARAM != NULL);
@@ -891,7 +928,7 @@ void ProcessScriptFunction(UObject* Context, UFunction* Function, FFrame& Stack,
 			Stack.Step(Stack.Object, NULL);
 
 			CA_SUPPRESS(6263)
-			FOutParmRec* Out = (FOutParmRec*)FMemory_Alloca(sizeof(FOutParmRec));
+			FOutParmRec* Out = (FOutParmRec*)UE_VSTACK_ALLOC(NewStack.CachedThreadVirtualStackAllocator, sizeof(FOutParmRec));
 			// set the address and property in the out param info
 			// warning: Stack.MostRecentPropertyAddress could be NULL for optional out parameters
 			// if that's the case, we use the extra memory allocated for the out param in the function's locals
@@ -1008,7 +1045,7 @@ void UObject::CallFunction( FFrame& Stack, RESULT_DECL, UFunction* Function )
 		if (FunctionCallspace & FunctionCallspace::Remote)
 		{
 			// Call native networkable function.
-			uint8* Buffer = (uint8*)FMemory_Alloca_Aligned(Function->ParmsSize, Function->GetMinAlignment());
+			uint8* Buffer = (uint8*)UE_VSTACK_ALLOC_ALIGNED(Stack.CachedThreadVirtualStackAllocator, Function->ParmsSize, Function->GetMinAlignment());
 
 			SavedCode = Stack.Code; // Since this is native, we need to rollback the stack if we are calling both remotely and locally
 
@@ -1978,10 +2015,12 @@ void UObject::ProcessEvent( UFunction* Function, void* Parms )
 			Frame = Function->GetOuterUClassUnchecked()->GetPersistentUberGraphFrame(this, Function);
 		}
 #endif
+		FVirtualStackAllocator* VirtualStackAllocator = FBlueprintContext::GetThreadSingleton()->GetVirtualStackAllocator();
+		UE_VSTACK_MAKE_FRAME(ProcessEventBookmark, VirtualStackAllocator);
 		const bool bUsePersistentFrame = (NULL != Frame);
 		if (!bUsePersistentFrame)
 		{
-			Frame = (uint8*)FMemory_Alloca_Aligned(Function->PropertiesSize, Function->GetMinAlignment());
+			Frame = (uint8*)UE_VSTACK_ALLOC_ALIGNED(VirtualStackAllocator, Function->PropertiesSize, Function->GetMinAlignment());
 			// zero the local property memory
 			FMemory::Memzero(Frame + Function->ParmsSize, Function->PropertiesSize - Function->ParmsSize);
 		}
@@ -2008,7 +2047,7 @@ void UObject::ProcessEvent( UFunction* Function, void* Parms )
 				if ( Property->HasAnyPropertyFlags(CPF_OutParm) )
 				{
 					CA_SUPPRESS(6263)
-					FOutParmRec* Out = (FOutParmRec*)FMemory_Alloca(sizeof(FOutParmRec));
+					FOutParmRec* Out = (FOutParmRec*)UE_VSTACK_ALLOC(VirtualStackAllocator, sizeof(FOutParmRec));
 					// set the address and property in the out param info
 					// note that since C++ doesn't support "optional out" we can ignore that here
 					Out->PropAddr = Property->ContainerPtrToValuePtr<uint8>(Parms);
@@ -2610,7 +2649,7 @@ DEFINE_FUNCTION(UObject::execSwitchValue)
 
 	bool bProperCaseUsed = false;
 	{
-		auto LocalTempIndexMem = (uint8*)FMemory_Alloca_Aligned(IndexProperty->GetSize(), IndexProperty->GetMinAlignment());
+		auto LocalTempIndexMem = (uint8*)UE_VSTACK_ALLOC_ALIGNED(Stack.CachedThreadVirtualStackAllocator, IndexProperty->GetSize(), IndexProperty->GetMinAlignment());
 		IndexProperty->InitializeValue(LocalTempIndexMem);
 		for (int32 CaseIndex = 0; CaseIndex < NumCases; ++CaseIndex)
 		{
@@ -2738,20 +2777,20 @@ DEFINE_FUNCTION(UObject::execLet)
 
 		if (LocallyKnownProperty)
 		{
-			LocalTempResult = (uint8*)FMemory_Alloca_Aligned(LocallyKnownProperty->GetSize(), LocallyKnownProperty->GetMinAlignment());
+			LocalTempResult = (uint8*)UE_VSTACK_ALLOC_ALIGNED(Stack.CachedThreadVirtualStackAllocator, LocallyKnownProperty->GetSize(), LocallyKnownProperty->GetMinAlignment());
 			LocallyKnownProperty->InitializeValue(LocalTempResult);
 			Stack.MostRecentPropertyAddress = LocalTempResult;
 		}
 		else
 		{
-			Stack.MostRecentPropertyAddress = (uint8*)FMemory_Alloca(1024);
+			Stack.MostRecentPropertyAddress = (uint8*)UE_VSTACK_ALLOC(Stack.CachedThreadVirtualStackAllocator, 1024);
 			FMemory::Memzero(Stack.MostRecentPropertyAddress, sizeof(FString));
 		}
 	}
 	else if (LocallyKnownProperty && LocallyKnownProperty->HasSetter())
 	{
 		// We can't assign a value directly to a property if it's got a setter or getter
-		LocalTempResult = (uint8*)FMemory_Alloca_Aligned(LocallyKnownProperty->GetSize(), LocallyKnownProperty->GetMinAlignment());
+		LocalTempResult = (uint8*)UE_VSTACK_ALLOC_ALIGNED(Stack.CachedThreadVirtualStackAllocator, LocallyKnownProperty->GetSize(), LocallyKnownProperty->GetMinAlignment());
 		LocallyKnownProperty->InitializeValue(LocalTempResult);
 		PreviousPropertyAddress = Stack.MostRecentPropertyAddress;
 		Stack.MostRecentPropertyAddress = LocalTempResult;
@@ -3171,7 +3210,7 @@ public:
 		const FMulticastScriptDelegate* DelegateAddr = (DelegateProp ? DelegateProp->GetMulticastDelegate(Stack.MostRecentPropertyAddress) : nullptr);
 
 		//Fill parameters
-		uint8* Parameters = (uint8*)FMemory_Alloca_Aligned(SignatureFunction->ParmsSize, SignatureFunction->GetMinAlignment());
+		uint8* Parameters = (uint8*)UE_VSTACK_ALLOC_ALIGNED(Stack.CachedThreadVirtualStackAllocator, SignatureFunction->ParmsSize, SignatureFunction->GetMinAlignment());
 		FMemory::Memzero(Parameters, SignatureFunction->ParmsSize);
 		for (FProperty* Property = (FProperty*)SignatureFunction->ChildProperties; *Stack.Code != EX_EndFunctionParms; Property = (FProperty*)Property->Next)
 		{
diff --git a/Engine/Source/Runtime/CoreUObject/Public/UObject/Script.h b/Engine/Source/Runtime/CoreUObject/Public/UObject/Script.h
--- a/Engine/Source/Runtime/CoreUObject/Public/UObject/Script.h
+++ b/Engine/Source/Runtime/CoreUObject/Public/UObject/Script.h
@@ -11,6 +11,7 @@
 #include "Stats/Stats.h"
 #include "Misc/EnumClassFlags.h"
 #include "Misc/CoreMisc.h"
+#include "Memory/VirtualStackAllocator.h"
 
 struct FFrame;
 
@@ -61,6 +62,23 @@ typedef uint16 CodeSkipSizeType;
 typedef uint32 CodeSkipSizeType;
 #endif
 
+// Context object for data and utilities that may be needed throughout BP execution
+// In the future, it would be preferable for this not to be a thread singleton but to have
+// clearer initialization/termination semantics and per-thread tuning for the stack allocator
+class FBlueprintContext
+{
+public:
+
+	COREUOBJECT_API static FBlueprintContext* GetThreadSingleton();
+
+	FBlueprintContext();
+
+	FVirtualStackAllocator* GetVirtualStackAllocator() { return &VirtualStackAllocator; }
+
+private:
+
+	FVirtualStackAllocator VirtualStackAllocator;
+};
 
 //
 // Blueprint VM intrinsic return value declaration.
diff --git a/Engine/Source/Runtime/CoreUObject/Public/UObject/Stack.h b/Engine/Source/Runtime/CoreUObject/Public/UObject/Stack.h
--- a/Engine/Source/Runtime/CoreUObject/Public/UObject/Stack.h
+++ b/Engine/Source/Runtime/CoreUObject/Public/UObject/Stack.h
@@ -12,9 +12,25 @@
 #include "Templates/Casts.h"
 #include "UObject/UnrealType.h"
 #include "UObject/EnumProperty.h"
+#include "Memory/VirtualStackAllocator.h"
 
 DECLARE_LOG_CATEGORY_EXTERN(LogScriptFrame, Warning, All);
 
+#ifndef UE_USE_VIRTUAL_STACK_ALLOCATOR_FOR_SCRIPT_VM
+#define UE_USE_VIRTUAL_STACK_ALLOCATOR_FOR_SCRIPT_VM 0
+#endif 
+
+#if UE_USE_VIRTUAL_STACK_ALLOCATOR_FOR_SCRIPT_VM
+FORCEINLINE void* UeVstackAllocHelper(FVirtualStackAllocator* Allocator, size_t Size, size_t Align) { return ((Size == 0) ? 0 : Allocator->Allocate(Size, (Align < 16) ? 16 : Align)); }
+#define UE_VSTACK_MAKE_FRAME(Name, VirtualStackAllocatorPtr) FScopedStackAllocatorBookmark Name = VirtualStackAllocatorPtr->CreateScopedBookmark()
+#define UE_VSTACK_ALLOC(VirtualStackAllocatorPtr, Size) UeVstackAllocHelper((VirtualStackAllocatorPtr), (Size), 0) 
+#define UE_VSTACK_ALLOC_ALIGNED(VirtualStackAllocatorPtr, Size, Align) UeVstackAllocHelper((VirtualStackAllocatorPtr), (Size), (Align))
+#else
+#define UE_VSTACK_MAKE_FRAME(Name, VirtualStackAllocatorPtr)
+#define UE_VSTACK_ALLOC(VirtualStackAllocatorPtr, Size) FMemory_Alloca(Size)
+#define UE_VSTACK_ALLOC_ALIGNED(VirtualStackAllocatorPtr, Size, Align) FMemory_Alloca_Aligned(Size, Align)
+#endif
+
 /**
  * Property data type enums.
  * @warning: if values in this enum are modified, you must update:
@@ -111,6 +127,10 @@ public:
 	/** Currently executed native function */
 	UFunction* CurrentNativeFunction;
 
+#if UE_USE_VIRTUAL_STACK_ALLOCATOR_FOR_SCRIPT_VM
+	FVirtualStackAllocator* CachedThreadVirtualStackAllocator;
+#endif
+
 	/** Previous tracking frame */
 	FFrame* PreviousTrackingFrame;
 
@@ -270,6 +290,16 @@ inline FFrame::FFrame( UObject* InObject, UFunction* InNode, void* InLocals, FFr
 		DepthCounter = (InPreviousFrame->DepthCounter < MAX_uint8) ? InPreviousFrame->DepthCounter + 1 : MAX_uint8;
 	}
 #endif
+#if UE_USE_VIRTUAL_STACK_ALLOCATOR_FOR_SCRIPT_VM
+	if (InPreviousFrame == nullptr)
+	{
+		CachedThreadVirtualStackAllocator = FBlueprintContext::GetThreadSingleton()->GetVirtualStackAllocator();
+	}
+	else
+	{
+		CachedThreadVirtualStackAllocator = InPreviousFrame->CachedThreadVirtualStackAllocator;
+	}
+#endif
 }
 
 template<typename T>
