Adding UStruct::GetStructPathName() (part of a bigger change required for a fix)

diff --git a/Engine/Source/Runtime/CoreUObject/Public/UObject/Class.h b/Engine/Source/Runtime/CoreUObject/Public/UObject/Class.h
--- a/Engine/Source/Runtime/CoreUObject/Public/UObject/Class.h
+++ b/Engine/Source/Runtime/CoreUObject/Public/UObject/Class.h
@@ -377,6 +377,16 @@ public:
 	// UField interface.
 	virtual void AddCppProperty(FProperty* Property) override;
 
+	/**
+	 * Returns struct path name as a package + struct FName pair
+	 */
+	FORCEINLINE FTopLevelAssetPath GetStructPathName() const
+	{
+		// Some day this check may actually be relevant
+		checkf(GetOuter() == GetOutermost(), TEXT("Only top level objects are supported by FTopLevelAssetPath. This object is a subobject: \"%s\""), *GetPathName());
+		return FTopLevelAssetPath(GetOuter()->GetFName(), GetFName());
+	}
+
 	/** Searches property link chain for a property with the specified name */
 	FProperty* FindPropertyByName(FName InName) const;
 
@@ -2714,9 +2724,7 @@ public:
 	 */
 	FORCEINLINE FTopLevelAssetPath GetClassPathName() const
 	{
-		// Some day this check may actually be relevant
-		checkf(GetOuter() == GetOutermost(), TEXT("Only top level classes are supported by FTopLevelAssetPath. This class is nested: %s"), *GetPathName());
-		return FTopLevelAssetPath(GetOuter()->GetFName(), GetFName());
+		return GetStructPathName();
 	}
 
 #if WITH_EDITOR
