AsyncLoading2: Keep the full package path for loaded packages instead of just the package name

diff --git a/Engine/Source/Runtime/CoreUObject/Private/Serialization/AsyncLoading2.cpp b/Engine/Source/Runtime/CoreUObject/Private/Serialization/AsyncLoading2.cpp
--- a/Engine/Source/Runtime/CoreUObject/Private/Serialization/AsyncLoading2.cpp
+++ b/Engine/Source/Runtime/CoreUObject/Private/Serialization/AsyncLoading2.cpp
@@ -264,7 +264,7 @@ if ((ELogVerbosity::Type(ELogVerbosity::Verbosity & ELogVerbosity::VerbosityMask
 	UE_LOG(LogStreaming, Verbosity, LogDesc TEXT(": %s (0x%llX) %s (0x%llX) - ") Format, \
 		*(PackageDesc).UPackageName.ToString(), \
 		(PackageDesc).UPackageId.ValueForDebugging(), \
-		*(PackageDesc).PackageNameToLoad.ToString(), \
+		*(PackageDesc).PackagePathToLoad.GetPackageFName().ToString(), \
 		(PackageDesc).PackageIdToLoad.ValueForDebugging(), \
 		##__VA_ARGS__); \
 }
@@ -356,22 +356,18 @@ struct FPackageRequest
 	int32 RequestId = -1;
 	int32 Priority = -1;
 	FName CustomName;
-	FName PackageNameToLoad;
-	FName UPackageName;
-	FPackageId PackageIdToLoad;
+	FPackagePath PackagePath;
 	TUniquePtr<FLoadPackageAsyncDelegate> PackageLoadedDelegate;
 	FPackageRequest* Next = nullptr;
 
-	static FPackageRequest Create(const int32 RequestId, const int32 Priority, FName PackageNameToLoad, FPackageId PackageIdToLoad, FName CustomName, TUniquePtr<FLoadPackageAsyncDelegate> PackageLoadedDelegate)
+	static FPackageRequest Create(int32 RequestId, int32 Priority, const FPackagePath& PackagePath, FName CustomName, TUniquePtr<FLoadPackageAsyncDelegate> PackageLoadedDelegate)
 	{
 		return FPackageRequest
 		{
 			RequestId,
 			Priority,
 			CustomName,
-			PackageNameToLoad,
-			FName(),
-			PackageIdToLoad,
+			PackagePath,
 			MoveTemp(PackageLoadedDelegate),
 			nullptr
 		};
@@ -392,9 +388,9 @@ struct FAsyncPackageDesc2
 	// The name of the UPackage being loaded
 	// Set to none for imported packages up until the package summary has been serialized
 	FName UPackageName;
-	// The name of the package being loaded from disk
-	// Set to none up until the package summary has been serialized
-	FName PackageNameToLoad;
+	// The package path of the package being loaded from disk
+	// Set to none for imported packages up until the package summary has been serialized
+	FPackagePath PackagePathToLoad;
 	// Packages with a a custom name can't be imported
 	bool bCanBeImported;
 
@@ -403,6 +399,7 @@ struct FAsyncPackageDesc2
 		int32 Priority,
 		FName UPackageName,
 		FPackageId PackageIdToLoad,
+		FPackagePath&& PackagePathToLoad,
 		bool bHasCustomName)
 	{
 		return FAsyncPackageDesc2
@@ -412,13 +409,12 @@ struct FAsyncPackageDesc2
 			FPackageId::FromName(UPackageName),
 			PackageIdToLoad,
 			UPackageName,
-			FName(),
+			MoveTemp(PackagePathToLoad),
 			!bHasCustomName
 		};
 	}
 
 	static FAsyncPackageDesc2 FromPackageImport(
-		int32 RequestID,
 		int32 Priority,
 		FPackageId ImportedPackageId,
 		FPackageId PackageIdToLoad,
@@ -426,12 +422,12 @@ struct FAsyncPackageDesc2
 	{
 		return FAsyncPackageDesc2
 		{
-			RequestID,
+			INDEX_NONE,
 			Priority,
 			ImportedPackageId,
 			PackageIdToLoad,
 			UPackageName,
-			FName(),
+			FPackagePath(),
 			true
 		};
 	}
@@ -2715,27 +2711,29 @@ bool FAsyncLoadingThread2::CreateAsyncPackagesFromQueue(FAsyncLoadingThreadState
 			++NumDequeued;
 		
 			FPackageRequest& Request = OptionalRequest.GetValue();
+			FName PackageNameToLoad = Request.PackagePath.GetPackageFName();
 			TCHAR NameBuffer[FName::StringBufferSize];
-			uint32 NameLen = Request.PackageNameToLoad.ToString(NameBuffer);
+			uint32 NameLen = PackageNameToLoad.ToString(NameBuffer);
 			FStringView PackageNameStr = FStringView(NameBuffer, NameLen);
 			if (!FPackageName::IsValidLongPackageName(PackageNameStr))
 			{
 				FString NewPackageNameStr;
 				if (FPackageName::TryConvertFilenameToLongPackageName(FString(PackageNameStr), NewPackageNameStr))
 				{
-					Request.PackageNameToLoad = *NewPackageNameStr;
-					Request.PackageIdToLoad = FPackageId::FromName(Request.PackageNameToLoad);
+					PackageNameToLoad = *NewPackageNameStr;
 				}
 			}
 
-			Request.UPackageName = Request.PackageNameToLoad;
+			FPackageId PackageIdToLoad = FPackageId::FromName(PackageNameToLoad);
+			FName UPackageName = PackageNameToLoad;
 			{
 				FName SourcePackageName;
 				FPackageId RedirectedToPackageId;
-				if (PackageStore->GetPackageRedirectInfo(Request.PackageIdToLoad, SourcePackageName, RedirectedToPackageId))
+				if (PackageStore->GetPackageRedirectInfo(PackageIdToLoad, SourcePackageName, RedirectedToPackageId))
 				{
-					Request.PackageIdToLoad = RedirectedToPackageId;
-					Request.UPackageName = SourcePackageName;
+					PackageIdToLoad = RedirectedToPackageId;
+					Request.PackagePath.Empty(); // We no longer know the path but it will be set again when serializing the package summary
+					UPackageName = SourcePackageName;
 				}
 			}
 
@@ -2752,24 +2750,23 @@ bool FAsyncLoadingThread2::CreateAsyncPackagesFromQueue(FAsyncLoadingThreadState
 						Request.CustomName = *NewPackageNameStr;
 					}
 				}
-				Request.UPackageName = Request.CustomName;
+				UPackageName = Request.CustomName;
 			}
 
-			FPackageId PackageIdToLoad = Request.PackageIdToLoad;
 			FPackageStoreEntry PackageEntry;
 			EPackageStoreEntryStatus PackageStatus = PackageStore->GetPackageStoreEntry(PackageIdToLoad, PackageEntry);
 			if (PackageStatus == EPackageStoreEntryStatus::Missing)
 			{
 				// While there is an active load request for (InName=/Temp/PackageABC_abc, InPackageToLoadFrom=/Game/PackageABC), then allow these requests too:
 				// (InName=/Temp/PackageA_abc, InPackageToLoadFrom=/Temp/PackageABC_abc) and (InName=/Temp/PackageABC_xyz, InPackageToLoadFrom=/Temp/PackageABC_abc)
-				FAsyncPackage2* Package = GetAsyncPackage(Request.PackageIdToLoad);
+				FAsyncPackage2* Package = GetAsyncPackage(PackageIdToLoad);
 				if (Package)
 				{
 					PackageIdToLoad = Package->Desc.PackageIdToLoad;
 					PackageStatus = PackageStore->GetPackageStoreEntry(PackageIdToLoad, PackageEntry);
 				}
 			}
-			FAsyncPackageDesc2 PackageDesc = FAsyncPackageDesc2::FromPackageRequest(Request.RequestId, Request.Priority, Request.UPackageName, PackageIdToLoad, !Request.CustomName.IsNone());
+			FAsyncPackageDesc2 PackageDesc = FAsyncPackageDesc2::FromPackageRequest(Request.RequestId, Request.Priority, UPackageName, PackageIdToLoad, MoveTemp(Request.PackagePath), !Request.CustomName.IsNone());
 			if (PackageStatus == EPackageStoreEntryStatus::Missing)
 			{
 				QueueMissingPackage(PackageDesc, MoveTemp(Request.PackageLoadedDelegate));
@@ -3416,7 +3413,7 @@ void FAsyncPackage2::ImportPackagesRecursiveInner(FIoBatch& IoBatch, IPackageSto
 
 		FAsyncPackage2* ImportedPackage = nullptr;
 		bool bInserted = false;
-		FAsyncPackageDesc2 PackageDesc = FAsyncPackageDesc2::FromPackageImport(INDEX_NONE, Desc.Priority, ImportedPackageId, ImportedPackageIdToLoad, ImportedPackageUPackageName);
+		FAsyncPackageDesc2 PackageDesc = FAsyncPackageDesc2::FromPackageImport(Desc.Priority, ImportedPackageId, ImportedPackageIdToLoad, ImportedPackageUPackageName);
 		if (PackageRef.AreAllPublicExportsLoaded())
 		{
 			ImportedPackage = AsyncLoadingThread.FindAsyncPackage(ImportedPackageId);
@@ -3692,9 +3689,11 @@ EAsyncPackageState::Type FAsyncPackage2::Event_ProcessPackageSummary(FAsyncLoadi
 			}
 		}
 
-		check(Package->Desc.PackageNameToLoad.IsNone());
-		Package->Desc.PackageNameToLoad = Package->HeaderData.PackageName;
-		check(Package->Desc.PackageIdToLoad == FPackageId::FromName(Package->Desc.PackageNameToLoad));
+		check(Package->Desc.PackageIdToLoad == FPackageId::FromName(Package->HeaderData.PackageName));
+		if (Package->Desc.PackagePathToLoad.IsEmpty())
+		{
+			Package->Desc.PackagePathToLoad = FPackagePath::FromPackageNameUnchecked(Package->HeaderData.PackageName);
+		}
 		// Imported packages won't have a UPackage name set unless they were redirected, in which case they will have the source package name
 		if (Package->Desc.UPackageName.IsNone())
 		{
@@ -3703,7 +3702,7 @@ EAsyncPackageState::Type FAsyncPackage2::Event_ProcessPackageSummary(FAsyncLoadi
 		check(Package->Desc.UPackageId == FPackageId::FromName(Package->Desc.UPackageName));
 		Package->CreateUPackage(Package->HeaderData.PackageSummary, Package->HeaderData.VersioningInfo.GetPtrOrNull());
 
-		TRACE_LOADTIME_PACKAGE_SUMMARY(Package, Package->Desc.PackageNameToLoad, Package->HeaderData.PackageSummary->HeaderSize, Package->HeaderData.ImportMap.Num(), Package->HeaderData.ExportMap.Num());
+		TRACE_LOADTIME_PACKAGE_SUMMARY(Package, Package->HeaderData.PackageName, Package->HeaderData.PackageSummary->HeaderSize, Package->HeaderData.ImportMap.Num(), Package->HeaderData.ExportMap.Num());
 	}
 
 	Package->AsyncPackageLoadingState = EAsyncPackageLoadingState2::SetupDependencies;
@@ -3836,9 +3835,10 @@ EAsyncPackageState::Type FAsyncPackage2::Event_ProcessExportBundle(FAsyncLoading
 			Ar.ExternalReadDependencies = &Package->ExternalReadDependencies;
 
 			// Check if the package is instanced
-			if (Package->Desc.UPackageName != Package->Desc.PackageNameToLoad)
+			FName PackageNameToLoad = Package->Desc.PackagePathToLoad.GetPackageFName();
+			if (Package->Desc.UPackageName != PackageNameToLoad)
 			{
-				Package->Desc.PackageNameToLoad.ToString(Ar.InstancedPackageSourceName);
+				PackageNameToLoad.ToString(Ar.InstancedPackageSourceName);
 				Package->Desc.UPackageName.ToString(Ar.InstancedPackageInstanceName);
 			}
 		}
@@ -3944,7 +3944,7 @@ UObject* FAsyncPackage2::EventDrivenIndexToObject(const FAsyncPackageHeaderData&
 		UE_CLOG(!Result, LogStreaming, Warning, TEXT("Missing %s import 0x%llX for package %s"),
 			Index.IsScriptImport() ? TEXT("script") : TEXT("package"),
 			Index.Value(),
-			*Desc.PackageNameToLoad.ToString());
+			*Desc.PackagePathToLoad.GetPackageFName().ToString());
 	}
 #if DO_CHECK
 	if (bCheckSerialized && !IsFullyLoadedObj(Result))
@@ -5835,13 +5835,13 @@ FAsyncPackage2::~FAsyncPackage2()
 	ImportStore.ReleasePackageReference(Desc);
 
 	checkf(RefCount == 0, TEXT("RefCount is not 0 when deleting package %s"),
-		*Desc.PackageNameToLoad.ToString());
+		*Desc.PackagePathToLoad.GetPackageFName().ToString());
 
 	checkf(RequestIDs.Num() == 0, TEXT("MarkRequestIDsAsComplete() has not been called for package %s"),
-		*Desc.PackageNameToLoad.ToString());
+		*Desc.PackagePathToLoad.GetPackageFName().ToString());
 	
 	checkf(ConstructedObjects.Num() == 0, TEXT("ClearConstructedObjects() has not been called for package %s"),
-		*Desc.PackageNameToLoad.ToString());
+		*Desc.PackagePathToLoad.GetPackageFName().ToString());
 
 	FMemory::Free(Data.MemoryBuffer);
 }
@@ -6030,7 +6030,7 @@ void FAsyncPackage2::CreateUPackage(const FZenPackageSummary* PackageSummary, co
 			bCreatedLinkerRoot = true;
 		}
 		LinkerRoot->SetFlags(RF_Public | RF_WasLoaded);
-		LinkerRoot->SetLoadedPath(FPackagePath::FromPackageNameUnchecked(Desc.PackageNameToLoad));
+		LinkerRoot->SetLoadedPath(Desc.PackagePathToLoad);
 		LinkerRoot->SetCanBeImportedFlag(Desc.bCanBeImported);
 		LinkerRoot->SetPackageId(Desc.UPackageId);
 		LinkerRoot->SetPackageFlagsTo(PackageSummary->PackageFlags | PKG_Cooked);
@@ -6221,7 +6221,6 @@ int32 FAsyncLoadingThread2::LoadPackage(const FPackagePath& InPackagePath, FName
 	LazyInitializeFromLoadPackage();
 	
 	const FName PackageNameToLoad = InPackagePath.GetPackageFName();
-	const FPackageId PackageIdToLoad = FPackageId::FromName(PackageNameToLoad);
 	if (InCustomName == PackageNameToLoad)
 	{
 		InCustomName = NAME_None;
@@ -6243,7 +6242,7 @@ int32 FAsyncLoadingThread2::LoadPackage(const FPackagePath& InPackagePath, FName
 		? MakeUnique<FLoadPackageAsyncDelegate>(MoveTemp(InCompletionDelegate))
 		: TUniquePtr<FLoadPackageAsyncDelegate>();
 
-	PackageRequestQueue.Enqueue(FPackageRequest::Create(RequestId, InPackagePriority, PackageNameToLoad, PackageIdToLoad, InCustomName, MoveTemp(CompletionDelegate)));
+	PackageRequestQueue.Enqueue(FPackageRequest::Create(RequestId, InPackagePriority, InPackagePath, InCustomName, MoveTemp(CompletionDelegate)));
 	++QueuedPackagesCounter;
 	++PackagesWithRemainingWorkCounter;
 
