EditorBulkDatas need to create new guids to avoid duplication when they are loading from an instanced load, since instanced loads duplicate data from a source package into a new package of a different name.
#rnx

diff --git a/Engine/Source/Runtime/CoreUObject/Private/Serialization/EditorBulkData.cpp b/Engine/Source/Runtime/CoreUObject/Private/Serialization/EditorBulkData.cpp
--- a/Engine/Source/Runtime/CoreUObject/Private/Serialization/EditorBulkData.cpp
+++ b/Engine/Source/Runtime/CoreUObject/Private/Serialization/EditorBulkData.cpp
@@ -959,10 +959,15 @@ void FEditorBulkData::Serialize(FArchive& Ar, UObject* Owner, bool bAllowRegiste
 		}
 		else if (Ar.IsLoading())
 		{
-			if (Ar.HasAllPortFlags(PPF_Duplicate) && BulkDataId.IsValid())
+			if (BulkDataId.IsValid())
 			{
-				// When duplicating BulkDatas we need to create a new BulkDataId to respect the uniqueness contract
-				BulkDataId = CreateUniqueGuid(BulkDataId, Owner, TEXT("PPF_Duplicate serialization"));
+				FLinkerLoad* LinkerLoad = Cast<FLinkerLoad>(Ar.GetLinker());
+				if (Ar.HasAllPortFlags(PPF_Duplicate) ||
+					(LinkerLoad && LinkerLoad->GetInstancingContext().IsInstanced()))
+				{
+					// When duplicating BulkDatas we need to create a new BulkDataId to respect the uniqueness contract
+					BulkDataId = CreateUniqueGuid(BulkDataId, Owner, TEXT("PPF_Duplicate serialization"));
+				}
 			}
 
 			OffsetInFile = INDEX_NONE;
