Modify FGCObjectInfo::TryResolveObject to use only FNames and no FStrings to search for objects.
Requires base casses to handle recursion for script packages and for finding UClass::StaticClass().
The class of class is class, and the outer of the class UPackage is a UPackage.

diff --git a/Engine/Source/Runtime/CoreUObject/Private/UObject/GCObjectInfo.cpp b/Engine/Source/Runtime/CoreUObject/Private/UObject/GCObjectInfo.cpp
--- a/Engine/Source/Runtime/CoreUObject/Private/UObject/GCObjectInfo.cpp
+++ b/Engine/Source/Runtime/CoreUObject/Private/UObject/GCObjectInfo.cpp
@@ -5,10 +5,34 @@
 =============================================================================*/
 
 #include "UObject/GCObjectInfo.h"
+#include "UObject/Class.h"
+#include "Templates/Casts.h"
 
 UObject* FGCObjectInfo::TryResolveObject()
 {
-	return StaticFindObject(UObject::StaticClass(), nullptr, *GetPathName());
+	if (this == Class)
+	{
+		// Recursion base case, the class of class is class 
+		return UClass::StaticClass();
+	}
+
+	if (Outer == nullptr)
+	{
+		// We are a package, those are the only objects that are allowed to have no outer
+		return StaticFindObjectFast(UObject::StaticClass(), nullptr, Name, false, false);
+	}
+
+	UClass* ResolvedClass = Cast<UClass>(Class->TryResolveObject());
+	if (!ResolvedClass)
+	{
+		return nullptr;
+	}
+	
+	if (UObject* OuterObj = Outer->TryResolveObject())
+	{
+		return StaticFindObjectFast(ResolvedClass, OuterObj, Name, true, false);
+	}
+	return nullptr;
 }
 
 void FGCObjectInfo::GetPathName(FStringBuilderBase& ResultString) const
