New utility functions for ANY_PACKAGE removal

- UClass::TryFindTypeSlow
- UClass::TryFindTypeSlowSafe
- UClass::TryConvertShortTypeNameToPathName
- UClass::TryFixShortClassNameExportPath
- Additional options for StaticFindFirstObject
- Added a new overload for FObjectPropertyBase::GetExportPath

diff --git a/Engine/Source/Runtime/CoreUObject/Private/UObject/Class.cpp b/Engine/Source/Runtime/CoreUObject/Private/UObject/Class.cpp
--- a/Engine/Source/Runtime/CoreUObject/Private/UObject/Class.cpp
+++ b/Engine/Source/Runtime/CoreUObject/Private/UObject/Class.cpp
@@ -52,6 +52,7 @@
 #include "Math/InterpCurvePoint.h"
 #include "UObject/UE5MainStreamObjectVersion.h"
 #include "UObject/TopLevelAssetPath.h"
+#include "HAL/PlatformStackWalk.h"
 
 // This flag enables some expensive class tree validation that is meant to catch mutations of 
 // the class tree outside of SetSuperStruct. It has been disabled because loading blueprints 
@@ -5756,6 +5757,103 @@
 	}
 }
 
+UField* UClass::TryFindTypeSlow(UClass* TypeClass, const FString& InPathNameOrShortName, EFindFirstObjectOptions InOptions)
+{
+	checkf(TypeClass && TypeClass->IsChildOf(UField::StaticClass()), TEXT("TryFindType requires a valid TypeClass parameter which is a subclass of UField. \"%s\" provided."), *GetFullNameSafe(TypeClass));
+
+	UField* FoundType = nullptr;
+	if (!InPathNameOrShortName.IsEmpty())
+	{
+		if (!FPackageName::IsShortPackageName(InPathNameOrShortName))
+		{
+			FoundType = (UField*)StaticFindObject(TypeClass, nullptr, *InPathNameOrShortName, !!(InOptions & EFindFirstObjectOptions::ExactClass));
+		}
+		else
+		{
+			// RobM: I can't decide if this should be an ensure. Maybe in the future?
+			TStringBuilder<1024> Callstack;
+			ANSICHAR Buffer[1024];			
+			uint64 StackFrames[10];
+			uint32 NumStackFrames = FPlatformStackWalk::CaptureStackBackTrace(StackFrames, UE_ARRAY_COUNT(StackFrames));
+			for (uint32 Idx = 0; Idx < NumStackFrames; Idx++)
+			{
+				Buffer[0] = '\0';
+				FPlatformStackWalk::ProgramCounterToHumanReadableString(Idx, StackFrames[Idx], Buffer, sizeof(Buffer));
+				ANSICHAR* TrimmedBuffer = FCStringAnsi::Strstr(Buffer, "!");
+				if (!TrimmedBuffer)
+				{
+					TrimmedBuffer = Buffer;
+				}
+				else
+				{
+					TrimmedBuffer++;
+				}
+				Callstack.Append(TrimmedBuffer);
+				Callstack.Append(TEXT("\r\n"));
+			}
+			
+			FoundType = (UField*)StaticFindFirstObject(TypeClass, *InPathNameOrShortName, InOptions | EFindFirstObjectOptions::EnsureIfAmbiguous | EFindFirstObjectOptions::NativeFirst, ELogVerbosity::Error, TEXT("TryFindType"));
+
+			UE_LOG(LogClass, Warning, TEXT("Short type name \"%s\" provided for TryFindType. Please convert it to a path name (suggested: \"%s\"). Callstack:\r\n%s"), *InPathNameOrShortName, *GetPathNameSafe(FoundType), Callstack.GetData());
+		}
+	}
+	return FoundType;
+}
+
+UField* UClass::TryFindTypeSlowSafe(UClass* TypeClass, const FString& InPathNameOrShortName, EFindFirstObjectOptions InOptions)
+{
+	if (!UE::IsSavingPackage(nullptr) && !IsGarbageCollectingAndLockingUObjectHashTables())
+	{
+		return TryFindTypeSlow(TypeClass, InPathNameOrShortName, InOptions);
+	}
+	return nullptr;
+}
+
+FTopLevelAssetPath UClass::TryConvertShortTypeNameToPathName(UClass* TypeClass, const FString& InShortTypeName, ELogVerbosity::Type AmbiguousMessageVerbosity /*= ELogVerbosity::NoLogging*/, const TCHAR* AmbiguousClassMessage /*= nullptr*/)
+{
+	checkf(TypeClass && TypeClass->IsChildOf(UField::StaticClass()), TEXT("TryConvertShortTypeNameToPathName requires a valid TypeClass parameter which is a subclass of UField. \"%s\" provided."), *GetFullNameSafe(TypeClass));
+
+	if (InShortTypeName.IsEmpty())
+	{
+		return FTopLevelAssetPath();
+	}
+	else if (!FPackageName::IsShortPackageName(InShortTypeName))
+	{
+		return FTopLevelAssetPath(InShortTypeName);
+	}
+	else
+	{
+		FTopLevelAssetPath Result;
+		UField* FoundType = (UField*)StaticFindFirstObject(TypeClass, *InShortTypeName, EFindFirstObjectOptions::EnsureIfAmbiguous | EFindFirstObjectOptions::NativeFirst, AmbiguousMessageVerbosity, AmbiguousClassMessage);
+		if (FoundType)
+		{
+			// UField does not define a GetFieldPathName so do what GetStructPathName does (this is the only place that needs it)
+			checkf(FoundType->GetOuter() == FoundType->GetOutermost(), TEXT("Trying to construct FTopLevelAssetPath for nested type: \"%s\""), *FoundType->GetPathName());
+			Result = FTopLevelAssetPath(FoundType->GetOuter()->GetFName(), FoundType->GetFName());
+		}
+		return Result;
+	}
+}
+
+bool UClass::TryFixShortClassNameExportPath(FString& InOutExportPathToFix, ELogVerbosity::Type AmbiguousMessageVerbosity /*= ELogVerbosity::NoLogging*/, const TCHAR* AmbiguousClassMessage /*= nullptr*/)
+{
+	FString ClassName;
+	FString ObjectPath;
+	if (FPackageName::ParseExportTextPath(*InOutExportPathToFix, &ClassName, &ObjectPath))
+	{
+		if (FPackageName::IsShortPackageName(ClassName))
+		{
+			FTopLevelAssetPath ClassPathName = UClass::TryConvertShortTypeNameToPathName<UClass>(ClassName, AmbiguousMessageVerbosity, AmbiguousClassMessage);
+			if (!ClassPathName.IsNull())
+			{
+				InOutExportPathToFix = FObjectPropertyBase::GetExportPath(ClassPathName, ObjectPath);
+				return true;
+			}
+		}
+	}
+	return false;
+}
+
 #if WITH_EDITOR || HACK_HEADER_GENERATOR
 void UClass::GetHideFunctions(TArray<FString>& OutHideFunctions) const
 {
diff --git a/Engine/Source/Runtime/CoreUObject/Private/UObject/PropertyBaseObject.cpp b/Engine/Source/Runtime/CoreUObject/Private/UObject/PropertyBaseObject.cpp
--- a/Engine/Source/Runtime/CoreUObject/Private/UObject/PropertyBaseObject.cpp
+++ b/Engine/Source/Runtime/CoreUObject/Private/UObject/PropertyBaseObject.cpp
@@ -179,8 +179,13 @@
 	Super::AddReferencedObjects( Collector );
 }
 
-FString FObjectPropertyBase::GetExportPath(const UObject* Object, const UObject* Parent, const UObject* ExportRootScope, const uint32 PortFlags)
+FString FObjectPropertyBase::GetExportPath(FTopLevelAssetPath ClassPathName, const FString& ObjectPathName)
 {
+	return FString::Printf(TEXT("%s.%s'%s'"), *ClassPathName.GetPackageName().ToString(), *ClassPathName.GetAssetName().ToString(), *ObjectPathName);
+}
+
+FString FObjectPropertyBase::GetExportPath(const UObject* Object, const UObject* Parent /*= nullptr*/, const UObject* ExportRootScope /*= nullptr*/, const uint32 PortFlags /*= PPF_None*/)
+{
 	bool bExportFullyQualified = true;
 
 	// When exporting from one package or graph to another package or graph, we don't want to fully qualify the name, as it may refer
@@ -216,13 +221,13 @@
 	// Take the path name relative to the stopping point outermost ptr.
 	// This is so that cases like a component referencing a component in another actor work correctly when pasted
 	FString PathName = Object->GetPathName(StopOuter);
-	int32 ResultIdx = 0;
 	// Object names that contain invalid characters and paths that contain spaces must be put into quotes to be handled correctly
 	if (PortFlags & PPF_Delimited)
 	{
 		PathName = FString::Printf(TEXT("\"%s\""), *PathName.ReplaceQuotesWithEscapedQuotes());
 	}
-	return FString::Printf( TEXT("%s'%s'"), *Object->GetClass()->GetName(), *PathName );
+	FTopLevelAssetPath ClassPathName = Object->GetClass()->GetClassPathName();
+	return GetExportPath(ClassPathName, *PathName);
 }
 
 void FObjectPropertyBase::ExportText_Internal( FString& ValueStr, const void* PropertyValueOrContainer, EPropertyPointerType PropertyPointerType, const void* DefaultValue, UObject* Parent, int32 PortFlags, UObject* ExportRootScope ) const
diff --git a/Engine/Source/Runtime/CoreUObject/Private/UObject/UObjectGlobals.cpp b/Engine/Source/Runtime/CoreUObject/Private/UObject/UObjectGlobals.cpp
--- a/Engine/Source/Runtime/CoreUObject/Private/UObject/UObjectGlobals.cpp
+++ b/Engine/Source/Runtime/CoreUObject/Private/UObject/UObjectGlobals.cpp
@@ -530,12 +530,20 @@
 UObject* StaticFindFirstObject(UClass* Class, const TCHAR* Name, EFindFirstObjectOptions Options /*= EFindFirstObjectOptions::None*/, ELogVerbosity::Type AmbiguousMessageVerbosity /*= ELogVerbosity::Warning*/, const TCHAR* InCurrentOperation /*= nullptr*/)
 {
 	UObject* Result = nullptr;
-	// Don't resolve the name since we're searching in any package
-	FString InName = Name;
-	ConstructorHelpers::StripObjectClass(InName);
-	FName ObjectName(*InName, FNAME_Add);
+	FName ObjectName;
+	if (!FCString::Strchr(Name, TCHAR('\'')))
+	{
+		// Skip unnecessary allocations in StripObjectClass
+		ObjectName = FName(Name, FNAME_Add);
+	}
+	else
+	{
+		FString InName = Name;
+		ConstructorHelpers::StripObjectClass(InName);
+		ObjectName = FName(*InName, FNAME_Add);
+	}
 
-	if (AmbiguousMessageVerbosity == ELogVerbosity::NoLogging)
+	if (AmbiguousMessageVerbosity == ELogVerbosity::NoLogging && !(Options & (EFindFirstObjectOptions::NativeFirst | EFindFirstObjectOptions::EnsureIfAmbiguous)))
 	{
 		Result = StaticFindFirstObjectFastInternal(Class, ObjectName, !!(Options & EFindFirstObjectOptions::ExactClass));
 	}
@@ -544,37 +552,71 @@
 		TArray<UObject*> FoundObjects;
 		if (StaticFindAllObjectsFast(FoundObjects, Class, ObjectName, !!(Options & EFindFirstObjectOptions::ExactClass)))
 		{
-			Result = FoundObjects[0];
-			if (AmbiguousMessageVerbosity != ELogVerbosity::NoLogging && FoundObjects.Num() > 1)
-			{
-				TStringBuilder<256> Message;
-				Message.Append(TEXT("StaticFindFirstObject: Ambiguous object name "));
-				Message.Append(Name);
-				if (InCurrentOperation)
+			if (FoundObjects.Num() > 1)
+			{				
+				if (!!(Options & EFindFirstObjectOptions::NativeFirst))
 				{
-					Message.Append(TEXT(" while "));
-					Message.Append(InCurrentOperation);
-				}
-				Message.Append(TEXT(", will return "));
-				Message.Append(Result->GetPathName());
-				Message.Append(TEXT(" but could also be: "));
-				for (int32 ObjectIndex = 1; ObjectIndex < FoundObjects.Num(); ++ObjectIndex)
-				{
-					Message.Append(FoundObjects[ObjectIndex]->GetPathName());
-					if (ObjectIndex < FoundObjects.Num() - 1)
+					// Prioritize native class instances or native type objects
+					for (UObject* FoundObject : FoundObjects)
 					{
-						Message.Append(TEXT(","));
+						if (FoundObject->IsA<UField>())
+						{
+							// If we were looking for a 'type' (UEnum / UClass / UScriptStruct) object priotize native types
+							if (FoundObject->GetOutermost()->HasAnyPackageFlags(PKG_CompiledIn))
+							{
+								Result = FoundObject;
+								break;
+							}
+						}
+						else if (FoundObject->GetClass()->GetOutermost()->HasAnyPackageFlags(PKG_CompiledIn))
+						{
+							Result = FoundObject;
+							// Don't break yet, maybe we can find a native type (see above) which is usually what we're after anyway
+						}
+					}
+					if (!Result)
+					{
+						Result = FoundObjects[0];
 					}
 				}
 
-				if (AmbiguousMessageVerbosity == ELogVerbosity::Fatal)
+				if (AmbiguousMessageVerbosity != ELogVerbosity::NoLogging || !!(Options & EFindFirstObjectOptions::EnsureIfAmbiguous))
 				{
-					UE_LOG(LogUObjectGlobals, Fatal, TEXT("%s"), *Message.ToString());
+					TStringBuilder<256> Message;
+					Message.Append(TEXT("StaticFindFirstObject: Ambiguous object name "));
+					Message.Append(Name);
+					if (InCurrentOperation)
+					{
+						Message.Append(TEXT(" while "));
+						Message.Append(InCurrentOperation);
+					}
+					Message.Append(TEXT(", will return "));
+					Message.Append(Result->GetPathName());
+					Message.Append(TEXT(" but could also be: "));
+					for (int32 ObjectIndex = 1; ObjectIndex < FoundObjects.Num(); ++ObjectIndex)
+					{
+						Message.Append(FoundObjects[ObjectIndex]->GetPathName());
+						if (ObjectIndex < FoundObjects.Num() - 1)
+						{
+							Message.Append(TEXT(","));
+						}
+					}
+
+					if (AmbiguousMessageVerbosity == ELogVerbosity::Fatal)
+					{
+						UE_LOG(LogUObjectGlobals, Fatal, TEXT("%s"), *Message.ToString());
+					}
+					else if (AmbiguousMessageVerbosity != ELogVerbosity::NoLogging)
+					{
+						GLog->CategorizedLogf(TEXT("LogUObjectGlobals"), AmbiguousMessageVerbosity, TEXT("%s"), *Message.ToString());
+					}
+
+					ensureAlwaysMsgf(!(Options & EFindFirstObjectOptions::EnsureIfAmbiguous), TEXT("%s"), *Message.ToString());
 				}
-				else
-				{
-					GLog->CategorizedLogf(TEXT("LogUObjectGlobals"), AmbiguousMessageVerbosity, TEXT("%s"), *Message.ToString());
-				}
+			}
+			else
+			{
+				Result = FoundObjects[0];
 			}
 		}
 	}
diff --git a/Engine/Source/Runtime/CoreUObject/Public/UObject/Class.h b/Engine/Source/Runtime/CoreUObject/Public/UObject/Class.h
--- a/Engine/Source/Runtime/CoreUObject/Public/UObject/Class.h
+++ b/Engine/Source/Runtime/CoreUObject/Public/UObject/Class.h
@@ -2727,6 +2727,88 @@
 		return GetStructPathName();
 	}
 
+	/**
+	 * Utility function that tries to find a type (class/struct/enum) given a path name or a short name however it will throw a warning (with a callstack) if it's the latter.
+	 * Useful for tracking down and fixing remaining places that use short class names.
+	 * This function is slow and should not be used in performance critical situations
+	 * @param TypeClass Class of the expected type. Must be a subclass of UField
+	 * @param InShortNameOrPathName Class name
+	 * @return Pointer to a type object if successful, nullptr otherwise
+	 */
+	static UField* TryFindTypeSlow(UClass* TypeClass, const FString& InPathNameOrShortName, EFindFirstObjectOptions InOptions = EFindFirstObjectOptions::None);
+
+	/**
+	 * Utility function that tries to find a type (class/struct/enum) given a path name or a short name however it will throw a warning (with a callstack) if it's the latter.
+	 * Useful for tracking down and fixing remaining places that use short class names.
+	 * This function is slow and should not be used in performance critical situations
+	 * @param InShortNameOrPathName Class name
+	 * @return Pointer to a type object if successful, false otherwise
+	 */
+	template <typename T>
+	static T* TryFindTypeSlow(const FString& InShortNameOrPathName, EFindFirstObjectOptions InOptions = EFindFirstObjectOptions::None)
+	{
+		return (T*)TryFindTypeSlow(T::StaticClass(), InShortNameOrPathName, InOptions);
+	}
+
+	/**
+	 * Utility function that tries to find a type (class/struct/enum) given a path name or a short name however it will throw a warning (with a callstack) if it's the latter.
+	 * Useful for tracking down and fixing remaining places that use short class names.
+	 * This function is slow and should not be used in performance critical situations
+	 * This version of TryFindType will not assert when saving package or when collecting garbage
+	 * @param TypeClass Class of the expected type. Must be a subclass of UField
+	 * @param InShortNameOrPathName Class name
+	 * @return Pointer to a type object if successful, nullptr otherwise
+	 */
+	static UField* TryFindTypeSlowSafe(UClass* TypeClass, const FString& InPathNameOrShortName, EFindFirstObjectOptions InOptions = EFindFirstObjectOptions::None);
+
+	/**
+	 * Utility function that tries to find a type (class/struct/enum) given a path name or a short name however it will throw a warning (with a callstack) if it's the latter.
+	 * Useful for tracking down and fixing remaining places that use short class names.
+	 * This function is slow and should not be used in performance critical situations
+	 * This version of TryFindType will not assert when saving package or when collecting garbage
+	 * @param InShortNameOrPathName Class name
+	 * @return Pointer to a type object if successful, nullptr otherwise
+	 */
+	template <typename T>
+	static T* TryFindTypeSlowSafe(const FString& InShortNameOrPathName, EFindFirstObjectOptions InOptions = EFindFirstObjectOptions::None)
+	{
+		return (T*)TryFindTypeSafe(T::StaticClass(), InShortNameOrPathName, InOptions);
+	}
+
+	/**
+	 * Tries to convert short class name to class path name. This will only work if the class exists in memory
+	 * This function is slow and should not be used in performance critical situations
+	 * @param TypeClass Expected Class of the type that needs to be converted (UEnum / UStruct / UClass)
+	 * @param InShortClassName Short class name. If this parameter is already a path name then it's going to be returned as FTopLevelAssetPath
+	 * @param AmbiguousMessageVerbosity Verbosity with which to log a message if class name is ambiguous 
+	 * @param AmbiguousClassMessage Additional message to log when class name is ambiguous (e.g. current operation) 
+	 * @return Class path name if successful. Empty FTopLevelAssetPath otherwise.
+	 */
+	static FTopLevelAssetPath TryConvertShortTypeNameToPathName(UClass* TypeClass, const FString& InShortTypeName, ELogVerbosity::Type AmbiguousMessageVerbosity = ELogVerbosity::NoLogging, const TCHAR* AmbiguousClassMessage = nullptr);
+
+	/**
+	 * Tries to convert short class name to class path name. This will only work if the class exists in memory
+	 * @param InShortClassName Short class name. If this parameter is already a path name then it's going to be returned as FTopLevelAssetPath
+	 * @param AmbiguousMessageVerbosity Verbosity with which to log a message if class name is ambiguous
+	 * @param AmbiguousClassMessage Additional message to log when class name is ambiguous (e.g. current operation)
+	 * @return Class path name if successful. Empty FTopLevelAssetPath otherwise.
+	 */
+	template <typename T>
+	static FTopLevelAssetPath TryConvertShortTypeNameToPathName(const FString& InShortTypeName, ELogVerbosity::Type AmbiguousMessageVerbosity = ELogVerbosity::NoLogging, const TCHAR* AmbiguousClassMessage = nullptr)
+	{
+		return TryConvertShortTypeNameToPathName(T::StaticClass(), InShortTypeName, AmbiguousMessageVerbosity, AmbiguousClassMessage);
+	}
+
+	/**
+	 * Tries to fix an export path containing short class name. Will not modify the input InOutExportPathToFix if a fixup
+	 * was not necessary or unsuccessful
+	 * @param InOutExportPathToFix Export path (Class'/Path/To.Object')
+	 * @param AmbiguousMessageVerbosity Verbosity with which to log a message if class name is ambiguous
+	 * @param AmbiguousClassMessage Additional message to log when class name is ambiguous (e.g. current operation)
+	 * @return True if the short path was successfully fixed. False if the provided export path did not contain short class name or the short path could not be fixed.
+	 */
+	static bool TryFixShortClassNameExportPath(FString& InOutExportPathToFix, ELogVerbosity::Type AmbiguousMessageVerbosity = ELogVerbosity::NoLogging, const TCHAR* AmbiguousClassMessage = nullptr);
+
 #if WITH_EDITOR
 	/** Provides access to C++ type info. */
 	const ICppClassTypeInfo* GetCppTypeInfo() const
diff --git a/Engine/Source/Runtime/CoreUObject/Public/UObject/UObjectGlobals.h b/Engine/Source/Runtime/CoreUObject/Public/UObject/UObjectGlobals.h
--- a/Engine/Source/Runtime/CoreUObject/Public/UObject/UObjectGlobals.h
+++ b/Engine/Source/Runtime/CoreUObject/Public/UObject/UObjectGlobals.h
@@ -372,13 +372,16 @@
 enum class EFindFirstObjectOptions
 {
 	None = 0, // Unused / defaults to Quiet
-	ExactClass = 1 << 1 // Whether to require an exact match with the passed in class
+	ExactClass = 1 << 1, // Whether to require an exact match with the passed in class
+	NativeFirst = 1 << 2, // If multiple results are found, prioritize native classes or native class instances
+	EnsureIfAmbiguous = 1 << 3 // Ensure if multiple results are found
 };
 ENUM_CLASS_FLAGS(EFindFirstObjectOptions);
 
 /**
- * Tries to find the first objects matching the search paramters in memory. This will handle fully qualified paths of the form /path/packagename.object:subobject and resolve references for you.
- *
+ * Tries to find the first object matching the search paramters in memory. This will handle fully qualified paths of the form /path/packagename.object:subobject and resolve references for you.
+ * If multiple objects share the same name the returned object is random and not based on its time of creation unless otherwise specified in Options (see EFindFirstObjectOptions::NativeFirst)
+ * 
  * @param	Class						The to be found object's class
  * @param	Name						The object path to search for an object, relative to InOuter
  * @param	Options						Search options
@@ -392,6 +395,7 @@
 /**
  * Tries to find the first objects matching the search paramters in memory. This will handle fully qualified paths of the form /path/packagename.object:subobject and resolve references for you.
  * This version of StaticFindFirstObject will not assert on GIsSavingPackage or IsGarbageCollecting()
+ * If multiple objects share the same name the returned object is random and not based on its time of creation unless otherwise specified in Options (see EFindFirstObjectOptions::NativeFirst)
  * 
  * @param	Class						The to be found object's class
  * @param	Name						The object path to search for an object, relative to InOuter
diff --git a/Engine/Source/Runtime/CoreUObject/Public/UObject/UnrealType.h b/Engine/Source/Runtime/CoreUObject/Public/UObject/UnrealType.h
--- a/Engine/Source/Runtime/CoreUObject/Public/UObject/UnrealType.h
+++ b/Engine/Source/Runtime/CoreUObject/Public/UObject/UnrealType.h
@@ -2659,9 +2659,24 @@
 	static bool ParseObjectPropertyValue( const FProperty* Property, UObject* OwnerObject, UClass* RequiredMetaClass, uint32 PortFlags, const TCHAR*& Buffer, UObject*& out_ResolvedValue, FUObjectSerializeContext* InSerializeContext = nullptr, bool bAllowAnyPackage = true );
 	static UObject* FindImportedObject( const FProperty* Property, UObject* OwnerObject, UClass* ObjectClass, UClass* RequiredMetaClass, const TCHAR* Text, uint32 PortFlags = 0, FUObjectSerializeContext* InSerializeContext = nullptr, bool bAllowAnyPackage = true );
 	
-	// Returns the qualified export path for a given object, parent, and export root scope
-	static FString GetExportPath(const UObject* Object, const UObject* Parent, const UObject* ExportRootScope, const uint32 PortFlags);
+	/**
+	 * Returns the qualified export path for a given object, parent, and export root scope
+	 * @param Object Object to get the export path for
+	 * @param Parent Outer of the Object used as a root object for generating the Objects path name
+	 * @param ExportRootScope Similar to Parent but used when exporting from one package or graph to another package or graph
+	 * @param PortFlags Property port flags
+	 * @return A string representing the export path of an object, usually in the form of /ClassPackage.ClassName'/Package/Path.Object'
+	 */
+	static FString GetExportPath(const UObject* Object, const UObject* Parent = nullptr, const UObject* ExportRootScope = nullptr, const uint32 PortFlags = PPF_None);
 
+	/**
+	 * Returns the qualified export path given a class path name and object path name
+	 * @param ClassPathName Class path name
+	 * @param ObjectPathName of the Object used as a root object for generating the Objects path name
+	 * @return A string representing the export path of an object in the form of /ClassPackage.ClassName'/Package/Path.Object'
+	 */
+	static FString GetExportPath(FTopLevelAssetPath ClassPathName, const FString& ObjectPathName);
+
 	// Helper method for sharing code with FObjectPtrProperty even though one doesn't inherit from the other
 	static bool StaticIdentical(UObject* A, UObject* B, uint32 PortFlags);
 
