Add -EditorDomain and -NoEditorDomain command line parsing

#rnx

diff --git a/Engine/Source/Runtime/CoreUObject/Private/UObject/PackageResourceManager.cpp b/Engine/Source/Runtime/CoreUObject/Private/UObject/PackageResourceManager.cpp
--- a/Engine/Source/Runtime/CoreUObject/Private/UObject/PackageResourceManager.cpp
+++ b/Engine/Source/Runtime/CoreUObject/Private/UObject/PackageResourceManager.cpp
@@ -270,6 +270,17 @@
 		{
 			return false;
 		}
+
+		if (FParse::Param(FCommandLine::Get(), TEXT("NoEditorDomain")))
+		{
+			return false;
+		}
+
+		if (FParse::Param(FCommandLine::Get(), TEXT("EditorDomain")))
+		{
+			return true;
+		}
+
 		bool bEnabledByConfig = true;
 		GConfig->GetBool(TEXT("EditorDomain"), TEXT("EditorDomainEnabled"), bEnabledByConfig, GEditorIni);
 		if (GConfig->GetBool(TEXT("CookSettings"), TEXT("EditorDomainEnabled"), bEnabledByConfig, GEditorIni))
