Fixing static analysis warning which revealed a bug in StaticFindFirstObject

diff --git a/Engine/Source/Runtime/CoreUObject/Private/UObject/UObjectGlobals.cpp b/Engine/Source/Runtime/CoreUObject/Private/UObject/UObjectGlobals.cpp
--- a/Engine/Source/Runtime/CoreUObject/Private/UObject/UObjectGlobals.cpp
+++ b/Engine/Source/Runtime/CoreUObject/Private/UObject/UObjectGlobals.cpp
@@ -574,10 +574,10 @@ UObject* StaticFindFirstObject(UClass* Class, const TCHAR* Name, EFindFirstObjec
 							// Don't break yet, maybe we can find a native type (see above) which is usually what we're after anyway
 						}
 					}
-					if (!Result)
-					{
-						Result = FoundObjects[0];
-					}
+				}
+				if (!Result)
+				{
+					Result = FoundObjects[0];
 				}
 
 				if (AmbiguousMessageVerbosity != ELogVerbosity::NoLogging || !!(Options & EFindFirstObjectOptions::EnsureIfAmbiguous))
