Fix for copy+pasted description for UStruct::DestroyStruct.

diff --git a/Engine/Source/Runtime/CoreUObject/Public/UObject/Class.h b/Engine/Source/Runtime/CoreUObject/Public/UObject/Class.h
--- a/Engine/Source/Runtime/CoreUObject/Public/UObject/Class.h
+++ b/Engine/Source/Runtime/CoreUObject/Public/UObject/Class.h
@@ -454,7 +454,7 @@ public:
 	virtual void InitializeStruct(void* Dest, int32 ArrayDim = 1) const;
 	
 	/**
-	 * Destroy a struct in memory. This may be done by calling the native destructor and then the constructor or individually reinitializing properties
+	 * Destroy a struct in memory. This may be done by calling the native destructor or individually destroying properties
 	 *
 	 * @param	Dest		Pointer to memory to destory
 	 * @param	ArrayDim	Number of elements in the array
