Add IsGarbage function.

diff --git a/Engine/Source/Runtime/CoreUObject/Public/UObject/GCObjectInfo.h b/Engine/Source/Runtime/CoreUObject/Public/UObject/GCObjectInfo.h
--- a/Engine/Source/Runtime/CoreUObject/Public/UObject/GCObjectInfo.h
+++ b/Engine/Source/Runtime/CoreUObject/Public/UObject/GCObjectInfo.h
@@ -96,6 +96,13 @@
 		return HasAnyInternalFlags(EInternalObjectFlags::Native);
 	}
 
+	bool IsGarbage() const
+	{
+		PRAGMA_DISABLE_DEPRECATION_WARNINGS
+		return HasAnyInternalFlags(EInternalObjectFlags::Garbage | EInternalObjectFlags::PendingKill);
+		PRAGMA_ENABLE_DEPRECATION_WARNINGS
+	}
+
 	int32 GetOwnerIndex() const
 	{
 		return ClusterRootIndex;
