Make ImplementsGetWorld check thread-safe

#rnx

diff --git a/Engine/Source/Runtime/CoreUObject/Private/UObject/Obj.cpp b/Engine/Source/Runtime/CoreUObject/Private/UObject/Obj.cpp
--- a/Engine/Source/Runtime/CoreUObject/Private/UObject/Obj.cpp
+++ b/Engine/Source/Runtime/CoreUObject/Private/UObject/Obj.cpp
@@ -57,6 +57,7 @@
 #include "HAL/FileManager.h"
 #include "HAL/LowLevelMemTracker.h"
 #include "HAL/LowLevelMemStats.h"
+#include "Misc/ScopeRWLock.h"
 #include "Misc/PackageAccessTracking.h"
 #include "Misc/PackageAccessTrackingOps.h"
 
@@ -856,8 +857,8 @@
 #if WITH_ENGINE
 
 #if DO_CHECK
-// Used to check to see if a derived class actually implemented GetWorld() or not, but only on the game thread
-bool bGetWorldOverridden = false;
+// Used to check to see if a derived class actually implemented GetWorld() or not
+thread_local bool bGetWorldOverridden = false;
 #endif
 
 class UWorld* UObject::GetWorld() const
@@ -868,10 +869,7 @@
 	}
 
 #if DO_CHECK
-	if (IsInGameThread())
-	{
-		bGetWorldOverridden = false;
-	}
+	bGetWorldOverridden = false;
 #endif
 	return nullptr;
 }
@@ -879,21 +877,19 @@
 class UWorld* UObject::GetWorldChecked(bool& bSupported) const
 {
 #if DO_CHECK
-	const bool bGameThread = IsInGameThread();
-	if (bGameThread)
-	{
-		bGetWorldOverridden = true;
-	}
+	bGetWorldOverridden = true;
 #endif
 
 	UWorld* World = GetWorld();
 
 #if DO_CHECK
-	if (bGameThread && !bGetWorldOverridden)
+	if (!bGetWorldOverridden)
 	{
+		static FRWLock       ReportedClassesLock;
 		static TSet<UClass*> ReportedClasses;
 
 		UClass* UnsupportedClass = GetClass();
+		FWriteScopeLock ReportedClassesScopeLock(ReportedClassesLock);
 		if (!ReportedClasses.Contains(UnsupportedClass))
 		{
 			UClass* SuperClass = UnsupportedClass->GetSuperClass();
@@ -910,7 +906,7 @@
 		}
 	}
 
-	bSupported = bGameThread ? bGetWorldOverridden : (World != nullptr);
+	bSupported = bGetWorldOverridden;
 	check(World && bSupported);
 #else
 	bSupported = World != nullptr;
@@ -922,7 +918,6 @@
 bool UObject::ImplementsGetWorld() const
 {
 #if DO_CHECK
-	check(IsInGameThread());
 	bGetWorldOverridden = true;
 	GetWorld();
 	return bGetWorldOverridden;
