Fix potential deadlock by moving the ForEachObject sections that locks the hash table outside of another lock

#rnx

diff --git a/Engine/Source/Runtime/CoreUObject/Private/Misc/ExclusiveLoadPackageTimeTracker.cpp b/Engine/Source/Runtime/CoreUObject/Private/Misc/ExclusiveLoadPackageTimeTracker.cpp
--- a/Engine/Source/Runtime/CoreUObject/Private/Misc/ExclusiveLoadPackageTimeTracker.cpp
+++ b/Engine/Source/Runtime/CoreUObject/Private/Misc/ExclusiveLoadPackageTimeTracker.cpp
@@ -55,6 +55,26 @@
 
 void FExclusiveLoadPackageTimeTracker::InternalPopLoadPackage(UPackage* LoadedPackage, UObject* LoadedAsset)
 {
+	// Resolve the classname first to avoid having to take a lock on the object hash table inside another lock
+	// which can cause deadlock.
+	FName ClassName = UnknownAssetName;
+	if (LoadedAsset)
+	{
+		ClassName = LoadedAsset->GetClass()->GetFName();
+	}
+	else if (LoadedPackage)
+	{
+		ForEachObjectWithPackage(LoadedPackage, [&ClassName](UObject* Object)
+		{
+			if (Object && Object->IsAsset())
+			{
+				ClassName = Object->GetClass()->GetFName();
+				return false;
+			}
+			return true;
+		}, false);
+	}
+
 	FScopeLock Lock(&TimesCritical);
 	if (ensure(TimeStack.Num() > 0))
 	{
@@ -62,25 +82,6 @@
 		FLoadTime& Time = TimeStack.Last();
 		Time.ExclusiveTime += (CurrentTime - Time.LastStartTime);
 
-		FName ClassName = UnknownAssetName;
-		if (LoadedAsset)
-		{
-			ClassName = LoadedAsset->GetClass()->GetFName();
-		}
-		else if (LoadedPackage)
-		{
-
-			ForEachObjectWithPackage(LoadedPackage, [&ClassName](UObject* Object)
-			{
-				if (Object && Object->IsAsset())
-				{
-					ClassName = Object->GetClass()->GetFName();
-					return false;
-				}
-				return true;
-			}, false);
-		}
-
 		double InclusiveTime = CurrentTime - Time.OriginalStartTime;
 		FLoadTime* LoadTimePtr = LoadTimes.Find(Time.TimeName);
 		if ( LoadTimePtr )
