Refactor reference chain searching to handle multiple objects at once to reduce time spent scanning all objects and creating/deleting temporary structures.
Add ShortestToGarbage option to deduplicate reference chains by the root reference & the first garbage object encountered.
Do not use the 'Visited' counter for the FGCObject referencer so that we can find multiple chains that terminate there since it handles many disparate types of reference.
When removing duplicate chains, if the chain root is the GC object referencer, use the previous object in the chain as a key instead.
Sort chains by length before deduplicating in shortest/longest mode.
Truncate GCObjectReferencer strings in case they overflow FName's capacity.

diff --git a/Engine/Source/Runtime/CoreUObject/Private/UObject/ReferenceChainSearch.cpp b/Engine/Source/Runtime/CoreUObject/Private/UObject/ReferenceChainSearch.cpp
--- a/Engine/Source/Runtime/CoreUObject/Private/UObject/ReferenceChainSearch.cpp
+++ b/Engine/Source/Runtime/CoreUObject/Private/UObject/ReferenceChainSearch.cpp
@@ -40,9 +40,21 @@ FReferenceChainSearch::FGraphNode* FReferenceChainSearch::FindOrAddNode(UObject*
 	return FindOrAddNode(FGCObjectInfo::FindOrAddInfoHelper(InObjectToFindNodeFor, ObjectToInfoMap));
 }
 
-int32 FReferenceChainSearch::BuildReferenceChains(FGraphNode* TargetNode, TArray<FReferenceChain*>& ProducedChains, int32 ChainDepth, const int32 VisitCounter, EReferenceChainSearchMode SearchMode)
+int32 FReferenceChainSearch::BuildReferenceChainsRecursive(FGraphNode* TargetNode, TArray<FReferenceChain*>& ProducedChains, int32 ChainDepth, const int32 VisitCounter, EReferenceChainSearchMode SearchMode, FGraphNode* GCObjReferencerNode)
 {
 	int32 ProducedChainsCount = 0;
+
+	// Always create a chain for the FGCObject referencer node even if we've found another path to it, because it handles many types of references
+	if (TargetNode == GCObjReferencerNode)
+	{
+		// This is a root so we can construct a chain from this node up to the target node
+		FReferenceChain* Chain = new FReferenceChain(ChainDepth);
+		Chain->InsertNode(TargetNode);
+		ProducedChains.Add(Chain);
+		ProducedChainsCount = 1;
+		return ProducedChainsCount;
+	}
+
 	if (TargetNode->Visited != VisitCounter)
 	{
 		TargetNode->Visited = VisitCounter;
@@ -53,10 +65,10 @@ int32 FReferenceChainSearch::BuildReferenceChains(FGraphNode* TargetNode, TArray
 			for (FGraphNode* ReferencedByNode : TargetNode->ReferencedByObjects)
 			{
 				// For each of the referencers of this node, duplicate the current chain and continue processing
-				if (ReferencedByNode->Visited != VisitCounter)
+				if (ReferencedByNode->Visited != VisitCounter || ReferencedByNode == GCObjReferencerNode)
 				{
 					int32 OldChainsCount = ProducedChains.Num();
-					int32 NewChainsCount = BuildReferenceChains(ReferencedByNode, ProducedChains, ChainDepth + 1, VisitCounter, SearchMode);
+					int32 NewChainsCount = BuildReferenceChainsRecursive(ReferencedByNode, ProducedChains, ChainDepth + 1, VisitCounter, SearchMode, GCObjReferencerNode);
 					// Insert the current node to all chains produced recursively
 					for (int32 NewChainIndex = OldChainsCount; NewChainIndex < (NewChainsCount + OldChainsCount); ++NewChainIndex)
 					{
@@ -80,15 +92,15 @@ int32 FReferenceChainSearch::BuildReferenceChains(FGraphNode* TargetNode, TArray
 	return ProducedChainsCount;
 }
 
-void FReferenceChainSearch::RemoveChainsWithDuplicatedRoots(TArray<FReferenceChain*>& AllChains)
+void FReferenceChainSearch::RemoveChainsWithDuplicatedRoots(TArray<FReferenceChain*>& AllChains, FGraphNode* GCObjReferencerNode)
 {
 	// This is going to be rather slow but it depends on the number of chains which shouldn't be too bad (usually)
 	for (int32 FirstChainIndex = 0; FirstChainIndex < AllChains.Num(); ++FirstChainIndex)
 	{
-		const FGraphNode* RootNode = AllChains[FirstChainIndex]->GetRootNode();
+		const FGraphNode* RootNode = AllChains[FirstChainIndex]->GetRootNode(GCObjReferencerNode);
 		for (int32 SecondChainIndex = AllChains.Num() - 1; SecondChainIndex > FirstChainIndex; --SecondChainIndex)
 		{
-			if (AllChains[SecondChainIndex]->GetRootNode() == RootNode)
+			if (AllChains[SecondChainIndex]->GetRootNode(GCObjReferencerNode) == RootNode)
 			{
 				delete AllChains[SecondChainIndex];
 				AllChains.RemoveAt(SecondChainIndex);
@@ -97,17 +109,72 @@ void FReferenceChainSearch::RemoveChainsWithDuplicatedRoots(TArray<FReferenceCha
 	}
 }
 
-typedef TPair<FReferenceChainSearch::FGraphNode*, FReferenceChainSearch::FGraphNode*> FRootAndReferencerPair;
+void FReferenceChainSearch::RemoveDuplicateGarbageChains(TArray<FReferenceChain*>& AllChains, FGraphNode* GCObjReferencerNode)
+{
+	typedef TPair<FGraphNode*, FGraphNode*> FGarbageChainKey; 
+
+	TMap<FGarbageChainKey, FReferenceChain*> GarbageChains;
+	TArray<FReferenceChain*> NoGarbageChains;
+
+	for (int32 ChainIndex = 0; ChainIndex < AllChains.Num(); ++ChainIndex)
+	{
+		FReferenceChain* Chain = AllChains[ChainIndex];
+		const int32 GarbageIndex = Chain->Nodes.FindLastByPredicate([](FGraphNode* Node){ return Node->ObjectInfo->IsGarbage(); });
+
+		if (GarbageIndex == INDEX_NONE)
+		{
+			NoGarbageChains.Add(Chain);
+			continue;
+		}
+
+		FGraphNode* Garbage = Chain->Nodes[GarbageIndex];
+		FGraphNode* ChainReferencer = Chain->Nodes.Last();
+		if (ChainReferencer == GCObjReferencerNode && Chain->Nodes.Num() > 2)
+		{
+			ChainReferencer = Chain->Nodes[Chain->Nodes.Num() - 2];
+		}
 
-void FReferenceChainSearch::RemoveDuplicatedChains(TArray<FReferenceChain*>& AllChains)
+		FGarbageChainKey Key(Garbage, ChainReferencer);
+		if (FReferenceChain* Existing = GarbageChains.FindRef(Key))
+		{
+			int32 ExistingLen = Existing->Nodes.Num() - Existing->Nodes.Find(Garbage);
+			int32 NewLen = Chain->Nodes.Num() - Chain->Nodes.Find(Garbage);
+			if (NewLen < ExistingLen)
+			{
+				GarbageChains.FindOrAdd(Key, Chain);
+			}
+		}
+		else
+		{
+			GarbageChains.Add(Key, Chain);
+		}
+	}
+	AllChains.Reset();
+	GarbageChains.GenerateValueArray(AllChains);
+	AllChains.Append(NoGarbageChains);
+}
+
+typedef TTuple<FReferenceChainSearch::FGraphNode*, FReferenceChainSearch::FGraphNode*, FReferenceChainSearch::FGraphNode*> FChainDedupKey;
+
+void FReferenceChainSearch::RemoveDuplicatedChains(TArray<FReferenceChain*>& AllChains, FGraphNode* GCObjectReferencerNode)
 {
 	// We consider chains identical if the direct referencer of the target node and the root node are identical
-	TMap<FRootAndReferencerPair, FReferenceChain*> UniqueChains;
-
+	TMap<FChainDedupKey, FReferenceChain*> UniqueChains;
+	
 	for (int32 ChainIndex = 0; ChainIndex < AllChains.Num(); ++ChainIndex)
 	{
 		FReferenceChain* Chain = AllChains[ChainIndex];
-		FRootAndReferencerPair ChainRootAndReferencer(Chain->Nodes[1], Chain->Nodes.Last());
+		FGraphNode* ChainRoot = Chain->Nodes[1];
+
+		// If the chain referencers is the GUObjectReferencer then go one down to find a more unique object to use for deduplication
+		FGraphNode* ChainReferencer = Chain->Nodes.Last();
+		if (ChainReferencer == GCObjectReferencerNode && Chain->Nodes.Num() > 2)
+		{
+			ChainReferencer = Chain->Nodes[Chain->Nodes.Num()-2];
+		}
+
+		FChainDedupKey ChainRootAndReferencer(Chain->TargetNode, ChainRoot, ChainReferencer);
+
 		FReferenceChain** ExistingChain = UniqueChains.Find(ChainRootAndReferencer);
 		if (ExistingChain)
 		{
@@ -130,62 +197,76 @@ void FReferenceChainSearch::RemoveDuplicatedChains(TArray<FReferenceChain*>& All
 	UniqueChains.GenerateValueArray(AllChains);
 }
 
-void FReferenceChainSearch::BuildReferenceChains(FGraphNode* TargetNode, TArray<FReferenceChain*>& Chains, EReferenceChainSearchMode SearchMode)
-{	
-	TArray<FReferenceChain*> AllChains;
-
-	// Recursively construct reference chains	
+void FReferenceChainSearch::BuildReferenceChains(TConstArrayView<FGraphNode*> TargetNodes, TArray<FReferenceChain*>& Chains, EReferenceChainSearchMode SearchMode, FGraphNode* GCObjReferencerNode)
+{
 	int32 VisitCounter = 0;
-	for (FGraphNode* ReferencedByNode : TargetNode->ReferencedByObjects)
+	TArray<FReferenceChain*> AllChains;
+	for (FGraphNode* TargetNode : TargetNodes)
 	{
-		TargetNode->Visited = ++VisitCounter;
-		
-		AllChains.Reset();
-		const int32 MinChainDepth = 2; // The chain will contain at least the TargetNode and the ReferencedByNode
-		BuildReferenceChains(ReferencedByNode, AllChains, MinChainDepth, VisitCounter, SearchMode);
-		for (FReferenceChain* Chain : AllChains)
-		{
-			Chain->InsertNode(TargetNode);
-		}
+		TArray<FReferenceChain*> ThisTargetChains;
 
-		// Filter based on search mode	
-		if (!!(SearchMode & EReferenceChainSearchMode::ExternalOnly))
+		// Recursively construct reference chains
+		for (FGraphNode* ReferencedByNode : TargetNode->ReferencedByObjects)
 		{
-			for (int32 ChainIndex = AllChains.Num() - 1; ChainIndex >= 0; --ChainIndex)
+			TargetNode->Visited = ++VisitCounter;
+		
+			AllChains.Reset();
+			const int32 MinChainDepth = 2; // The chain will contain at least the TargetNode and the ReferencedByNode
+			BuildReferenceChainsRecursive(ReferencedByNode, AllChains, MinChainDepth, VisitCounter, SearchMode, GCObjReferencerNode);
+
+			for (FReferenceChain* Chain : AllChains)
+			{
+				Chain->TargetNode = TargetNode; // Store the node that created this chain in case we later truncate the chain
+				Chain->InsertNode(TargetNode);
+			}
+
+			// Filter based on search mode	
+			if (!!(SearchMode & EReferenceChainSearchMode::ExternalOnly))
 			{
-				FReferenceChain* Chain = AllChains[ChainIndex];	
-				if (!Chain->IsExternal())
+				for (int32 ChainIndex = AllChains.Num() - 1; ChainIndex >= 0; --ChainIndex)
 				{
-					// Discard the chain
-					delete Chain;
-					AllChains.RemoveAtSwap(ChainIndex);
+					FReferenceChain* Chain = AllChains[ChainIndex];	
+					if (!Chain->IsExternal())
+					{
+						// Discard the chain
+						delete Chain;
+						AllChains.RemoveAtSwap(ChainIndex);
+					}
 				}
 			}
+
+			ThisTargetChains.Append(AllChains);
 		}
 
-		Chains.Append(AllChains);
-	}
+		// Sort all chains based on the search criteria
+		if (!(SearchMode & EReferenceChainSearchMode::Longest))
+		{
+			// Sort from the shortest to the longest chain
+			ThisTargetChains.Sort([](const FReferenceChain& LHS, const FReferenceChain& RHS) { return LHS.Num() < RHS.Num(); });
+		}
+		else
+		{
+			// Sort from the longest to the shortest chain
+			ThisTargetChains.Sort([](const FReferenceChain& LHS, const FReferenceChain& RHS) { return LHS.Num() > RHS.Num(); });
+		}
 
-	// Reject duplicates
-	if (!!(SearchMode & (EReferenceChainSearchMode::Longest | EReferenceChainSearchMode::Shortest)))
-	{
-		RemoveChainsWithDuplicatedRoots(Chains);
-	}
-	else
-	{
-		RemoveDuplicatedChains(Chains);
-	}
+		if (!!(SearchMode & (EReferenceChainSearchMode::ShortestToGarbage)))
+		{
+			RemoveDuplicateGarbageChains(ThisTargetChains, GCObjReferencerNode);
+		}
+		// Remove duplicated roots per target object
+		else if (!!(SearchMode & (EReferenceChainSearchMode::Longest | EReferenceChainSearchMode::Shortest)))
+		{
+			RemoveChainsWithDuplicatedRoots(ThisTargetChains, GCObjReferencerNode);
+		}
 
-	// Sort all chains based on the search criteria
-	if (!(SearchMode & EReferenceChainSearchMode::Longest))
-	{
-		// Sort from the shortest to the longest chain
-		Chains.Sort([](const FReferenceChain& LHS, const FReferenceChain& RHS) { return LHS.Num() < RHS.Num(); });
+		Chains.Append(ThisTargetChains);
 	}
-	else
+
+	// If we didn't remove dupe-root chains, remove overall duplicated chains on the entire set of results
+	if (!(SearchMode & (EReferenceChainSearchMode::Longest | EReferenceChainSearchMode::Shortest | EReferenceChainSearchMode::ShortestToGarbage)))
 	{
-		// Sort from the longest to the shortest chain
-		Chains.Sort([](const FReferenceChain& LHS, const FReferenceChain& RHS) { return LHS.Num() > RHS.Num(); });
+		RemoveDuplicatedChains(Chains, GCObjReferencerNode);
 	}
 
 	// Finally, fill extended reference info for the remaining chains
@@ -195,17 +276,20 @@ void FReferenceChainSearch::BuildReferenceChains(FGraphNode* TargetNode, TArray<
 	}
 }
 
-void FReferenceChainSearch::BuildReferenceChainsForDirectReferences(FGraphNode* TargetNode, TArray<FReferenceChain*>& AllChains, EReferenceChainSearchMode SearchMode)
+void FReferenceChainSearch::BuildReferenceChainsForDirectReferences(TConstArrayView<FGraphNode*> TargetNodes, TArray<FReferenceChain*>& AllChains, EReferenceChainSearchMode SearchMode, FGraphNode* GCObjReferencerNode)
 {
-	for (FGraphNode* ReferencedByNode : TargetNode->ReferencedByObjects)
+	for (FGraphNode* TargetNode : TargetNodes) 
 	{
-		if (!(SearchMode & EReferenceChainSearchMode::ExternalOnly) || !ReferencedByNode->ObjectInfo->IsIn(TargetNode->ObjectInfo))
+		for (FGraphNode* ReferencedByNode : TargetNode->ReferencedByObjects)
 		{
-			FReferenceChain* Chain = new FReferenceChain();
-			Chain->AddNode(TargetNode);
-			Chain->AddNode(ReferencedByNode);
-			Chain->FillReferenceInfo();
-			AllChains.Add(Chain);
+			if (!(SearchMode & EReferenceChainSearchMode::ExternalOnly) || !ReferencedByNode->ObjectInfo->IsIn(TargetNode->ObjectInfo))
+			{
+				FReferenceChain* Chain = new FReferenceChain();
+				Chain->AddNode(TargetNode);
+				Chain->AddNode(ReferencedByNode);
+				Chain->FillReferenceInfo();
+				AllChains.Add(Chain);
+			}
 		}
 	}
 }
@@ -471,15 +555,13 @@ bool FReferenceChainSearch::FReferenceChain::IsExternal() const
 */
 class FDirectReferenceProcessor : public FSimpleReferenceProcessorBase
 {	
-	UObject* ObjectToFindReferencesTo;
 	TSet<FReferenceChainSearch::FObjectReferenceInfo>& ReferencedObjects;
 	TMap<UObject*, FGCObjectInfo*>& ObjectToInfoMap;
 
 public:
 
-	FDirectReferenceProcessor(UObject* InObjectToFindReferencesTo, TSet<FReferenceChainSearch::FObjectReferenceInfo>& InReferencedObjects, TMap<UObject*, FGCObjectInfo*>& InObjectToInfoMap)
-		: ObjectToFindReferencesTo(InObjectToFindReferencesTo)
-		, ReferencedObjects(InReferencedObjects)
+	FDirectReferenceProcessor(TSet<FReferenceChainSearch::FObjectReferenceInfo>& InReferencedObjects, TMap<UObject*, FGCObjectInfo*>& InObjectToInfoMap)
+		: ReferencedObjects(InReferencedObjects)
 		, ObjectToInfoMap(InObjectToInfoMap)
 	{
 	}
@@ -509,6 +591,12 @@ public:
 						FString RefName;
 						if (FGCObject::GGCObjectReferencer->GetReferencerName(Object, RefName, true))
 						{
+							// In case FGCObjects misbehave and give us long strings, truncate so we can make them an FName.
+							if (RefName.Len() >= NAME_SIZE)
+							{
+								RefName.LeftInline(NAME_SIZE-1);
+							}
+
 							RefInfo.ReferencerName = FName(*RefName);
 						}
 						else if (ReferencingObject)
@@ -544,14 +632,22 @@ public:
 	}
 };
 
-FReferenceChainSearch::FReferenceChainSearch(UObject* InObjectToFindReferencesTo, EReferenceChainSearchMode Mode /*= EReferenceChainSearchMode::PrintResults*/)
-	: ObjectToFindReferencesTo(InObjectToFindReferencesTo)
+FReferenceChainSearch::FReferenceChainSearch(UObject* InObjectToFindReferencesTo, EReferenceChainSearchMode Mode /*= EReferenceChainSearchMode::PrintResults*/) 
+: FReferenceChainSearch(TConstArrayView<UObject*>(&InObjectToFindReferencesTo, 1), Mode)
+{
+}
+
+FReferenceChainSearch::FReferenceChainSearch(TConstArrayView<UObject*> InObjectsToFindReferencesTo, EReferenceChainSearchMode Mode /*= EReferenceChainSearchMode::PrintResults*/)
+	: ObjectsToFindReferencesTo(InObjectsToFindReferencesTo)
 	, SearchMode(Mode)
 {
 	FSlowHeartBeatScope DisableHangDetection; // This function can be very slow
 
-	check(ObjectToFindReferencesTo);
-	ObjectInfoToFindReferencesTo = FGCObjectInfo::FindOrAddInfoHelper(ObjectToFindReferencesTo, ObjectToInfoMap);
+	for( UObject* Obj : InObjectsToFindReferencesTo )
+	{
+		check(Obj);
+		ObjectInfosToFindReferencesTo.Add(FGCObjectInfo::FindOrAddInfoHelper(Obj, ObjectToInfoMap));
+	}
 
 	// First pass is to find all direct references for each object
 	FindDirectReferencesForObjects();
@@ -577,22 +673,40 @@ FReferenceChainSearch::~FReferenceChainSearch()
 
 void FReferenceChainSearch::PerformSearch()
 {
-	FGraphNode* ObjectNodeToFindReferencesTo = FindOrAddNode(ObjectToFindReferencesTo);
-	check(ObjectNodeToFindReferencesTo);
+	TArray<FGraphNode*> GraphNodesToFind;
+	GraphNodesToFind.Reserve(ObjectsToFindReferencesTo.Num());
+	for (UObject* Obj : ObjectsToFindReferencesTo)
+	{
+		FGraphNode* ObjectNodeToFindReferencesTo = FindOrAddNode(Obj);
+		check(ObjectNodeToFindReferencesTo);
+		GraphNodesToFind.Add(ObjectNodeToFindReferencesTo);
+	}
+
+	FGraphNode* GCObjectReferencerGraphNode = nullptr;
+	if (FGCObject::GGCObjectReferencer)
+	{
+		FGCObjectInfo* ObjectInfo = FGCObjectInfo::FindOrAddInfoHelper(FGCObject::GGCObjectReferencer, ObjectToInfoMap);
+		GCObjectReferencerGraphNode = AllNodes.FindRef(ObjectInfo);
+	}
 
 	// Now it's time to build the reference chain from all of the objects that reference the object to find references to
 	if (!(SearchMode & EReferenceChainSearchMode::Direct))
 	{		
-		BuildReferenceChains(ObjectNodeToFindReferencesTo, ReferenceChains, SearchMode);
+		BuildReferenceChains(GraphNodesToFind, ReferenceChains, SearchMode, GCObjectReferencerGraphNode);
 	}
 	else
 	{
-		BuildReferenceChainsForDirectReferences(ObjectNodeToFindReferencesTo, ReferenceChains, SearchMode);
+		BuildReferenceChainsForDirectReferences(GraphNodesToFind, ReferenceChains, SearchMode, GCObjectReferencerGraphNode);
 	}
 }
 
 #if ENABLE_GC_HISTORY
 void FReferenceChainSearch::PerformSearchFromGCSnapshot(UObject* InObjectToFindReferencesTo, FGCSnapshot& InSnapshot)
+{
+	PerformSearchFromGCSnapshot(TConstArrayView<UObject*>(&InObjectToFindReferencesTo, 1), InSnapshot);
+}
+
+void FReferenceChainSearch::PerformSearchFromGCSnapshot(TConstArrayView<UObject*> InObjectsToFindReferencesTo, FGCSnapshot& InSnapshot)
 {
 	FSlowHeartBeatScope DisableHangDetection; // This function can be very slow
 
@@ -601,8 +715,12 @@ void FReferenceChainSearch::PerformSearchFromGCSnapshot(UObject* InObjectToFindR
 	// Temporarily move the generated object info structs. We don't want to copy everything to minimize memory usage and save a few ms
 	ObjectToInfoMap = MoveTemp(InSnapshot.ObjectToInfoMap);
 
-	ObjectToFindReferencesTo = InObjectToFindReferencesTo;
-	ObjectInfoToFindReferencesTo = FGCObjectInfo::FindOrAddInfoHelper(ObjectToFindReferencesTo, ObjectToInfoMap);
+	ObjectsToFindReferencesTo = InObjectsToFindReferencesTo;
+	ObjectInfosToFindReferencesTo.Reset();
+	for (UObject* Obj : InObjectsToFindReferencesTo)
+	{
+		ObjectInfosToFindReferencesTo.Add(FGCObjectInfo::FindOrAddInfoHelper(Obj, ObjectToInfoMap));
+	}
 	FGCObjectInfo* GCObjectReferencerInfo = FGCObjectInfo::FindOrAddInfoHelper(FGCObject::GGCObjectReferencer, ObjectToInfoMap);
 
 	// We can avoid copying object infos but we need to regenerate direct reference infos
@@ -642,7 +760,7 @@ void FReferenceChainSearch::PerformSearchFromGCSnapshot(UObject* InObjectToFindR
 void FReferenceChainSearch::FindDirectReferencesForObjects()
 {
 	TSet<FObjectReferenceInfo> ReferencedObjects;
-	FDirectReferenceProcessor Processor(ObjectToFindReferencesTo, ReferencedObjects, ObjectToInfoMap);
+	FDirectReferenceProcessor Processor(ReferencedObjects, ObjectToInfoMap);
 	TFastReferenceCollector<
 		FDirectReferenceProcessor, 
 		FDirectReferenceCollector, 
@@ -674,54 +792,76 @@ void FReferenceChainSearch::FindDirectReferencesForObjects()
 	}
 }
 
-void FReferenceChainSearch::PrintResults(bool bDumpAllChains /*= false*/) const
+int32 FReferenceChainSearch::PrintResults(bool bDumpAllChains /*= false*/, UObject* TargetObject /*= nullptr*/) const
 {
-	PrintResults([](FCallbackParams& Params) { return true; }, bDumpAllChains);
+	return PrintResults([](FCallbackParams& Params) { return true; }, bDumpAllChains, TargetObject);
 }
 
-void FReferenceChainSearch::PrintResults(TFunctionRef<bool(FCallbackParams& Params)> ReferenceCallback, bool bDumpAllChains /*= false*/) const
+int32 FReferenceChainSearch::PrintResults(TFunctionRef<bool(FCallbackParams& Params)> ReferenceCallback, bool bDumpAllChains /*= false*/, UObject* TargetObject /*= nullptr*/) const
 {
-	if (ReferenceChains.Num())
-	{
-		FSlowHeartBeatScope DisableHangDetection; // This function can be very slow
+	FSlowHeartBeatScope DisableHangDetection; // This function can be very slow
+
+	const int32 MaxChainsToPrint = 100;
+	int32 NumPrintedChains = 0;
 
-		const int32 MaxChainsToPrint = 100;
-		int32 NumPrintedChains = 0;
+	for (FReferenceChain* Chain : ReferenceChains)
+	{
+		if (TargetObject != nullptr && Chain->TargetNode->ObjectInfo->TryResolveObject() != TargetObject)
+		{
+			continue;
+		}
 
-		for (FReferenceChain* Chain : ReferenceChains)
+		if (bDumpAllChains || NumPrintedChains < MaxChainsToPrint)
 		{
-			if (bDumpAllChains || NumPrintedChains < MaxChainsToPrint)
-			{
-				DumpChain(Chain, ReferenceCallback, *GLog);
-				NumPrintedChains++;
-			}
-			else
-			{
-				UE_LOG(LogReferenceChain, Log, TEXT("Referenced by %d more reference chain(s)."), ReferenceChains.Num() - NumPrintedChains);
-				break;
-			}
+			DumpChain(Chain, ReferenceCallback, *GLog);
+			NumPrintedChains++;
+		}
+		else
+		{
+			UE_LOG(LogReferenceChain, Log, TEXT("Referenced by %d more reference chain(s)."), ReferenceChains.Num() - NumPrintedChains);
+			break;
 		}
 	}
-	else
+
+	if(NumPrintedChains == 0)
 	{
-		check(ObjectInfoToFindReferencesTo);
-		UE_LOG(LogReferenceChain, Log, TEXT("%s%s is not currently reachable."),
-			*GetObjectFlags(ObjectInfoToFindReferencesTo),
-			*ObjectInfoToFindReferencesTo->GetFullName()
-		);
+		if (TargetObject)
+		{
+			FGCObjectInfo ObjInfo(TargetObject);
+			UE_LOG(LogReferenceChain, Log, TEXT("%s%s is not currently reachable."),
+				*GetObjectFlags(&ObjInfo),
+				*ObjInfo.GetFullName()
+			);
+		}
+		else
+		{
+			UE_LOG(LogReferenceChain, Log, TEXT("No target objects are currently reachable."));
+		}
 	}
+
+	return NumPrintedChains;
 }
 
-FString FReferenceChainSearch::GetRootPath() const
+FString FReferenceChainSearch::GetRootPath(UObject* TargetObject /*= nullptr*/) const
 {
-	return GetRootPath([](FCallbackParams& Params) { return true; });
+	return GetRootPath([](FCallbackParams& Params) { return true; }, TargetObject);
 }
 
-FString FReferenceChainSearch::GetRootPath(TFunctionRef<bool(FCallbackParams& Params)> ReferenceCallback) const
+FString FReferenceChainSearch::GetRootPath(TFunctionRef<bool(FCallbackParams& Params)> ReferenceCallback, UObject* TargetObject /*= nullptr*/) const
 {
+	FReferenceChain* Chain = nullptr;
 	if (ReferenceChains.Num())
 	{
-		FReferenceChain* Chain = ReferenceChains[0];
+		Chain = ReferenceChains[0];
+		if( TargetObject )
+		{
+			int32 Index = ReferenceChains.IndexOfByPredicate([=](const FReferenceChain* Chain){return Chain->TargetNode->ObjectInfo->TryResolveObject() == TargetObject; });
+			Chain = Index != INDEX_NONE ? ReferenceChains[Index] : nullptr;
+		}
+	}
+	
+	if (Chain)
+	{
 		FStringOutputDevice OutString;
 		OutString.SetAutoEmitLineTerminator(true);
 		DumpChain(Chain, ReferenceCallback, OutString);
@@ -729,10 +869,18 @@ FString FReferenceChainSearch::GetRootPath(TFunctionRef<bool(FCallbackParams& Pa
 	}
 	else
 	{
-		return FString::Printf(TEXT("%s%s is not currently reachable."),
-			*GetObjectFlags(ObjectInfoToFindReferencesTo),
-			*ObjectInfoToFindReferencesTo->GetFullName()
-		);
+		if (TargetObject)
+		{
+			FGCObjectInfo ObjectInfo(TargetObject);
+			return FString::Printf(TEXT("%s%s is not currently reachable."),
+				*GetObjectFlags(&ObjectInfo),
+				*ObjectInfo.GetFullName()
+			);
+		}
+		else
+		{
+			return TEXT("No target objects are currently reachable.");
+		}
 	}
 }
 
diff --git a/Engine/Source/Runtime/CoreUObject/Public/UObject/ReferenceChainSearch.h b/Engine/Source/Runtime/CoreUObject/Public/UObject/ReferenceChainSearch.h
--- a/Engine/Source/Runtime/CoreUObject/Public/UObject/ReferenceChainSearch.h
+++ b/Engine/Source/Runtime/CoreUObject/Public/UObject/ReferenceChainSearch.h
@@ -22,6 +22,8 @@ enum class EReferenceChainSearchMode
 	Direct = 1 << 3,
 	// Returns complete chains. (Ignoring non GC objects)
 	FullChain = 1 << 4,
+	// Returns the shortest path to a garbage object from which the target object is reachable
+	ShortestToGarbage = 1 << 5,
 
 	// Print results
 	PrintResults = 1 << 16,
@@ -159,6 +161,9 @@ public:
 	{
 		friend class FReferenceChainSearch;
 
+		/** The target nodes that caused this chain to be created. Usually this will be Nodes[0] unless the chain was truncated by request. */
+		FGraphNode* TargetNode;
+
 		/** Nodes in this reference chain. The first node is the target object and the last one is a root object */
 		TArray<FGraphNode*> Nodes;
 		/** Reference information for Nodes */
@@ -188,6 +193,15 @@ public:
 		{
 			return Nodes[NodeIndex];
 		}
+		FGraphNode* GetRootNode(FGraphNode* Exclude) const
+		{
+			if (Nodes.Last() == Exclude && Nodes.Num() > 2)
+			{
+				return Nodes[Nodes.Num()-2];
+			}
+
+			return Nodes.Last();
+		}
 		FGraphNode* GetRootNode() const
 		{
 			return Nodes.Last();
@@ -234,6 +248,7 @@ public:
 
 	/** Constructs a new search engine and finds references to the specified object */
 	COREUOBJECT_API explicit FReferenceChainSearch(UObject* InObjectToFindReferencesTo, EReferenceChainSearchMode Mode = EReferenceChainSearchMode::PrintResults);
+	COREUOBJECT_API explicit FReferenceChainSearch(TConstArrayView<UObject*> InObjectsToFindReferencesTo, EReferenceChainSearchMode Mode = EReferenceChainSearchMode::PrintResults);
 
 	/** Constructs a new search engine but does not find references to any objects until one of the PerformSearch*() functions is called */
 	COREUOBJECT_API explicit FReferenceChainSearch(EReferenceChainSearchMode Mode);
@@ -244,29 +259,32 @@ public:
 #if ENABLE_GC_HISTORY
 	/** Searches for references in a previous GC run snapshot temporarily acquiring its object info */
 	COREUOBJECT_API void PerformSearchFromGCSnapshot(UObject* InObjectToFindReferencesTo, FGCSnapshot& InSnapshot);
+	COREUOBJECT_API void PerformSearchFromGCSnapshot(TConstArrayView<UObject*> InObjectsToFindReferencesTo, FGCSnapshot& InSnapshot);
 #endif //ENABLE_GC_HISTORY
 
 	/**
 	 * Dumps results to log
 	 * @param bDumpAllChains - if set to false, the output will be trimmed to the first 100 reference chains
+	 * @returns The number of results printed.
 	 */
-	COREUOBJECT_API void PrintResults(bool bDumpAllChains = false) const;
+	COREUOBJECT_API int32 PrintResults(bool bDumpAllChains = false, UObject* TargetObject=nullptr) const;
 
 	/**
 	 * Dumps results to log
 	 * @param ReferenceCallback - function called when processing each reference, if true is returned the next reference will be processed otherwise printing will be aborted
 	 * @param bDumpAllChains - if set to false, the output will be trimmed to the first 100 reference chains
+	 * @returns The number of results printed.
 	 */
-	COREUOBJECT_API void PrintResults(TFunctionRef<bool(FCallbackParams& Params)> ReferenceCallback, bool bDumpAllChains = false) const;
+	COREUOBJECT_API int32 PrintResults(TFunctionRef<bool(FCallbackParams& Params)> ReferenceCallback, bool bDumpAllChains = false, UObject* TargetObject = nullptr) const;
 
 	/** Returns a string with a short report explaining the root path, will contain newlines */
-	COREUOBJECT_API FString GetRootPath() const;
+	COREUOBJECT_API FString GetRootPath(UObject* TargetObject = nullptr) const;
 
 	/** 
 	 * Returns a string with a short report explaining the root path, will contain newlines 
 	 * @param ReferenceCallback - function called when processing each reference, if true is returned the next reference will be processed otherwise printing will be aborted
 	 */
-	COREUOBJECT_API FString GetRootPath(TFunctionRef<bool(FCallbackParams& Params)> ReferenceCallback) const;
+	COREUOBJECT_API FString GetRootPath(TFunctionRef<bool(FCallbackParams& Params)> ReferenceCallback, UObject* TargetObject = nullptr) const;
 
 	/** Returns all reference chains */
 	COREUOBJECT_API const TArray<FReferenceChain*>& GetReferenceChains() const
@@ -276,15 +294,16 @@ public:
 
 private:
 
-	/** The object we're going to look for references to */
-	UObject* ObjectToFindReferencesTo = nullptr;
-	FGCObjectInfo* ObjectInfoToFindReferencesTo = nullptr;
+	/** The objects we're going to look for references to */
+	TArray<UObject*> ObjectsToFindReferencesTo;
+	TArray<FGCObjectInfo*> ObjectInfosToFindReferencesTo;
 
 	/** Search mode and options */
 	EReferenceChainSearchMode SearchMode = EReferenceChainSearchMode::Default;
 
-	/** All reference chain found during the search */
+	/** All reference chain found during the search, one per entry in ObjectsToFindReferencesTo */
 	TArray<FReferenceChain*> ReferenceChains;
+
 	/** All nodes created during the search */
 	TMap<FGCObjectInfo*, FGraphNode*> AllNodes;
 	/** Maps UObject pointers to object info structs */
@@ -304,18 +323,22 @@ private:
 	FGraphNode* FindOrAddNode(FGCObjectInfo* InObjectInfo);
 
 	/** Builds reference chains */
-	static int32 BuildReferenceChains(FGraphNode* TargetNode, TArray<FReferenceChain*>& ProducedChains, int32 ChainDepth, const int32 VisitCounter, EReferenceChainSearchMode SearchMode);
+	static int32 BuildReferenceChainsRecursive(FGraphNode* TargetNodes, TArray<FReferenceChain*>& ProducedChains, int32 ChainDepth, const int32 VisitCounter, EReferenceChainSearchMode SearchMode, FGraphNode* GCObjReferencerNode);
 	/** Builds reference chains */
-	static void BuildReferenceChains(FGraphNode* TargetNode, TArray<FReferenceChain*>& AllChains, EReferenceChainSearchMode SearchMode);
+	static void BuildReferenceChains(TConstArrayView<FGraphNode*> TargetNodes, TArray<FReferenceChain*>& AllChains, EReferenceChainSearchMode SearchMode, FGraphNode* GCObjReferencerNode);
 	/** Builds reference chains for direct references only */
-	static void BuildReferenceChainsForDirectReferences(FGraphNode* TargetNode, TArray<FReferenceChain*>& AllChains, EReferenceChainSearchMode SearchMode);
+	static void BuildReferenceChainsForDirectReferences(TConstArrayView<FGraphNode*> TargetNodes, TArray<FReferenceChain*>& AllChains, EReferenceChainSearchMode SearchMode, FGraphNode* GCObjReferencerNode);
 	/** Leaves only chains with unique root objects */
-	static void RemoveChainsWithDuplicatedRoots(TArray<FReferenceChain*>& AllChains);
+	static void RemoveChainsWithDuplicatedRoots(TArray<FReferenceChain*>& AllChains, FGraphNode* GCObjReferencerNode);
 	/** Leaves only unique chains */
-	static void RemoveDuplicatedChains(TArray<FReferenceChain*>& AllChains);
+	static void RemoveDuplicatedChains(TArray<FReferenceChain*>& AllChains, FGraphNode* GCObjReferencerNode);
+	/** Deduplicates chains that have the same root and first garbage object encountered */
+	static void RemoveDuplicateGarbageChains(TArray<FReferenceChain*>& AllChains, FGraphNode* GCObjReferencerNode);
 
 	/** Returns a string with all flags (we care about) set on an object */
 	static FString GetObjectFlags(FGCObjectInfo* InObject);
+	static FString GetObjectFlags(EInternalObjectFlags InternalFlags, EObjectFlags Flags);
+
 	/** Dumps a reference chain to log */
 	static void DumpChain(FReferenceChainSearch::FReferenceChain* Chain, TFunctionRef<bool(FCallbackParams& Params)> ReferenceCallback, FOutputDevice& Out);
 };
diff --git a/Engine/Source/Runtime/Engine/Classes/Engine/Engine.h b/Engine/Source/Runtime/Engine/Classes/Engine/Engine.h
--- a/Engine/Source/Runtime/Engine/Classes/Engine/Engine.h
+++ b/Engine/Source/Runtime/Engine/Classes/Engine/Engine.h
@@ -3193,7 +3193,8 @@
 	 * @param ObjectToFindReferencesTo World or its package (or any object from the world package that should've been destroyed)
 	 * @param Options Determines how the stale references messages should be logged
 	 */
-	static void FindAndPrintStaleReferencesToObject(UObject* ObjectToFindReferencesTo, EPrintStaleReferencesOptions Options);
+	static FString FindAndPrintStaleReferencesToObject(UObject* ObjectToFindReferencesTo, EPrintStaleReferencesOptions Options);
+	static TArray<FString> FindAndPrintStaleReferencesToObjects(TConstArrayView<UObject*> ObjectsToFindReferencesTo, EPrintStaleReferencesOptions Options);
 
 	FWorldContext& CreateNewWorldContext(EWorldType::Type WorldType);
 
diff --git a/Engine/Source/Runtime/Engine/Private/UnrealEngine.cpp b/Engine/Source/Runtime/Engine/Private/UnrealEngine.cpp
--- a/Engine/Source/Runtime/Engine/Private/UnrealEngine.cpp
+++ b/Engine/Source/Runtime/Engine/Private/UnrealEngine.cpp
@@ -15214,6 +15214,7 @@
 	TRACE_CPUPROFILER_EVENT_SCOPE(VerifyLoadMapWorldCleanup);
 	// All worlds at this point should be the CurrentWorld of some context, preview worlds, or streaming level
 	// worlds that are owned by the CurrentWorld of some context.
+	TArray<UObject*> LeakedWorlds;
 	for( TObjectIterator<UWorld> It; It; ++It )
 	{
 		UWorld* World = *It;
@@ -15222,18 +15223,29 @@
 		{
 			if ((World->PersistentLevel == nullptr || !WorldHasValidContext(World->PersistentLevel->OwningWorld)) && !IsWorldDuplicate(World))
 			{
-				UE_LOG(LogLoad, Error, TEXT("Previously active world %s not cleaned up by garbage collection!"), *World->GetPathName());
-				UE_LOG(LogLoad, Error, TEXT("Once a world has become active, it cannot be reused and must be destroyed and reloaded. World referenced by:"));
+				LeakedWorlds.Add(World);
+			}
+		}
+	}
 
-				FindAndPrintStaleReferencesToObject(World, UObjectBaseUtility::IsPendingKillEnabled() ? EPrintStaleReferencesOptions::Fatal : (EPrintStaleReferencesOptions::Error | EPrintStaleReferencesOptions::Ensure));
-			}
+	if (LeakedWorlds.Num())
+	{
+		UE_LOG(LogLoad, Error, TEXT("Some previously active worlds were not cleaned up by garbage collection!"));
+		UE_LOG(LogLoad, Error, TEXT("Once a world has become active, it cannot be reused and must be destroyed and reloaded. Dumping reference chains:"));
+
+		FindAndPrintStaleReferencesToObjects(LeakedWorlds, UObjectBaseUtility::IsPendingKillEnabled() ? EPrintStaleReferencesOptions::Fatal : (EPrintStaleReferencesOptions::Error | EPrintStaleReferencesOptions::Ensure));
+
+		if (!UObjectBaseUtility::IsPendingKillEnabled())
+		{
+			UE_LOG(LogTemp, Fatal, TEXT("Shutting down because of world leaks."));
 		}
 	}
 }
 
-static void PrintStaleReferenceChainsAndFindReferencingObjects(UObject* ObjectToFindReferencesTo, FReferenceChainSearch& RefChainSearch, FGCObjectInfo*& OutGarbageObject, FGCObjectInfo*& OutReferencingObject)
+static bool PrintStaleReferenceChainsAndFindReferencingObjects(UObject* ObjectToFindReferencesTo, FReferenceChainSearch& RefChainSearch, FGCObjectInfo*& OutGarbageObject, FGCObjectInfo*& OutReferencingObject)
 {
-	RefChainSearch.PrintResults([&ObjectToFindReferencesTo, &OutGarbageObject, &OutReferencingObject](FReferenceChainSearch::FCallbackParams& Params)
+	UE_LOG(LogLoad, Error, TEXT("Printing reference chains leading to %s: "), *ObjectToFindReferencesTo->GetFullName());
+	return RefChainSearch.PrintResults([&ObjectToFindReferencesTo, &OutGarbageObject, &OutReferencingObject](FReferenceChainSearch::FCallbackParams& Params)
 		{
 			check(Params.Object);
 			if (!Params.Object->IsValid())
@@ -15258,7 +15270,7 @@
 				}
 				return true;
 			}
-		}, false);
+		}, false, ObjectToFindReferencesTo) != 0;
 }
 
 static FString GetPathToStaleObjectReferencer(UObject* ObjectToFindReferencesTo, FReferenceChainSearch& RefChainSearch)
@@ -15310,95 +15322,115 @@
 	FindAndPrintStaleReferencesToObject(ObjectToFindReferencesTo, Options);
 }
 
-void UEngine::FindAndPrintStaleReferencesToObject(UObject* ObjectToFindReferencesTo, EPrintStaleReferencesOptions Options)
+FString UEngine::FindAndPrintStaleReferencesToObject(UObject* ObjectToFindReferencesTo, EPrintStaleReferencesOptions Options)
 {
-	checkf(ObjectToFindReferencesTo, TEXT("ObjectToFindReferencesTo cannot be null"));
+	return FindAndPrintStaleReferencesToObjects(TConstArrayView<UObject*>(&ObjectToFindReferencesTo, 1), Options)[0];
+}
+
+TArray<FString> UEngine::FindAndPrintStaleReferencesToObjects(TConstArrayView<UObject*> ObjectsToFindReferencesTo, EPrintStaleReferencesOptions Options)
+{
+	TArray<FString> Paths; 
+	checkf(ObjectsToFindReferencesTo.Num() > 0, TEXT("ObjectsToFindReferencesTo cannot be empty"));
+	for (UObject* Obj : ObjectsToFindReferencesTo)
+	{
+		checkf(Obj, TEXT("Object to find references to cannot be null"));
+	}
+
 	const ELogVerbosity::Type OptionsToVerbosityMapping[] = { ELogVerbosity::NoLogging, ELogVerbosity::Log, ELogVerbosity::Display, ELogVerbosity::Warning, ELogVerbosity::Error, ELogVerbosity::Fatal };
 	const int32 VerbosityIndex = int32(Options & EPrintStaleReferencesOptions::VerbosityMask);
 	checkf(VerbosityIndex < UE_ARRAY_COUNT(OptionsToVerbosityMapping), TEXT("Unsupported verbosity option provided: %d"), VerbosityIndex);
 	ELogVerbosity::Type Verbosity = OptionsToVerbosityMapping[VerbosityIndex];
 
-	UE_LOG(LogLoad, Log, TEXT("Looking for existing references to %s..."), *ObjectToFindReferencesTo->GetFullName());
+	UE_LOG(LogLoad, Log, TEXT("Beginning reference chain search..."));
+	for (UObject* Obj : ObjectsToFindReferencesTo)
+	{
+		UE_LOG(LogLoad, Log, TEXT(" - %s"), *Obj->GetFullName());
+	}
 
-	bool bReferenceChainFound = false;
-	FGCObjectInfo* OutGarbageObject = nullptr;
-	FGCObjectInfo* OutReferencingObject = nullptr;
-	FReferenceChainSearch RefChainSearch(ObjectToFindReferencesTo, EReferenceChainSearchMode::Default);
+	FReferenceChainSearch RefChainSearch(ObjectsToFindReferencesTo, EReferenceChainSearchMode::ShortestToGarbage);
 #if ENABLE_GC_HISTORY
-	FReferenceChainSearch HistorySearch(EReferenceChainSearchMode::Default); // we may not need it but it has to live in this scope so that FGCObjectInfos are being kept alive
+	FReferenceChainSearch HistorySearch(EReferenceChainSearchMode::ShortestToGarbage); // we may not need it but it has to live in this scope so that FGCObjectInfos are being kept alive
 #endif
+	
+	for (UObject* ObjectToFindReferencesTo : ObjectsToFindReferencesTo)
+	{
+		FGCObjectInfo* OutGarbageObject = nullptr;
+		FGCObjectInfo* OutReferencingObject = nullptr;
 
-	PrintStaleReferenceChainsAndFindReferencingObjects(ObjectToFindReferencesTo, RefChainSearch, OutGarbageObject, OutReferencingObject);
+		bool bReferenceChainFound = PrintStaleReferenceChainsAndFindReferencingObjects(ObjectToFindReferencesTo, RefChainSearch, OutGarbageObject, OutReferencingObject);
 
-	FString PathToCulprit;
-	FString GarbageErrorMessage;
-	if (RefChainSearch.GetReferenceChains().Num())
-	{
-		bReferenceChainFound = true;
-		PathToCulprit = GetPathToStaleObjectReferencer(ObjectToFindReferencesTo, RefChainSearch);
-		FString GarbageObjectName = OutGarbageObject ? OutGarbageObject->GetFullName() : ObjectToFindReferencesTo->GetFullName();
-		checkf(OutReferencingObject, TEXT("No object referencing %s found even though we have a valid reference chain"), *ObjectToFindReferencesTo->GetPathName());
-		GarbageErrorMessage = FString::Printf(TEXT("Object %s is being referenced by %s"), *GarbageObjectName, *OutReferencingObject->GetFullName());
-	}
+		FString PathToCulprit;
+		FString GarbageErrorMessage;
+		if (bReferenceChainFound)
+		{
+			PathToCulprit = GetPathToStaleObjectReferencer(ObjectToFindReferencesTo, RefChainSearch);
+			FString GarbageObjectName = OutGarbageObject ? OutGarbageObject->GetFullName() : ObjectToFindReferencesTo->GetFullName();
+			checkf(OutReferencingObject, TEXT("No object referencing %s found even though we have a valid reference chain"), *ObjectToFindReferencesTo->GetPathName());
+			GarbageErrorMessage = FString::Printf(TEXT("Object %s is being referenced by %s"), *GarbageObjectName, *OutReferencingObject->GetFullName());
+		}
 #if ENABLE_GC_HISTORY
-	else
-	{
-		FGCSnapshot* LastGCSnapshot = FGCHistory::Get().GetLastSnapshot();
-		if (LastGCSnapshot)
+		else
 		{
-			UE_LOG(LogLoad, Log, TEXT("Looking for references to %s in the last GC run history..."), *ObjectToFindReferencesTo->GetFullName());
+			FGCSnapshot* LastGCSnapshot = FGCHistory::Get().GetLastSnapshot();
+			if (LastGCSnapshot)
+			{
+				UE_LOG(LogLoad, Log, TEXT("Looking for references to %s in the last GC run history..."), *ObjectToFindReferencesTo->GetFullName());
 
-			HistorySearch.PerformSearchFromGCSnapshot(ObjectToFindReferencesTo, *LastGCSnapshot);
+				HistorySearch.PerformSearchFromGCSnapshot(ObjectsToFindReferencesTo, *LastGCSnapshot);
 
-			OutGarbageObject = nullptr;
-			OutReferencingObject = nullptr;
-			PrintStaleReferenceChainsAndFindReferencingObjects(ObjectToFindReferencesTo, HistorySearch, OutGarbageObject, OutReferencingObject);
-			if (HistorySearch.GetReferenceChains().Num())
-			{
-				bReferenceChainFound = true;
-				PathToCulprit = GetPathToStaleObjectReferencer(ObjectToFindReferencesTo, RefChainSearch);
-				FString GarbageObjectName = OutGarbageObject ? OutGarbageObject->GetFullName() : ObjectToFindReferencesTo->GetFullName();
-				checkf(OutReferencingObject, TEXT("No object referencing %s found even though we have a valid reference chain"), *ObjectToFindReferencesTo->GetPathName());
-				GarbageErrorMessage = FString::Printf(TEXT("Garbage object %s was previously being referenced by %s"), *GarbageObjectName, *OutReferencingObject->GetFullName());
+				OutGarbageObject = nullptr;
+				OutReferencingObject = nullptr;
+				bReferenceChainFound = PrintStaleReferenceChainsAndFindReferencingObjects(ObjectToFindReferencesTo, HistorySearch, OutGarbageObject, OutReferencingObject);
+				if (bReferenceChainFound)
+				{
+					bReferenceChainFound = true;
+					PathToCulprit = GetPathToStaleObjectReferencer(ObjectToFindReferencesTo, HistorySearch);
+					FString GarbageObjectName = OutGarbageObject ? OutGarbageObject->GetFullName() : ObjectToFindReferencesTo->GetFullName();
+					checkf(OutReferencingObject, TEXT("No object referencing %s found even though we have a valid reference chain"), *ObjectToFindReferencesTo->GetPathName());
+					GarbageErrorMessage = FString::Printf(TEXT("Garbage object %s was previously being referenced by %s"), *GarbageObjectName, *OutReferencingObject->GetFullName());
+				}
 			}
 		}
-	}
 #endif // ENABLE_GC_HISTORY
 
-	if (!bReferenceChainFound)
-	{
-		if (OutGarbageObject)
+		if (!bReferenceChainFound)
 		{
-			GarbageErrorMessage = FString::Printf(TEXT("Garbage object %s is not referenced so it may have a flag set that's preventing it from being destroyed (see log for details)"), *OutGarbageObject->GetFullName());
+			if (OutGarbageObject)
+			{
+				GarbageErrorMessage = FString::Printf(TEXT("Garbage object %s is not referenced so it may have a flag set that's preventing it from being destroyed (see log for details)"), *OutGarbageObject->GetFullName());
+				PathToCulprit = OutGarbageObject->GetFullName();
+			}
+			else
+			{
+				GarbageErrorMessage = TEXT("However it's not referenced by any object. It may have a flag set that's preventing it from being destroyed (see log for details)");
+			}
+		}
+
+		Paths.Add(PathToCulprit);
+
+		if ((Options & EPrintStaleReferencesOptions::Ensure) == EPrintStaleReferencesOptions::Ensure)
+		{
+			ensureAlwaysMsgf(false, TEXT("Old %s not cleaned up by GC! %s:") LINE_TERMINATOR TEXT("%s"),
+				*ObjectToFindReferencesTo->GetFullName(),
+				*GarbageErrorMessage,
+				*PathToCulprit);
+		}
+		else if (Verbosity == ELogVerbosity::Fatal)
+		{
+			UE_LOG(LogLoad, Fatal, TEXT("Old %s not cleaned up by GC! %s:") LINE_TERMINATOR TEXT("%s"),
+				*ObjectToFindReferencesTo->GetFullName(), 
+				*GarbageErrorMessage,
+				*PathToCulprit);
 		}
 		else
 		{
-			GarbageErrorMessage = TEXT("However it's not referenced by any object. It may have a flag set that's preventing it from being destroyed (see log for details)");
+			GLog->CategorizedLogf(TEXT("LogLoad"), Verbosity, TEXT("Old %s not cleaned up by GC! %s:") LINE_TERMINATOR TEXT("%s"),
+				*ObjectToFindReferencesTo->GetFullName(), 
+				*GarbageErrorMessage,
+				*PathToCulprit);
 		}
 	}
-
-	if ((Options & EPrintStaleReferencesOptions::Ensure) == EPrintStaleReferencesOptions::Ensure)
-	{
-		ensureAlwaysMsgf(false, TEXT("Old %s not cleaned up by GC! %s:") LINE_TERMINATOR TEXT("%s"),
-			*ObjectToFindReferencesTo->GetFullName(),
-			*GarbageErrorMessage,
-			*PathToCulprit);
-	}
-
-	if (Verbosity == ELogVerbosity::Fatal)
-	{
-		UE_LOG(LogLoad, Fatal, TEXT("Old %s not cleaned up by GC! %s:") LINE_TERMINATOR TEXT("%s"),
-			*ObjectToFindReferencesTo->GetFullName(), 
-			*GarbageErrorMessage,
-			*PathToCulprit);
-	}
-	else
-	{
-		GLog->CategorizedLogf(TEXT("LogLoad"), Verbosity, TEXT("Old %s not cleaned up by GC! %s:") LINE_TERMINATOR TEXT("%s"),
-			*ObjectToFindReferencesTo->GetFullName(), 
-			*GarbageErrorMessage,
-			*PathToCulprit);
-	}
+	return Paths;
 }
 
 /*-----------------------------------------------------------------------------
