Reverting to old object export text format as the new format shouldn't have been introduced before the remaining ANY_PACKAGE removal changed are submitted

diff --git a/Engine/Source/Runtime/CoreUObject/Private/UObject/PropertyBaseObject.cpp b/Engine/Source/Runtime/CoreUObject/Private/UObject/PropertyBaseObject.cpp
--- a/Engine/Source/Runtime/CoreUObject/Private/UObject/PropertyBaseObject.cpp
+++ b/Engine/Source/Runtime/CoreUObject/Private/UObject/PropertyBaseObject.cpp
@@ -181,7 +181,7 @@
 
 FString FObjectPropertyBase::GetExportPath(FTopLevelAssetPath ClassPathName, const FString& ObjectPathName)
 {
-	return FString::Printf(TEXT("%s.%s'%s'"), *ClassPathName.GetPackageName().ToString(), *ClassPathName.GetAssetName().ToString(), *ObjectPathName);
+	return FString::Printf(TEXT("%s'%s'"), *ClassPathName.GetAssetName().ToString(), *ObjectPathName);
 }
 
 FString FObjectPropertyBase::GetExportPath(const UObject* Object, const UObject* Parent /*= nullptr*/, const UObject* ExportRootScope /*= nullptr*/, const uint32 PortFlags /*= PPF_None*/)
