Removed disable optimization pragma and unused header

diff --git a/Engine/Source/Runtime/CoreUObject/Private/Serialization/BulkData.cpp b/Engine/Source/Runtime/CoreUObject/Private/Serialization/BulkData.cpp
--- a/Engine/Source/Runtime/CoreUObject/Private/Serialization/BulkData.cpp
+++ b/Engine/Source/Runtime/CoreUObject/Private/Serialization/BulkData.cpp
@@ -13,8 +13,6 @@
 #include "UObject/LinkerSave.h"
 #include "UObject/PackageResourceManager.h"
 
-PRAGMA_DISABLE_OPTIMIZATION
-
 /** Whether to track information of how bulk data is being used */
 #define TRACK_BULKDATA_USE 0
 
@@ -2232,5 +2230,3 @@ void FFormatContainer::SerializeAttemptMappedLoad(FArchive& Ar, UObject* Owner)
 		Bulk.Serialize(Ar, Owner, -1, true);
 	}
 }
-
-PRAGMA_ENABLE_OPTIMIZATION
diff --git a/Engine/Source/Runtime/CoreUObject/Public/Serialization/BulkData.h b/Engine/Source/Runtime/CoreUObject/Public/Serialization/BulkData.h
--- a/Engine/Source/Runtime/CoreUObject/Public/Serialization/BulkData.h
+++ b/Engine/Source/Runtime/CoreUObject/Public/Serialization/BulkData.h
@@ -10,7 +10,6 @@
 #include "Misc/PackagePath.h"
 #include "Misc/PackageSegment.h"
 #include "Serialization/FileRegions.h"
-#include "Templates/RefCounting.h"
 #include "IO/PackageId.h"
 #include "IO/IoDispatcher.h"
 
