Adding automated startup tests to check if class / struct / enum / property metadata contains values with short type names which will help with ANY_PACKAGE removal

(disabled until all ANY_PACKAGE removal changes are in)

diff --git a/Engine/Config/BaseEngine.ini b/Engine/Config/BaseEngine.ini
--- a/Engine/Config/BaseEngine.ini
+++ b/Engine/Config/BaseEngine.ini
@@ -103,6 +103,22 @@ EngineModuleReflectedUninitializedPropertyVerbosity=Error
 ProjectModuleReflectedUninitializedPropertyVerbosity=Error
 ObjectReferenceReflectedUninitializedPropertyVerbosity=Error
 
+[CoreUObject.ShortTypeNameInMetaDataCheck]
+;+MetaDataKeys=AllowedClasses
+;+MetaDataKeys=DisallowedClasses
+;+MetaDataKeys=BaseStruct
+;+MetaDataKeys=BaseClass
+;+MetaDataKeys=HasNativeMake
+;+MetaDataKeys=HasnativeBreak
+;+MetaDataKeys=RowType
+;+MetaDataKeys=BitmaskEnum
+;+MetaDataKeys=ActorClassToPoll
+;+MetaDataKeys=ComponentClass
+;+MetaDataKeys=MustImplement
+;+MetaDataKeys=MetaClass
+;+MetaDataKeys=MetaStruct
+;+MetaDataKeys=RequiredAssetDataTags
+
 [/Script/Engine.Engine]
 ConsoleClassName=/Script/Engine.Console
 GameViewportClientClassName=/Script/Engine.GameViewportClient
diff --git a/Engine/Source/Runtime/CoreUObject/Private/UObject/Class.cpp b/Engine/Source/Runtime/CoreUObject/Private/UObject/Class.cpp
--- a/Engine/Source/Runtime/CoreUObject/Private/UObject/Class.cpp
+++ b/Engine/Source/Runtime/CoreUObject/Private/UObject/Class.cpp
@@ -3472,23 +3472,24 @@
 	}
 }
 
+static FString GetFieldLocation(const UField* Field)
+{
+	check(Field);
+	UPackage* ScriptPackage = Field->GetOutermost();
+	FString StructLocation = FString::Printf(TEXT(" Module:%s"), *FPackageName::GetShortName(ScriptPackage->GetName()));
+#if WITH_EDITORONLY_DATA
+	static const FName NAME_ModuleRelativePath(TEXT("ModuleRelativePath"));
+	const FString& ModuleRelativeIncludePath = Field->GetMetaData(NAME_ModuleRelativePath);
+	if (!ModuleRelativeIncludePath.IsEmpty())
+	{
+		StructLocation += FString::Printf(TEXT(" File:%s"), *ModuleRelativeIncludePath);
+	}
+#endif
+	return StructLocation;
+};
+
 int32 FStructUtils::AttemptToFindUninitializedScriptStructMembers()
 {
-	auto GetStructLocation = [](const UScriptStruct* ScriptStruct) -> FString {
-		check(ScriptStruct);
-		UPackage* ScriptPackage = ScriptStruct->GetOutermost();
-		FString StructLocation = FString::Printf(TEXT(" Module:%s"), *FPackageName::GetShortName(ScriptPackage->GetName()));
-#if WITH_EDITORONLY_DATA
-		static const FName NAME_ModuleRelativePath(TEXT("ModuleRelativePath"));
-		const FString& ModuleRelativeIncludePath = ScriptStruct->GetMetaData(NAME_ModuleRelativePath);
-		if (!ModuleRelativeIncludePath.IsEmpty())
-		{
-			StructLocation += FString::Printf(TEXT(" File:%s"), *ModuleRelativeIncludePath);
-		}
-#endif
-		return StructLocation;
-	};
-
 	auto DetermineIfModuleIsEngine = [](const UScriptStruct* ScriptStruct) -> bool
 	{
 		UPackage* ScriptPackage = ScriptStruct->GetOutermost();
@@ -3595,7 +3596,7 @@
 				{
 					++UninitializedObjectPropertyCount;
 				}
-				UE_LOG(LogClass, Warning, TEXT("%s %s%s::%s is not initialized properly even though its struct probably has a custom default constructor.%s"), *Property->GetClass()->GetName(), ScriptStruct->GetPrefixCPP(), *ScriptStruct->GetName(), *Property->GetNameCPP(), *GetStructLocation(ScriptStruct));
+				UE_LOG(LogClass, Warning, TEXT("%s %s%s::%s is not initialized properly even though its struct probably has a custom default constructor.%s"), *Property->GetClass()->GetName(), ScriptStruct->GetPrefixCPP(), *ScriptStruct->GetName(), *Property->GetNameCPP(), *GetFieldLocation(ScriptStruct));
 			}
 			for (const FProperty* Property : UninitializedPropertiesNoInit)
 			{
@@ -3612,7 +3613,7 @@
 					}
 
 #if !NO_LOGGING
-					FMsg::Logf(__FILE__, __LINE__, LogClass.GetCategoryName(), Verbosity, TEXT("%s %s%s::%s is not initialized properly.%s"), *Property->GetClass()->GetName(), ScriptStruct->GetPrefixCPP(), *ScriptStruct->GetName(), *Property->GetNameCPP(), *GetStructLocation(ScriptStruct));
+					FMsg::Logf(__FILE__, __LINE__, LogClass.GetCategoryName(), Verbosity, TEXT("%s %s%s::%s is not initialized properly.%s"), *Property->GetClass()->GetName(), ScriptStruct->GetPrefixCPP(), *ScriptStruct->GetName(), *Property->GetNameCPP(), *GetFieldLocation(ScriptStruct));
 #endif
 				}
 			}
@@ -3634,9 +3635,9 @@
 	TEXT("Finds USTRUCT() structs that fail to initialize reflected member variables"),
 	FConsoleCommandWithWorldAndArgsDelegate::CreateStatic(
 		[](const TArray<FString>& Params, UWorld* World)
-{
-	FStructUtils::AttemptToFindUninitializedScriptStructMembers();
-}));
+		{
+			FStructUtils::AttemptToFindUninitializedScriptStructMembers();
+		}));
 
 IMPLEMENT_SIMPLE_AUTOMATION_TEST(FAutomationTestAttemptToFindUninitializedScriptStructMembers, "UObject.Class AttemptToFindUninitializedScriptStructMembers", EAutomationTestFlags::EditorContext | EAutomationTestFlags::ApplicationContextMask | EAutomationTestFlags::ServerContext | EAutomationTestFlags::SmokeFilter)
 bool FAutomationTestAttemptToFindUninitializedScriptStructMembers::RunTest(const FString& Parameters)
@@ -3654,8 +3655,213 @@
 	}
 }
 
+#if WITH_EDITORONLY_DATA
+/** 
+* Checks if MetaData value contains short type name 
+* @param MetaDataKey MetaData key name
+* @param MetaDataValue Value stored under MetaData key
+* @param OutParsedMetaDataValue If short type name is stored in MetaData value this parameter will contain the short type name
+* @returns true if MetaData value contains short type name
+**/
+static bool CheckIfMetaDataValueIsShortTypeName(FName MetaDataKey, const FString& MetaDataValue, FString& OutParsedMetaDataValue)
+{
+	// Some keys require additional parsing
+	static const FName NAME_RequiredAssetDataTags(TEXT("RequiredAssetDataTags"));
+
+	if (MetaDataKey == NAME_RequiredAssetDataTags) // @see SPropertyEditorAsset::InitializeAssetDataTags
+	{
+		TArray<FString> RequiredAssetDataTagsAndValues;
+		MetaDataValue.ParseIntoArray(RequiredAssetDataTagsAndValues, TEXT(","), true);
+
+		for (const FString& TagAndOptionalValueString : RequiredAssetDataTagsAndValues)
+		{
+			TArray<FString> TagAndOptionalValue;
+			TagAndOptionalValueString.ParseIntoArray(TagAndOptionalValue, TEXT("="), true);
+			if (TagAndOptionalValue.Num() == 2 && TagAndOptionalValue[0] == TEXT("RowStructure"))
+			{
+				if (FPackageName::IsShortPackageName(TagAndOptionalValue[1]))
+				{
+					OutParsedMetaDataValue = TagAndOptionalValue[1];
+					return true;
+				}
+			}
+		}
+	}
+	else if (FPackageName::IsShortPackageName(MetaDataValue))
+	{
+		OutParsedMetaDataValue = MetaDataValue;
+		return true;
+	}
+	return false;
+}
+
+/** Tries to find a path name for short type name **/
+static FString GetSuggestedPathNameForTypeShortName(const FString& ShortName)
+{
+	UField* FoundType = FindFirstObject<UField>(*ShortName);
+	if (FoundType)
+	{
+		return FString::Printf(TEXT("Suggested pathname: \"%s\"."), *FoundType->GetPathName());
+	}
+	else
+	{
+		return FString(TEXT("No valid type found in memory."));
+	}
+}
+
+/** Logs a field that defines MetaData entry with a short type name **/
+template <typename T>
+static void LogMetaDataShortTypeName(const T* Field, FName MetaDataKey, const FString& MetaDataValue)
+{
+	if constexpr (std::is_base_of_v<T, UField>)
+	{
+		if (const UFunction* Func = Cast<UFunction>(Field))
+		{
+			UStruct* FuncOwner = CastChecked<UStruct>(Func->GetOuter());
+			FMsg::Logf(__FILE__, __LINE__, LogClass.GetCategoryName(), ELogVerbosity::Warning, TEXT("Function %s%s::%s defines MetaData key \"%s\" which contains short type name \"%s\". %s%s"),
+				FuncOwner->GetPrefixCPP(),
+				*FuncOwner->GetName(),
+				*Func->GetName(),
+				*MetaDataKey.ToString(),
+				*MetaDataValue,
+				*GetSuggestedPathNameForTypeShortName(MetaDataValue),
+				*GetFieldLocation(Func));
+		}
+		else
+		{
+			FMsg::Logf(__FILE__, __LINE__, LogClass.GetCategoryName(), ELogVerbosity::Warning, TEXT("%s %s%s defines MetaData key \"%s\" which contains short type name \"%s\". %s%s"),
+				*Field->GetClass()->GetName(),
+				Field->template IsA<UStruct>() ? CastChecked<UStruct>(Field)->GetPrefixCPP() : TEXT(""),
+				*Field->GetName(),
+				*MetaDataKey.ToString(),
+				*MetaDataValue,
+				*GetSuggestedPathNameForTypeShortName(MetaDataValue),
+				*GetFieldLocation(Field));
+		}
+	}
+	else if constexpr (std::is_same_v<T, FProperty>)
+	{
+		UStruct* OwnerStruct = Field->GetOwnerStruct();
+		FMsg::Logf(__FILE__, __LINE__, LogClass.GetCategoryName(), ELogVerbosity::Warning, TEXT("Property %s %s%s::%s defines MetaData key \"%s\" which contains short type name \"%s\". %s%s"),
+			*Field->GetClass()->GetName(),
+			OwnerStruct->GetPrefixCPP(),
+			*OwnerStruct->GetName(),
+			*Field->GetName(),
+			*MetaDataKey.ToString(),
+			*MetaDataValue,
+			*GetSuggestedPathNameForTypeShortName(MetaDataValue),
+			*GetFieldLocation(OwnerStruct));
+	}
+}
+
+/** Checks MetaData values for known MetaData keys to see if they contain short type names **/
+template <typename T>
+static int32 FindFieldMetaDataShortTypeNames(const T* Field, const TArray<FName>& MetaDataKeys)
+{
+	int32 ShortTypeNameCount = 0;
+	for (FName MetaDataKey : MetaDataKeys)
+	{
+		const FString* MetaData = Field->FindMetaData(MetaDataKey);
+		if (MetaData && !MetaData->IsEmpty())
+		{
+			FString ParsedMetaDataValue;
+			int32 CommaIndex = -1;
+			// Some keys can store multiple type names separated with a comma
+			if (MetaData->FindChar(TCHAR(','), CommaIndex))
+			{
+				TArray<FString> MetaDataValues;
+				MetaData->ParseIntoArrayWS(MetaDataValues, TEXT(","), true);
+				for (const FString& MetaDataValue : MetaDataValues)
+				{
+					if (CheckIfMetaDataValueIsShortTypeName(MetaDataKey, MetaDataValue, ParsedMetaDataValue))
+					{
+						LogMetaDataShortTypeName(Field, MetaDataKey, ParsedMetaDataValue);
+						ShortTypeNameCount++;
+					}
+				}
+			}
+			else if (CheckIfMetaDataValueIsShortTypeName(MetaDataKey, *MetaData, ParsedMetaDataValue))
+			{
+				LogMetaDataShortTypeName(Field, MetaDataKey, ParsedMetaDataValue);
+				ShortTypeNameCount++;
+			}
+		}
+	}
+	return ShortTypeNameCount;
+}
+
+int32 FStructUtils::AttemptToFindShortTypeNamesInMetaData()
+{
+	static struct FShortTypeNameMetaDataSettings
+	{
+		TArray<FName> MetaDataKeys;
+		FShortTypeNameMetaDataSettings()
+		{
+			TArray<FString> ConfigMetaDataKeys;
+			if (GConfig->GetArray(TEXT("CoreUObject.ShortTypeNameInMetaDataCheck"), TEXT("MetaDataKeys"), ConfigMetaDataKeys, GEngineIni))
+			{
+				MetaDataKeys.Reserve(ConfigMetaDataKeys.Num());
+				for (const FString& Key : ConfigMetaDataKeys)
+				{
+					MetaDataKeys.Add(FName(*Key));
+				}
+			}
+		}
+	} Settings;
+
+	int32 ShortTypeNamesCount = 0;
+
+	if (Settings.MetaDataKeys.Num())
+	{
+		for (TObjectIterator<UField> FieldIt; FieldIt; ++FieldIt)
+		{
+			// Check MetaData stored with classes / enums / structs
+			ShortTypeNamesCount += FindFieldMetaDataShortTypeNames(*FieldIt, Settings.MetaDataKeys);
+
+			// Check struct members but no need to check function parameters
+			UStruct* StructWithMembers = Cast<UStruct>(*FieldIt);
+			if (StructWithMembers && !StructWithMembers->IsA<UFunction>())
+			{
+				for (TFieldIterator<FProperty> PropertyIt(StructWithMembers, EFieldIterationFlags::None); PropertyIt; ++PropertyIt)
+				{
+					ShortTypeNamesCount += FindFieldMetaDataShortTypeNames(*PropertyIt, Settings.MetaDataKeys);
+				}
+			}
+		}
+	}
+
+	if (ShortTypeNamesCount > 0)
+	{
+		UE_LOG(LogClass, Display, TEXT("%i short type names in reflected types' MetaData"), ShortTypeNamesCount);
+	}
+
+	return ShortTypeNamesCount;
+}
+
+FAutoConsoleCommandWithWorldAndArgs GCmdListShortTypeNamesInMetaData(
+	TEXT("CoreUObject.AttemptToFindShortTypeNamesInMetaData"),
+	TEXT("Finds short type names stored in known MetaData entries"),
+	FConsoleCommandWithWorldAndArgsDelegate::CreateStatic(
+		[](const TArray<FString>& Params, UWorld* World)
+		{
+			FStructUtils::AttemptToFindShortTypeNamesInMetaData();
+		}));
+
+IMPLEMENT_SIMPLE_AUTOMATION_TEST(FAutomationTestAttemptToFindShortTypeNamesInMetaData, "UObject.Class AttemptToFindShortTypeNamesInMetaData", EAutomationTestFlags::EditorContext | EAutomationTestFlags::ApplicationContextMask | EAutomationTestFlags::ServerContext | EAutomationTestFlags::SmokeFilter)
+bool FAutomationTestAttemptToFindShortTypeNamesInMetaData::RunTest(const FString& Parameters)
+{
+	// This test is not necessary when running under UHT so just skip it in that config.
+#if !HACK_HEADER_GENERATOR
+	return FStructUtils::AttemptToFindShortTypeNamesInMetaData() == 0;
+#else
+	return true;
 #endif
+}
 
+#endif // WITH_EDITORONLY_DATA
+
+#endif // !(UE_BUILD_TEST || UE_BUILD_SHIPPING)
+
 IMPLEMENT_CORE_INTRINSIC_CLASS(UScriptStruct, UStruct,
 	{
 	}
diff --git a/Engine/Source/Runtime/CoreUObject/Public/UObject/Class.h b/Engine/Source/Runtime/CoreUObject/Public/UObject/Class.h
--- a/Engine/Source/Runtime/CoreUObject/Public/UObject/Class.h
+++ b/Engine/Source/Runtime/CoreUObject/Public/UObject/Class.h
@@ -3655,7 +3655,12 @@
 #if !(UE_BUILD_TEST || UE_BUILD_SHIPPING)
 	/** Looks for uninitialized script struct pointers. Returns the number found */
 	COREUOBJECT_API static int32 AttemptToFindUninitializedScriptStructMembers();
-#endif
+
+#if WITH_EDITORONLY_DATA
+	/** Looks for short type names within struct metadata. Returns the number found */
+	COREUOBJECT_API static int32 AttemptToFindShortTypeNamesInMetaData();
+#endif // WITH_EDITORONLY_DATA
+#endif // !(UE_BUILD_TEST || UE_BUILD_SHIPPING)
 };
 
 // Helper struct to test if member initialization tests work properly
