BulkData - enable serialize bulk data size as 64bit integer by setting BULKDATA_Size64Bit

diff --git a/Engine/Source/Runtime/CoreUObject/Private/Serialization/BulkData.cpp b/Engine/Source/Runtime/CoreUObject/Private/Serialization/BulkData.cpp
--- a/Engine/Source/Runtime/CoreUObject/Private/Serialization/BulkData.cpp
+++ b/Engine/Source/Runtime/CoreUObject/Private/Serialization/BulkData.cpp
@@ -1519,10 +1519,6 @@ void FBulkData::Serialize(FArchive& Ar, UObject* Owner, bool bAttemptFileMapping
 			{
 				SetBulkDataFlagsOn(LocalBulkDataFlags, BULKDATA_Size64Bit);
 			}
-			else
-			{
-				ClearBulkDataFlagsOn(LocalBulkDataFlags, BULKDATA_Size64Bit);
-			}
 			// Remove single element serialization requirement before saving out bulk data flags.
 			ClearBulkDataFlagsOn(LocalBulkDataFlags, BULKDATA_ForceSingleElementSerialization);
 
