- Replacing many uses of LoadGlobalIniFile and LoadExternalIniFile with FCOnfigContext
#p4v-preflight-copy 20293528

diff --git a/Engine/Plugins/Developer/TraceDataFiltering/Source/TraceDataFiltering/Private/TraceDataFilteringModule.cpp b/Engine/Plugins/Developer/TraceDataFiltering/Source/TraceDataFiltering/Private/TraceDataFilteringModule.cpp
--- a/Engine/Plugins/Developer/TraceDataFiltering/Source/TraceDataFiltering/Private/TraceDataFilteringModule.cpp
+++ b/Engine/Plugins/Developer/TraceDataFiltering/Source/TraceDataFiltering/Private/TraceDataFilteringModule.cpp
@@ -14,6 +14,7 @@
 #include "Styling/AppStyle.h"
 #include "Widgets/Docking/SDockTab.h"
 #include "Misc/ConfigCacheIni.h"
+#include "Misc/ConfigContext.h"
 
 #include "TraceServices/ModuleService.h"
 #include "TraceServices/AnalysisService.h"
@@ -35,7 +36,7 @@ void FTraceFilteringModule::StartupModule()
 
 	FEventFilterStyle::Initialize();
 
-	FConfigCacheIni::LoadGlobalIniFile(TraceFiltersIni, TEXT("TraceDataFilters"));
+	FConfigContext::ReadIntoGConfig().Load(TEXT("TraceDataFilters"), TraceFiltersIni);
 
 #if WITH_EDITOR
 	const FSlateIcon TabIcon(FEventFilterStyle::GetStyleSetName(), "EventFilter.TabIcon");
diff --git a/Engine/Plugins/Developer/TraceSourceFiltering/Source/SourceFilteringEditor/Private/SourceFilteringEditorModule.cpp b/Engine/Plugins/Developer/TraceSourceFiltering/Source/SourceFilteringEditor/Private/SourceFilteringEditorModule.cpp
--- a/Engine/Plugins/Developer/TraceSourceFiltering/Source/SourceFilteringEditor/Private/SourceFilteringEditorModule.cpp
+++ b/Engine/Plugins/Developer/TraceSourceFiltering/Source/SourceFilteringEditor/Private/SourceFilteringEditorModule.cpp
@@ -11,6 +11,7 @@
 #include "Framework/Docking/TabManager.h"
 #include "Widgets/Docking/SDockTab.h"
 #include "Misc/ConfigCacheIni.h"
+#include "Misc/ConfigContext.h"
 #include "WorkspaceMenuStructureModule.h"
 #include "WorkspaceMenuStructure.h"
 #include "Features/IModularFeatures.h"
@@ -60,7 +61,7 @@ void FSourceFilteringEditorModule::StartupModule()
 	FSourceFilteringCommands::Register();
 
 	// Populate static ini path
-	FConfigCacheIni::LoadGlobalIniFile(SourceFiltersIni, TEXT("TraceSourceFilters"));
+	FConfigContext::ReadIntoGConfig().Load(TEXT("TraceSourceFilters"), SourceFiltersIni);
 	
 	IUnrealInsightsModule& UnrealInsightsModule = FModuleManager::LoadModuleChecked<IUnrealInsightsModule>("TraceInsights");	
 	FOnRegisterMajorTabExtensions& TimingProfilerLayoutExtension = UnrealInsightsModule.OnRegisterMajorTabExtension(FInsightsManagerTabs::TimingProfilerTabId);
diff --git a/Engine/Plugins/Editor/CryptoKeys/Source/CryptoKeys/Private/CryptoKeysSettings.cpp b/Engine/Plugins/Editor/CryptoKeys/Source/CryptoKeys/Private/CryptoKeysSettings.cpp
--- a/Engine/Plugins/Editor/CryptoKeys/Source/CryptoKeys/Private/CryptoKeysSettings.cpp
+++ b/Engine/Plugins/Editor/CryptoKeys/Source/CryptoKeys/Private/CryptoKeysSettings.cpp
@@ -3,6 +3,7 @@
 #include "CryptoKeysHelpers.h"
 #include "Settings/ProjectPackagingSettings.h"
 #include "Misc/ConfigCacheIni.h"
+#include "Misc/ConfigContext.h"
 #include "Misc/Base64.h"
 
 UCryptoKeysSettings::UCryptoKeysSettings()
@@ -16,20 +17,20 @@ UCryptoKeysSettings::UCryptoKeysSettings()
 
 		if (GConfig->IsReadyForUse())
 		{
-			FString EncryptionIni;
-			FConfigCacheIni::LoadGlobalIniFile(EncryptionIni, TEXT("Encryption"));
+			FConfigFile EncryptionConfig;
+			FConfigContext::ReadIntoLocalFile(EncryptionConfig).Load(TEXT("Encryption"));
 
 			FString OldEncryptionKey;
-			if (GConfig->GetString(TEXT("Core.Encryption"), TEXT("aes.key"), OldEncryptionKey, EncryptionIni))
+			if (EncryptionConfig.GetString(TEXT("Core.Encryption"), TEXT("aes.key"), OldEncryptionKey))
 			{
 				EncryptionKey = FBase64::Encode(OldEncryptionKey);
 			}
 
 			FString OldSigningModulus, OldSigningPublicExponent, OldSigningPrivateExponent;
 
-			bEnablePakSigning = GConfig->GetString(TEXT("Core.Encryption"), TEXT("rsa.privateexp"), OldSigningPrivateExponent, EncryptionIni)
-				&& GConfig->GetString(TEXT("Core.Encryption"), TEXT("rsa.publicexp"), OldSigningPublicExponent, EncryptionIni)
-				&& GConfig->GetString(TEXT("Core.Encryption"), TEXT("rsa.modulus"), OldSigningModulus, EncryptionIni);
+			bEnablePakSigning = EncryptionConfig.GetString(TEXT("Core.Encryption"), TEXT("rsa.privateexp"), OldSigningPrivateExponent)
+				&& EncryptionConfig.GetString(TEXT("Core.Encryption"), TEXT("rsa.publicexp"), OldSigningPublicExponent)
+				&& EncryptionConfig.GetString(TEXT("Core.Encryption"), TEXT("rsa.modulus"), OldSigningModulus);
 
 			if (bEnablePakSigning)
 			{
diff --git a/Engine/Plugins/Editor/GameplayTagsEditor/Source/GameplayTagsEditor/Private/GameplayTagsEditorModule.cpp b/Engine/Plugins/Editor/GameplayTagsEditor/Source/GameplayTagsEditor/Private/GameplayTagsEditorModule.cpp
--- a/Engine/Plugins/Editor/GameplayTagsEditor/Source/GameplayTagsEditor/Private/GameplayTagsEditorModule.cpp
+++ b/Engine/Plugins/Editor/GameplayTagsEditor/Source/GameplayTagsEditor/Private/GameplayTagsEditorModule.cpp
@@ -268,8 +268,6 @@ public:
 
 				// Reload off disk
 				GConfig->LoadFile(TagList->ConfigFileName);
-				//FString DestFileName;
-				//FConfigCacheIni::LoadGlobalIniFile(DestFileName, *FString::Printf(TEXT("Tags/%s"), *Source->SourceName.ToString()), nullptr, true);
 
 				// Explicitly remove user tags section
 				GConfig->EmptySection(TEXT("UserTags"), TagList->ConfigFileName);
diff --git a/Engine/Plugins/Experimental/DefaultInstallBundleManager/Source/Private/DefaultInstallBundleManager.cpp b/Engine/Plugins/Experimental/DefaultInstallBundleManager/Source/Private/DefaultInstallBundleManager.cpp
--- a/Engine/Plugins/Experimental/DefaultInstallBundleManager/Source/Private/DefaultInstallBundleManager.cpp
+++ b/Engine/Plugins/Experimental/DefaultInstallBundleManager/Source/Private/DefaultInstallBundleManager.cpp
@@ -14,6 +14,7 @@
 
 #include "Misc/Paths.h"
 #include "Misc/ConfigCacheIni.h"
+#include "Misc/ConfigContext.h"
 #include "Misc/DateTime.h"
 #include "Misc/CoreDelegates.h"
 #include "Misc/CommandLine.h"
@@ -165,7 +166,7 @@ void FDefaultInstallBundleManager::SetMustWaitForPSOCache(FBundleInfo& BundleInf
 void LoadKeychainFromInI(FKeyChain& OutCryptoSettings)
 {
 	FConfigFile ConfigFile;
-	FConfigCacheIni::LoadExternalIniFile(ConfigFile, TEXT("Crypto"), *FPaths::EngineConfigDir(), *FPaths::Combine(FPaths::ProjectDir(), TEXT("Config/")), true, TEXT("Windows"));
+	FConfigContext::ReadIntoLocalFile(ConfigFile, TEXT("Windows")).Load(TEXT("Crypto"));
 	if (ConfigFile.Num())
 	{
 		static const TCHAR* SectionName = TEXT("/Script/CryptoKeys.CryptoKeysSettings");
@@ -198,7 +199,7 @@ FDefaultInstallBundleManager::FDefaultInstallBundleManager(const TCHAR* InConfig
 {
 	// Init Config Settings
 	const TCHAR* DefaultConfigBaseName = TEXT("InstallBundleManager");
-	FConfigCacheIni::LoadGlobalIniFile(GInstallBundleManagerIni, InConfigBaseName ? InConfigBaseName : DefaultConfigBaseName);
+	FConfigContext::ReadIntoGConfig().Load(InConfigBaseName ? InConfigBaseName : DefaultConfigBaseName, GInstallBundleManagerIni);
 
 	TickHandle = FTSTicker::GetCoreTicker().AddTicker(FTickerDelegate::CreateRaw(this, &FDefaultInstallBundleManager::Tick));
 
diff --git a/Engine/Source/Developer/LocalizationService/Private/LocalizationServiceHelpers.cpp b/Engine/Source/Developer/LocalizationService/Private/LocalizationServiceHelpers.cpp
--- a/Engine/Source/Developer/LocalizationService/Private/LocalizationServiceHelpers.cpp
+++ b/Engine/Source/Developer/LocalizationService/Private/LocalizationServiceHelpers.cpp
@@ -4,6 +4,7 @@
 #include "ILocalizationServiceProvider.h"
 #include "Misc/Paths.h"
 #include "Misc/ConfigCacheIni.h"
+#include "Misc/ConfigContext.h"
 #include "ILocalizationServiceModule.h"
 
 #define LOCTEXT_NAMESPACE "LocalizationServiceHelpers"
@@ -22,8 +23,8 @@ namespace LocalizationServiceHelpers
 			static FString LocalizationServiceSettingsIni;
 			if (LocalizationServiceSettingsIni.Len() == 0)
 			{
-				const FString LocalizationServiceSettingsDir = FPaths::GeneratedConfigDir();
-				FConfigCacheIni::LoadGlobalIniFile(LocalizationServiceSettingsIni, TEXT("LocalizationServiceSettings"), NULL, false, false, true, true, *LocalizationServiceSettingsDir);
+				FConfigContext Context = FConfigContext::ReadIntoGConfig();
+				Context.Load(TEXT("LocalizationServiceSettings"), LocalizationServiceSettingsIni);
 			}
 			return LocalizationServiceSettingsIni;
 		}
@@ -34,8 +35,9 @@ namespace LocalizationServiceHelpers
 		static FString LocalizationServiceGlobalSettingsIni;
 		if (LocalizationServiceGlobalSettingsIni.Len() == 0)
 		{
-			const FString LocalizationServiceSettingsDir = FPaths::EngineSavedDir() + TEXT("Config/");
-			FConfigCacheIni::LoadGlobalIniFile(LocalizationServiceGlobalSettingsIni, TEXT("LocalizationServiceSettings"), NULL, false, false, true, true, *LocalizationServiceSettingsDir);
+			FConfigContext Context = FConfigContext::ReadIntoGConfig();
+			Context.GeneratedConfigDir = FPaths::EngineSavedDir() + TEXT("Config/");
+			Context.Load(TEXT("LocalizationServiceSettings"), LocalizationServiceGlobalSettingsIni);
 		}
 		return LocalizationServiceGlobalSettingsIni;
 	}
diff --git a/Engine/Source/Developer/Profiler/Private/ProfilerManager.h b/Engine/Source/Developer/Profiler/Private/ProfilerManager.h
--- a/Engine/Source/Developer/Profiler/Private/ProfilerManager.h
+++ b/Engine/Source/Developer/Profiler/Private/ProfilerManager.h
@@ -9,6 +9,7 @@
 #include "Misc/Guid.h"
 #include "Framework/Commands/UICommandList.h"
 #include "Misc/ConfigCacheIni.h"
+#include "Misc/ConfigContext.h"
 #include "IProfilerClient.h"
 #include "ISessionManager.h"
 #include "ProfilerDataSource.h"
@@ -43,7 +44,7 @@ public:
 
 	void LoadFromConfig()
 	{
-		FConfigCacheIni::LoadGlobalIniFile(ProfilerSettingsIni, TEXT("ProfilerSettings"));
+		FConfigContext::ReadIntoGConfig().Load(TEXT("ProfilerSettings"), ProfilerSettingsIni);
 
 		GConfig->GetBool(TEXT("Profiler.ProfilerOptions"), TEXT("bShowCoalescedViewModesInEventGraph"), bShowCoalescedViewModesInEventGraph, ProfilerSettingsIni);
 	}
diff --git a/Engine/Source/Developer/Settings/Private/SettingsSection.cpp b/Engine/Source/Developer/Settings/Private/SettingsSection.cpp
--- a/Engine/Source/Developer/Settings/Private/SettingsSection.cpp
+++ b/Engine/Source/Developer/Settings/Private/SettingsSection.cpp
@@ -3,6 +3,7 @@
 #include "SettingsSection.h"
 #include "Misc/Paths.h"
 #include "Misc/ConfigCacheIni.h"
+#include "Misc/ConfigContext.h"
 #include "UObject/Class.h"
 
 
@@ -180,7 +181,7 @@ bool FSettingsSection::ResetDefaults()
 		GConfig->EmptySection(*SettingsObject->GetClass()->GetPathName(), ConfigName);
 		GConfig->Flush(false);
 
-		FConfigCacheIni::LoadGlobalIniFile(ConfigName, *FPaths::GetBaseFilename(ConfigName), nullptr, true);
+		FConfigContext::ForceReloadIntoGConfig().Load(*FPaths::GetBaseFilename(ConfigName));
 
 		SettingsObject->ReloadConfig(nullptr, nullptr, UE::LCPF_PropagateToInstances|UE::LCPF_PropagateToChildDefaultObjects);
 
diff --git a/Engine/Source/Developer/SourceControl/Private/SourceControlHelpers.cpp b/Engine/Source/Developer/SourceControl/Private/SourceControlHelpers.cpp
--- a/Engine/Source/Developer/SourceControl/Private/SourceControlHelpers.cpp
+++ b/Engine/Source/Developer/SourceControl/Private/SourceControlHelpers.cpp
@@ -4,7 +4,7 @@
 #include "ISourceControlState.h"
 #include "HAL/FileManager.h"
 #include "Misc/Paths.h"
-#include "Misc/ConfigCacheIni.h"
+#include "Misc/ConfigContext.h"
 #include "ISourceControlOperation.h"
 #include "SourceControlOperations.h"
 #include "ISourceControlProvider.h"
@@ -1490,8 +1490,8 @@ const FString& USourceControlHelpers::GetSettingsIni()
 		static FString SourceControlSettingsIni;
 		if (SourceControlSettingsIni.Len() == 0)
 		{
-			const FString SourceControlSettingsDir = FPaths::GeneratedConfigDir();
-			FConfigCacheIni::LoadGlobalIniFile(SourceControlSettingsIni, TEXT("SourceControlSettings"), nullptr, false, false, true, true, *SourceControlSettingsDir);
+			FConfigContext Context = FConfigContext::ReadIntoGConfig();
+			Context.Load(TEXT("SourceControlSettings"), SourceControlSettingsIni);
 		}
 		return SourceControlSettingsIni;
 	}
@@ -1503,8 +1503,9 @@ const FString& USourceControlHelpers::GetGlobalSettingsIni()
 	static FString SourceControlGlobalSettingsIni;
 	if (SourceControlGlobalSettingsIni.Len() == 0)
 	{
-		const FString SourceControlSettingsDir = FPaths::EngineSavedDir() + TEXT("Config/");
-		FConfigCacheIni::LoadGlobalIniFile(SourceControlGlobalSettingsIni, TEXT("SourceControlSettings"), nullptr, false, false, true, true, *SourceControlSettingsDir);
+		FConfigContext Context = FConfigContext::ReadIntoGConfig();
+		Context.GeneratedConfigDir = FPaths::EngineSavedDir() + TEXT("Config/");
+		Context.Load(TEXT("SourceControlSettings"), SourceControlGlobalSettingsIni);
 	}
 	return SourceControlGlobalSettingsIni;
 }
diff --git a/Engine/Source/Developer/TraceInsights/Private/Insights/InsightsSettings.h b/Engine/Source/Developer/TraceInsights/Private/Insights/InsightsSettings.h
--- a/Engine/Source/Developer/TraceInsights/Private/Insights/InsightsSettings.h
+++ b/Engine/Source/Developer/TraceInsights/Private/Insights/InsightsSettings.h
@@ -4,6 +4,7 @@
 
 #include "CoreMinimal.h"
 #include "Misc/ConfigCacheIni.h"
+#include "Misc/ConfigContext.h"
 #include "ProfilingDebugging/MiscTrace.h"
 
 /** Contains all settings for the Unreal Insights, accessible through the main manager. */
@@ -35,7 +36,7 @@ public:
 
 	void LoadFromConfig()
 	{
-		if (!FConfigCacheIni::LoadGlobalIniFile(SettingsIni, TEXT("UnrealInsightsSettings")))
+		if (!FConfigContext::ReadIntoGConfig().Load(TEXT("UnrealInsightsSettings"), SettingsIni))
 		{
 			return;
 		}
diff --git a/Engine/Source/Developer/TraceServices/Private/Model/DbgHelpResolver.cpp b/Engine/Source/Developer/TraceServices/Private/Model/DbgHelpResolver.cpp
--- a/Engine/Source/Developer/TraceServices/Private/Model/DbgHelpResolver.cpp
+++ b/Engine/Source/Developer/TraceServices/Private/Model/DbgHelpResolver.cpp
@@ -12,6 +12,7 @@
 #include "HAL/RunnableThread.h"
 #include "Logging/LogMacros.h"
 #include "Misc/ConfigCacheIni.h"
+#include "Misc/ConfigContext.h"
 #include "Misc/Guid.h"
 #include "Misc/PathViews.h"
 #include "Misc/ScopeLock.h"
@@ -48,7 +49,8 @@ FDbgHelpResolver::FDbgHelpResolver(IAnalysisSession& InSession)
 
 	// Paths from configuration
 	FString SettingsIni;
-	if (FConfigCacheIni::LoadGlobalIniFile(SettingsIni, TEXT("UnrealInsightsSettings")))
+
+	if (FConfigContext::ReadIntoGConfig().Load(TEXT("UnrealInsightsSettings"), SettingsIni))
 	{
 		GConfig->GetArray(TEXT("Insights.MemoryProfiler"), TEXT("SymbolSearchPaths"), SymbolSearchPaths, SettingsIni);
 	}
diff --git a/Engine/Source/Developer/TraceServices/Private/Model/SymslibResolver.cpp b/Engine/Source/Developer/TraceServices/Private/Model/SymslibResolver.cpp
--- a/Engine/Source/Developer/TraceServices/Private/Model/SymslibResolver.cpp
+++ b/Engine/Source/Developer/TraceServices/Private/Model/SymslibResolver.cpp
@@ -15,6 +15,7 @@
 #include "HAL/PlatformTLS.h"
 #include "Logging/LogMacros.h"
 #include "Misc/ConfigCacheIni.h"
+#include "Misc/ConfigContext.h"
 #include "Misc/CString.h"
 #include "Misc/Paths.h"
 #include "Misc/PathViews.h"
@@ -489,7 +490,7 @@ FSymslibResolver::FSymslibResolver(IAnalysisSession& InSession)
 
 	// Paths from configuration
 	FString SettingsIni;
-	if (FConfigCacheIni::LoadGlobalIniFile(SettingsIni, TEXT("UnrealInsightsSettings")))
+	if (FConfigContext::ReadIntoGConfig().Load(TEXT("UnrealInsightsSettings"), SettingsIni))
 	{
 		GConfig->GetArray(TEXT("Insights.MemoryProfiler"), TEXT("SymbolSearchPaths"), SymbolSearchPaths, SettingsIni);
 	}
diff --git a/Engine/Source/Editor/DetailCustomizations/Private/ConfigEditorPropertyDetails.cpp b/Engine/Source/Editor/DetailCustomizations/Private/ConfigEditorPropertyDetails.cpp
--- a/Engine/Source/Editor/DetailCustomizations/Private/ConfigEditorPropertyDetails.cpp
+++ b/Engine/Source/Editor/DetailCustomizations/Private/ConfigEditorPropertyDetails.cpp
@@ -2,6 +2,7 @@
 
 #include "ConfigEditorPropertyDetails.h"
 #include "Misc/ConfigCacheIni.h"
+#include "Misc/ConfigContext.h"
 #include "Modules/ModuleManager.h"
 #include "UObject/UnrealType.h"
 #include "PropertyEditorModule.h"
@@ -102,8 +103,7 @@ void FConfigPropertyHelperDetails::OnPropertyValueChanged(UObject* Object, FProp
 			NewFile.UpdateSinglePropertyInSection(*ConfigIniName, *PropertyName, *SectionName);
 
 			// reload the file, so that it refresh the cache internally.
-			FString FinalIniFileName;
-			GConfig->LoadGlobalIniFile(FinalIniFileName, *OriginalProperty->GetOwnerClass()->ClassConfigName.ToString(), NULL, true);
+			FConfigContext::ForceReloadIntoGConfig().Load(*OriginalProperty->GetOwnerClass()->ClassConfigName.ToString());
 
 			// Update the CDO, as this change might have had an impact on it's value.
 			OriginalProperty->GetOwnerClass()->GetDefaultObject()->ReloadConfig();
diff --git a/Engine/Source/Editor/LocalizationDashboard/Private/LocalizationTargetDetailCustomization.cpp b/Engine/Source/Editor/LocalizationDashboard/Private/LocalizationTargetDetailCustomization.cpp
--- a/Engine/Source/Editor/LocalizationDashboard/Private/LocalizationTargetDetailCustomization.cpp
+++ b/Engine/Source/Editor/LocalizationDashboard/Private/LocalizationTargetDetailCustomization.cpp
@@ -6,6 +6,7 @@
 #include "Misc/MessageDialog.h"
 #include "Internationalization/Culture.h"
 #include "Misc/ConfigCacheIni.h"
+#include "Misc/ConfigContext.h"
 #include "UObject/UnrealType.h"
 #include "DesktopPlatformModule.h"
 #include "Framework/Application/SlateApplication.h"
@@ -672,8 +673,7 @@ void FLocalizationTargetDetailCustomization::SetLoadingPolicy(const ELocalizatio
 		}
 
 		// Reload the updated file into the config system.
-		FString FinalIniFileName;
-		GConfig->LoadGlobalIniFile(FinalIniFileName, *LoadingPolicyConfig.BaseConfigName, nullptr, /*bForceReload*/true);
+		FConfigContext::ForceReloadIntoGConfig().Load(*LoadingPolicyConfig.BaseConfigName);
 	};
 
 	for (const FLocalizationTargetLoadingPolicyConfig& LoadingPolicyConfig : LoadingPolicyConfigs)
diff --git a/Engine/Source/Editor/UnrealEd/Private/Factories/PackFactory.cpp b/Engine/Source/Editor/UnrealEd/Private/Factories/PackFactory.cpp
--- a/Engine/Source/Editor/UnrealEd/Private/Factories/PackFactory.cpp
+++ b/Engine/Source/Editor/UnrealEd/Private/Factories/PackFactory.cpp
@@ -13,6 +13,7 @@
 #include "Serialization/MemoryWriter.h"
 #include "Serialization/MemoryReader.h"
 #include "Misc/ConfigCacheIni.h"
+#include "Misc/ConfigContext.h"
 #include "Misc/FeedbackContext.h"
 #include "Misc/App.h"
 #include "Modules/ModuleManager.h"
@@ -448,8 +449,7 @@ UObject* UPackFactory::FactoryCreateBinary
 
 						NewFile.UpdateSections(*EngineIniFilename, *RedirectsSection);
 
-						FString FinalIniFileName;
-						GConfig->LoadGlobalIniFile(FinalIniFileName, *RedirectsSection, NULL, true);
+						FConfigContext::ForceReloadIntoGConfig().Load(*RedirectsSection);
 
 						FLinkerLoad::AddGameNameRedirect(*LongOldGameName, *LongNewGameName);
 						FLinkerLoad::AddGameNameRedirect(*ConfigParameters.GameName, *LongNewGameName);
diff --git a/Engine/Source/Editor/UnrealEd/Private/StaticLightingSystem/StaticLightingSystem.cpp b/Engine/Source/Editor/UnrealEd/Private/StaticLightingSystem/StaticLightingSystem.cpp
--- a/Engine/Source/Editor/UnrealEd/Private/StaticLightingSystem/StaticLightingSystem.cpp
+++ b/Engine/Source/Editor/UnrealEd/Private/StaticLightingSystem/StaticLightingSystem.cpp
@@ -10,6 +10,7 @@
 #include "Misc/Paths.h"
 #include "Misc/Guid.h"
 #include "Misc/ConfigCacheIni.h"
+#include "Misc/ConfigContext.h"
 #include "HAL/IConsoleManager.h"
 #include "Misc/ScopedSlowTask.h"
 #include "Misc/App.h"
@@ -634,7 +635,8 @@ bool FStaticLightingSystem::BeginLightmassProcess()
 			UE_LOG(LogStaticLightingSystem, Warning, TEXT("WorldSettings.bForceNoPrecomputedLighting is true, Skipping Lighting Build!"));
 		}
 		
-		FConfigCacheIni::LoadGlobalIniFile(GLightmassIni, TEXT("Lightmass"), NULL, true);
+		FConfigContext::ForceReloadIntoGConfig().Load(TEXT("Lightmass"), GLightmassIni);
+
 		verify(GConfig->GetBool(TEXT("DevOptions.StaticLighting"), TEXT("bUseBilinearFilterLightmaps"), GUseBilinearLightmaps, GLightmassIni));
 		verify(GConfig->GetBool(TEXT("DevOptions.StaticLighting"), TEXT("bAllowCropping"), GAllowLightmapCropping, GLightmassIni));
 		verify(GConfig->GetBool(TEXT("DevOptions.StaticLighting"), TEXT("bRebuildDirtyGeometryForLighting"), bRebuildDirtyGeometryForLighting, GLightmassIni));
diff --git a/Engine/Source/Runtime/Core/Private/Misc/ConfigCacheIni.cpp b/Engine/Source/Runtime/Core/Private/Misc/ConfigCacheIni.cpp
--- a/Engine/Source/Runtime/Core/Private/Misc/ConfigCacheIni.cpp
+++ b/Engine/Source/Runtime/Core/Private/Misc/ConfigCacheIni.cpp
@@ -3573,8 +3573,7 @@ void FConfigCacheIni::InitializeConfigSystem()
 	if (!FParse::Param(FCommandLine::Get(), TEXT("textconfig")) &&
 		FConfigCacheIni::CreateGConfigFromSaved(nullptr))
 	{
-		FConfigContext Context = FConfigContext::ReadIntoGConfig();
-		Context.bForceReload = true;
+		FConfigContext Context = FConfigContext::ForceReloadIntoGConfig();
 		// Force reload GameUserSettings because they may be saved to disk on consoles/similar platforms
 		// So the safest thing to do is to re-read the file after binary configs load.
 		Context.Load(TEXT("GameUserSettings"), GGameUserSettingsIni);
diff --git a/Engine/Source/Runtime/Core/Private/Misc/ConfigContext.cpp b/Engine/Source/Runtime/Core/Private/Misc/ConfigContext.cpp
--- a/Engine/Source/Runtime/Core/Private/Misc/ConfigContext.cpp
+++ b/Engine/Source/Runtime/Core/Private/Misc/ConfigContext.cpp
@@ -122,6 +122,12 @@ const FConfigContext::FPerPlatformDirs& FConfigContext::GetPerPlatformDirs(const
 
 bool FConfigContext::Load(const TCHAR* InBaseIniName, FString& OutFinalFilename)
 {
+	// for single file loads, just return early of the file doesn't exist
+	if (!bIsHierarchicalConfig && !DoesConfigFileExistWrapper(InBaseIniName, IniCacheSet))
+	{
+		return false;
+	}
+
 	if (bCacheOnNextLoad || BaseIniName != InBaseIniName)
 	{
 		ResetBaseIni(InBaseIniName);
@@ -292,8 +298,17 @@ bool FConfigContext::PerformLoad()
 	// if bIsBaseIniName is false, that means the .ini is a ready-to-go .ini file, and just needs to be loaded into the FConfigFile
 	if (!bIsHierarchicalConfig)
 	{
-		// generate path to the .ini file (not a Default ini, IniName is the complete name of the file, without path)
-		DestIniFilename = FString::Printf(TEXT("%s/%s.ini"), *ProjectConfigDir, *BaseIniName);
+		// if the ini name passed in already is a full path, just use it
+		if (BaseIniName.EndsWith(TEXT(".ini")))
+		{
+			DestIniFilename = BaseIniName;
+			BaseIniName = FPaths::GetBaseFilename(BaseIniName);
+		}
+		else
+		{
+			// generate path to the .ini file (not a Default ini, IniName is the complete name of the file, without path)
+			DestIniFilename = FString::Printf(TEXT("%s/%s.ini"), *ProjectConfigDir, *BaseIniName);
+		}
 
 		// load the .ini file straight up
 		LoadAnIniFile(*DestIniFilename, *ConfigFile);
diff --git a/Engine/Source/Runtime/Core/Private/Misc/RemoteConfigIni.cpp b/Engine/Source/Runtime/Core/Private/Misc/RemoteConfigIni.cpp
--- a/Engine/Source/Runtime/Core/Private/Misc/RemoteConfigIni.cpp
+++ b/Engine/Source/Runtime/Core/Private/Misc/RemoteConfigIni.cpp
@@ -7,6 +7,7 @@
 #include "Misc/FileHelper.h"
 #include "Misc/Paths.h"
 #include "Misc/ConfigCacheIni.h"
+#include "Misc/ConfigContext.h"
 #include "Misc/App.h"
 
 // Globals
@@ -368,7 +369,7 @@ void FRemoteConfig::FinishRead(const TCHAR* Filename)
 		FString DestFileName(Filename);
 		GRemoteConfigIOManager.GetReadData(Filename, *IOInfo);
 		IOInfo->bWasProcessed = true;
-		FConfigCacheIni::LoadGlobalIniFile(DestFileName, IOInfo->DefaultIniFile);
+		FConfigContext::ReadIntoGConfig().Load(IOInfo->DefaultIniFile,DestFileName);
 	}
 }
 
diff --git a/Engine/Source/Runtime/CoreUObject/Private/UObject/Class.cpp b/Engine/Source/Runtime/CoreUObject/Private/UObject/Class.cpp
--- a/Engine/Source/Runtime/CoreUObject/Private/UObject/Class.cpp
+++ b/Engine/Source/Runtime/CoreUObject/Private/UObject/Class.cpp
@@ -6,6 +6,7 @@
 #include "Misc/ScopeLock.h"
 #include "Serialization/MemoryWriter.h"
 #include "Misc/ConfigCacheIni.h"
+#include "Misc/ConfigContext.h"
 #include "Misc/OutputDeviceHelper.h"
 #include "Misc/FeedbackContext.h"
 #include "Misc/OutputDeviceConsole.h"
@@ -5958,7 +5959,7 @@ const FString UClass::GetConfigName() const
 	{
 		// generate the class ini name, and make sure it's up to date
 		FString ConfigGameName;
-		FConfigCacheIni::LoadGlobalIniFile(ConfigGameName, *ClassConfigName.ToString());
+		FConfigContext::ReadIntoGConfig().Load(*ClassConfigName.ToString(), ConfigGameName);
 		return ConfigGameName;
 	}
 }
diff --git a/Engine/Source/Runtime/CoreUObject/Private/UObject/Obj.cpp b/Engine/Source/Runtime/CoreUObject/Private/UObject/Obj.cpp
--- a/Engine/Source/Runtime/CoreUObject/Private/UObject/Obj.cpp
+++ b/Engine/Source/Runtime/CoreUObject/Private/UObject/Obj.cpp
@@ -12,6 +12,7 @@
 #include "Logging/LogScopedCategoryAndVerbosityOverride.h"
 #include "Stats/Stats.h"
 #include "Misc/ConfigCacheIni.h"
+#include "Misc/ConfigContext.h"
 #include "Misc/CoreDelegates.h"
 #include "Misc/App.h"
 #include "Misc/ITransaction.h"
@@ -1851,22 +1852,8 @@ FString GetConfigFilename( UObject* SourceObject )
 {
 	checkSlow(SourceObject);
 
-#if 0 // If you find that you are sad this code is gone, please contact RobM or JoeG. We could not find a case for its existence
-	// and it broke/made cumbersome perobjectconfig file specification for files that were outside the transient package
-	if (UsesPerObjectConfig(SourceObject) && SourceObject->GetOutermost() != GetTransientPackage())
-	{
-		// if this is a PerObjectConfig object that is not contained by the transient package,
-		// load the class's package's ini file
-		FString PerObjectConfigName;
-		FConfigCacheIni::LoadGlobalIniFile(PerObjectConfigName, *SourceObject->GetOutermost()->GetName());
-		return PerObjectConfigName;
-	}
-	else
-#endif
-	{
-	// otherwise look at the class to get the config name
+	// look at the class to get the config name
 	return SourceObject->GetClass()->GetConfigName();
-	}
 }
 
 namespace UE { namespace Object { namespace Private {
@@ -2851,8 +2838,7 @@ void UObject::UpdateSingleSectionOfConfigFile(const FString& ConfigIniName)
 	// then we don't want to touch GConfig
 	if (OverridePlatform.Len() == 0)
 	{
-		FString FinalIniFileName;
-		GConfig->LoadGlobalIniFile(FinalIniFileName, *GetClass()->ClassConfigName.ToString(), NULL, true);
+		FConfigContext::ForceReloadIntoGConfig().Load(*GetClass()->ClassConfigName.ToString());
 	}
 }
 
@@ -2929,8 +2915,7 @@ void UObject::UpdateSinglePropertyInConfigFile(const FProperty* InProperty, cons
 		// then we don't want to touch GConfig
 		if (OverridePlatform.Len() == 0)
 		{
-			FString FinalIniFileName;
-			GConfig->LoadGlobalIniFile(FinalIniFileName, *GetClass()->ClassConfigName.ToString(), NULL, true);
+			FConfigContext::ForceReloadIntoGConfig().Load(*GetClass()->ClassConfigName.ToString());
 		}
 	}
 	else
diff --git a/Engine/Source/Runtime/Engine/Private/GameUserSettings.cpp b/Engine/Source/Runtime/Engine/Private/GameUserSettings.cpp
--- a/Engine/Source/Runtime/Engine/Private/GameUserSettings.cpp
+++ b/Engine/Source/Runtime/Engine/Private/GameUserSettings.cpp
@@ -3,6 +3,7 @@
 #include "GameFramework/GameUserSettings.h"
 #include "HAL/FileManager.h"
 #include "Misc/ConfigCacheIni.h"
+#include "Misc/ConfigContext.h"
 #include "Misc/FileHelper.h"
 #include "HAL/IConsoleManager.h"
 #include "GenericPlatform/GenericApplication.h"
@@ -582,8 +583,9 @@ void UGameUserSettings::SaveSettings()
 
 void UGameUserSettings::LoadConfigIni(bool bForceReload/*=false*/)
 {
-	// Load .ini, allowing merging
-	FConfigCacheIni::LoadGlobalIniFile(GGameUserSettingsIni, TEXT("GameUserSettings"), nullptr, bForceReload, false, true, true, *FPaths::GeneratedConfigDir());
+	FConfigContext Context = FConfigContext::ReadIntoGConfig();
+	Context.bForceReload = bForceReload;
+	Context.Load(TEXT("GameUserSettings"), GGameUserSettingsIni);
 }
 
 void UGameUserSettings::PreloadResolutionSettings(bool bAllowCmdLineOverrides /*= true*/)
diff --git a/Engine/Source/Runtime/Engine/Private/RuntimeOptionsBase.cpp b/Engine/Source/Runtime/Engine/Private/RuntimeOptionsBase.cpp
--- a/Engine/Source/Runtime/Engine/Private/RuntimeOptionsBase.cpp
+++ b/Engine/Source/Runtime/Engine/Private/RuntimeOptionsBase.cpp
@@ -5,6 +5,7 @@
 #include "Misc/Paths.h"
 #include "Misc/FileHelper.h"
 #include "Misc/ConfigCacheIni.h"
+#include "Misc/ConfigContext.h"
 #include "Interfaces/IPluginManager.h"
 #include "Misc/CommandLine.h"
 #include "HAL/IConsoleManager.h"
@@ -38,8 +39,7 @@ void URuntimeOptionsBase::InitializeRuntimeOptions()
 		ensureAlwaysMsgf(NumNonAbstractVersions == 1, TEXT("Error in %s hierarchy; properties should only be introduced in leaf non-abstract subclasses of URuntimeOptionsBase"), *GetClass()->GetName());
 
 		// Load the .ini file if it hasn't already been loaded
-		FString RuntimeOptionsIniName;
-		FConfigCacheIni::LoadGlobalIniFile(/*out*/ RuntimeOptionsIniName, *GetClass()->ClassConfigName.ToString(), nullptr);
+		FConfigContext::ReadIntoGConfig().Load(*GetClass()->ClassConfigName.ToString());
 	
 		// Apply command line overrides and expose properties as console variables
 		ApplyCommandlineOverrides();
