Add a check(this) to IsChildOf due to it being undefined behavior

Some compilers (clang13) are now optimizing away the this == nullptr checks, causing crashes. MSVC does not currently do so, so adding this check to catch these issues sooner rather then on clang13 only platforms

diff --git a/Engine/Source/Runtime/CoreUObject/Private/UObject/Class.cpp b/Engine/Source/Runtime/CoreUObject/Private/UObject/Class.cpp
--- a/Engine/Source/Runtime/CoreUObject/Private/UObject/Class.cpp
+++ b/Engine/Source/Runtime/CoreUObject/Private/UObject/Class.cpp
@@ -2193,6 +2193,10 @@ void UStruct::TagSubobjects(EObjectFlags NewFlags)
 #if USTRUCT_FAST_ISCHILDOF_COMPARE_WITH_OUTERWALK || USTRUCT_FAST_ISCHILDOF_IMPL == USTRUCT_ISCHILDOF_OUTERWALK
 bool UStruct::IsChildOf( const UStruct* SomeBase ) const
 {
+	// If you're looking at this check it is due to calling IsChildOf with a this nullptr. *MAKE* sure you do not call this function
+	// with a this nullptr. It is undefined behavior, and some compilers, clang13 have started to optimize out this == nullptr checks.
+	check(this);
+
 	if (SomeBase == nullptr)
 	{
 		return false;
