Dump the Blueprint callstack when MostRecentPropertyAddress is null.

#UE-130566
#8455

diff --git a/Engine/Source/Runtime/CoreUObject/Private/UObject/ScriptCore.cpp b/Engine/Source/Runtime/CoreUObject/Private/UObject/ScriptCore.cpp
--- a/Engine/Source/Runtime/CoreUObject/Private/UObject/ScriptCore.cpp
+++ b/Engine/Source/Runtime/CoreUObject/Private/UObject/ScriptCore.cpp
@@ -933,7 +933,7 @@ void ProcessScriptFunction(UObject* Context, UFunction* Function, FFrame& Stack,
 			// warning: Stack.MostRecentPropertyAddress could be NULL for optional out parameters
 			// if that's the case, we use the extra memory allocated for the out param in the function's locals
 			// so there's always a valid address
-			ensure(Stack.MostRecentPropertyAddress); // possible problem - output param values on local stack are neither initialized nor cleaned.
+			ensureMsgf(Stack.MostRecentPropertyAddress, TEXT("MostRecentPropertyAddress was null. Blueprint callstack:\n%s"), *Stack.GetScriptCallstack()); // possible problem - output param values on local stack are neither initialized nor cleaned.
 			Out->PropAddr = (Stack.MostRecentPropertyAddress != NULL) ? Stack.MostRecentPropertyAddress : Property->ContainerPtrToValuePtr<uint8>(NewStack.Locals);
 			Out->Property = Property;
 
