DerivedData: Added a new API in place of the use of IoDispatcher

diff --git a/Engine/Source/Runtime/CoreUObject/Private/Serialization/DerivedData.cpp b/Engine/Source/Runtime/CoreUObject/Private/Serialization/DerivedData.cpp
--- a/Engine/Source/Runtime/CoreUObject/Private/Serialization/DerivedData.cpp
+++ b/Engine/Source/Runtime/CoreUObject/Private/Serialization/DerivedData.cpp
@@ -2,191 +2,1120 @@
 
 #include "Serialization/DerivedData.h"
 
+#include "Containers/ArrayView.h"
+#include "Hash/xxhash.h"
+#include "IO/IoDispatcher.h"
 #include "Memory/SharedBuffer.h"
+#include "Misc/StringBuilder.h"
+#include "Misc/TVariant.h"
+#include "String/BytesToHex.h"
+#include <atomic>
 
 #if WITH_EDITORONLY_DATA
-
-#include "DerivedDataBuildDefinition.h"
+#include "Compression/CompressedBuffer.h"
+#include "DerivedDataCache.h"
 #include "DerivedDataCacheKey.h"
+#include "DerivedDataRequestOwner.h"
+#include "DerivedDataRequestTypes.h"
+#include "DerivedDataSharedString.h"
+#include "DerivedDataValue.h"
 #include "DerivedDataValueId.h"
-#include "Serialization/EditorDerivedData.h"
-#include "Serialization/EditorDerivedDataIoStore.h"
-#include "Templates/SharedPointer.h"
+#include "IO/IoHash.h"
+#include "Memory/CompositeBuffer.h"
 #include "UObject/LinkerSave.h"
+#endif // WITH_EDITORONLY_DATA
 
-namespace UE::DerivedData::IoStore
+namespace UE::DerivedData::Private { class FIoResponseDispatcher; }
+
+#if WITH_EDITORONLY_DATA
+namespace UE::DerivedData::Private
 {
 
-TSharedPtr<IEditorDerivedDataIoStore> GEditorDerivedDataIoStore;
+struct FCacheKeyWithId
+{
+	FCacheKey Key;
+	FValueId Id;
 
-} // UE::DerivedData::IoStore
+	inline FCacheKeyWithId(const FCacheKey& InKey, const FValueId& InId)
+		: Key(InKey)
+		, Id(InId)
+	{
+	}
 
-#endif // WITH_EDITORONLY_DATA
+	inline bool operator==(const FCacheKeyWithId& Other) const
+	{
+		return Key == Other.Key && Id == Other.Id;
+	}
+};
 
-///////////////////////////////////////////////////////////////////////////////////////////////////////////////////////
+struct FCompositeBufferWithHash
+{
+	FCompositeBuffer Buffer;
+	FIoHash Hash;
 
-namespace UE
+	template <typename... ArgTypes>
+	inline explicit FCompositeBufferWithHash(ArgTypes&&... Args)
+		: Buffer(Forward<ArgTypes>(Args)...)
+		, Hash(FIoHash::HashBuffer(Buffer))
+	{
+	}
+
+	inline bool operator==(const FCompositeBufferWithHash& Other) const
+	{
+		return Hash == Other.Hash;
+	}
+};
+
+class FEditorData
 {
+public:
+	template <typename... ArgTypes>
+	inline FEditorData(const FSharedString& InName, ArgTypes&&... Args)
+		: Name(InName)
+		, Data(Forward<ArgTypes>(Args)...)
+	{
+	}
 
-void FDerivedData::Serialize(FArchive& Ar, UObject* Owner)
+	const FSharedString& GetName() const { return Name; }
+
+	friend FStringBuilderBase& operator<<(FStringBuilderBase& Builder, const FEditorData& EditorData);
+
+	bool ReferenceEquals(const FEditorData& Other) const;
+	uint32 ReferenceHash() const;
+
+	void Serialize(FArchive& Ar, UObject* Owner) const;
+
+	template <typename VisitorType>
+	inline void Visit(VisitorType&& Visitor) const
+	{
+		::Visit(Visitor, Data);
+	}
+
+private:
+	FSharedString Name;
+	TVariant<FCompositeBufferWithHash, FCompressedBuffer, FCacheKeyWithId> Data;
+};
+
+FStringBuilderBase& operator<<(FStringBuilderBase& Builder, const FEditorData& EditorData)
 {
-	if (Ar.IsPersistent() && !Ar.IsObjectReferenceCollector() && !Ar.ShouldSkipBulkData())
+	struct FVisitor
 	{
-#if WITH_EDITORONLY_DATA
-		if (Ar.IsSaving())
+		inline void operator()(const FCompositeBufferWithHash& BufferWithHash) const
+		{
+			Builder
+				<< TEXTVIEW("Buffer: Size ") << BufferWithHash.Buffer.GetSize()
+				<< TEXTVIEW(" Hash ") << BufferWithHash.Hash;
+		}
+
+		inline void operator()(const FCompressedBuffer& Buffer) const
+		{
+			Builder
+				<< TEXTVIEW("Buffer: CompressedSize ") << Buffer.GetCompressedSize()
+				<< TEXTVIEW(" Size ") << Buffer.GetRawSize()
+				<< TEXTVIEW(" Hash ") << Buffer.GetRawHash();
+		}
+
+		inline void operator()(const FCacheKeyWithId& CacheKeyWithId) const
+		{
+			Builder << TEXTVIEW("Cache: Key ") << CacheKeyWithId.Key;
+			if (CacheKeyWithId.Id.IsValid())
+			{
+				Builder << TEXTVIEW(" ID ") << CacheKeyWithId.Id;
+			}
+		}
+
+		FStringBuilderBase& Builder;
+	};
+
+	Visit(FVisitor{Builder}, EditorData.Data);
+	if (!EditorData.Name.IsEmpty())
+	{
+		Builder << TEXTVIEW(" for ") << EditorData.Name;
+	}
+	return Builder;
+}
+
+bool FEditorData::ReferenceEquals(const FEditorData& Other) const
+{
+	if (Data.GetIndex() != Other.Data.GetIndex())
+	{
+		return false;
+	}
+
+	if (const FCompositeBufferWithHash* BufferWithHash = Data.TryGet<FCompositeBufferWithHash>())
+	{
+		return Other.Data.Get<FCompositeBufferWithHash>().Hash == BufferWithHash->Hash;
+	}
+
+	if (const FCompressedBuffer* Buffer = Data.TryGet<FCompressedBuffer>())
+	{
+		return Other.Data.Get<FCompressedBuffer>().GetRawHash() == Buffer->GetRawHash();
+	}
+
+	if (const FCacheKeyWithId* CacheKeyWithId = Data.TryGet<FCacheKeyWithId>())
+	{
+		return Other.Data.Get<FCacheKeyWithId>() == *CacheKeyWithId;
+	}
+
+	checkNoEntry();
+	return false;
+}
+
+uint32 FEditorData::ReferenceHash() const
+{
+	struct FVisitor
+	{
+		inline void operator()(const FCompositeBufferWithHash& BufferWithHash)
+		{
+			Hash = GetTypeHash(BufferWithHash.Hash);
+		}
+
+		inline void operator()(const FCompressedBuffer& Buffer)
+		{
+			Hash = GetTypeHash(Buffer.GetRawHash());
+		}
+
+		inline void operator()(const FCacheKeyWithId& CacheKeyWithId)
+		{
+			Hash = HashCombineFast(GetTypeHash(CacheKeyWithId.Key), GetTypeHash(CacheKeyWithId.Id));
+		}
+		uint32 Hash = 0;
+	};
+
+	FVisitor Visitor;
+	Visit(Visitor);
+	return Visitor.Hash;
+}
+
+void FEditorData::Serialize(FArchive& Ar, UObject* Owner) const
+{
+	checkf(Ar.IsSaving() && Ar.IsCooking(), TEXT("FEditorData for FDerivedData only supports saving to cooked packages."));
+
+	struct FVisitor
+	{
+		inline explicit FVisitor(FArchive& Ar)
+			: Linker(Cast<FLinkerSave>(Ar.GetLinker()))
 		{
-			checkf(Ar.IsCooking(), TEXT("Serializing FDerivedData is only supported for cooked packages."));
-			FLinkerSave* const Linker = Cast<FLinkerSave>(Ar.GetLinker());
 			checkf(Linker, TEXT("Serializing FDerivedData requires a linker."));
-			FIoChunkId SaveId = EditorData->Save(*Linker);
-			Ar << SaveId;
+			checkf(Ar.IsCooking(), TEXT("Serializing FDerivedData is only supported for cooked packages."));
 		}
-		else
+
+		inline void operator()(const FCompositeBufferWithHash& BufferWithHash)
+		{
+			ChunkId = Linker->AddDerivedData(FValue::Compress(BufferWithHash.Buffer).GetData());
+		}
+
+		inline void operator()(const FCompressedBuffer& Buffer)
+		{
+			ChunkId = Linker->AddDerivedData(Buffer);
+		}
+
+		inline void operator()(const FCacheKeyWithId& CacheKeyWithId)
+		{
+			ChunkId = Linker->AddDerivedData(CacheKeyWithId.Key, CacheKeyWithId.Id);
+		}
+
+		FLinkerSave* Linker;
+		uint64 ChunkOffset = 0;
+		uint64 ChunkSize = 0;
+		FIoChunkId ChunkId;
+	};
+
+	static_assert(sizeof(FIoChunkId) == 12);
+	FVisitor Visitor(Ar);
+	Visit(Visitor);
+
+	Ar << Visitor.ChunkOffset;
+	Ar << Visitor.ChunkSize;
+	Ar.Serialize(&Visitor.ChunkId, sizeof(Visitor.ChunkId));
+}
+
+} // UE::DerivedData::Private
+#endif // WITH_EDITORONLY_DATA
+
+namespace UE
+{
+
+FStringBuilderBase& operator<<(FStringBuilderBase& Builder, const FDerivedData& Data)
+{
+#if WITH_EDITORONLY_DATA
+	using namespace DerivedData::Private;
+	if (Data.EditorData)
+	{
+		return Builder << *Data.EditorData;
+	}
 #endif
+
+	if (Data)
+	{
+		Builder << TEXTVIEW("Chunk: ID ");
+		String::BytesToHexLower(MakeArrayView(Data.ChunkId), Builder);
+		if (Data.ChunkOffset)
 		{
-			Ar << ChunkId;
+			Builder << TEXTVIEW(" / Offset ") << Data.ChunkOffset;
 		}
-		Ar << Flags;
+		if (Data.ChunkSize != MAX_uint64)
+		{
+			Builder << TEXTVIEW(" / Size ") << Data.ChunkSize;
+		}
+		return Builder;
 	}
+
+	return Builder << TEXT("Null");
 }
 
-FUniqueBuffer FDerivedData::LoadData(FArchive& Ar, FDerivedDataBufferAllocator Allocator)
+bool FDerivedData::ReferenceEquals(const FDerivedData& Other) const
 {
-	check(Ar.IsLoading());
-	unimplemented();
-	return FUniqueBuffer();
+#if WITH_EDITORONLY_DATA
+	if (EditorData && Other.EditorData)
+	{
+		return EditorData->ReferenceEquals(*Other.EditorData);
+	}
+#endif
+
+	if (ChunkOffset == Other.ChunkOffset && ChunkSize == Other.ChunkSize && MakeArrayView(ChunkId) == Other.ChunkId)
+	{
+		return true;
+	}
+
+	return false;
 }
 
+uint32 FDerivedData::ReferenceHash() const
+{
 #if WITH_EDITORONLY_DATA
+	if (EditorData)
+	{
+		return EditorData->ReferenceHash();
+	}
+#endif
+
+	FXxHash64Builder Builder;
+	Builder.Update(&ChunkOffset, sizeof(ChunkOffset));
+	Builder.Update(&ChunkSize, sizeof(ChunkSize));
+	Builder.Update(MakeMemoryView(ChunkId));
+	return uint32(Builder.Finalize().Hash);
+}
 
-void FDerivedData::SaveData(FArchive& Ar, const FDerivedData& Data)
+void FDerivedData::Serialize(FArchive& Ar, UObject* Owner)
+{
+	if (!Ar.IsPersistent() || Ar.IsObjectReferenceCollector() || Ar.ShouldSkipBulkData())
+	{
+		return;
+	}
+
+#if WITH_EDITORONLY_DATA
+	if (EditorData)
+	{
+		EditorData->Serialize(Ar, Owner);
+		Ar << Flags;
+		return;
+	}
+#endif
+
+	Ar << ChunkOffset;
+	Ar << ChunkSize;
+	Ar.Serialize(ChunkId, sizeof(ChunkId));
+	Ar << Flags;
+}
+
+#if WITH_EDITORONLY_DATA
+
+FDerivedData::FDerivedData(const DerivedData::FSharedString& Name, const FSharedBuffer& Data)
+	: EditorData(MakePimpl<DerivedData::Private::FEditorData, EPimplPtrMode::DeepCopy>(
+		Name, TInPlaceType<DerivedData::Private::FCompositeBufferWithHash>(), Data))
+	, Flags(EDerivedDataFlags::Required)
+{
+}
+
+FDerivedData::FDerivedData(const DerivedData::FSharedString& Name, const FCompositeBuffer& Data)
+	: EditorData(MakePimpl<DerivedData::Private::FEditorData, EPimplPtrMode::DeepCopy>(
+		Name, TInPlaceType<DerivedData::Private::FCompositeBufferWithHash>(), Data))
+	, Flags(EDerivedDataFlags::Required)
+{
+}
+
+FDerivedData::FDerivedData(const DerivedData::FSharedString& Name, const FCompressedBuffer& Data)
+	: EditorData(MakePimpl<DerivedData::Private::FEditorData, EPimplPtrMode::DeepCopy>(
+		Name, TInPlaceType<FCompressedBuffer>(), Data))
+	, Flags(EDerivedDataFlags::Required)
+{
+}
+
+FDerivedData::FDerivedData(const DerivedData::FSharedString& Name, const DerivedData::FCacheKey& Key)
+	: EditorData(MakePimpl<DerivedData::Private::FEditorData, EPimplPtrMode::DeepCopy>(
+		Name, TInPlaceType<DerivedData::Private::FCacheKeyWithId>(), Key, DerivedData::FValueId::Null))
+	, Flags(EDerivedDataFlags::Required)
 {
-	checkf(Data.EditorData, TEXT("Inline serialization is not supported for non-editor derived data references."));
-	check(Ar.IsSaving());
-	unimplemented();
 }
 
-FDerivedData::FDerivedData() = default;
-FDerivedData::FDerivedData(FDerivedData&& Other) = default;
-FDerivedData& FDerivedData::operator=(FDerivedData&& Other) = default;
+FDerivedData::FDerivedData(
+	const DerivedData::FSharedString& Name,
+	const DerivedData::FCacheKey& Key,
+	const DerivedData::FValueId& ValueId)
+	: EditorData(MakePimpl<DerivedData::Private::FEditorData, EPimplPtrMode::DeepCopy>(
+		Name, TInPlaceType<DerivedData::Private::FCacheKeyWithId>(), Key, ValueId))
+	, Flags(EDerivedDataFlags::Required)
+{
+}
 
-FDerivedData::FDerivedData(const FDerivedData& Other)
+const DerivedData::FSharedString& FDerivedData::GetName() const
 {
-	if (Other.EditorData)
+	return EditorData ? EditorData->GetName() : DerivedData::FSharedString::Empty;
+}
+
+void FDerivedData::SetFlags(EDerivedDataFlags InFlags)
+{
+	Flags = InFlags;
+}
+
+#endif // WITH_EDITORONLY_DATA
+
+} // UE
+
+///////////////////////////////////////////////////////////////////////////////////////////////////////////////////////
+
+namespace UE::DerivedData::Private
+{
+
+#if WITH_EDITORONLY_DATA
+static inline EDerivedDataIoStatus ConvertToIoStatus(EStatus Status)
+{
+	switch (Status)
+	{
+	case EStatus::Ok:
+		return EDerivedDataIoStatus::Ok;
+	case EStatus::Error:
+		return EDerivedDataIoStatus::Error;
+	case EStatus::Canceled:
+		return EDerivedDataIoStatus::Canceled;
+	default:
+		return EDerivedDataIoStatus::Unknown;
+	}
+}
+
+static inline EPriority ConvertToDerivedDataPriority(FDerivedDataIoPriority Priority)
+{
+	if (Priority == FDerivedDataIoPriority::Blocking())
+	{
+		return EPriority::Blocking;
+	}
+	if (FDerivedDataIoPriority::Highest().InterpolateTo(FDerivedDataIoPriority::High(), 0.8f) < Priority)
+	{
+		return EPriority::Highest;
+	}
+	if (FDerivedDataIoPriority::High().InterpolateTo(FDerivedDataIoPriority::Normal(), 0.6f) < Priority)
+	{
+		return EPriority::High;
+	}
+	if (FDerivedDataIoPriority::Normal().InterpolateTo(FDerivedDataIoPriority::Low(), 0.4f) < Priority)
 	{
-		EditorData = Other.EditorData->Clone();
-		ChunkId = DerivedData::IoStore::GEditorDerivedDataIoStore->AddData(EditorData.Get());
+		return EPriority::Normal;
+	}
+	if (FDerivedDataIoPriority::Low().InterpolateTo(FDerivedDataIoPriority::Lowest(), 0.2f) < Priority)
+	{
+		return EPriority::Low;
 	}
 	else
 	{
-		ChunkId = Other.ChunkId;
+		return EPriority::Lowest;
 	}
 }
+#endif
 
-FDerivedData& FDerivedData::operator=(const FDerivedData& Other)
+static inline int32 ConvertToIoDispatcherPriority(FDerivedDataIoPriority Priority)
 {
-	if (EditorData)
+	static_assert(IoDispatcherPriority_Min == FDerivedDataIoPriority::Lowest().Value);
+	static_assert(IoDispatcherPriority_Max == FDerivedDataIoPriority::Blocking().Value);
+	static_assert(IoDispatcherPriority_Medium == FDerivedDataIoPriority::Normal().Value);
+	return Priority.Value;
+}
+
+///////////////////////////////////////////////////////////////////////////////////////////////////////////////////////
+
+enum class EIoRequestType : uint8
+{
+	Read,
+	Cache,
+	Exists,
+};
+
+struct FIoChunkState
+{
+	FIoRequest Request;
+	uint64 ChunkOffset = 0;
+	uint64 ChunkSize = 0;
+	FIoChunkId ChunkId;
+	std::atomic<bool> bCanceled = false;
+};
+
+struct FIoEditorState
+{
+#if WITH_EDITORONLY_DATA
+	FEditorData EditorData;
+	FCacheKey CacheKey;
+	FIoHash Hash;
+
+	inline explicit FIoEditorState(const FEditorData& InEditorData)
+		: EditorData(InEditorData)
 	{
-		DerivedData::IoStore::GEditorDerivedDataIoStore->RemoveData(ChunkId);
 	}
-	if (Other.EditorData)
+#endif
+};
+
+struct FIoRequestState
+{
+	TVariant<TYPE_OF_NULLPTR, FIoChunkState, FIoEditorState> State;
+	FDerivedDataIoOptions Options;
+	FSharedBuffer Data;
+	uint64 Size = 0;
+	EIoRequestType Type = EIoRequestType::Read;
+	std::atomic<EDerivedDataIoStatus> Status = EDerivedDataIoStatus::Unknown;
+
+	inline void SetStatus(EDerivedDataIoStatus NewStatus)
 	{
-		EditorData = Other.EditorData->Clone();
-		ChunkId = DerivedData::IoStore::GEditorDerivedDataIoStore->AddData(EditorData.Get());
+		const bool bOk = NewStatus == EDerivedDataIoStatus::Ok;
+		Status.store(NewStatus, bOk ? std::memory_order_release : std::memory_order_relaxed);
 	}
-	else
+};
+
+class FIoResponse
+{
+public:
+	static FDerivedDataIoRequest Queue(
+		TPimplPtr<FIoResponse>& Self,
+		const FDerivedData& Data,
+		const FDerivedDataIoOptions& Options,
+		EIoRequestType Type);
+
+	static void Dispatch(
+		TPimplPtr<FIoResponse>& Self,
+		FDerivedDataIoResponse& OutResponse,
+		FDerivedDataIoPriority Priority,
+		FDerivedDataIoComplete&& OnComplete);
+
+	static inline FIoRequestState* TryGetRequest(const TPimplPtr<FIoResponse>& Self, const FDerivedDataIoRequest Handle)
 	{
-		EditorData.Reset();
-		ChunkId = Other.ChunkId;
+		return Self && Self->Requests.IsValidIndex(Handle.Index) ? &Self->Requests[Handle.Index] : nullptr;
 	}
-	return *this;
+
+	~FIoResponse();
+
+	void SetPriority(FDerivedDataIoPriority Priority);
+	bool Cancel();
+	bool Poll() const;
+
+	EDerivedDataIoStatus GetOverallStatus() const;
+
+private:
+	inline void BeginRequest() { RemainingRequests.fetch_add(1, std::memory_order_relaxed); }
+	void EndRequest();
+
+#if WITH_EDITORONLY_DATA
+	FRequestOwner Owner{EPriority::Normal};
+#endif
+	TArray<FIoRequestState> Requests;
+	std::atomic<uint32> RemainingRequests = 0;
+	std::atomic<EDerivedDataIoStatus> OverallStatus = EDerivedDataIoStatus::Unknown;
+	FDerivedDataIoComplete ResponseComplete;
+
+	friend FIoResponseDispatcher;
+};
+
+class FIoResponseDispatcher
+{
+public:
+	static void Dispatch(FIoResponse& Response, FDerivedDataIoPriority Priority);
+
+private:
+	FIoResponseDispatcher() = default;
+
+	void DispatchChunk(FIoResponse& Response, FIoRequestState& Request, FIoChunkState& Chunk);
+
+	static FIoReadOptions MakeIoReadOptions(const FIoChunkState& State, const FDerivedDataIoOptions& Options);
+	static void OnIoRequestComplete(FIoResponse& Response, FIoRequestState& Request, TIoStatusOr<FIoBuffer> Buffer);
+
+#if WITH_EDITORONLY_DATA
+	void DispatchEditor(FIoResponse& Response, FIoRequestState& Request, FIoEditorState& Editor, int32 RequestIndex);
+	static void OnCacheRequestComplete(FIoResponse& Response, FCacheGetChunkResponse&& Chunk);
+
+	TArray<FCacheGetChunkRequest> CacheRequests;
+#endif
+
+	FIoBatch Batch;
+	FDerivedDataIoPriority Priority;
+};
+
+///////////////////////////////////////////////////////////////////////////////////////////////////////////////////////
+
+FDerivedDataIoRequest FIoResponse::Queue(
+	TPimplPtr<FIoResponse>& Self,
+	const FDerivedData& Data,
+	const FDerivedDataIoOptions& Options,
+	const EIoRequestType Type)
+{
+	if (!Self)
+	{
+		Self = MakePimpl<FIoResponse>();
+	}
+
+	const int32 Index = Self->Requests.AddDefaulted();
+	FIoRequestState& RequestState = Self->Requests[Index];
+	RequestState.Options = Options;
+	RequestState.Type = Type;
+
+	FDerivedDataIoRequest Handle;
+	Handle.Index = Index;
+
+#if WITH_EDITORONLY_DATA
+	if (Data.EditorData)
+	{
+		RequestState.State.Emplace<FIoEditorState>(*Data.EditorData);
+		return Handle;
+	}
+#endif
+
+	if (Data)
+	{
+		RequestState.State.Emplace<FIoChunkState>();
+		FIoChunkState& ChunkState = RequestState.State.Get<FIoChunkState>();
+		ChunkState.ChunkOffset = Data.ChunkOffset;
+		ChunkState.ChunkSize = Data.ChunkSize;
+		ChunkState.ChunkId.Set(MakeMemoryView(Data.ChunkId));
+	}
+
+	return Handle;
 }
 
-FDerivedData::FDerivedData(const FSharedBuffer& Data)
-	: EditorData(DerivedData::Private::MakeEditorDerivedData(Data))
-	, ChunkId(DerivedData::IoStore::GEditorDerivedDataIoStore->AddData(EditorData.Get()))
+void FIoResponse::Dispatch(
+	TPimplPtr<FIoResponse>& InResponse,
+	FDerivedDataIoResponse& OutResponse,
+	FDerivedDataIoPriority Priority,
+	FDerivedDataIoComplete&& OnComplete)
 {
+	// An empty batch completes immediately.
+	if (!InResponse)
+	{
+		OutResponse.Reset();
+		if (OnComplete)
+		{
+			Invoke(OnComplete);
+		}
+	}
+
+	// Assign OutResponse early because OnComplete may reference it.
+	FIoResponse& Self = *InResponse;
+	Self.ResponseComplete = MoveTemp(OnComplete);
+	OutResponse.Response = MoveTemp(InResponse);
+
+	// The Begin/End pair blocks completion until dispatch is complete.
+	Self.BeginRequest();
+	FIoResponseDispatcher::Dispatch(Self, Priority);
+	Self.EndRequest();
 }
 
-FDerivedData::FDerivedData(const FCompositeBuffer& Data)
-	: EditorData(DerivedData::Private::MakeEditorDerivedData(Data))
-	, ChunkId(DerivedData::IoStore::GEditorDerivedDataIoStore->AddData(EditorData.Get()))
+FIoResponse::~FIoResponse()
 {
+	verifyf(Cancel(), TEXT("Requests must be complete before the response is destroyed but it has %u remaining."),
+		RemainingRequests.load(std::memory_order_relaxed));
 }
 
-FDerivedData::FDerivedData(const FCompressedBuffer& Data)
-	: EditorData(DerivedData::Private::MakeEditorDerivedData(Data))
-	, ChunkId(DerivedData::IoStore::GEditorDerivedDataIoStore->AddData(EditorData.Get()))
+void FIoResponse::SetPriority(FDerivedDataIoPriority Priority)
 {
+#if WITH_EDITORONLY_DATA
+	Owner.SetPriority(ConvertToDerivedDataPriority(Priority));
+#endif
+
+	const int32 IoDispatcherPriority = ConvertToIoDispatcherPriority(Priority);
+	for (FIoRequestState& Request : Requests)
+	{
+		if (FIoChunkState* ChunkState = Request.State.TryGet<FIoChunkState>())
+		{
+			ChunkState->Request.UpdatePriority(IoDispatcherPriority);
+		}
+	}
 }
 
-FDerivedData::FDerivedData(const DerivedData::FSharedString& Name, const DerivedData::FCacheKey& Key)
-	: EditorData(DerivedData::Private::MakeEditorDerivedData(Name, Key, DerivedData::FValueId::Null))
-	, ChunkId(DerivedData::IoStore::GEditorDerivedDataIoStore->AddData(EditorData.Get()))
+bool FIoResponse::Cancel()
 {
+	if (Poll())
+	{
+		return true;
+	}
+
+#if WITH_EDITORONLY_DATA
+	// Request cancellation is synchronous but is expected to be very fast.
+	Owner.Cancel();
+#endif
+
+	for (FIoRequestState& Request : Requests)
+	{
+		if (FIoChunkState* ChunkState = Request.State.TryGet<FIoChunkState>())
+		{
+			// Request cancellation only once because every call wakes the dispatcher.
+			if (!ChunkState->bCanceled.exchange(true, std::memory_order_relaxed))
+			{
+				ChunkState->Request.Cancel();
+			}
+		}
+	}
+
+	return Poll();
 }
 
-FDerivedData::FDerivedData(const DerivedData::FSharedString& Name, const DerivedData::FCacheKey& Key, const DerivedData::FValueId& ValueId)
-	: EditorData(DerivedData::Private::MakeEditorDerivedData(Name, Key, ValueId))
-	, ChunkId(DerivedData::IoStore::GEditorDerivedDataIoStore->AddData(EditorData.Get()))
+bool FIoResponse::Poll() const
 {
+	return OverallStatus.load(std::memory_order_relaxed) != EDerivedDataIoStatus::Unknown;
 }
 
-FDerivedData::FDerivedData(const DerivedData::FBuildDefinition& BuildDefinition, const DerivedData::FValueId& ValueId)
-	: EditorData(DerivedData::Private::MakeEditorDerivedData(BuildDefinition, ValueId))
-	, ChunkId(DerivedData::IoStore::GEditorDerivedDataIoStore->AddData(EditorData.Get()))
+EDerivedDataIoStatus FIoResponse::GetOverallStatus() const
 {
+	return OverallStatus.load(std::memory_order_relaxed);
 }
 
-FDerivedData::FDerivedData(FStringView CacheKey, FStringView CacheContext)
-	: EditorData(DerivedData::Private::MakeEditorDerivedData(CacheKey, CacheContext))
-	, ChunkId(DerivedData::IoStore::GEditorDerivedDataIoStore->AddData(EditorData.Get()))
+void FIoResponse::EndRequest()
 {
+	if (RemainingRequests.fetch_sub(1, std::memory_order_acq_rel) == 1)
+	{
+		// Calculate the overall status for the response.
+		static_assert(EDerivedDataIoStatus::Ok < EDerivedDataIoStatus::Error);
+		static_assert(EDerivedDataIoStatus::Error < EDerivedDataIoStatus::Canceled);
+		static_assert(EDerivedDataIoStatus::Canceled < EDerivedDataIoStatus::Unknown);
+		EDerivedDataIoStatus Status = EDerivedDataIoStatus::Ok;
+		for (FIoRequestState& Request : Requests)
+		{
+			Status = FMath::Max(Status, Request.Status.load(std::memory_order_relaxed));
+		}
+		OverallStatus.store(Status, std::memory_order_relaxed);
+
+		// Invoke the completion callback, but move it to the stack first because it may delete the response.
+		if (FDerivedDataIoComplete OnComplete = MoveTemp(ResponseComplete))
+		{
+			Invoke(OnComplete);
+		}
+	}
 }
 
-FDerivedData::~FDerivedData()
+///////////////////////////////////////////////////////////////////////////////////////////////////////////////////////
+
+void FIoResponseDispatcher::Dispatch(FIoResponse& Response, FDerivedDataIoPriority Priority)
 {
-	if (EditorData)
+#if WITH_EDITORONLY_DATA
+	FRequestBarrier Barrier(Response.Owner);
+#endif
+
+	FIoResponseDispatcher Dispatcher;
+	Dispatcher.Priority = Priority;
+
+	int32 RequestIndex = -1;
+	for (FIoRequestState& Request : Response.Requests)
 	{
-		using namespace DerivedData::IoStore;
-		if (IEditorDerivedDataIoStore* EditorBackend = GEditorDerivedDataIoStore.Get())
+		++RequestIndex;
+		if (FIoChunkState* ChunkState = Request.State.TryGet<FIoChunkState>())
 		{
-			verify(EditorBackend->RemoveData(ChunkId) == EditorData.Get());
+			Dispatcher.DispatchChunk(Response, Request, *ChunkState);
+		}
+	#if WITH_EDITORONLY_DATA
+		else if (FIoEditorState* EditorState = Request.State.TryGet<FIoEditorState>())
+		{
+			Dispatcher.DispatchEditor(Response, Request, *EditorState, RequestIndex);
+		}
+	#endif
+		else
+		{
+			Request.SetStatus(EDerivedDataIoStatus::Error);
 		}
 	}
+
+	Dispatcher.Batch.Issue();
+
+#if WITH_EDITORONLY_DATA
+	if (!Dispatcher.CacheRequests.IsEmpty())
+	{
+		GetCache().GetChunks(Dispatcher.CacheRequests, Response.Owner,
+			[Response = &Response](FCacheGetChunkResponse&& Chunk)
+			{
+				OnCacheRequestComplete(*Response, MoveTemp(Chunk));
+			});
+	}
+#endif
 }
 
-void FDerivedData::SetFlags(EDerivedDataFlags InFlags)
+void FIoResponseDispatcher::DispatchChunk(FIoResponse& Response, FIoRequestState& Request, FIoChunkState& Chunk)
 {
-	checkf(!EnumHasAllFlags(InFlags, EDerivedDataFlags::Optional | EDerivedDataFlags::MemoryMapped),
-		TEXT("Optional and MemoryMapped cannot be used together on derived data references."));
-	Flags = InFlags;
+	if (Request.Type == EIoRequestType::Read)
+	{
+		Response.BeginRequest();
+		Chunk.Request = Batch.ReadWithCallback(
+			Chunk.ChunkId,
+			MakeIoReadOptions(Chunk, Request.Options),
+			ConvertToIoDispatcherPriority(Priority),
+			[Response = &Response, &Request](TIoStatusOr<FIoBuffer> Buffer)
+			{
+				OnIoRequestComplete(*Response, Request, MoveTemp(Buffer));
+			});
+	}
+	else
+	{
+		TIoStatusOr<uint64> Size = FIoDispatcher::Get().GetSizeForChunk(Chunk.ChunkId);
+		if (Size.IsOk())
+		{
+			const uint64 TotalSize = Size.ValueOrDie();
+			const uint64 RequestOffset = Request.Options.GetOffset();
+			const uint64 AvailableSize = RequestOffset <= TotalSize ? TotalSize - RequestOffset : 0;
+			Request.Size = FMath::Min(Request.Options.GetSize(), AvailableSize);
+			Request.SetStatus(EDerivedDataIoStatus::Ok);
+		}
+		else
+		{
+			Request.SetStatus(EDerivedDataIoStatus::Error);
+		}
+	}
+}
+
+FIoReadOptions FIoResponseDispatcher::MakeIoReadOptions(
+	const FIoChunkState& State,
+	const FDerivedDataIoOptions& Options)
+{
+	const uint64 LocalOffset = Options.GetOffset();
+	const uint64 TotalOffset = State.ChunkOffset + LocalOffset;
+
+	FIoReadOptions ReadOptions;
+	ReadOptions.SetTargetVa(Options.GetTarget());
+
+	if (Options.GetSize() == MAX_uint64)
+	{
+		if (State.ChunkSize == MAX_uint64)
+		{
+			ReadOptions.SetRange(TotalOffset, MAX_uint64);
+		}
+		else
+		{
+			ReadOptions.SetRange(TotalOffset, LocalOffset <= State.ChunkSize ? State.ChunkSize - LocalOffset : 0);
+		}
+	}
+	else
+	{
+		ReadOptions.SetRange(TotalOffset, Options.GetSize());
+	}
+
+	return ReadOptions;
+}
+
+void FIoResponseDispatcher::OnIoRequestComplete(
+	FIoResponse& Response,
+	FIoRequestState& Request,
+	TIoStatusOr<FIoBuffer> StatusOrBuffer)
+{
+	FIoChunkState& Chunk = Request.State.Get<FIoChunkState>();
+
+	if (StatusOrBuffer.IsOk())
+	{
+		FIoBuffer Data = StatusOrBuffer.ConsumeValueOrDie();
+		const uint64 DataSize = Data.GetSize();
+
+		// Return a view of the target when one was provided, otherwise take ownership of the buffer.
+		if (Request.Options.GetTarget())
+		{
+			Request.Data = FSharedBuffer::MakeView(Request.Options.GetTarget(), DataSize);
+		}
+		else
+		{
+			Request.Data = FSharedBuffer::TakeOwnership(Data.Release().ConsumeValueOrDie(), DataSize, FMemory::Free);
+		}
+
+		Request.Size = DataSize;
+
+		Request.SetStatus(EDerivedDataIoStatus::Ok);
+	}
+	else
+	{
+		const bool bCanceled = StatusOrBuffer.Status().GetErrorCode() == EIoErrorCode::Cancelled;
+		Request.SetStatus(bCanceled ? EDerivedDataIoStatus::Canceled : EDerivedDataIoStatus::Error);
+	}
+
+	Chunk.Request.Release();
+	Response.EndRequest();
+}
+
+#if WITH_EDITORONLY_DATA
+
+void FIoResponseDispatcher::DispatchEditor(
+	FIoResponse& Response,
+	FIoRequestState& Request,
+	FIoEditorState& Editor,
+	int32 RequestIndex)
+{
+	struct FVisitor
+	{
+		void operator()(const FCompositeBufferWithHash& BufferWithHash) const
+		{
+			const uint64 TotalSize = BufferWithHash.Buffer.GetSize();
+			Editor->Hash = BufferWithHash.Hash;
+
+			const uint64 RequestOffset = Request->Options.GetOffset();
+			const uint64 AvailableSize = RequestOffset <= TotalSize ? TotalSize - RequestOffset : 0;
+			const uint64 RequestSize = FMath::Min(Request->Options.GetSize(), AvailableSize);
+			Request->Size = RequestSize;
+
+			if (Request->Type == EIoRequestType::Read)
+			{
+				auto Execute = [Response = Response, Request = Request, &BufferWithHash, RequestSize]
+				{
+					if (void* Target = Request->Options.GetTarget())
+					{
+						const FMutableMemoryView TargetView(Target, RequestSize);
+						BufferWithHash.Buffer.CopyTo(TargetView, Request->Options.GetOffset());
+						Request->Data = FSharedBuffer::MakeView(TargetView);
+					}
+					else
+					{
+						Request->Data = BufferWithHash.Buffer.Mid(Request->Options.GetOffset(), RequestSize).ToShared();
+					}
+					Request->SetStatus(EDerivedDataIoStatus::Ok);
+					Response->EndRequest();
+				};
+
+				// Execute small copy tasks inline to avoid task overhead.
+				Response->BeginRequest();
+				if (RequestSize <= 64 * 1024)
+				{
+					Execute();
+				}
+				else
+				{
+					FRequestBarrier Barrier(Response->Owner);
+					Response->Owner.LaunchTask(TEXT("DerivedDataCopy"), MoveTemp(Execute));
+				}
+			}
+			else
+			{
+				Request->SetStatus(EDerivedDataIoStatus::Ok);
+			}
+		}
+
+		void operator()(const FCompressedBuffer& Buffer) const
+		{
+			const uint64 TotalSize = Buffer.GetRawSize();
+			Editor->Hash = Buffer.GetRawHash();
+
+			const uint64 RequestOffset = Request->Options.GetOffset();
+			const uint64 AvailableSize = RequestOffset <= TotalSize ? TotalSize - RequestOffset : 0;
+			const uint64 RequestSize = FMath::Min(Request->Options.GetSize(), AvailableSize);
+			Request->Size = RequestSize;
+
+			if (Request->Type == EIoRequestType::Read)
+			{
+				auto Execute = [Response = Response, Request = Request, &Buffer, RequestSize]
+				{
+					FCompressedBufferReader Reader(Buffer);
+					if (void* Target = Request->Options.GetTarget())
+					{
+						const FMutableMemoryView TargetView(Target, RequestSize);
+						if (Reader.TryDecompressTo(TargetView, Request->Options.GetOffset()))
+						{
+							Request->Data = FSharedBuffer::MakeView(TargetView);
+						}
+					}
+					else
+					{
+						Request->Data = Reader.Decompress(Request->Options.GetOffset(), RequestSize);
+					}
+					Request->SetStatus(Request->Data ? EDerivedDataIoStatus::Ok : EDerivedDataIoStatus::Error);
+					Response->EndRequest();
+				};
+
+				// Execute small decompression tasks inline to avoid task overhead.
+				Response->BeginRequest();
+				if (RequestSize <= 16 * 1024)
+				{
+					Execute();
+				}
+				else
+				{
+					FRequestBarrier Barrier(Response->Owner);
+					Response->Owner.LaunchTask(TEXT("DerivedDataDecompress"), MoveTemp(Execute));
+				}
+			}
+			else
+			{
+				Request->SetStatus(EDerivedDataIoStatus::Ok);
+			}
+		}
+
+		void operator()(const FCacheKeyWithId& CacheKeyWithId) const
+		{
+			FCacheGetChunkRequest& Chunk = Dispatcher->CacheRequests.AddDefaulted_GetRef();
+			Chunk.Name = Editor->EditorData.GetName();
+			Chunk.Key = CacheKeyWithId.Key;
+			Chunk.Id = CacheKeyWithId.Id;
+			Chunk.RawOffset = Request->Options.GetOffset();
+			Chunk.RawSize = Request->Options.GetSize();
+			Chunk.Policy =
+				Request->Type == EIoRequestType::Read  ? (ECachePolicy::Default) :
+				Request->Type == EIoRequestType::Cache ? (ECachePolicy::Default | ECachePolicy::SkipData) :
+				                                         (ECachePolicy::Query   | ECachePolicy::SkipData);
+			Chunk.UserData = uint64(RequestIndex);
+			Response->BeginRequest();
+		}
+
+		FIoResponseDispatcher* Dispatcher;
+		FIoResponse* Response;
+		FIoRequestState* Request;
+		FIoEditorState* Editor;
+		int32 RequestIndex;
+	};
+
+	FVisitor Visitor{this, &Response, &Request, &Editor, RequestIndex};
+	Editor.EditorData.Visit(Visitor);
+}
+
+void FIoResponseDispatcher::OnCacheRequestComplete(FIoResponse& Response, FCacheGetChunkResponse&& Chunk)
+{
+	FIoRequestState& Request = Response.Requests[int32(Chunk.UserData)];
+	FIoEditorState& Editor = Request.State.Get<FIoEditorState>();
+	Editor.Hash = Chunk.RawHash;
+	Request.Size = Chunk.RawSize;
+	Request.Data = MoveTemp(Chunk.RawData);
+	Request.SetStatus(ConvertToIoStatus(Chunk.Status));
+	Response.EndRequest();
 }
 
 #endif // WITH_EDITORONLY_DATA
 
-} // UE
+} // UE::DerivedData::Private
 
 ///////////////////////////////////////////////////////////////////////////////////////////////////////////////////////
 
-namespace UE::DerivedData::IoStore
+namespace UE
 {
 
-void InitializeIoDispatcher()
+void FDerivedDataIoResponse::SetPriority(FDerivedDataIoPriority Priority)
 {
-#if WITH_EDITORONLY_DATA
-	FIoDispatcher& IoDispatcher = FIoDispatcher::Get();
-	TSharedRef<IEditorDerivedDataIoStore> EditorBackend = CreateEditorDerivedDataIoStore();
-	GEditorDerivedDataIoStore = EditorBackend;
-	IoDispatcher.Mount(EditorBackend);
-#endif // WITH_EDITORONLY_DATA
+	if (Response)
+	{
+		Response->SetPriority(Priority);
+	}
+}
+
+bool FDerivedDataIoResponse::Cancel()
+{
+	return Response ? Response->Cancel() : true;
 }
 
-void TearDownIoDispatcher()
+bool FDerivedDataIoResponse::Poll() const
 {
+	return Response ? Response->Poll() : true;
+}
+
+EDerivedDataIoStatus FDerivedDataIoResponse::GetOverallStatus() const
+{
+	return Response ? Response->GetOverallStatus() : EDerivedDataIoStatus::Ok;
+}
+
+EDerivedDataIoStatus FDerivedDataIoResponse::GetStatus(FDerivedDataIoRequest Handle) const
+{
+	using namespace DerivedData::Private;
+	if (FIoRequestState* Request = FIoResponse::TryGetRequest(Response, Handle))
+	{
+		return Request->Status.load(std::memory_order_relaxed);
+	}
+	return EDerivedDataIoStatus::Error;
+}
+
+FSharedBuffer FDerivedDataIoResponse::GetData(FDerivedDataIoRequest Handle) const
+{
+	using namespace DerivedData::Private;
+	if (FIoRequestState* Request = FIoResponse::TryGetRequest(Response, Handle))
+	{
+		if (Request->Status.load(std::memory_order_acquire) == EDerivedDataIoStatus::Ok)
+		{
+			return Request->Data;
+		}
+	}
+	return FSharedBuffer();
+}
+
+uint64 FDerivedDataIoResponse::GetSize(FDerivedDataIoRequest Handle) const
+{
+	using namespace DerivedData::Private;
+	if (FIoRequestState* Request = FIoResponse::TryGetRequest(Response, Handle))
+	{
+		if (Request->Status.load(std::memory_order_acquire) == EDerivedDataIoStatus::Ok)
+		{
+			return Request->Size;
+		}
+	}
+	return 0;
+}
+
+#if WITH_EDITORONLY_DATA
+const FIoHash* FDerivedDataIoResponse::GetHash(FDerivedDataIoRequest Handle) const
+{
+	using namespace DerivedData::Private;
+	if (FIoRequestState* Request = FIoResponse::TryGetRequest(Response, Handle))
+	{
+		if (Request->Status.load(std::memory_order_relaxed) == EDerivedDataIoStatus::Ok)
+		{
+			if (FIoEditorState* EditorState = Request->State.TryGet<FIoEditorState>())
+			{
+				return &EditorState->Hash;
+			}
+		}
+	}
+	return nullptr;
+}
+#endif // WITH_EDITORONLY_DATA
+
 #if WITH_EDITORONLY_DATA
-	GEditorDerivedDataIoStore.Reset();
+const DerivedData::FCacheKey* FDerivedDataIoResponse::GetCacheKey(FDerivedDataIoRequest Handle) const
+{
+	using namespace DerivedData;
+	using namespace DerivedData::Private;
+	if (FIoRequestState* Request = FIoResponse::TryGetRequest(Response, Handle))
+	{
+		if (Request->Status.load(std::memory_order_relaxed) == EDerivedDataIoStatus::Ok)
+		{
+			if (FIoEditorState* EditorState = Request->State.TryGet<FIoEditorState>())
+			{
+				if (EditorState->CacheKey != FCacheKey::Empty)
+				{
+					return &EditorState->CacheKey;
+				}
+			}
+		}
+	}
+	return nullptr;
+}
 #endif // WITH_EDITORONLY_DATA
+
+///////////////////////////////////////////////////////////////////////////////////////////////////////////////////////
+
+FDerivedDataIoRequest FDerivedDataIoBatch::Read(const FDerivedData& Data, const FDerivedDataIoOptions& Options)
+{
+	using namespace DerivedData::Private;
+	return FIoResponse::Queue(Response, Data, Options, EIoRequestType::Read);
 }
 
-} // UE::DerivedData::IoStore
+FDerivedDataIoRequest FDerivedDataIoBatch::Cache(const FDerivedData& Data, const FDerivedDataIoOptions& Options)
+{
+	using namespace DerivedData::Private;
+	return FIoResponse::Queue(Response, Data, Options, EIoRequestType::Cache);
+}
+
+FDerivedDataIoRequest FDerivedDataIoBatch::Exists(const FDerivedData& Data, const FDerivedDataIoOptions& Options)
+{
+	using namespace DerivedData::Private;
+	return FIoResponse::Queue(Response, Data, Options, EIoRequestType::Exists);
+}
+
+void FDerivedDataIoBatch::Dispatch(
+	FDerivedDataIoResponse& OutResponse,
+	FDerivedDataIoPriority Priority,
+	FDerivedDataIoComplete&& OnComplete)
+{
+	using namespace DerivedData::Private;
+	FIoResponse::Dispatch(Response, OutResponse, Priority, MoveTemp(OnComplete));
+}
+
+} // UE
diff --git a/Engine/Source/Runtime/CoreUObject/Private/Serialization/EditorDerivedData.cpp b/Engine/Source/Runtime/CoreUObject/Private/Serialization/EditorDerivedData.cpp
deleted file mode 100644
diff --git a/Engine/Source/Runtime/CoreUObject/Private/Serialization/EditorDerivedData.h b/Engine/Source/Runtime/CoreUObject/Private/Serialization/EditorDerivedData.h
deleted file mode 100644
diff --git a/Engine/Source/Runtime/CoreUObject/Private/Serialization/EditorDerivedDataIoStore.cpp b/Engine/Source/Runtime/CoreUObject/Private/Serialization/EditorDerivedDataIoStore.cpp
deleted file mode 100644
diff --git a/Engine/Source/Runtime/CoreUObject/Private/Serialization/EditorDerivedDataIoStore.h b/Engine/Source/Runtime/CoreUObject/Private/Serialization/EditorDerivedDataIoStore.h
deleted file mode 100644
diff --git a/Engine/Source/Runtime/CoreUObject/Public/Serialization/DerivedData.h b/Engine/Source/Runtime/CoreUObject/Public/Serialization/DerivedData.h
--- a/Engine/Source/Runtime/CoreUObject/Public/Serialization/DerivedData.h
+++ b/Engine/Source/Runtime/CoreUObject/Public/Serialization/DerivedData.h
@@ -3,90 +3,376 @@
 #pragma once
 
 #include "Containers/StringFwd.h"
-#include "IO/IoDispatcher.h"
 #include "Memory/MemoryFwd.h"
+#include "Memory/MemoryView.h"
 #include "Misc/EnumClassFlags.h"
-#include "Templates/UniquePtr.h"
+#include "Templates/Function.h"
+#include "Templates/PimplPtr.h"
 
 #define UE_API COREUOBJECT_API
 
 class FArchive;
 class FCompressedBuffer;
 class UObject;
+struct FIoHash;
 
-template <typename FuncType> class TFunctionRef;
-
-namespace UE::DerivedData { class FBuildDefinition; }
 namespace UE::DerivedData { struct FCacheKey; }
 namespace UE::DerivedData { struct FValueId; }
 namespace UE::DerivedData { template <typename CharType> class TSharedString; }
 namespace UE::DerivedData { using FSharedString = TSharedString<TCHAR>; }
-namespace UE::DerivedData::Private { class FEditorDerivedData; }
+namespace UE::DerivedData::Private { class FEditorData; }
+namespace UE::DerivedData::Private { class FIoResponse; }
 
 namespace UE
 {
 
+/** Flags that modify how this derived data is saved and how it may be loaded. */
 enum class EDerivedDataFlags : uint32
 {
 	None            = 0,
+	/** Stage the referenced data with required content. May combine with Optional to stage into both locations. */
 	Required        = 1 << 0,
+	/** Stage the referenced data with optional content. */
 	Optional        = 1 << 1,
+	/** Stage the referenced data with alignment padding to allow memory-mapped access on supported platforms. */
 	MemoryMapped    = 1 << 2,
 };
 
 ENUM_CLASS_FLAGS(EDerivedDataFlags);
 
-using FDerivedDataBufferAllocator = TFunctionRef<FUniqueBuffer (uint64 Size)>;
-
+/**
+ * Derived Data Reference
+ *
+ * Provides a consistent way to reference derived data:
+ * - That has the same API for the editor as for builds with cooked packages.
+ * - That has the same API to load from the cache as from a raw buffer, a compressed buffer, or a cooked package.
+ * - That supports asynchronous, event-driven, prioritized loading of derived data.
+ * - That supports saving to a cooked package without loading from the cache, when using Zen.
+ * - That replaces serialized bulk data, such as FByteBulkData, in a cooked package.
+ */
 class FDerivedData
 {
 public:
-	inline bool HasData() const { return ChunkId != FIoChunkId::InvalidChunkId; }
+	/** A null reference. */
+	static const FDerivedData Null;
 
-	inline const FIoChunkId& GetChunkId() const { return ChunkId; }
+	/** Constructs a null reference. */
+	FDerivedData() = default;
 
-	inline EDerivedDataFlags GetFlags() const { return Flags; }
+	/** Resets the reference to null. */
+	inline void Reset() { *this = FDerivedData(); }
 
-	UE_API void Serialize(FArchive& Ar, UObject* Owner);
+	/** Returns true if this is a null reference. */
+	inline bool IsNull() const;
 
-	UE_API static FUniqueBuffer LoadData(FArchive& Ar, FDerivedDataBufferAllocator Allocator);
+	/** Returns true if this is a non-null reference. */
+	inline explicit operator bool() const { return !IsNull(); }
 
-#if WITH_EDITORONLY_DATA
-	UE_API static void SaveData(FArchive& Ar, const FDerivedData& Data);
+	/** Returns true if this is a non-null reference. */
+	inline bool HasData() const { return !IsNull(); }
 
-	UE_API FDerivedData();
-	UE_API ~FDerivedData();
+	/** Returns the flags, which are mainly relevant for staged references. */
+	inline EDerivedDataFlags GetFlags() const { return Flags; }
 
-	UE_API FDerivedData(FDerivedData&& Other);
-	UE_API FDerivedData(const FDerivedData& Other);
-	UE_API FDerivedData& operator=(FDerivedData&& Other);
-	UE_API FDerivedData& operator=(const FDerivedData& Other);
+	/** Appends the name and description of this reference to the builder. */
+	UE_API friend FStringBuilderBase& operator<<(FStringBuilderBase& Builder, const FDerivedData& DerivedData);
 
-	/** Reference a value that is stored in a buffer. */
-	UE_API explicit FDerivedData(const FSharedBuffer& Data);
-	UE_API explicit FDerivedData(const FCompositeBuffer& Data);
-	UE_API explicit FDerivedData(const FCompressedBuffer& Data);
+	/** Returns true if this and the other are equivalent references. */
+	UE_API bool ReferenceEquals(const FDerivedData& Other) const;
+
+	/** Returns a hash that factors in the same members as ReferenceEquals. */
+	UE_API uint32 ReferenceHash() const;
+
+	/** Serializes this reference to or from a cooked package. */
+	UE_API void Serialize(FArchive& Ar, UObject* Owner);
+
+#if WITH_EDITORONLY_DATA
+	/** References a value that is stored in a buffer. */
+	UE_API FDerivedData(const DerivedData::FSharedString& Name, const FSharedBuffer& Data);
+	UE_API FDerivedData(const DerivedData::FSharedString& Name, const FCompositeBuffer& Data);
+	UE_API FDerivedData(const DerivedData::FSharedString& Name, const FCompressedBuffer& Data);
 
-	/** Reference a value that was saved using ICache::PutValue. */
+	/** References a value that was saved using ICache::PutValue. */
 	UE_API FDerivedData(const DerivedData::FSharedString& Name, const DerivedData::FCacheKey& Key);
-	/** Reference a value in a record that was saved using ICache::Put. */
+	/** References a value in a record that was saved using ICache::Put. */
 	UE_API FDerivedData(const DerivedData::FSharedString& Name, const DerivedData::FCacheKey& Key, const DerivedData::FValueId& ValueId);
-	/** Reference a value in the output from the build function for the build definition. */
-	UE_API FDerivedData(const DerivedData::FBuildDefinition& BuildDefinition, const DerivedData::FValueId& ValueId);
 
-	/** Reference a value that was saved using FDerivedDataCacheInterface::Put. */
-	UE_API FDerivedData(FStringView CacheKey, FStringView CacheContext);
+	/** Returns the name of the reference if available. */
+	UE_API const DerivedData::FSharedString& GetName() const;
 
-	/** Overwrite the existing flags. */
+	/** Overwrites the existing flags. */
 	UE_API void SetFlags(EDerivedDataFlags Flags);
 
 private:
-	TUniquePtr<DerivedData::Private::FEditorDerivedData> EditorData;
+	TPimplPtr<DerivedData::Private::FEditorData, EPimplPtrMode::DeepCopy> EditorData;
 #endif // WITH_EDITORONLY_DATA
 
 private:
-	FIoChunkId ChunkId = FIoChunkId::InvalidChunkId;
-	EDerivedDataFlags Flags = EDerivedDataFlags::Required;
+	// Chunk members can only be non-zero when loading from a cooked package.
+
+	/** The offset of this data within its chunk. */
+	uint64 ChunkOffset = 0;
+	/** The size of this data within its chunk. */
+	uint64 ChunkSize = 0;
+	/** The FIoChunkId. */
+	uint8 ChunkId[12]{};
+
+	/** The flags that modify how this derived data is saved and how it may be loaded. */
+	EDerivedDataFlags Flags = EDerivedDataFlags::None;
+
+	friend DerivedData::Private::FIoResponse;
+};
+
+inline const FDerivedData FDerivedData::Null;
+
+inline bool FDerivedData::IsNull() const
+{
+#if WITH_EDITORONLY_DATA
+	return !EditorData && ChunkId[11] == 0;
+#else
+	// Check the last byte of ChunkId, which is type of the FIoChunkId, where 0 is EIoChunkType::Invalid.
+	return ChunkId[11] == 0;
+#endif
+}
+
+///////////////////////////////////////////////////////////////////////////////////////////////////////////////////////
+
+/** Priority for scheduling an operation on a Derived Data Reference. */
+struct FDerivedDataIoPriority
+{
+	/**
+	 * Interpolates between this priority and a target priority.
+	 *
+	 * Alpha must be in [0, 1] where 0 returns this priority and 1 returns the target priority.
+	 * Blocking() is returned if either priority is blocking.
+	 */
+	constexpr FDerivedDataIoPriority InterpolateTo(const FDerivedDataIoPriority Target, const float Alpha) const
+	{
+		// Blocking input produces Blocking output.
+		if (*this == Blocking() || Target == Blocking())
+		{
+			return Blocking();
+		}
+		return {Value + int32((int64(Target.Value) - int64(Value)) * double(Alpha))};
+	}
+
+	static constexpr FDerivedDataIoPriority Lowest()   { return {MIN_int32}; }
+	static constexpr FDerivedDataIoPriority Low()      { return {MIN_int32 / 2}; }
+	static constexpr FDerivedDataIoPriority Normal()   { return {0}; }
+	static constexpr FDerivedDataIoPriority High()     { return {MAX_int32 / 2}; }
+	static constexpr FDerivedDataIoPriority Highest()  { return {MAX_int32 - 1}; }
+	static constexpr FDerivedDataIoPriority Blocking() { return {MAX_int32}; }
+
+	friend constexpr bool operator==(FDerivedDataIoPriority A, FDerivedDataIoPriority B) { return A.Value == B.Value; }
+	friend constexpr bool operator!=(FDerivedDataIoPriority A, FDerivedDataIoPriority B) { return A.Value != B.Value; }
+	friend constexpr bool operator<(FDerivedDataIoPriority A, FDerivedDataIoPriority B) { return A.Value < B.Value; }
+	friend constexpr uint32 GetTypeHash(const FDerivedDataIoPriority& Priority) { return uint32(Priority.Value); }
+
+	/** Priority. Lower values have lower priority and 0 is equivalent to Normal(). */
+	int32 Value = 0;
+};
+
+/** Status of an operation on a Derived Data Reference. */
+enum class EDerivedDataIoStatus : uint8
+{
+	/** The operation completed successfully. */
+	Ok,
+	/** The operation completed unsuccessfully. */
+	Error,
+	/** The operation was canceled before it completed. */
+	Canceled,
+	/** The operation has not completed or is otherwise in an unknown state. */
+	Unknown,
+};
+
+/**
+ * Request on a Derived Data Reference that may be used to query the associated response.
+ */
+class FDerivedDataIoRequest
+{
+public:
+	/** Resets this to a null request. */
+	constexpr inline void Reset() { Index = -1; }
+
+	/** Returns true if this is a null request. */
+	constexpr inline bool IsNull() const { return Index == -1; }
+
+	/** Returns true if this is a non-null request. */
+	constexpr inline explicit operator bool() const { return !IsNull(); }
+
+private:
+	int32 Index = -1;
+
+	friend DerivedData::Private::FIoResponse;
+};
+
+/**
+ * Tracks one or more associated requests dispatched by FDerivedDataIoBatch.
+ *
+ * Any const member functions are thread-safe, while any non-const member functions are not.
+ *
+ * Asserts that requests are complete upon destruction or assignment.
+ */
+class FDerivedDataIoResponse
+{
+public:
+	/** Resets this to a null response. Asserts that requests are complete. */
+	inline void Reset() { Response.Reset(); }
+
+	/** Returns true if this is a null response . */
+	inline bool IsNull() const { return !Response; }
+
+	/** Returns true if this is a non-null response. */
+	inline explicit operator bool() const { return !IsNull(); }
+
+	/** Sets the priority of the associated requests if they are not complete. */
+	UE_API void SetPriority(FDerivedDataIoPriority Priority);
+
+	/**
+	 * Cancels the associated requests if they are not complete.
+	 *
+	 * @return true if cancellation completed synchronously, or false if further polling is needed.
+	 */
+	UE_API bool Cancel();
+
+	/** Returns true if the associated requests are complete and the optional callback was invoked. */
+	UE_API bool Poll() const;
+
+	/**
+	 * Returns the overall status of the requests. See EDerivedDataIoStatus.
+	 *
+	 * The status of the overall response is, in priority order,
+	 * - Unknown if any request status is Unknown, or
+	 * - Canceled if any request status is Canceled, or
+	 * - Error if any request status is Error, otherwise
+	 * - Ok because every request status must be Ok.
+	 */
+	UE_API EDerivedDataIoStatus GetOverallStatus() const;
+
+	/** Returns the status of the request. See EDerivedDataIoStatus. */
+	UE_API EDerivedDataIoStatus GetStatus(FDerivedDataIoRequest Handle) const;
+
+	/** Returns the data for the request. Null if executing, failed, or canceled. */
+	UE_API FSharedBuffer GetData(FDerivedDataIoRequest Handle) const;
+
+	/** Returns the size for the request. Zero if not available, executing, failed, or canceled. */
+	UE_API uint64 GetSize(FDerivedDataIoRequest Handle) const;
+
+#if WITH_EDITORONLY_DATA
+	/** Returns the hash of the entire derived data. Null if not available, executing, failed, or canceled. */
+	UE_API const FIoHash* GetHash(FDerivedDataIoRequest Handle) const;
+
+	/** Returns the cache key for the request. Null if not applicable, executing, failed, or canceled. */
+	UE_API const DerivedData::FCacheKey* GetCacheKey(FDerivedDataIoRequest Handle) const;
+#endif
+
+private:
+	TPimplPtr<DerivedData::Private::FIoResponse> Response;
+
+	friend DerivedData::Private::FIoResponse;
+};
+
+/** Options for operations on a Derived Data Reference. */
+class FDerivedDataIoOptions
+{
+public:
+	/** Reads the entirety of the referenced data into a buffer that is allocated on demand. */
+	FDerivedDataIoOptions() = default;
+
+	/**
+	 * Reads referenced data into the target view, with an optional offset into the source data.
+	 *
+	 * The size of the data from the source offset must be at least the size of the target view.
+	 */
+	constexpr explicit FDerivedDataIoOptions(FMutableMemoryView TargetView, uint64 SourceOffset = 0)
+		: Target(TargetView.GetData())
+		, Size(TargetView.GetSize())
+		, Offset(SourceOffset)
+	{
+	}
+
+	/**
+	 * Reads referenced data into a buffer that is allocated on demand, with an optional offset into the source data.
+	 *
+	 * The size of the data from the source offset must be at least the given source size, if not MAX_uint64.
+	 */
+	constexpr explicit FDerivedDataIoOptions(uint64 SourceOffset, uint64 SourceSize = MAX_uint64)
+		: Size(SourceSize)
+		, Offset(SourceOffset)
+	{
+	}
+
+	constexpr void* GetTarget() const { return Target; }
+	constexpr uint64 GetSize() const { return Size; }
+	constexpr uint64 GetOffset() const { return Offset; }
+
+private:
+	/** An optional target address to read the data into. Size is required when a target is provided. */
+	void* Target = nullptr;
+
+	/** The size to read, starting from the offset. Set MAX_uint64 to read to the end. */
+	uint64 Size = MAX_uint64;
+
+	/** The offset into the data at which to start reading. */
+	uint64 Offset = 0;
+};
+
+using FDerivedDataIoComplete = TUniqueFunction<void ()>;
+
+/**
+ * Batch of requests to access Derived Data References.
+ *
+ * Gather related requests into a batch before issuing it to reduce the overhead of accessing related data.
+ */
+class FDerivedDataIoBatch
+{
+public:
+	/** Resets this to an empty batch. */
+	inline void Reset() { Response.Reset(); }
+
+	/** Returns true if this batch has no queued requests. */
+	inline bool IsEmpty() const { return !Response; }
+
+	/**
+	 * Reads the derived data into memory. Not executed until Dispatch() is invoked.
+	 *
+	 * The data, size, and status will be available on the response.
+	 *
+	 * @return A request to query the response with. May be discarded by the caller if not needed.
+	 */
+	UE_API FDerivedDataIoRequest Read(const FDerivedData& Data, const FDerivedDataIoOptions& Options = {});
+
+	/**
+	 * Caches the derived data into local storage without loading it. Not executed until Dispatch() is invoked.
+	 *
+	 * The size and status will be available on the response. Status is Ok if it caches and Error if not.
+	 *
+	 * @return A request to query the response with. May be discarded by the caller if not needed.
+	 */
+	UE_API FDerivedDataIoRequest Cache(const FDerivedData& Data, const FDerivedDataIoOptions& Options = {});
+
+	/**
+	 * Caches the derived data into local storage without loading it. Not executed until Dispatch() is invoked.
+	 *
+	 * The size and status will be available on the response. Status is Ok if it exists and Error if not.
+	 *
+	 * @return A request to query the response with. May be discarded by the caller if not needed.
+	 */
+	UE_API FDerivedDataIoRequest Exists(const FDerivedData& Data, const FDerivedDataIoOptions& Options = {});
+
+	/**
+	 * Dispatches the requests that have been queued to this batch and resets this batch to empty.
+	 *
+	 * The response must be kept alive until the batch completes, otherwise it will be automatically canceled.
+	 *
+	 * @param OutResponse   Response that is assigned before any requests begin executing.
+	 * @param Priority      Priority of this batch relative to other batches.
+	 * @param OnComplete    Optional completion callback invoked when the entire batch is complete or canceled.
+	 */
+	UE_API void Dispatch(FDerivedDataIoResponse& OutResponse, FDerivedDataIoPriority Priority = {}, FDerivedDataIoComplete&& OnComplete = {});
+
+private:
+	TPimplPtr<DerivedData::Private::FIoResponse> Response;
 };
 
 } // UE
@@ -94,8 +380,13 @@ private:
 namespace UE::DerivedData::IoStore
 {
 
-UE_API void InitializeIoDispatcher();
-UE_API void TearDownIoDispatcher();
+inline void InitializeIoDispatcher()
+{
+}
+
+inline void TearDownIoDispatcher()
+{
+}
 
 } // UE::DerivedData::IoStore
 
diff --git a/Engine/Source/Runtime/CoreUObject/Tests/Serialization/DerivedDataTest.cpp b/Engine/Source/Runtime/CoreUObject/Tests/Serialization/DerivedDataTest.cpp
new file mode 100644
--- /dev/null
+++ b/Engine/Source/Runtime/CoreUObject/Tests/Serialization/DerivedDataTest.cpp
@@ -0,0 +1,314 @@
+// Copyright Epic Games, Inc. All Rights Reserved.
+
+#if WITH_LOW_LEVEL_TESTS
+
+#include "Serialization/DerivedData.h"
+
+#include "HAL/Event.h"
+#include "IO/IoHash.h"
+#include "Memory/CompositeBuffer.h"
+#include "Memory/SharedBuffer.h"
+#include "Misc/StringBuilder.h"
+#include "TestHarness.h"
+
+#if WITH_EDITORONLY_DATA
+#include "DerivedDataCache.h"
+#include "DerivedDataRequestOwner.h"
+#include "DerivedDataSharedString.h"
+#include "DerivedDataValue.h"
+#endif
+
+namespace UE::DerivedDataTest::Private
+{
+
+static FSharedBuffer MakeSharedBuffer(uint64& Counter, uint64 Size)
+{
+	FUniqueBuffer Unique = FUniqueBuffer::Alloc(Size);
+	TArrayView<uint64> Values(static_cast<uint64*>(Unique.GetData()), Size / sizeof(uint64));
+	uint64 Index = 0;
+	for (uint64& Value : Values)
+	{
+		Value = ++Index * Size + Counter;
+	}
+	FMemory::Memset(Values.GetData() + Values.Num(), 0, Size - Values.Num() * sizeof(uint64));
+	++Counter;
+	return Unique.MoveToShared();
+}
+
+} // UE::DerivedDataTest::Private
+
+namespace UE
+{
+
+TEST_CASE("CoreUObject::Serialization::DerivedData", "[CoreUObject][Serialization]")
+{
+	using namespace DerivedData;
+	using namespace DerivedDataTest::Private;
+	uint64 Counter = 0;
+
+	SECTION("Null")
+	{
+		const FDerivedData DerivedData;
+
+		CHECK(DerivedData.IsNull());
+		CHECK_FALSE(DerivedData);
+		CHECK(DerivedData.GetFlags() == EDerivedDataFlags::None);
+		CHECK(WriteToString<64>(DerivedData) == TEXTVIEW("Null"));
+		CHECK(DerivedData.ReferenceEquals(FDerivedData::Null));
+		CHECK(DerivedData.ReferenceHash() == FDerivedData::Null.ReferenceHash());
+	#if WITH_EDITORONLY_DATA
+		CHECK(DerivedData.GetName().IsEmpty());
+	#endif
+
+		FEventRef Event;
+
+		FDerivedDataIoResponse Response;
+		CHECK(Response.IsNull());
+		CHECK_FALSE(Response);
+
+		FDerivedDataIoRequest Requests[3];
+		CHECK(Requests[0].IsNull());
+		CHECK_FALSE(Requests[0]);
+
+		{
+			FDerivedDataIoBatch Batch;
+			CHECK(Batch.IsEmpty());
+			Requests[0] = Batch.Read(DerivedData, {});
+			Requests[1] = Batch.Cache(DerivedData, {});
+			Requests[2] = Batch.Exists(DerivedData, {});
+			CHECK_FALSE(Batch.IsEmpty());
+			Batch.Dispatch(Response, {}, [&Event] { Event->Trigger(); });
+			CHECK(Batch.IsEmpty());
+		}
+
+		Event->Wait();
+
+		CHECK_FALSE(Response.IsNull());
+		CHECK(Response);
+
+		CHECK(Response.Poll());
+		CHECK(Response.Cancel());
+		CHECK(Response.GetOverallStatus() == EDerivedDataIoStatus::Error);
+
+		for (const FDerivedDataIoRequest& Request : Requests)
+		{
+			CHECK(Response.GetStatus(Request) == EDerivedDataIoStatus::Error);
+			CHECK(Response.GetData(Request).IsNull());
+			CHECK(Response.GetSize(Request) == 0);
+		#if WITH_EDITORONLY_DATA
+			CHECK(Response.GetHash(Request) == nullptr);
+			CHECK(Response.GetCacheKey(Request) == nullptr);
+		#endif
+		}
+	}
+
+#if WITH_EDITORONLY_DATA
+	SECTION("CompositeBuffer")
+	{
+		constexpr uint64 HalfSize = 32 * 1024;
+		const FCompositeBuffer Source(MakeSharedBuffer(Counter, HalfSize), MakeSharedBuffer(Counter, HalfSize));
+		const FIoHash SourceHash = FIoHash::HashBuffer(Source);
+		const FDerivedData DerivedData(TEXTVIEW("CompositeBuffer"), Source);
+		const FDerivedData DerivedDataCopy(TEXTVIEW("CompositeBufferCopy"), Source.ToShared());
+
+		CHECK_FALSE(DerivedData.IsNull());
+		CHECK(DerivedData);
+		CHECK(DerivedData.GetFlags() == EDerivedDataFlags::Required);
+		CHECK(WriteToString<128>(DerivedData).ToView() == WriteToString<128>(TEXTVIEW("Buffer: Size "),
+			HalfSize * 2, TEXTVIEW(" Hash "), SourceHash, TEXTVIEW(" for CompositeBuffer")));
+		CHECK(DerivedData.ReferenceEquals(DerivedDataCopy));
+		CHECK(DerivedData.ReferenceHash() == DerivedDataCopy.ReferenceHash());
+		CHECK(DerivedData.GetName() == TEXTVIEW("CompositeBuffer"));
+
+		FEventRef Event;
+
+		FUniqueBuffer Target = FUniqueBuffer::Alloc(HalfSize);
+		FDerivedDataIoResponse Response;
+		FDerivedDataIoRequest Requests[6];
+		{
+			FDerivedDataIoBatch Batch;
+			Requests[0] = Batch.Read(DerivedData, {});
+			Requests[1] = Batch.Cache(DerivedData, {});
+			Requests[2] = Batch.Exists(DerivedData, {});
+			Requests[3] = Batch.Exists(DerivedData, FDerivedDataIoOptions(HalfSize));
+			Requests[4] = Batch.Read(DerivedData, FDerivedDataIoOptions(HalfSize, HalfSize));
+			Requests[5] = Batch.Read(DerivedData, FDerivedDataIoOptions(Target, HalfSize / 2));
+			Batch.Dispatch(Response, {}, [&Event] { Event->Trigger(); });
+		}
+		Event->Wait();
+
+		CHECK(Response.Poll());
+		CHECK(Response.Cancel());
+		CHECK(Response.GetOverallStatus() == EDerivedDataIoStatus::Ok);
+
+		for (const FDerivedDataIoRequest& Request : Requests)
+		{
+			CHECK(Response.GetStatus(Request) == EDerivedDataIoStatus::Ok);
+			const FIoHash* RequestHash = Response.GetHash(Request);
+			CHECKED_IF(RequestHash)
+			{
+				CHECK(*RequestHash == SourceHash);
+			}
+			CHECK(Response.GetCacheKey(Request) == nullptr);
+		}
+
+		CHECK(Response.GetData(Requests[0]).GetSize() == HalfSize * 2);
+		CHECK(Response.GetData(Requests[1]).IsNull());
+		CHECK(Response.GetData(Requests[2]).IsNull());
+		CHECK(Response.GetData(Requests[3]).IsNull());
+		CHECK(Response.GetData(Requests[4]) == Source.GetSegments()[1]);
+		CHECK(FCompositeBuffer(Response.GetData(Requests[5])).EqualBytes(Source.Mid(HalfSize / 2, HalfSize)));
+
+		CHECK(Response.GetSize(Requests[0]) == HalfSize * 2);
+		CHECK(Response.GetSize(Requests[1]) == HalfSize * 2);
+		CHECK(Response.GetSize(Requests[2]) == HalfSize * 2);
+		CHECK(Response.GetSize(Requests[3]) == HalfSize);
+		CHECK(Response.GetSize(Requests[4]) == HalfSize);
+		CHECK(Response.GetSize(Requests[5]) == HalfSize);
+
+		Response.Reset();
+		CHECK(Response.IsNull());
+		CHECK(Response.GetStatus(Requests[0]) == EDerivedDataIoStatus::Error);
+		CHECK(Response.GetData(Requests[0]).IsNull());
+	}
+
+	SECTION("CompressedBuffer")
+	{
+		constexpr uint64 HalfSize = 16 * 1024;
+		const FSharedBuffer RawSource = MakeSharedBuffer(Counter, HalfSize * 2);
+		const FCompressedBuffer Source = FValue::Compress(RawSource, HalfSize).GetData();
+		const FDerivedData DerivedData(TEXTVIEW("CompressedBuffer"), Source);
+		const FDerivedData DerivedDataCopy(TEXTVIEW("CompressedBufferCopy"), Source);
+
+		CHECK_FALSE(DerivedData.IsNull());
+		CHECK(DerivedData);
+		CHECK(DerivedData.GetFlags() == EDerivedDataFlags::Required);
+		CHECK(WriteToString<128>(DerivedData).ToView() == WriteToString<128>(TEXTVIEW("Buffer: CompressedSize "),
+			Source.GetCompressedSize(), TEXTVIEW(" Size "), Source.GetRawSize(), TEXTVIEW(" Hash "),
+			Source.GetRawHash(), TEXTVIEW(" for CompressedBuffer")));
+		CHECK(DerivedData.ReferenceEquals(DerivedDataCopy));
+		CHECK(DerivedData.ReferenceHash() == DerivedDataCopy.ReferenceHash());
+		CHECK(DerivedData.GetName() == TEXTVIEW("CompressedBuffer"));
+
+		FEventRef Event;
+
+		FUniqueBuffer Target = FUniqueBuffer::Alloc(HalfSize);
+		FDerivedDataIoResponse Response;
+		FDerivedDataIoRequest Requests[6];
+		{
+			FDerivedDataIoBatch Batch;
+			Requests[0] = Batch.Read(DerivedData, {});
+			Requests[1] = Batch.Cache(DerivedData, {});
+			Requests[2] = Batch.Exists(DerivedData, {});
+			Requests[3] = Batch.Exists(DerivedData, FDerivedDataIoOptions(HalfSize));
+			Requests[4] = Batch.Read(DerivedData, FDerivedDataIoOptions(HalfSize, HalfSize));
+			Requests[5] = Batch.Read(DerivedData, FDerivedDataIoOptions(Target, HalfSize / 2));
+			Batch.Dispatch(Response, {}, [&Event] { Event->Trigger(); });
+		}
+		Event->Wait();
+
+		CHECK(Response.Poll());
+		CHECK(Response.Cancel());
+		CHECK(Response.GetOverallStatus() == EDerivedDataIoStatus::Ok);
+
+		for (const FDerivedDataIoRequest& Request : Requests)
+		{
+			CHECK(Response.GetStatus(Request) == EDerivedDataIoStatus::Ok);
+			const FIoHash* RequestHash = Response.GetHash(Request);
+			CHECKED_IF(RequestHash)
+			{
+				CHECK(*RequestHash == Source.GetRawHash());
+			}
+			CHECK(Response.GetCacheKey(Request) == nullptr);
+		}
+
+		CHECK(Response.GetData(Requests[0]).GetSize() == HalfSize * 2);
+		CHECK(Response.GetData(Requests[1]).IsNull());
+		CHECK(Response.GetData(Requests[2]).IsNull());
+		CHECK(Response.GetData(Requests[3]).IsNull());
+		CHECK(Response.GetData(Requests[4]).GetView().EqualBytes(RawSource.GetView().Right(HalfSize)));
+		CHECK(Response.GetData(Requests[5]).GetView().EqualBytes(RawSource.GetView().Mid(HalfSize / 2, HalfSize)));
+
+		CHECK(Response.GetSize(Requests[0]) == HalfSize * 2);
+		CHECK(Response.GetSize(Requests[1]) == HalfSize * 2);
+		CHECK(Response.GetSize(Requests[2]) == HalfSize * 2);
+		CHECK(Response.GetSize(Requests[3]) == HalfSize);
+		CHECK(Response.GetSize(Requests[4]) == HalfSize);
+		CHECK(Response.GetSize(Requests[5]) == HalfSize);
+	}
+
+	SECTION("CacheValue")
+	{
+		constexpr uint64 HalfSize = 4 * 1024;
+		const FSharedBuffer RawValue = MakeSharedBuffer(Counter, HalfSize * 2);
+		const FValue Value = FValue::Compress(RawValue, HalfSize);
+		const FCacheKey Key{FCacheBucket(ANSITEXTVIEW("Test")), Value.GetRawHash()};
+
+		FRequestOwner PutOwner(EPriority::Blocking);
+		GetCache().PutValue({{{TEXTVIEW("CacheValue")}, Key, Value}}, PutOwner);
+		PutOwner.Wait();
+
+		const FDerivedData DerivedData(TEXTVIEW("CacheValue"), Key);
+		const FDerivedData DerivedDataCopy(TEXTVIEW("CacheValueCopy"), Key);
+
+		CHECK_FALSE(DerivedData.IsNull());
+		CHECK(DerivedData);
+		CHECK(DerivedData.GetFlags() == EDerivedDataFlags::Required);
+		CHECK(WriteToString<128>(DerivedData).ToView() == WriteToString<128>(TEXTVIEW("Cache: Key "), Key,
+			TEXTVIEW(" for CacheValue")));
+		CHECK(DerivedData.ReferenceEquals(DerivedDataCopy));
+		CHECK(DerivedData.ReferenceHash() == DerivedDataCopy.ReferenceHash());
+		CHECK(DerivedData.GetName() == TEXTVIEW("CacheValue"));
+
+		FEventRef Event;
+
+		FUniqueBuffer Target = FUniqueBuffer::Alloc(HalfSize);
+		FDerivedDataIoResponse Response;
+		FDerivedDataIoRequest Requests[6];
+		{
+			FDerivedDataIoBatch Batch;
+			Requests[0] = Batch.Read(DerivedData, {});
+			Requests[1] = Batch.Cache(DerivedData, {});
+			Requests[2] = Batch.Exists(DerivedData, {});
+			Requests[3] = Batch.Exists(DerivedData, FDerivedDataIoOptions(HalfSize));
+			Requests[4] = Batch.Read(DerivedData, FDerivedDataIoOptions(HalfSize, HalfSize));
+			Requests[5] = Batch.Read(DerivedData, FDerivedDataIoOptions(Target, HalfSize / 2));
+			Batch.Dispatch(Response, {}, [&Event] { Event->Trigger(); });
+		}
+		Event->Wait();
+
+		CHECK(Response.Poll());
+		CHECK(Response.Cancel());
+		CHECK(Response.GetOverallStatus() == EDerivedDataIoStatus::Ok);
+
+		for (const FDerivedDataIoRequest& Request : Requests)
+		{
+			CHECK(Response.GetStatus(Request) == EDerivedDataIoStatus::Ok);
+			const FIoHash* RequestHash = Response.GetHash(Request);
+			CHECKED_IF(RequestHash)
+			{
+				CHECK(*RequestHash == Value.GetRawHash());
+			}
+			CHECK(Response.GetCacheKey(Request) == nullptr);
+		}
+
+		CHECK(Response.GetData(Requests[0]).GetSize() == HalfSize * 2);
+		CHECK(Response.GetData(Requests[1]).IsNull());
+		CHECK(Response.GetData(Requests[2]).IsNull());
+		CHECK(Response.GetData(Requests[3]).IsNull());
+		CHECK(Response.GetData(Requests[4]).GetView().EqualBytes(RawValue.GetView().Right(HalfSize)));
+		CHECK(Response.GetData(Requests[5]).GetView().EqualBytes(RawValue.GetView().Mid(HalfSize / 2, HalfSize)));
+
+		CHECK(Response.GetSize(Requests[0]) == HalfSize * 2);
+		CHECK(Response.GetSize(Requests[1]) == HalfSize * 2);
+		CHECK(Response.GetSize(Requests[2]) == HalfSize * 2);
+		CHECK(Response.GetSize(Requests[3]) == HalfSize);
+		CHECK(Response.GetSize(Requests[4]) == HalfSize);
+		CHECK(Response.GetSize(Requests[5]) == HalfSize);
+	}
+#endif // WITH_EDITORONLY_DATA
+}
+
+} // UE
+
+#endif // WITH_LOW_LEVEL_TESTS
