AsyncLoading2: Add support for having multiple PackageStore backends

diff --git a/Engine/Source/Runtime/Core/Private/IO/PackageStore.cpp b/Engine/Source/Runtime/Core/Private/IO/PackageStore.cpp
--- a/Engine/Source/Runtime/Core/Private/IO/PackageStore.cpp
+++ b/Engine/Source/Runtime/Core/Private/IO/PackageStore.cpp
@@ -116,3 +116,83 @@ FPackageStoreEntryResource FPackageStoreEntryResource::FromCbObject(const FCbObj
 
 	return Entry;
 }
+
+thread_local int32 FPackageStore::ThreadReadCount = 0;
+
+FPackageStoreReadScope::FPackageStoreReadScope(FPackageStore& InPackageStore)
+	: PackageStore(InPackageStore)
+{
+	if (!PackageStore.ThreadReadCount)
+	{
+		for (const TSharedRef<IPackageStoreBackend>& Backend : PackageStore.Backends)
+		{
+			Backend->BeginRead();
+		}
+	}
+	++PackageStore.ThreadReadCount;
+}
+
+FPackageStoreReadScope::~FPackageStoreReadScope()
+{
+	check(PackageStore.ThreadReadCount > 0);
+	if (--PackageStore.ThreadReadCount == 0)
+	{
+		for (const TSharedRef<IPackageStoreBackend>& Backend : PackageStore.Backends)
+		{
+			Backend->EndRead();
+		}
+	}
+}
+
+FPackageStore::FPackageStore()
+	: BackendContext(MakeShared<FPackageStoreBackendContext>())
+{
+	
+}
+
+FPackageStore& FPackageStore::Get()
+{
+	static FPackageStore Instance;
+	return Instance;
+}
+
+
+void FPackageStore::Mount(TSharedRef<IPackageStoreBackend> Backend)
+{
+	check(IsInGameThread());
+	Backends.Add(Backend);
+	Backend->OnMounted(BackendContext);
+}
+
+EPackageStoreEntryStatus FPackageStore::GetPackageStoreEntry(FPackageId PackageId, FPackageStoreEntry& OutPackageStoreEntry)
+{
+	check(ThreadReadCount);
+	for (const TSharedRef<IPackageStoreBackend>& Backend : Backends)
+	{
+		EPackageStoreEntryStatus Status = Backend->GetPackageStoreEntry(PackageId, OutPackageStoreEntry);
+		if (Status >= EPackageStoreEntryStatus::Pending)
+		{
+			return Status;
+		}
+	}
+	return EPackageStoreEntryStatus::Missing;
+}
+
+bool FPackageStore::GetPackageRedirectInfo(FPackageId PackageId, FName& OutSourcePackageName, FPackageId& OutRedirectedToPackageId)
+{
+	check(ThreadReadCount);
+	for (const TSharedRef<IPackageStoreBackend>& Backend : Backends)
+	{
+		if (Backend->GetPackageRedirectInfo(PackageId, OutSourcePackageName, OutRedirectedToPackageId))
+		{
+			return true;
+		}
+	}
+	return false;
+}
+
+FPackageStoreBackendContext::FPendingEntriesAddedEvent& FPackageStore::OnPendingEntriesAdded()
+{
+	return BackendContext->PendingEntriesAdded;
+}
+
diff --git a/Engine/Source/Runtime/Core/Private/Misc/CoreDelegates.cpp b/Engine/Source/Runtime/Core/Private/Misc/CoreDelegates.cpp
--- a/Engine/Source/Runtime/Core/Private/Misc/CoreDelegates.cpp
+++ b/Engine/Source/Runtime/Core/Private/Misc/CoreDelegates.cpp
@@ -204,7 +204,9 @@ FCoreDelegates::FPreloadPackageShaderMaps FCoreDelegates::PreloadPackageShaderMa
 FCoreDelegates::FReleasePreloadedPackageShaderMaps FCoreDelegates::ReleasePreloadedPackageShaderMaps;
 FCoreDelegates::FOnLogVerbosityChanged FCoreDelegates::OnLogVerbosityChanged;
 
+PRAGMA_DISABLE_DEPRECATION_WARNINGS
 FCoreDelegates::FCreatePackageStore FCoreDelegates::CreatePackageStore;
+PRAGMA_ENABLE_DEPRECATION_WARNINGS
 
 FCoreDelegates::FApplicationNetworkInitializationChanged FCoreDelegates::ApplicationNetworkInitializationChanged;
 
diff --git a/Engine/Source/Runtime/Core/Public/IO/PackageStore.h b/Engine/Source/Runtime/Core/Public/IO/PackageStore.h
--- a/Engine/Source/Runtime/Core/Public/IO/PackageStore.h
+++ b/Engine/Source/Runtime/Core/Public/IO/PackageStore.h
@@ -32,9 +32,9 @@ struct FPackageStoreExportInfo
 enum class EPackageStoreEntryStatus
 {
 	None,
-	Ok,
-	Pending,
 	Missing,
+	Pending,
+	Ok,
 };
 
 /**
@@ -124,7 +124,7 @@ struct FPackageStoreEntryResource
 /**
  * Stores information about available packages that can be loaded.
  */
-class IPackageStore
+class UE_DEPRECATED(5.1, "Use IPackageStoreBackend instead") IPackageStore
 {
 public:
 	/* Destructor. */
@@ -152,8 +152,10 @@ public:
 	virtual FEntriesAddedEvent& OnPendingEntriesAdded() = 0;
 };
 
-class FPackageStoreBase
+class UE_DEPRECATED(5.1, "Use IPackageStoreBackend instead") FPackageStoreBase
+PRAGMA_DISABLE_DEPRECATION_WARNINGS
 	: public IPackageStore
+PRAGMA_ENABLE_DEPRECATION_WARNINGS
 {
 public:
 	virtual FEntriesAddedEvent& OnPendingEntriesAdded() override
@@ -163,4 +165,77 @@ public:
 
 protected:
 	FEntriesAddedEvent PendingEntriesAdded;
-};
\ No newline at end of file
+};
+
+class FPackageStoreBackendContext
+{
+public:
+	/* Event broadcasted when pending entries are completed and added to the package store */
+	DECLARE_EVENT(FPackageStoreBackendContext, FPendingEntriesAddedEvent);
+	FPendingEntriesAddedEvent PendingEntriesAdded;
+};
+
+/**
+ * Package store backend interface.
+ */
+class IPackageStoreBackend
+{
+public:
+	/* Destructor. */
+	virtual ~IPackageStoreBackend() { }
+
+	/** Called when the backend is mounted */
+	virtual void OnMounted(TSharedRef<const FPackageStoreBackendContext> Context) = 0;
+
+	/** Called when the loader enters a package store read scope. */
+	virtual void BeginRead() = 0;
+
+	/** Called when the loader exits a package store read scope. */
+	virtual void EndRead() = 0;
+
+	/* Returns the package store entry data with export info and imported packages for the specified package ID. */
+	virtual EPackageStoreEntryStatus GetPackageStoreEntry(FPackageId PackageId, FPackageStoreEntry& OutPackageStoreEntry) = 0;
+
+	/* Returns the redirected package ID and source package name for the specified package ID if it's being redirected. */
+	virtual bool GetPackageRedirectInfo(FPackageId PackageId, FName& OutSourcePackageName, FPackageId& OutRedirectedToPackageId) = 0;
+};
+
+/**
+ * Stores information about available packages that can be loaded.
+ */
+class FPackageStore
+{
+public:
+	CORE_API static FPackageStore& Get();
+
+	/* Mount a package store backend. */
+	CORE_API void Mount(TSharedRef<IPackageStoreBackend> Backend);
+
+	/* Returns the package store entry data with export info and imported packages for the specified package ID. */
+	CORE_API EPackageStoreEntryStatus GetPackageStoreEntry(FPackageId PackageId, FPackageStoreEntry& OutPackageStoreEntry);
+
+	/* Returns the redirected package ID and source package name for the specified package ID if it's being redirected. */
+	CORE_API bool GetPackageRedirectInfo(FPackageId PackageId, FName& OutSourcePackageName, FPackageId& OutRedirectedToPackageId);
+
+	CORE_API FPackageStoreBackendContext::FPendingEntriesAddedEvent& OnPendingEntriesAdded();
+
+private:
+	FPackageStore();
+
+	friend class FPackageStoreReadScope;
+
+	TSharedRef<FPackageStoreBackendContext> BackendContext;
+	TArray<TSharedRef<IPackageStoreBackend>> Backends;
+
+	static thread_local int32 ThreadReadCount;
+};
+
+class FPackageStoreReadScope
+{
+public:
+	CORE_API FPackageStoreReadScope(FPackageStore& InPackageStore);
+	CORE_API ~FPackageStoreReadScope();
+
+private:
+	FPackageStore& PackageStore;
+};
diff --git a/Engine/Source/Runtime/Core/Public/Misc/CoreDelegates.h b/Engine/Source/Runtime/Core/Public/Misc/CoreDelegates.h
--- a/Engine/Source/Runtime/Core/Public/Misc/CoreDelegates.h
+++ b/Engine/Source/Runtime/Core/Public/Misc/CoreDelegates.h
@@ -644,8 +644,12 @@ public:
 	DECLARE_MULTICAST_DELEGATE_ThreeParams(FOnLogVerbosityChanged, const FLogCategoryName& /* CategoryName */, ELogVerbosity::Type /* OldVerbosity */, ELogVerbosity::Type /* NewVerbosity */);
 	static FOnLogVerbosityChanged OnLogVerbosityChanged;
 
+	UE_DEPRECATED(5.1, "Use FPackageStore::Mount() instead")
 	DECLARE_DELEGATE_RetVal(TSharedPtr<class IPackageStore>, FCreatePackageStore);
+	UE_DEPRECATED(5.1, "Use FPackageStore::Mount() instead")
+	PRAGMA_DISABLE_DEPRECATION_WARNINGS
 	static FCreatePackageStore CreatePackageStore;
+	PRAGMA_ENABLE_DEPRECATION_WARNINGS
 
 	// Called immediately before the parent process will start responding to signals to fork
 	static FSimpleMulticastDelegate OnParentBeginFork;
diff --git a/Engine/Source/Runtime/CoreUObject/Private/Misc/PackageName.cpp b/Engine/Source/Runtime/CoreUObject/Private/Misc/PackageName.cpp
--- a/Engine/Source/Runtime/CoreUObject/Private/Misc/PackageName.cpp
+++ b/Engine/Source/Runtime/CoreUObject/Private/Misc/PackageName.cpp
@@ -1656,26 +1656,36 @@ FPackageName::EPackageLocationFilter FPackageName::DoesPackageExistEx(const FPac
 
 	if (((uint8)Filter & (uint8)EPackageLocationFilter::IoDispatcher))
 	{
+		bool bFoundInIoDispatcher = false;
 		if (DoesPackageExistOverrideDelegate.IsBound())
 		{
 			if (DoesPackageExistOverrideDelegate.Execute(PackagePath.GetPackageFName()))
 			{
-				if (OutPackagePath)
+				bFoundInIoDispatcher = true;
+			}
+		}
+		else if (FIoDispatcher::IsInitialized())
+		{
+			bFoundInIoDispatcher = FIoDispatcher::Get().DoesChunkExist(CreatePackageDataChunkId(FPackageId::FromName(PackagePath.GetPackageFName())));
+		}
+
+		if (bFoundInIoDispatcher)
+		{
+			if (OutPackagePath)
+			{
+				*OutPackagePath = PackagePath;
+				if (OutPackagePath->GetHeaderExtension() == EPackageExtension::Unspecified)
 				{
-					*OutPackagePath = PackagePath;
-					if (OutPackagePath->GetHeaderExtension() == EPackageExtension::Unspecified)
-					{
-						OutPackagePath->SetHeaderExtension(EPackageExtension::EmptyString);
-					}
+					OutPackagePath->SetHeaderExtension(EPackageExtension::EmptyString);
 				}
+			}
 
-				Result = EPackageLocationFilter((uint8)Result | (uint8)EPackageLocationFilter::IoDispatcher);
+			Result = EPackageLocationFilter((uint8)Result | (uint8)EPackageLocationFilter::IoDispatcher);
 
-				// if we just want to find any existence, then we are done
-				if (Filter == EPackageLocationFilter::Any)
-				{
-					return Result;
-				}
+			// if we just want to find any existence, then we are done
+			if (Filter == EPackageLocationFilter::Any)
+			{
+				return Result;
 			}
 		}
 	}
diff --git a/Engine/Source/Runtime/CoreUObject/Private/Serialization/AsyncLoading2.cpp b/Engine/Source/Runtime/CoreUObject/Private/Serialization/AsyncLoading2.cpp
--- a/Engine/Source/Runtime/CoreUObject/Private/Serialization/AsyncLoading2.cpp
+++ b/Engine/Source/Runtime/CoreUObject/Private/Serialization/AsyncLoading2.cpp
@@ -759,19 +759,15 @@ private:
 	// Packages in active loading or completely loaded packages, with Desc.DiskPackageName as key.
 	// Does not track temp packages with custom UPackage names, since they are never imorted by other packages.
 	TMap<FPackageId, FLoadedPackageRef> Packages;
-	IPackageStore* PackageStore = nullptr;
+	FPackageStore& PackageStore;
 
 public:
-	FLoadedPackageStore()
+	FLoadedPackageStore(FPackageStore& InPackageStore)
+		: PackageStore(InPackageStore)
 	{
 		Packages.Reserve(32768);
 	}
 
-	void Initialize(IPackageStore& InPackageStore)
-	{
-		PackageStore = &InPackageStore;
-	}
-
 	int32 NumTracked() const
 	{
 		return Packages.Num();
@@ -1687,7 +1683,7 @@ struct FAsyncPackage2
 	/** Creates GC clusters from loaded objects */
 	EAsyncPackageState::Type CreateClusters(FAsyncLoadingThreadState2& ThreadState);
 
-	void ImportPackagesRecursive(FIoBatch& IoBatch, IPackageStore& PackageStore);
+	void ImportPackagesRecursive(FIoBatch& IoBatch, FPackageStore& PackageStore);
 	void StartLoading(FIoBatch& IoBatch);
 
 #if WITH_IOSTORE_IN_EDITOR
@@ -1696,7 +1692,7 @@ struct FAsyncPackage2
 
 private:
 
-	void ImportPackagesRecursiveInner(FIoBatch& IoBatch, IPackageStore& PackageStore, const TArrayView<const FPackageId>& ImportedPackageIds, int32& ImportedPackageIndex);
+	void ImportPackagesRecursiveInner(FIoBatch& IoBatch, FPackageStore& PackageStore, const TArrayView<const FPackageId>& ImportedPackageIds, int32& ImportedPackageIndex);
 
 	uint8 PackageNodesMemory[EEventLoadNode2::Package_NumPhases * sizeof(FEventLoadNode2)];
 	/** Basic information associated with this package */
@@ -2091,7 +2087,7 @@ private:
 
 	IAsyncPackageLoader* UncookedPackageLoader;
 
-	TSharedPtr<IPackageStore> PackageStore;
+	FPackageStore& PackageStore;
 	FLoadedPackageStore LoadedPackageStore;
 	FGlobalImportStore GlobalImportStore;
 	TSpscQueue<FPackageRequest> PackageRequestQueue;
@@ -2596,25 +2592,11 @@ void FAsyncLoadingThread2::InitializeLoading()
 	FileOpenLogWrapper = (FPlatformFileOpenLog*)(FPlatformFileManager::Get().FindPlatformFile(FPlatformFileOpenLog::GetTypeName()));
 #endif
 
-	UE_CLOG(!FCoreDelegates::CreatePackageStore.IsBound(), LogStreaming, Fatal, TEXT("CreatePackageStore delegate not bound when initializing package loader"));
-	PackageStore = FCoreDelegates::CreatePackageStore.Execute();
-	UE_CLOG(!PackageStore.IsValid(), LogStreaming, Fatal, TEXT("CreatePackageStore delegate returned null when initializing package loader"));
-	PackageStore->Initialize();
-	PackageStore->OnPendingEntriesAdded().AddLambda([this]()
+	PackageStore.OnPendingEntriesAdded().AddLambda([this]()
 	{
 		AltZenaphore.NotifyOne();
 	});
 	
-	FPackageName::DoesPackageExistOverride().BindLambda([this](FName PackageName)
-	{
-		LazyInitializeFromLoadPackage();
-		
-		PackageStore->Lock();
-		bool bExists = PackageStore->DoesPackageExist(FPackageId::FromName(PackageName));
-		PackageStore->Unlock();
-		return bExists;
-	});
-
 	AsyncThreadReady.Increment();
 
 	UE_LOG(LogStreaming, Display, TEXT("AsyncLoading2 - Initialized"));
@@ -2678,16 +2660,16 @@ bool FAsyncLoadingThread2::CreateAsyncPackagesFromQueue(FAsyncLoadingThreadState
 
 	FIoBatch IoBatch = IoDispatcher.NewBatch();
 
-	PackageStore->Lock();
+	FPackageStoreReadScope _(PackageStore);
 
 	for (auto It = PendingPackages.CreateIterator(); It; ++It)
 	{
 		FAsyncPackage2* PendingPackage = *It;
 		FPackageStoreEntry PackageEntry;
-		if (EPackageStoreEntryStatus::Ok == PackageStore->GetPackageStoreEntry(PendingPackage->Desc.PackageIdToLoad, PackageEntry))
+		if (EPackageStoreEntryStatus::Ok == PackageStore.GetPackageStoreEntry(PendingPackage->Desc.PackageIdToLoad, PackageEntry))
 		{
 			InitializeAsyncPackageFromPackageStore(PendingPackage, PackageEntry);
-			PendingPackage->ImportPackagesRecursive(IoBatch, *PackageStore.Get());
+			PendingPackage->ImportPackagesRecursive(IoBatch, PackageStore);
 			PendingPackage->StartLoading(IoBatch);
 			It.RemoveCurrent();
 		}
@@ -2725,7 +2707,7 @@ bool FAsyncLoadingThread2::CreateAsyncPackagesFromQueue(FAsyncLoadingThreadState
 			{
 				FName SourcePackageName;
 				FPackageId RedirectedToPackageId;
-				if (PackageStore->GetPackageRedirectInfo(PackageIdToLoad, SourcePackageName, RedirectedToPackageId))
+				if (PackageStore.GetPackageRedirectInfo(PackageIdToLoad, SourcePackageName, RedirectedToPackageId))
 				{
 					PackageIdToLoad = RedirectedToPackageId;
 					Request.PackagePath.Empty(); // We no longer know the path but it will be set again when serializing the package summary
@@ -2750,7 +2732,7 @@ bool FAsyncLoadingThread2::CreateAsyncPackagesFromQueue(FAsyncLoadingThreadState
 			}
 
 			FPackageStoreEntry PackageEntry;
-			EPackageStoreEntryStatus PackageStatus = PackageStore->GetPackageStoreEntry(PackageIdToLoad, PackageEntry);
+			EPackageStoreEntryStatus PackageStatus = PackageStore.GetPackageStoreEntry(PackageIdToLoad, PackageEntry);
 			if (PackageStatus == EPackageStoreEntryStatus::Missing)
 			{
 				// While there is an active load request for (InName=/Temp/PackageABC_abc, InPackageToLoadFrom=/Game/PackageABC), then allow these requests too:
@@ -2759,7 +2741,7 @@ bool FAsyncLoadingThread2::CreateAsyncPackagesFromQueue(FAsyncLoadingThreadState
 				if (Package)
 				{
 					PackageIdToLoad = Package->Desc.PackageIdToLoad;
-					PackageStatus = PackageStore->GetPackageStoreEntry(PackageIdToLoad, PackageEntry);
+					PackageStatus = PackageStore.GetPackageStoreEntry(PackageIdToLoad, PackageEntry);
 				}
 			}
 			FAsyncPackageDesc2 PackageDesc = FAsyncPackageDesc2::FromPackageRequest(Request.RequestId, Request.Priority, UPackageName, PackageIdToLoad, MoveTemp(Request.PackagePath), !Request.CustomName.IsNone());
@@ -2788,7 +2770,7 @@ bool FAsyncLoadingThread2::CreateAsyncPackagesFromQueue(FAsyncLoadingThreadState
 						InitializeAsyncPackageFromPackageStore(Package, PackageEntry);
 						{
 							TRACE_CPUPROFILER_EVENT_SCOPE(ImportPackages);
-							Package->ImportPackagesRecursive(IoBatch, *PackageStore.Get());
+							Package->ImportPackagesRecursive(IoBatch, PackageStore);
 						}
 						Package->StartLoading(IoBatch);
 					}
@@ -2815,8 +2797,6 @@ bool FAsyncLoadingThread2::CreateAsyncPackagesFromQueue(FAsyncLoadingThreadState
 		}
 	}
 
-	PackageStore->Unlock();
-
 	IoBatch.Issue();
 	
 	return bPackagesCreated;
@@ -3308,7 +3288,7 @@ void FGlobalImportStore::FindAllScriptObjects()
 		ScriptObjects.Num(), (float)ScriptObjects.GetAllocatedSize() / 1024.f);
 }
 
-void FAsyncPackage2::ImportPackagesRecursiveInner(FIoBatch& IoBatch, IPackageStore& PackageStore, const TArrayView<const FPackageId>& ImportedPackageIds, int32& ImportedPackageIndex)
+void FAsyncPackage2::ImportPackagesRecursiveInner(FIoBatch& IoBatch, FPackageStore& PackageStore, const TArrayView<const FPackageId>& ImportedPackageIds, int32& ImportedPackageIndex)
 {
 	for (const FPackageId& ImportedPackageId : ImportedPackageIds)
 	{
@@ -3460,7 +3440,7 @@ void FAsyncPackage2::ImportPackagesRecursiveInner(FIoBatch& IoBatch, IPackageSto
 	}
 }
 
-void FAsyncPackage2::ImportPackagesRecursive(FIoBatch& IoBatch, IPackageStore& PackageStore)
+void FAsyncPackage2::ImportPackagesRecursive(FIoBatch& IoBatch, FPackageStore& PackageStore)
 {
 	if (AsyncPackageLoadingState >= EAsyncPackageLoadingState2::ImportPackages)
 	{
@@ -5194,6 +5174,8 @@ FAsyncLoadingThread2::FAsyncLoadingThread2(FIoDispatcher& InIoDispatcher, IAsync
 	: Thread(nullptr)
 	, IoDispatcher(InIoDispatcher)
 	, UncookedPackageLoader(InUncookedPackageLoader)
+	, PackageStore(FPackageStore::Get())
+	, LoadedPackageStore(PackageStore)
 {
 #if !WITH_IOSTORE_IN_EDITOR
 	IsEventDrivenLoaderEnabled(); // make sure the one time init inside runs
@@ -5326,8 +5308,6 @@ void FAsyncLoadingThread2::LazyInitializeFromLoadPackage()
 	{
 		GlobalImportStore.Initialize(IoDispatcher);
 	}
-
-	LoadedPackageStore.Initialize(*PackageStore.Get());
 }
 
 void FAsyncLoadingThread2::FinalizeInitialLoad()
diff --git a/Engine/Source/Runtime/PakFile/Private/FilePackageStore.cpp b/Engine/Source/Runtime/PakFile/Private/FilePackageStore.cpp
--- a/Engine/Source/Runtime/PakFile/Private/FilePackageStore.cpp
+++ b/Engine/Source/Runtime/PakFile/Private/FilePackageStore.cpp
@@ -27,9 +27,7 @@ namespace FilePackageStore
 	}
 }
 
-thread_local int32 FFilePackageStore::LockedOnThreadCount = 0;
-
-FFilePackageStore::FFilePackageStore()
+FFilePackageStoreBackend::FFilePackageStoreBackend()
 {
 #if WITH_EDITOR
 	OnContentPathMountedDelegateHandle = FPackageName::OnContentPathMounted().AddLambda([this](const FString& InAssetPath, const FString& InFilesystemPath)
@@ -56,20 +54,7 @@ FFilePackageStore::FFilePackageStore()
 				bNeedsUpdate = true;
 			}
 		});
-#endif //if WITH_EDITOR
-}
-
-FFilePackageStore::~FFilePackageStore()
-{
-#if WITH_EDITOR
-	FPackageName::OnContentPathMounted().Remove(OnContentPathMountedDelegateHandle);
-	FPackageName::OnContentPathDismounted().Remove(OnContentPathDismountedDelegateHandle);
-#endif
-}
 
-void FFilePackageStore::Initialize()
-{
-#if WITH_EDITOR
 	TArray<FString> RootPaths;
 	FPackageName::QueryRootContentPaths(RootPaths);
 	{
@@ -83,37 +68,30 @@ void FFilePackageStore::Initialize()
 #endif //if WITH_EDITOR
 }
 
-void FFilePackageStore::Lock()
+FFilePackageStoreBackend::~FFilePackageStoreBackend()
 {
-	if (!LockedOnThreadCount)
-	{
-		EntriesLock.ReadLock();
-		if (bNeedsUpdate)
-		{
-			Update();
-		}
-	}
-	++LockedOnThreadCount;
+#if WITH_EDITOR
+	FPackageName::OnContentPathMounted().Remove(OnContentPathMountedDelegateHandle);
+	FPackageName::OnContentPathDismounted().Remove(OnContentPathDismountedDelegateHandle);
+#endif
 }
 
-void FFilePackageStore::Unlock()
+void FFilePackageStoreBackend::BeginRead()
 {
-	check(LockedOnThreadCount > 0);
-	if (--LockedOnThreadCount == 0)
+	EntriesLock.ReadLock();
+	if (bNeedsUpdate)
 	{
-		EntriesLock.ReadUnlock();
+		Update();
 	}
 }
 
-bool FFilePackageStore::DoesPackageExist(FPackageId PackageId)
+void FFilePackageStoreBackend::EndRead()
 {
-	check(LockedOnThreadCount);
-	return PackageId.IsValid() && StoreEntriesMap.Contains(PackageId);
+	EntriesLock.ReadUnlock();
 }
 
-EPackageStoreEntryStatus FFilePackageStore::GetPackageStoreEntry(FPackageId PackageId, FPackageStoreEntry& OutPackageStoreEntry)
+EPackageStoreEntryStatus FFilePackageStoreBackend::GetPackageStoreEntry(FPackageId PackageId, FPackageStoreEntry& OutPackageStoreEntry)
 {
-	check(LockedOnThreadCount);
 #if WITH_EDITOR
 	const FUncookedPackage* FindUncookedPackage = UncookedPackagesMap.Find(PackageId);
 	if (FindUncookedPackage)
@@ -144,9 +122,8 @@ EPackageStoreEntryStatus FFilePackageStore::GetPackageStoreEntry(FPackageId Pack
 	return EPackageStoreEntryStatus::Missing;
 }
 
-bool FFilePackageStore::GetPackageRedirectInfo(FPackageId PackageId, FName& OutSourcePackageName, FPackageId& OutRedirectedToPackageId)
+bool FFilePackageStoreBackend::GetPackageRedirectInfo(FPackageId PackageId, FName& OutSourcePackageName, FPackageId& OutRedirectedToPackageId)
 {
-	check(LockedOnThreadCount);
 	TTuple<FName, FPackageId>* FindRedirect = RedirectsPackageMap.Find(PackageId);
 	if (FindRedirect)
 	{
@@ -175,7 +152,7 @@ bool FFilePackageStore::GetPackageRedirectInfo(FPackageId PackageId, FName& OutS
 	return false;
 }
 
-void FFilePackageStore::Mount(const FIoContainerHeader* ContainerHeader, uint32 Order)
+void FFilePackageStoreBackend::Mount(const FIoContainerHeader* ContainerHeader, uint32 Order)
 {
 	LLM_SCOPE(ELLMTag::AsyncLoading);
 	FWriteScopeLock _(EntriesLock);
@@ -191,7 +168,7 @@ void FFilePackageStore::Mount(const FIoContainerHeader* ContainerHeader, uint32
 	bNeedsUpdate = true;
 }
 
-void FFilePackageStore::Unmount(const FIoContainerHeader* ContainerHeader)
+void FFilePackageStoreBackend::Unmount(const FIoContainerHeader* ContainerHeader)
 {
 	FWriteScopeLock _(EntriesLock);
 	for (auto It = MountedContainers.CreateIterator(); It; ++It)
@@ -205,7 +182,7 @@ void FFilePackageStore::Unmount(const FIoContainerHeader* ContainerHeader)
 	}
 }
 
-void FFilePackageStore::Update()
+void FFilePackageStoreBackend::Update()
 {
 	LLM_SCOPE(ELLMTag::AsyncLoading);
 	TRACE_CPUPROFILER_EVENT_SCOPE(UpdateFilePackageStore);
@@ -303,7 +280,7 @@ void FFilePackageStore::Update()
 }
 
 #if WITH_EDITOR
-uint64 FFilePackageStore::AddUncookedPackagesFromRoot(const FString& RootPath)
+uint64 FFilePackageStoreBackend::AddUncookedPackagesFromRoot(const FString& RootPath)
 {
 	uint64 Count = 0;
 	FPackageName::IteratePackagesInDirectory(RootPath, [this, &RootPath, &Count](const TCHAR* InPackageFileName) -> bool
@@ -323,7 +300,7 @@ uint64 FFilePackageStore::AddUncookedPackagesFromRoot(const FString& RootPath)
 	return Count;
 }
 
-uint64 FFilePackageStore::RemoveUncookedPackagesFromRoot(const TSet<FString>& RootPaths)
+uint64 FFilePackageStoreBackend::RemoveUncookedPackagesFromRoot(const TSet<FString>& RootPaths)
 {
 	uint64 Count = 0;
 	TMap<FPackageId, FUncookedPackage> RemainingUncookedPackagesMap;
diff --git a/Engine/Source/Runtime/PakFile/Private/FilePackageStore.h b/Engine/Source/Runtime/PakFile/Private/FilePackageStore.h
--- a/Engine/Source/Runtime/PakFile/Private/FilePackageStore.h
+++ b/Engine/Source/Runtime/PakFile/Private/FilePackageStore.h
@@ -13,17 +13,19 @@ struct FFilePackageStoreEntry;
 /*
  * File/container based package store.
  */
-class FFilePackageStore
-	: public FPackageStoreBase
+class FFilePackageStoreBackend
+	: public IPackageStoreBackend
 {
 public:
-	FFilePackageStore();
-	virtual ~FFilePackageStore();
+	FFilePackageStoreBackend();
+	virtual ~FFilePackageStoreBackend();
 
-	virtual void Initialize() override;
-	virtual void Lock() override;
-	virtual void Unlock() override;
-	virtual bool DoesPackageExist(FPackageId PackageId) override;
+	virtual void OnMounted(TSharedRef<const FPackageStoreBackendContext>) override
+	{
+	}
+
+	virtual void BeginRead() override;
+	virtual void EndRead() override;
 	virtual EPackageStoreEntryStatus GetPackageStoreEntry(FPackageId PackageId, FPackageStoreEntry& OutPackageStoreEntry) override;
 	virtual bool GetPackageRedirectInfo(FPackageId PackageId, FName& OutSourcePackageName, FPackageId& OutRedirectedToPackageId) override;
 
@@ -71,5 +73,4 @@ private:
 	TMap<FPackageId, const FFilePackageStoreEntry*> OptionalSegmentStoreEntriesMap;
 #endif //if WITH_EDITOR
 
-	static thread_local int32 LockedOnThreadCount;
 };
diff --git a/Engine/Source/Runtime/PakFile/Private/IPlatformFilePak.cpp b/Engine/Source/Runtime/PakFile/Private/IPlatformFilePak.cpp
--- a/Engine/Source/Runtime/PakFile/Private/IPlatformFilePak.cpp
+++ b/Engine/Source/Runtime/PakFile/Private/IPlatformFilePak.cpp
@@ -7476,11 +7476,8 @@
 		FIoDispatcher& IoDispatcher = FIoDispatcher::Get();
 		IoDispatcherFileBackend = CreateIoDispatcherFileBackend();
 		IoDispatcher.Mount(IoDispatcherFileBackend.ToSharedRef());
-		FilePackageStore = MakeShared<FFilePackageStore>();
-		FCoreDelegates::CreatePackageStore.BindLambda([this]()
-		{
-			return FilePackageStore;
-		});
+		PackageStoreBackend = MakeShared<FFilePackageStoreBackend>();
+		FPackageStore::Get().Mount(PackageStoreBackend.ToSharedRef());
 
 		if (bShouldMountGlobal)
 		{
@@ -7785,7 +7782,7 @@
 				{
 					UE_LOG(LogPakFile, Display, TEXT("Mounted IoStore container \"%s\""), *UtocPath);
 					Pak->IoContainerHeader = MakeUnique<FIoContainerHeader>(MountResult.ConsumeValueOrDie());
-					FilePackageStore->Mount(Pak->IoContainerHeader.Get(), PakOrder);
+					PackageStoreBackend->Mount(Pak->IoContainerHeader.Get(), PakOrder);
 #if WITH_EDITOR
 					FString OptionalSegmentUtocPath = FPaths::ChangeExtension(InPakFilename, FString::Printf(TEXT("%s.utoc"), FPackagePath::GetOptionalSegmentExtensionModifier()));
 					if (FPlatformFileManager::Get().GetPlatformFile().FileExists(*OptionalSegmentUtocPath))
@@ -7794,7 +7791,7 @@
 						if (MountResult.IsOk())
 						{
 							Pak->OptionalSegmentIoContainerHeader = MakeUnique<FIoContainerHeader>(MountResult.ConsumeValueOrDie());
-							FilePackageStore->Mount(Pak->OptionalSegmentIoContainerHeader.Get(), PakOrder);
+							PackageStoreBackend->Mount(Pak->OptionalSegmentIoContainerHeader.Get(), PakOrder);
 							UE_LOG(LogPakFile, Display, TEXT("Mounted optional segment extension IoStore container \"%s\""), *OptionalSegmentUtocPath);
 						}
 						else
@@ -7884,14 +7881,14 @@
 		{
 			if (UnmountedPak)
 			{
-				FilePackageStore->Unmount(UnmountedPak->IoContainerHeader.Get());
+				PackageStoreBackend->Unmount(UnmountedPak->IoContainerHeader.Get());
 			}
 			FString ContainerPath = FPaths::ChangeExtension(InPakFilename, FString());
 			bRemovedContainerFile = IoDispatcherFileBackend->Unmount(*ContainerPath);
 #if WITH_EDITOR
 			if (UnmountedPak && UnmountedPak->OptionalSegmentIoContainerHeader.IsValid())
 			{
-				FilePackageStore->Unmount(UnmountedPak->OptionalSegmentIoContainerHeader.Get());
+				PackageStoreBackend->Unmount(UnmountedPak->OptionalSegmentIoContainerHeader.Get());
 				FString OptionalSegmentContainerPath = ContainerPath + FPackagePath::GetOptionalSegmentExtensionModifier();
 				IoDispatcherFileBackend->Unmount(*OptionalSegmentContainerPath);
 			}
diff --git a/Engine/Source/Runtime/PakFile/Public/IPlatformFilePak.h b/Engine/Source/Runtime/PakFile/Public/IPlatformFilePak.h
--- a/Engine/Source/Runtime/PakFile/Public/IPlatformFilePak.h
+++ b/Engine/Source/Runtime/PakFile/Public/IPlatformFilePak.h
@@ -22,7 +22,7 @@
 class FChunkCacheWorker;
 class IAsyncReadFileHandle;
 class FFileIoStore;
-class FFilePackageStore;
+class FFilePackageStoreBackend;
 struct FIoContainerHeader;
 
 PAKFILE_API DECLARE_LOG_CATEGORY_EXTERN(LogPakFile, Log, All);
@@ -2037,7 +2037,7 @@
 	/** The filename for the gameusersettings ini file, used for excluding ini files, but not gameusersettings */
 	FString GameUserSettingsIniFilename;
 	TSharedPtr<FFileIoStore> IoDispatcherFileBackend;
-	TSharedPtr<FFilePackageStore> FilePackageStore;
+	TSharedPtr<FFilePackageStoreBackend> PackageStoreBackend;
 
 	FTSTicker::FDelegateHandle RetireReadersHandle;
 
diff --git a/Engine/Source/Runtime/StorageServerClient/Private/CookOnTheFlyPackageStore.cpp b/Engine/Source/Runtime/StorageServerClient/Private/CookOnTheFlyPackageStore.cpp
--- a/Engine/Source/Runtime/StorageServerClient/Private/CookOnTheFlyPackageStore.cpp
+++ b/Engine/Source/Runtime/StorageServerClient/Private/CookOnTheFlyPackageStore.cpp
@@ -9,13 +9,14 @@
 #include "Serialization/MemoryReader.h"
 #include "CookOnTheFlyMessages.h"
 #include "CookOnTheFly.h"
+#include "Misc/PackageName.h"
 
-FCookOnTheFlyPackageStore::FCookOnTheFlyPackageStore(UE::Cook::ICookOnTheFlyServerConnection& InCookOnTheFlyServerConnection)
+FCookOnTheFlyPackageStoreBackend::FCookOnTheFlyPackageStoreBackend(UE::Cook::ICookOnTheFlyServerConnection& InCookOnTheFlyServerConnection)
 	: CookOnTheFlyServerConnection(InCookOnTheFlyServerConnection)
 {
 	// Index zero is invalid
 	PackageEntries.Add();
-	CookOnTheFlyServerConnection.OnMessage().AddRaw(this, &FCookOnTheFlyPackageStore::OnCookOnTheFlyMessage);
+	CookOnTheFlyServerConnection.OnMessage().AddRaw(this, &FCookOnTheFlyPackageStoreBackend::OnCookOnTheFlyMessage);
 
 	using namespace UE::Cook;
 	using namespace UE::ZenCookOnTheFly::Messaging;
@@ -39,16 +40,21 @@ FCookOnTheFlyPackageStore::FCookOnTheFlyPackageStore(UE::Cook::ICookOnTheFlyServ
 	{
 		UE_LOG(LogCookOnTheFly, Warning, TEXT("Failed to send 'GetCookedPackages' request"));
 	}
+
+	FPackageName::DoesPackageExistOverride().BindLambda([this](FName PackageName)
+	{
+		return DoesPackageExist(FPackageId::FromName(PackageName));
+	});
 }
 
-bool FCookOnTheFlyPackageStore::DoesPackageExist(FPackageId PackageId)
+bool FCookOnTheFlyPackageStoreBackend::DoesPackageExist(FPackageId PackageId)
 {
 	FScopeLock _(&CriticalSection);
 	FEntryInfo EntryInfo = PackageIdToEntryInfo.FindRef(PackageId);
 	return EntryInfo.Status != EPackageStoreEntryStatus::Missing;
 }
 
-EPackageStoreEntryStatus FCookOnTheFlyPackageStore::GetPackageStoreEntry(FPackageId PackageId, FPackageStoreEntry& OutPackageStoreEntry)
+EPackageStoreEntryStatus FCookOnTheFlyPackageStoreBackend::GetPackageStoreEntry(FPackageId PackageId, FPackageStoreEntry& OutPackageStoreEntry)
 {
 	using namespace UE::Cook;
 	using namespace UE::ZenCookOnTheFly::Messaging;
@@ -120,7 +126,7 @@ EPackageStoreEntryStatus FCookOnTheFlyPackageStore::GetPackageStoreEntry(FPackag
 	}
 }
 
-EPackageStoreEntryStatus FCookOnTheFlyPackageStore::CreatePackageStoreEntry(const FEntryInfo& EntryInfo, FPackageStoreEntry& OutPackageStoreEntry)
+EPackageStoreEntryStatus FCookOnTheFlyPackageStoreBackend::CreatePackageStoreEntry(const FEntryInfo& EntryInfo, FPackageStoreEntry& OutPackageStoreEntry)
 {
 	if (EntryInfo.Status == EPackageStoreEntryStatus::Ok)
 	{
@@ -135,7 +141,7 @@ EPackageStoreEntryStatus FCookOnTheFlyPackageStore::CreatePackageStoreEntry(cons
 	}
 }
 
-void FCookOnTheFlyPackageStore::AddPackages(TArray<FPackageStoreEntryResource> Entries, TArray<FPackageId> FailedPackageIds)
+void FCookOnTheFlyPackageStoreBackend::AddPackages(TArray<FPackageStoreEntryResource> Entries, TArray<FPackageId> FailedPackageIds)
 {
 	FScopeLock _(&CriticalSection);
 		
@@ -167,7 +173,7 @@ void FCookOnTheFlyPackageStore::AddPackages(TArray<FPackageStoreEntryResource> E
 	}
 }
 
-void FCookOnTheFlyPackageStore::OnCookOnTheFlyMessage(const UE::Cook::FCookOnTheFlyMessage& Message)
+void FCookOnTheFlyPackageStoreBackend::OnCookOnTheFlyMessage(const UE::Cook::FCookOnTheFlyMessage& Message)
 {
 	using namespace UE::ZenCookOnTheFly::Messaging;
 
@@ -183,7 +189,10 @@ void FCookOnTheFlyPackageStore::OnCookOnTheFlyMessage(const UE::Cook::FCookOnThe
 
 			AddPackages(MoveTemp(PackagesCookedMessage.CookedPackages), MoveTemp(PackagesCookedMessage.FailedPackages));
 
-			PendingEntriesAdded.Broadcast();
+			if (Context)
+			{
+				Context->PendingEntriesAdded.Broadcast();
+			}
 
 			break;
 		}
@@ -195,7 +204,7 @@ void FCookOnTheFlyPackageStore::OnCookOnTheFlyMessage(const UE::Cook::FCookOnThe
 	LastServerActivtyTime = FPlatformTime::Seconds();
 }
 
-void FCookOnTheFlyPackageStore::CheckActivity()
+void FCookOnTheFlyPackageStoreBackend::CheckActivity()
 {
 	const double TimeSinceLastClientActivity = FPlatformTime::Seconds() - LastClientActivtyTime;
 	const double TimeSinceLastServerActivity = FPlatformTime::Seconds() - LastServerActivtyTime;
diff --git a/Engine/Source/Runtime/StorageServerClient/Private/CookOnTheFlyPackageStore.h b/Engine/Source/Runtime/StorageServerClient/Private/CookOnTheFlyPackageStore.h
--- a/Engine/Source/Runtime/StorageServerClient/Private/CookOnTheFlyPackageStore.h
+++ b/Engine/Source/Runtime/StorageServerClient/Private/CookOnTheFlyPackageStore.h
@@ -9,8 +9,8 @@
 #include "IO/PackageStore.h"
 #include "CookOnTheFly.h"
 
-class FCookOnTheFlyPackageStore final
-	: public FPackageStoreBase
+class FCookOnTheFlyPackageStoreBackend final
+	: public IPackageStoreBackend
 {
 public:
 	struct FEntryInfo
@@ -25,21 +25,22 @@ public:
 		TAtomic<uint32> Failed{ 0 };
 	};
 
-	FCookOnTheFlyPackageStore(UE::Cook::ICookOnTheFlyServerConnection& InCookOnTheFlyServerConnection);
+	FCookOnTheFlyPackageStoreBackend(UE::Cook::ICookOnTheFlyServerConnection& InCookOnTheFlyServerConnection);
 
-	virtual void Initialize() override
+	virtual void OnMounted(TSharedRef<const FPackageStoreBackendContext> InContext) override
 	{
+		Context = InContext;
 	}
 
-	virtual void Lock() override
+	virtual void BeginRead() override
 	{
 	}
 
-	virtual void Unlock() override
+	virtual void EndRead() override
 	{
 	}
 
-	virtual bool DoesPackageExist(FPackageId PackageId) override;
+	virtual bool DoesPackageExist(FPackageId PackageId);
 	virtual EPackageStoreEntryStatus GetPackageStoreEntry(FPackageId PackageId, FPackageStoreEntry& OutPackageStoreEntry) override;
 	
 	virtual bool GetPackageRedirectInfo(FPackageId PackageId, FName& OutSourcePackageName, FPackageId& OutRedirectedToPackageId) override
@@ -54,6 +55,7 @@ private:
 	void CheckActivity();
 
 	UE::Cook::ICookOnTheFlyServerConnection& CookOnTheFlyServerConnection;
+	TSharedPtr<const FPackageStoreBackendContext> Context;
 	FCriticalSection CriticalSection;
 	TMap<FPackageId, FEntryInfo> PackageIdToEntryInfo;
 	TChunkedArray<FPackageStoreEntryResource> PackageEntries;
diff --git a/Engine/Source/Runtime/StorageServerClient/Private/StorageServerPackageStore.cpp b/Engine/Source/Runtime/StorageServerClient/Private/StorageServerPackageStore.cpp
--- a/Engine/Source/Runtime/StorageServerClient/Private/StorageServerPackageStore.cpp
+++ b/Engine/Source/Runtime/StorageServerClient/Private/StorageServerPackageStore.cpp
@@ -8,7 +8,7 @@
 
 #if !UE_BUILD_SHIPPING
 
-FStorageServerPackageStore::FStorageServerPackageStore(FStorageServerConnection& Connection)
+FStorageServerPackageStoreBackend::FStorageServerPackageStoreBackend(FStorageServerConnection& Connection)
 {
 	FIoChunkId HeaderChunkId = CreateIoChunkId(FIoContainerId::FromName(TEXT("global")).Value(), 0, EIoChunkType::ContainerHeader);
 	Connection.ReadChunkRequest(HeaderChunkId, 0, uint64(-1), [this](FStorageServerResponse& Response)
@@ -35,12 +35,7 @@ FStorageServerPackageStore::FStorageServerPackageStore(FStorageServerConnection&
 	});
 }
 
-bool FStorageServerPackageStore::DoesPackageExist(FPackageId PackageId)
-{
-	return PackageId.IsValid() && StoreEntriesMap.Contains(PackageId);
-}
-
-EPackageStoreEntryStatus FStorageServerPackageStore::GetPackageStoreEntry(FPackageId PackageId, FPackageStoreEntry& OutPackageStoreEntry)
+EPackageStoreEntryStatus FStorageServerPackageStoreBackend::GetPackageStoreEntry(FPackageId PackageId, FPackageStoreEntry& OutPackageStoreEntry)
 {
 	const FFilePackageStoreEntry* FindEntry = StoreEntriesMap.FindRef(PackageId);
 	if (FindEntry)
diff --git a/Engine/Source/Runtime/StorageServerClient/Private/StorageServerPackageStore.h b/Engine/Source/Runtime/StorageServerClient/Private/StorageServerPackageStore.h
--- a/Engine/Source/Runtime/StorageServerClient/Private/StorageServerPackageStore.h
+++ b/Engine/Source/Runtime/StorageServerClient/Private/StorageServerPackageStore.h
@@ -9,26 +9,25 @@
 class FStorageServerConnection;
 struct FFilePackageStoreEntry;
 
-class FStorageServerPackageStore
-	: public FPackageStoreBase
+class FStorageServerPackageStoreBackend
+	: public IPackageStoreBackend
 {
 public:
-	FStorageServerPackageStore(FStorageServerConnection& Connection);
-	virtual ~FStorageServerPackageStore() = default;
+	FStorageServerPackageStoreBackend(FStorageServerConnection& Connection);
+	virtual ~FStorageServerPackageStoreBackend() = default;
 
-	virtual void Initialize() override
+	virtual void OnMounted(TSharedRef<const FPackageStoreBackendContext> Context) override
 	{
 	}
 
-	virtual void Lock() override
+	virtual void BeginRead() override
 	{
 	}
 
-	virtual void Unlock() override
+	virtual void EndRead() override
 	{
 	}
 
-	virtual bool DoesPackageExist(FPackageId PackageId) override;
 	virtual EPackageStoreEntryStatus GetPackageStoreEntry(FPackageId PackageIde, FPackageStoreEntry& OutPackageStoreEntry) override;
 	
 	virtual bool GetPackageRedirectInfo(FPackageId PackageId, FName& OutSourcePackageName, FPackageId& OutRedirectedToPackageId) override
diff --git a/Engine/Source/Runtime/StorageServerClient/Private/StorageServerPlatformFile.cpp b/Engine/Source/Runtime/StorageServerClient/Private/StorageServerPlatformFile.cpp
--- a/Engine/Source/Runtime/StorageServerClient/Private/StorageServerPlatformFile.cpp
+++ b/Engine/Source/Runtime/StorageServerClient/Private/StorageServerPlatformFile.cpp
@@ -307,20 +307,16 @@ void FStorageServerPlatformFile::InitializeAfterProjectFilePath()
 			FIoDispatcher& IoDispatcher = FIoDispatcher::Get();
 			TSharedRef<FStorageServerIoDispatcherBackend> IoDispatcherBackend = MakeShared<FStorageServerIoDispatcherBackend>(*Connection.Get());
 			IoDispatcher.Mount(IoDispatcherBackend);
-
-			FCoreDelegates::CreatePackageStore.BindLambda([this]() -> TSharedPtr<IPackageStore>
-			{
 #if WITH_COTF
-				if (CookOnTheFlyServerConnection)
-				{
-					return MakeShared<FCookOnTheFlyPackageStore>(*CookOnTheFlyServerConnection.Get());
-				}
-				else
+			if (CookOnTheFlyServerConnection)
+			{
+				FPackageStore::Get().Mount(MakeShared<FCookOnTheFlyPackageStoreBackend>(*CookOnTheFlyServerConnection.Get()));
+			}
+			else
 #endif
-				{
-					return MakeShared<FStorageServerPackageStore>(*Connection.Get());
-				}
-			});
+			{
+				FPackageStore::Get().Mount(MakeShared<FStorageServerPackageStoreBackend>(*Connection.Get()));
+			}
 		}
 		else
 		{
diff --git a/Engine/Source/Runtime/StorageServerClient/StorageServerClient.Build.cs b/Engine/Source/Runtime/StorageServerClient/StorageServerClient.Build.cs
--- a/Engine/Source/Runtime/StorageServerClient/StorageServerClient.Build.cs
+++ b/Engine/Source/Runtime/StorageServerClient/StorageServerClient.Build.cs
@@ -9,6 +9,7 @@ public class StorageServerClient : ModuleRules
         PublicDependencyModuleNames.AddRange(
             new string[] {
                 "Core",
+                "CoreUObject",
                 "Sockets",
 				"CookOnTheFly"
             }
