Deprecating ANY_PACKAGE.

diff --git a/Engine/Source/Runtime/AssetRegistry/Public/AssetRegistry/IAssetRegistry.h b/Engine/Source/Runtime/AssetRegistry/Public/AssetRegistry/IAssetRegistry.h
--- a/Engine/Source/Runtime/AssetRegistry/Public/AssetRegistry/IAssetRegistry.h
+++ b/Engine/Source/Runtime/AssetRegistry/Public/AssetRegistry/IAssetRegistry.h
@@ -159,7 +159,10 @@ public:
 	 * @param bSearchSubClasses if true, all subclasses of the passed in class will be searched as well
 	 */
 	UFUNCTION(BlueprintCallable, BlueprintPure=false, Category = "AssetRegistry")
-	virtual bool GetAssetsByClass(FName ClassName, TArray<FAssetData>& OutAssetData, bool bSearchSubClasses = false) const = 0;
+	virtual bool GetAssetsByClass(FTopLevelAssetPath ClassPathName, TArray<FAssetData>& OutAssetData, bool bSearchSubClasses = false) const = 0;
+
+	UE_DEPRECATED(5.1, "Class names are now represented by path names. Please use a version of this function that uses FTopLevelAssetPath.")
+	virtual bool GetAssetsByClass(FName ClassPathName, TArray<FAssetData>& OutAssetData, bool bSearchSubClasses = false) const = 0;
 
 	/**
 	 * Gets asset data for all assets with the supplied tags, regardless of their value
@@ -336,9 +339,15 @@ public:
 	virtual void StripAssetRegistryKeyForObject(FName ObjectPath, FName Key) {}
 
 	/** Returns true if the specified ClassName's ancestors could be found. If so, OutAncestorClassNames is a list of all its ancestors. This can be slow if temporary caching mode is not on */
+	virtual bool GetAncestorClassNames(FTopLevelAssetPath ClassPathName, TArray<FTopLevelAssetPath>& OutAncestorClassNames) const = 0;
+
+	UE_DEPRECATED(5.1, "Class names are now represented by path names. Please use a version of this function that uses FTopLevelAssetPath.")
 	virtual bool GetAncestorClassNames(FName ClassName, TArray<FName>& OutAncestorClassNames) const = 0;
 
 	/** Returns the names of all classes derived by the supplied class names, excluding any classes matching the excluded class names. This can be slow if temporary caching mode is not on */
+	virtual void GetDerivedClassNames(const TArray<FTopLevelAssetPath>& ClassNames, const TSet<FTopLevelAssetPath>& ExcludedClassNames, TSet<FTopLevelAssetPath>& OutDerivedClassNames) const = 0;
+	
+	UE_DEPRECATED(5.1, "Class names are now represented by path names. Please use a version of this function that uses FTopLevelAssetPath.")
 	virtual void GetDerivedClassNames(const TArray<FName>& ClassNames, const TSet<FName>& ExcludedClassNames, TSet<FName>& OutDerivedClassNames) const = 0;
 
 	/** Gets a list of all paths that are currently cached */
@@ -703,6 +712,10 @@ namespace AssetRegistry
 	 */
 	ASSETREGISTRY_API const FAssetData* GetMostImportantAsset(TConstArrayView<const FAssetData*> PackageAssetDatas, bool bRequireOneTopLevelAsset);
 
+	// Wildcards (*) used when looking up assets in the asset registry
+	extern ASSETREGISTRY_API const FName WildcardFName;
+	extern ASSETREGISTRY_API const FTopLevelAssetPath WildcardPathName;
+
 } // namespace AssetRegistry
 } // namespace UE
 
diff --git a/Engine/Source/Runtime/CoreUObject/Private/AssetRegistry/AssetData.cpp b/Engine/Source/Runtime/CoreUObject/Private/AssetRegistry/AssetData.cpp
--- a/Engine/Source/Runtime/CoreUObject/Private/AssetRegistry/AssetData.cpp
+++ b/Engine/Source/Runtime/CoreUObject/Private/AssetRegistry/AssetData.cpp
@@ -52,10 +52,20 @@ static TSharedPtr<FAssetBundleData, ESPMode::ThreadSafe> ParseAssetBundles(const
 }}} // end namespace UE::AssetData::Private
 
 FAssetData::FAssetData(FName InPackageName, FName InPackagePath, FName InAssetName, FName InAssetClass, FAssetDataTagMap InTags, TArrayView<const int32> InChunkIDs, uint32 InPackageFlags)
+	: FAssetData(InPackageName, InPackagePath, InAssetName, FAssetData::TryConvertShortClassNameToPathName(InAssetClass), InTags, InChunkIDs, InPackageFlags)
+{
+}
+
+FAssetData::FAssetData(const FString& InLongPackageName, const FString& InObjectPath, FName InAssetClass, FAssetDataTagMap InTags, TArrayView<const int32> InChunkIDs, uint32 InPackageFlags)
+	: FAssetData(InLongPackageName, InObjectPath, FAssetData::TryConvertShortClassNameToPathName(InAssetClass), InTags, InChunkIDs, InPackageFlags)
+{
+}
+
+FAssetData::FAssetData(FName InPackageName, FName InPackagePath, FName InAssetName, FTopLevelAssetPath InAssetClassPathName, FAssetDataTagMap InTags, TArrayView<const int32> InChunkIDs, uint32 InPackageFlags)
 	: PackageName(InPackageName)
 	, PackagePath(InPackagePath)
 	, AssetName(InAssetName)
-	, AssetClass(InAssetClass)
+	, AssetClassPath(InAssetClassPathName)
 	, PackageFlags(InPackageFlags)
 	, ChunkIDs(MoveTemp(InChunkIDs))
 {
@@ -67,10 +77,10 @@ FAssetData::FAssetData(FName InPackageName, FName InPackagePath, FName InAssetNa
 	ObjectPath = FName(FStringView(ObjectPathStr));
 }
 
-FAssetData::FAssetData(const FString& InLongPackageName, const FString& InObjectPath, FName InAssetClass, FAssetDataTagMap InTags, TArrayView<const int32> InChunkIDs, uint32 InPackageFlags)
+FAssetData::FAssetData(const FString& InLongPackageName, const FString& InObjectPath, FTopLevelAssetPath InAssetClassPathName, FAssetDataTagMap InTags, TArrayView<const int32> InChunkIDs, uint32 InPackageFlags)
 	: ObjectPath(*InObjectPath)
 	, PackageName(*InLongPackageName)
-	, AssetClass(InAssetClass)
+	, AssetClassPath(InAssetClassPathName)
 	, PackageFlags(InPackageFlags)
 	, ChunkIDs(MoveTemp(InChunkIDs))
 {
@@ -105,7 +115,7 @@ FAssetData::FAssetData(const UObject* InAsset, FAssetData::ECreationFlags InCrea
 		PackageName = Package->GetFName();
 		PackagePath = FName(*FPackageName::GetLongPackagePath(Package->GetName()));
 		AssetName = InAsset->GetFName();
-		AssetClass = InAsset->GetClass()->GetFName();
+		AssetClassPath = InAsset->GetClass()->GetPathName();
 		ObjectPath = FName(*InAsset->GetPathName());
 
 		if (!EnumHasAnyFlags(InCreationFlags, FAssetData::ECreationFlags::SkipAssetRegistryTagsGathering))
@@ -192,12 +202,26 @@ FPrimaryAssetId FAssetData::GetPrimaryAssetId() const
 	return FPrimaryAssetId();
 }
 
-void FAssetData::SerializeForCacheInternal(FArchive& Ar, void (*SerializeTagsAndBundles)(FArchive& , FAssetData&))
+void FAssetData::SerializeForCacheInternal(FArchive& Ar, FAssetRegistryVersion::Type Version, void (*SerializeTagsAndBundles)(FArchive& , FAssetData&))
 {
 	// Serialize out the asset info
 	Ar << ObjectPath;
 	Ar << PackagePath;
-	Ar << AssetClass;
+
+	// Serialize the asset class. Legacy path is editor only for DiffAssetRegistriesCommandlet
+#if WITH_EDITORONLY_DATA
+	if (Version < FAssetRegistryVersion::ClassPaths)
+	{
+PRAGMA_DISABLE_DEPRECATION_WARNINGS
+		Ar << AssetClass;
+		AssetClassPath = FAssetData::TryConvertShortClassNameToPathName(AssetClass);
+PRAGMA_ENABLE_DEPRECATION_WARNINGS
+	}
+	else
+#endif 
+	{
+		Ar << AssetClassPath;
+	}
 
 	// These are derived from ObjectPath, we manually serialize them because they get pooled
 	Ar << PackageName;
@@ -218,6 +242,36 @@ void FAssetData::SerializeForCacheInternal(FArchive& Ar, void (*SerializeTagsAnd
 	Ar << PackageFlags;
 }
 
+FTopLevelAssetPath FAssetData::TryConvertShortClassNameToPathName(FName InClassName)
+{
+	FTopLevelAssetPath ClassPath;
+	if (!InClassName.IsNone())
+	{
+		FString ClassNameString(InClassName.ToString());
+		ClassPath = UClass::TryConvertShortTypeNameToPathName<UStruct>(ClassNameString, ELogVerbosity::Warning, TEXT("AssetRegistry trying to convert short name to path name"));
+		if (ClassPath.IsNull())
+		{
+			// In some cases the class name stored in asset registry tags have been redirected with ini class redirects
+			FString RedirectedName = FLinkerLoad::FindNewPathNameForClass(ClassNameString, false);
+			if (!FPackageName::IsShortPackageName(RedirectedName))
+			{
+				ClassPath = FTopLevelAssetPath(RedirectedName);
+			}
+			else
+			{
+				ClassPath = UClass::TryConvertShortTypeNameToPathName<UStruct>(RedirectedName, ELogVerbosity::Warning, TEXT("AssetRegistry trying to convert redirected short name to path name"));
+			}
+
+			if (ClassPath.IsNull())
+			{
+				// Fallback to a fake name but at least the class name will be preserved
+				ClassPath = FTopLevelAssetPath(TEXT("/Unknown"), InClassName);
+				UE_LOG(LogAssetData, Warning, TEXT("Failed to convert deprecated short class name \"%s\" to path name. Using \"%s\""), *InClassName.ToString(), *ClassPath.ToString());
+			}
+		}
+	}
+	return ClassPath;
+}
 
 bool FAssetRegistryVersion::SerializeVersion(FArchive& Ar, FAssetRegistryVersion::Type& Version)
 {
@@ -304,6 +358,41 @@ COREUOBJECT_API void FAssetPackageData::SerializeForCacheOldVersion(FArchive& Ar
 	SerializeForCacheInternal(Ar, *this, Version);
 }
 
+void FARFilter::PostSerialize(const FArchive& Ar)
+{
+	PRAGMA_DISABLE_DEPRECATION_WARNINGS
+#if WITH_EDITORONLY_DATA
+
+	auto ConvertShortClassNameToPathName = [](FName ShortClassFName)
+	{
+		FTopLevelAssetPath ClassPathName;
+		if (ShortClassFName != NAME_None)
+		{
+			FString ShortClassName = ShortClassFName.ToString();
+			ClassPathName = UClass::TryConvertShortTypeNameToPathName<UStruct>(*ShortClassName, ELogVerbosity::Warning, TEXT("FARFilter::PostSerialize"));
+			UE_CLOG(ClassPathName.IsNull(), LogAssetData, Error, TEXT("Failed to convert short class name %s to class path name."), *ShortClassName);
+		}
+		return ClassPathName;
+	};
+
+	for (FName ClassFName : ClassNames)
+	{
+		FTopLevelAssetPath ClassPathName = ConvertShortClassNameToPathName(ClassFName);
+		ClassPaths.Add(ClassPathName);
+	}
+	for (FName ClassFName : RecursiveClassesExclusionSet)
+	{
+		FTopLevelAssetPath ClassPathName = ConvertShortClassNameToPathName(ClassFName);
+		RecursiveClassPathsExclusionSet.Add(ClassPathName);
+	}
+
+	ClassNames.Empty();
+	RecursiveClassPathsExclusionSet.Empty();
+
+#endif // WITH_EDITORONLY_DATA
+PRAGMA_ENABLE_DEPRECATION_WARNINGS
+}
+
 namespace UE
 {
 namespace AssetRegistry
diff --git a/Engine/Source/Runtime/CoreUObject/Private/AssetRegistry/AssetDataTagMap.cpp b/Engine/Source/Runtime/CoreUObject/Private/AssetRegistry/AssetDataTagMap.cpp
--- a/Engine/Source/Runtime/CoreUObject/Private/AssetRegistry/AssetDataTagMap.cpp
+++ b/Engine/Source/Runtime/CoreUObject/Private/AssetRegistry/AssetDataTagMap.cpp
@@ -25,7 +25,7 @@ FAssetRegistryExportPath ParseExportPath(T ExportPath)
 	{
 		T ClassName;
 		verify(FPackageName::ParseExportTextPath(ExportPath, &ClassName, &ObjectPath));
-		Out.Class = FName(ClassName);
+		Out.ClassPath = FTopLevelAssetPath(ClassName);
 	}
 	else
 	{
@@ -60,7 +60,7 @@ FString FAssetRegistryExportPath::ToString() const
 
 FName FAssetRegistryExportPath::ToName() const
 {
-	if (Class.IsNone() && Object.IsNone())
+	if (ClassPath.IsNull() && Object.IsNone())
 	{
 		return Package;
 	}
@@ -72,21 +72,37 @@ FName FAssetRegistryExportPath::ToName() const
 
 void FAssetRegistryExportPath::ToString(FStringBuilderBase& Out) const
 {
-	if (!Class.IsNone())
+	if (!ClassPath.IsNull())
 	{
-		Out << Class << '\'';
+		Out << ClassPath << '\'';
 	}
 	Out << Package;
 	if (!Object.IsNone())
 	{
 		Out << '.' << Object;
 	}
-	if (!Class.IsNone())
+	if (!ClassPath.IsNull())
 	{
 		Out << '\'';
 	}
 }
 
+FString FAssetRegistryExportPath::ToPath() const
+{
+	TStringBuilder<256> Path;
+	ToPath(Path);
+	return FString(Path);
+}
+
+void FAssetRegistryExportPath::ToPath(FStringBuilderBase& Out) const
+{
+	Out << Package;
+	if (!Object.IsNone())
+	{
+		Out << '.' << Object;
+	}
+}
+
 static FAssetRegistryExportPath MakeNumberedPath(const FNameEntry& NameEntry)
 {
 	TStringBuilder<NAME_SIZE> Buffer;
@@ -103,17 +119,17 @@ static FAssetRegistryExportPath MakeNumberedPath(FName Name)
 
 static FArchive& operator<<(FArchive& Ar, FAssetRegistryExportPath& Path)
 {
-	return Ar << Path.Class << Path.Object << Path.Package;
+	return Ar << Path.ClassPath << Path.Object << Path.Package;
 }
 
 bool operator==(const FAssetRegistryExportPath& A, const FAssetRegistryExportPath& B)
 {
-	return A.Class == B.Class & A.Package == B.Package & A.Object == B.Object;
+	return A.ClassPath == B.ClassPath & A.Package == B.Package & A.Object == B.Object;
 }
 
 uint32 GetTypeHash(const FAssetRegistryExportPath& Path)
 {
-	return FixedTagPrivate::HashCombineQuick(GetTypeHash(Path.Class), GetTypeHash(Path.Package), GetTypeHash(Path.Object));
+	return FixedTagPrivate::HashCombineQuick(GetTypeHash(Path.ClassPath), GetTypeHash(Path.Package), GetTypeHash(Path.Object));
 }
 
 //////////////////////////////////////////////////////////////////////////
@@ -148,6 +164,11 @@ static FString LocalizeIfComplexString(const FString& Value)
 
 namespace FixedTagPrivate
 {
+	FArchive& operator<<(FArchive& Ar, FLegacyAssetRegistryExportPath& Path)
+	{
+		return Ar << Path.Class << Path.Object << Path.Package;
+	}
+
 	uint32 HashCaseSensitive(const TCHAR* Str, int32 Len)		
 	{
 		return CityHash64(reinterpret_cast<const char*>(Str), sizeof(TCHAR) * Len);
@@ -200,7 +221,8 @@ namespace FixedTagPrivate
 	
 	static bool EqualsInsensitive(const FNumberlessExportPath& A, const FNumberlessExportPath& B)
 	{
-		return	EqualsInsensitive(A.Class, B.Class) &		//-V792
+		return	EqualsInsensitive(A.ClassPackage, B.ClassPackage) &		//-V792
+				EqualsInsensitive(A.ClassObject, B.ClassObject) &		//-V792
 				EqualsInsensitive(A.Package, B.Package) &	//-V792
 				EqualsInsensitive(A.Object, B.Object);		//-V792
 	}
@@ -228,27 +250,37 @@ namespace FixedTagPrivate
 
 	static bool IsNumberless(FAssetRegistryExportPath Path)
 	{
-		return IsNumberless(Path.Class) & IsNumberless(Path.Object) & IsNumberless(Path.Package); //-V792
+		return IsNumberless(Path.ClassPath.GetAssetName()) & IsNumberless(Path.Object) & IsNumberless(Path.Package); //-V792
 	}
 
 	static FAssetRegistryExportPath MakeNumberedPath(const FNumberlessExportPath& Path)
 	{
 		FAssetRegistryExportPath Out;
-		Out.Class	= MakeNumberedName(Path.Class);
-		Out.Object	= MakeNumberedName(Path.Object);
-		Out.Package	= MakeNumberedName(Path.Package);
+		Out.ClassPath = FTopLevelAssetPath(MakeNumberedName(Path.ClassPackage), MakeNumberedName(Path.ClassObject));
+		Out.Object	  = MakeNumberedName(Path.Object);
+		Out.Package	  = MakeNumberedName(Path.Package);
 		return Out;
 	}
 	
 	static FNumberlessExportPath MakeNumberlessPath(const FAssetRegistryExportPath& Path)
 	{
 		FNumberlessExportPath Out;
-		Out.Class	= MakeNumberlessDisplayName(Path.Class);
-		Out.Object	= MakeNumberlessDisplayName(Path.Object);
-		Out.Package	= MakeNumberlessDisplayName(Path.Package);
+		Out.ClassPackage = MakeNumberlessDisplayName(Path.ClassPath.GetPackageName());
+		Out.ClassObject  = MakeNumberlessDisplayName(Path.ClassPath.GetAssetName());
+		Out.Object	     = MakeNumberlessDisplayName(Path.Object);
+		Out.Package	     = MakeNumberlessDisplayName(Path.Package);
 		return Out;
 	}
 	
+	static FAssetRegistryExportPath ConvertLegacyPath(const FLegacyAssetRegistryExportPath& Path)
+	{
+		FAssetRegistryExportPath Out;
+		Out.ClassPath = FAssetData::TryConvertShortClassNameToPathName(Path.Class);
+		Out.Object = Path.Object;
+		Out.Package = Path.Package;
+		return Out;
+	}
+
 	FString FNumberlessExportPath::ToString() const
 	{
 		return MakeNumberedPath(*this).ToString();
@@ -903,7 +935,8 @@ namespace FixedTagPrivate
 		View = MakeArrayView(reinterpret_cast<T*>(Data), View.Num());
 	}
 
-	class FSerializer
+	template <FAssetRegistryVersion::Type Version>
+	class TSerializer
 	{
 		template<class T>
 		void SaveItem(T Item)
@@ -934,6 +967,20 @@ namespace FixedTagPrivate
 			SaveItem(MakeNumberedPath(Path));
 		}
 
+		template<> FAssetRegistryExportPath LoadItem()
+		{
+			if (Version >= FAssetRegistryVersion::ClassPaths)
+			{
+				FAssetRegistryExportPath ExportPath;
+				Ar << ExportPath;
+				return ExportPath;
+			}
+			else
+			{
+				return ConvertLegacyPath(LoadItem<FLegacyAssetRegistryExportPath>());
+			}
+		}
+
 		template<> FNumberlessExportPath LoadItem()
 		{
 			return MakeNumberlessPath(LoadItem<FAssetRegistryExportPath>());
@@ -1042,12 +1089,15 @@ namespace FixedTagPrivate
 		static constexpr uint32 EndMagic		= 0x87654321;
 		static constexpr uint32 MaxViewAlignment = 16;
 	public:
-		FSerializer(FArchive& InAr) : Ar(InAr) {}
+		TSerializer(FArchive& InAr) // , FAssetRegistryVersion::Type InVersion = FAssetRegistryVersion::LatestVersion
+			: Ar(InAr)
+		{
+		}
 
 		void SaveTextData(TArrayView<const FText> Texts)
 		{
 			FArrayWriter Data;
-			FSerializer(Data).SaveViewData(Texts);
+			TSerializer<Version>(Data).SaveViewData(Texts);
 			SaveItem(Data.Num());
 			Ar.Serialize(Data.GetData(), Data.Num());
 		}
@@ -1171,10 +1221,10 @@ namespace FixedTagPrivate
 
 	void SaveStore(const FStoreData& Store, FArchive& Ar)
 	{
-		FSerializer(Ar).Save(Store);
+		TSerializer<FAssetRegistryVersion::LatestVersion>(Ar).Save(Store);
 	}
 
-	TRefCountPtr<const FStore> LoadStore(FArchive& Ar)
+	TRefCountPtr<const FStore> LoadStore(FArchive& Ar, FAssetRegistryVersion::Type Version /*= FAssetRegistryVersion::LatestVersion*/)
 	{
 		if (Ar.IsError())
 		{
@@ -1182,7 +1232,24 @@ namespace FixedTagPrivate
 		}
 
 		FStore* Store = GStores.CreateAndRegister();
-		FSerializer(Ar).Load(*Store);
+
+		// Versioning support inside of TSerializer was added with FAssetRegistryVersion::ClassPaths so everything before can just be
+		// serialized with FAssetRegistryVersion::ClassPaths - 1 = FAssetRegistryVersion::AddedChunkHashes.
+		// The main purpose of hardcoding version value and passing it as template argument is to eliminate
+		// branches in performance critical serialization code.
+		// Note that we branch in this function only because it's the only fucntion used to load legacy 
+		// asset registries in DiffAssetRegistries commandlet hence it's editor only code too
+#if WITH_EDITORONLY_DATA
+		if (Version < FAssetRegistryVersion::ClassPaths)
+		{
+			TSerializer<FAssetRegistryVersion::ObjectResourceOptionalVersionChange>(Ar).Load(*Store);
+		}
+		else // Add more branches should serialization format change
+#endif // WITH_EDITORONLY_DATA
+		{			
+			TSerializer<FAssetRegistryVersion::LatestVersion>(Ar).Load(*Store);
+		}
+		
 		return TRefCountPtr<const FStore>(Store);
 	}
 
@@ -1192,18 +1259,18 @@ namespace FixedTagPrivate
 
 	TFuture<void> FAsyncStoreLoader::ReadInitialDataAndKickLoad(FArchive& Ar, uint32 MaxWorkerTasks)
 	{
-		Order = FSerializer(Ar).LoadHeader(*Store);
+		Order = TSerializer<FAssetRegistryVersion::LatestVersion>(Ar).LoadHeader(*Store);
 
 		if (Order == ELoadOrder::TextFirst)
 		{
-			TArray<uint8> TextData = FSerializer(Ar).ReadTextData();
+			TArray<uint8> TextData = TSerializer<FAssetRegistryVersion::LatestVersion>(Ar).ReadTextData();
 
 			if (TextData.Num())
 			{
 				return Async(EAsyncExecution::TaskGraph, [Data = MoveTemp(TextData), OutStore = Store]()
 				{
 					FMemoryReader Reader(Data);
-					FSerializer(Reader).LoadTextData(*OutStore);
+					TSerializer<FAssetRegistryVersion::LatestVersion>(Reader).LoadTextData(*OutStore);
 				});
 			}
 		}
@@ -1215,7 +1282,7 @@ namespace FixedTagPrivate
 	{
 		if (Order.IsSet())
 		{
-			FSerializer(Ar).LoadFinalData(/* Out */ *Store, Order.GetValue());
+			TSerializer<FAssetRegistryVersion::LatestVersion>(Ar).LoadFinalData(/* Out */ *Store, Order.GetValue());
 
 			return TRefCountPtr<const FStore>(Store);
 		}
@@ -1455,13 +1522,13 @@ IMPLEMENT_SIMPLE_AUTOMATION_TEST(FAssetRegistryExportPathTest, "Engine.AssetRegi
 
 bool FAssetRegistryExportPathTest::RunTest(const FString& Parameters)
 {
-	TestEqual("Full numbered path",	FAssetRegistryExportPath("C_1\'P_2.O_3\'").ToString(),	"C_1\'P_2.O_3\'");
-	TestEqual("Full path",			FAssetRegistryExportPath("C\'P.O\'").ToString(),		"C\'P.O\'");
-	TestEqual("Package path",		FAssetRegistryExportPath("P.O").ToString(),				"P.O");
-	TestEqual("Object only path",	FAssetRegistryExportPath("O").ToString(),				"O");
-	TestEqual("Class parse",		FAssetRegistryExportPath("C\'P.O\'").Class,		FName("C"));
-	TestEqual("Package parse",		FAssetRegistryExportPath("C\'P.O\'").Package,	FName("P"));
-	TestEqual("Object parse",		FAssetRegistryExportPath("C\'P.O\'").Object,	FName("O"));
+	TestEqual("Full numbered path",	FAssetRegistryExportPath("/S/P.C_1\'P_2.O_3\'").ToString(),	"/S/P.C_1\'P_2.O_3\'");
+	TestEqual("Full path",			FAssetRegistryExportPath("/S/P.C\'P.O\'").ToString(),		"/S/P.C\'P.O\'");
+	TestEqual("Package path",		FAssetRegistryExportPath("P.O").ToString(),					"P.O");
+	TestEqual("Object only path",	FAssetRegistryExportPath("O").ToString(),					"O");
+	TestEqual("Class parse",		FAssetRegistryExportPath("/S/P.C\'P.O\'").ClassPath.ToString(),	"/S/P.C");
+	TestEqual("Package parse",		FAssetRegistryExportPath("/S/P.C\'P.O\'").Package,			FName("P"));
+	TestEqual("Object parse",		FAssetRegistryExportPath("/S/P.C\'P.O\'").Object,			FName("O"));
 
 	return true;
 }
@@ -1473,7 +1540,7 @@ IMPLEMENT_SIMPLE_AUTOMATION_TEST(FCompactExportPathTest, "Engine.AssetRegistry.F
 
 bool FCompactExportPathTest::RunTest(const FString& Parameters)
 {
-	const FAssetRegistryExportPath FullPath = FAssetRegistryExportPath("C\'P.O\'");
+	const FAssetRegistryExportPath FullPath = FAssetRegistryExportPath("/S/P.C\'P.O\'");
 
 	TestTrue("Roundtrip",	FullPath == MakeNumberedPath(MakeNumberlessPath(FullPath)));
 
@@ -1555,16 +1622,16 @@ bool FStoreTest::RunTest(const FString& Parameters)
 								{"Key_0",		"StringValue_0"}}));
 	LooseMaps.Add(MakeLooseMap({{"Name",		"NameValue"}, 
 								{"Name_0",		"NameValue_0"}}));
-	LooseMaps.Add(MakeLooseMap({{"FullPath",	"C\'P.O\'"}, 
+	LooseMaps.Add(MakeLooseMap({{"FullPath",	"/S/P.C\'P.O\'"}, 
 								{"PkgPath",		"P.O"},
 								{"ObjPath",		"O"}}));
-	LooseMaps.Add(MakeLooseMap({{"NumPath_0",	"C\'P.O_0\'"}, 
-								{"NumPath_1",	"C\'P_0.O\'"},
-								{"NumPath_2",	"C_0\'P.O\'"},
-								{"NumPath_3",	"C\'P_0.O_0\'"},
-								{"NumPath_4",	"C_0\'P_0.O\'"},
-								{"NumPath_5",	"C_0\'P.O_0\'"},
-								{"NumPath_6",	"C_0\'P_0.O_0\'"}}));
+	LooseMaps.Add(MakeLooseMap({{"NumPath_0",	"/S/P.C\'P.O_0\'"}, 
+								{"NumPath_1",	"/S/P.C\'P_0.O\'"},
+								{"NumPath_2",	"/S/P.C_0\'P.O\'"},
+								{"NumPath_3",	"/S/P.C\'P_0.O_0\'"},
+								{"NumPath_4",	"/S/P.C_0\'P_0.O\'"},
+								{"NumPath_5",	"/S/P.C_0\'P.O_0\'"},
+								{"NumPath_6",	"/S/P.C_0\'P_0.O_0\'"}}));
 	LooseMaps.Add(MakeLooseMap({{"SameSame",	"SameSame"}, 
 								{"AlsoSame",	"SameSame"}}));
 	LooseMaps.Add(MakeLooseMap({{"FilterKey1",	"FilterValue1"}, 
diff --git a/Engine/Source/Runtime/CoreUObject/Private/Blueprint/BlueprintSupport.cpp b/Engine/Source/Runtime/CoreUObject/Private/Blueprint/BlueprintSupport.cpp
--- a/Engine/Source/Runtime/CoreUObject/Private/Blueprint/BlueprintSupport.cpp
+++ b/Engine/Source/Runtime/CoreUObject/Private/Blueprint/BlueprintSupport.cpp
@@ -2375,9 +2375,18 @@ UObject* FLinkerLoad::FindImport(UClass* ImportClass, UObject* ImportOuter, cons
 	return Result;
 }
 
-UObject* FLinkerLoad::FindImportFast(UClass* ImportClass, UObject* ImportOuter, FName Name, bool bAnyPackage)
+UObject* FLinkerLoad::FindImportFast(UClass* ImportClass, UObject* ImportOuter, FName Name, bool bFindObjectbyName)
 {
-	UObject* Result = StaticFindObjectFast(ImportClass, ImportOuter, Name, false/*ExactClass*/, bAnyPackage);
+	UObject* Result = nullptr;
+	if (!bFindObjectbyName)
+	{
+		Result = StaticFindObjectFast(ImportClass, ImportOuter, Name, false/*ExactClass*/);
+	}
+	else
+	{
+		Result = StaticFindFirstObject(ImportClass, *Name.ToString(), EFindFirstObjectOptions::NativeFirst | EFindFirstObjectOptions::EnsureIfAmbiguous, ELogVerbosity::Warning, TEXT("FindImportFast"));
+	}
+
 	return Result;
 }
 
diff --git a/Engine/Source/Runtime/CoreUObject/Private/Internationalization/PackageLocalizationCache.cpp b/Engine/Source/Runtime/CoreUObject/Private/Internationalization/PackageLocalizationCache.cpp
--- a/Engine/Source/Runtime/CoreUObject/Private/Internationalization/PackageLocalizationCache.cpp
+++ b/Engine/Source/Runtime/CoreUObject/Private/Internationalization/PackageLocalizationCache.cpp
@@ -184,16 +184,16 @@ FPackageLocalizationCache::FPackageLocalizationCache()
 				for (const auto& SectionEntryPair : *AssetGroupClassesSection)
 				{
 					const FName GroupName = SectionEntryPair.Key;
-					const FName ClassName = *SectionEntryPair.Value.GetValue();
-
-					const auto* AssetClassGroupPair = AssetClassesToAssetGroups.FindByPredicate([&](const TTuple<FName, FName>& InAssetClassToAssetGroup)
+					const FTopLevelAssetPath ClassName = UClass::TryConvertShortTypeNameToPathName<UStruct>(SectionEntryPair.Value.GetValue(), ELogVerbosity::Fatal, TEXT("ReadAssetGroupClassSettings"));
+					
+					const auto* AssetClassGroupPair = AssetClassesToAssetGroups.FindByPredicate([&](const TTuple<FTopLevelAssetPath, FName>& InAssetClassToAssetGroup)
 					{
 						return InAssetClassToAssetGroup.Key == ClassName;
 					});
 
 					if (AssetClassGroupPair)
 					{
-						UE_CLOG(AssetClassGroupPair->Value != ClassName, LogPackageLocalizationCache, Warning, TEXT("Class '%s' was already assigned to asset group '%s', ignoring request to assign it to '%s' from the %s configuration."), *ClassName.ToString(), *AssetClassGroupPair->Value.ToString(), *GroupName.ToString(), InConfigLogName);
+						UE_CLOG(AssetClassGroupPair->Value != ClassName.GetAssetName(), LogPackageLocalizationCache, Warning, TEXT("Class '%s' was already assigned to asset group '%s', ignoring request to assign it to '%s' from the %s configuration."), *ClassName.ToString(), *AssetClassGroupPair->Value.ToString(), *GroupName.ToString(), InConfigLogName);
 					}
 					else
 					{
diff --git a/Engine/Source/Runtime/CoreUObject/Private/Internationalization/PackageLocalizationManager.cpp b/Engine/Source/Runtime/CoreUObject/Private/Internationalization/PackageLocalizationManager.cpp
--- a/Engine/Source/Runtime/CoreUObject/Private/Internationalization/PackageLocalizationManager.cpp
+++ b/Engine/Source/Runtime/CoreUObject/Private/Internationalization/PackageLocalizationManager.cpp
@@ -18,7 +18,7 @@ public:
 protected:
 	//~ FPackageLocalizationCache interface
 	virtual void FindLocalizedPackages(const TMap<FString, TArray<FString>>& NewSourceToLocalizedPaths, TMap<FName, TArray<FName>>& InOutSourcePackagesToLocalizedPackages) override;
-	virtual void FindAssetGroupPackages(const FName InAssetGroupName, const FName InAssetClassName) override;
+	virtual void FindAssetGroupPackages(const FName InAssetGroupName, const FTopLevelAssetPath& InAssetClassName) override;
 };
 
 void FDefaultPackageLocalizationCache::FindLocalizedPackages(const TMap<FString, TArray<FString>>& NewSourceToLocalizedPaths, TMap<FName, TArray<FName>>& InOutSourcePackagesToLocalizedPackages)
@@ -50,7 +50,7 @@ void FDefaultPackageLocalizationCache::FindLocalizedPackages(const TMap<FString,
 	}
 }
 
-void FDefaultPackageLocalizationCache::FindAssetGroupPackages(const FName InAssetGroupName, const FName InAssetClassName)
+void FDefaultPackageLocalizationCache::FindAssetGroupPackages(const FName InAssetGroupName, const FTopLevelAssetPath& InAssetClassName)
 {
 	// Not supported without the asset registry
 }
diff --git a/Engine/Source/Runtime/CoreUObject/Private/Misc/AssetRegistryInterface.cpp b/Engine/Source/Runtime/CoreUObject/Private/Misc/AssetRegistryInterface.cpp
--- a/Engine/Source/Runtime/CoreUObject/Private/Misc/AssetRegistryInterface.cpp
+++ b/Engine/Source/Runtime/CoreUObject/Private/Misc/AssetRegistryInterface.cpp
@@ -28,11 +28,11 @@ namespace Private
 }
 
 #if WITH_ENGINE && WITH_EDITOR
-	TSet<FName> SkipUncookedClasses;
-	TSet<FName> SkipCookedClasses;
+	TSet<FTopLevelAssetPath> SkipUncookedClasses;
+	TSet<FTopLevelAssetPath> SkipCookedClasses;
 	bool bInitializedSkipClasses = false;
 
-	void FFiltering::SetSkipClasses(const TSet<FName>& InSkipUncookedClasses, const TSet<FName>& InSkipCookedClasses)
+	void FFiltering::SetSkipClasses(const TSet<FTopLevelAssetPath>& InSkipUncookedClasses, const TSet<FTopLevelAssetPath>& InSkipCookedClasses)
 	{
 		bInitializedSkipClasses = true;
 		SkipUncookedClasses = InSkipUncookedClasses;
@@ -44,7 +44,7 @@ namespace Private
 namespace Utils
 {
 
-	bool ShouldSkipAsset(FName AssetClass, uint32 PackageFlags, const TSet<FName>& InSkipUncookedClasses, const TSet<FName>& InSkipCookedClasses)
+	bool ShouldSkipAsset(const FTopLevelAssetPath& AssetClass, uint32 PackageFlags, const TSet<FTopLevelAssetPath>& InSkipUncookedClasses, const TSet<FTopLevelAssetPath>& InSkipCookedClasses)
 	{
 		if (PackageFlags & PKG_ContainsNoAsset)
 		{
@@ -60,7 +60,7 @@ namespace Utils
 		return false;
 	}
 
-	bool ShouldSkipAsset(const UObject* InAsset, const TSet<FName>& InSkipUncookedClasses, const TSet<FName>& InSkipCookedClasses)
+	bool ShouldSkipAsset(const UObject* InAsset, const TSet<FTopLevelAssetPath>& InSkipUncookedClasses, const TSet<FTopLevelAssetPath>& InSkipCookedClasses)
 	{
 		if (!InAsset)
 		{
@@ -71,10 +71,10 @@ namespace Utils
 		{
 			return false;
 		}
-		return ShouldSkipAsset(InAsset->GetClass()->GetFName(), Package->GetPackageFlags(), InSkipUncookedClasses, InSkipCookedClasses);
+		return ShouldSkipAsset(InAsset->GetClass()->GetClassPathName(), Package->GetPackageFlags(), InSkipUncookedClasses, InSkipCookedClasses);
 	}
 
-	void PopulateSkipClasses(TSet<FName>& OutSkipUncookedClasses, TSet<FName>& OutSkipCookedClasses)
+	void PopulateSkipClasses(TSet<FTopLevelAssetPath>& OutSkipUncookedClasses, TSet<FTopLevelAssetPath>& OutSkipCookedClasses)
 	{
 		static const FName NAME_EnginePackage("/Script/Engine");
 		UPackage* EnginePackage = Cast<UPackage>(StaticFindObjectFast(UPackage::StaticClass(), nullptr, NAME_EnginePackage));
@@ -93,12 +93,12 @@ namespace Utils
 			}
 			else
 			{
-				OutSkipUncookedClasses.Add(NAME_BlueprintGeneratedClass);
+				OutSkipUncookedClasses.Add(BlueprintGeneratedClass->GetClassPathName());
 				for (TObjectIterator<UClass> It; It; ++It)
 				{
 					if (It->IsChildOf(BlueprintGeneratedClass) && !It->HasAnyClassFlags(CLASS_Abstract))
 					{
-						OutSkipUncookedClasses.Add(It->GetFName());
+						OutSkipUncookedClasses.Add(It->GetClassPathName());
 					}
 				}
 			}
@@ -118,12 +118,12 @@ namespace Utils
 			}
 			else
 			{
-				OutSkipCookedClasses.Add(NAME_Blueprint);
+				OutSkipCookedClasses.Add(BlueprintClass->GetClassPathName());
 				for (TObjectIterator<UClass> It; It; ++It)
 				{
 					if (It->IsChildOf(BlueprintClass) && !It->HasAnyClassFlags(CLASS_Abstract))
 					{
-						OutSkipCookedClasses.Add(It->GetFName());
+						OutSkipCookedClasses.Add(It->GetClassPathName());
 					}
 				}
 			}
@@ -133,7 +133,7 @@ namespace Utils
 }
 #endif
 
-	bool FFiltering::ShouldSkipAsset(FName AssetClass, uint32 PackageFlags)
+	bool FFiltering::ShouldSkipAsset(const FTopLevelAssetPath& AssetClass, uint32 PackageFlags)
 	{
 #if WITH_ENGINE && WITH_EDITOR
 		// We do not yet support having UBlueprintGeneratedClasses be assets when the UBlueprint is also
@@ -169,7 +169,7 @@ namespace Utils
 		{
 			return false;
 		}
-		return ShouldSkipAsset(InAsset->GetClass()->GetFName(), Package->GetPackageFlags());
+		return ShouldSkipAsset(InAsset->GetClass()->GetClassPathName(), Package->GetPackageFlags());
 	}
 
 	void FFiltering::MarkDirty()
diff --git a/Engine/Source/Runtime/CoreUObject/Private/Serialization/AsyncLoading.cpp b/Engine/Source/Runtime/CoreUObject/Private/Serialization/AsyncLoading.cpp
--- a/Engine/Source/Runtime/CoreUObject/Private/Serialization/AsyncLoading.cpp
+++ b/Engine/Source/Runtime/CoreUObject/Private/Serialization/AsyncLoading.cpp
@@ -1626,7 +1626,7 @@ FObjectImport* FAsyncPackage::FindExistingImport(int32 LocalImportIndex)
 			FObjectImport* ImportOuter = FindExistingImport(Import->OuterIndex.ToImport());
 			if (ImportOuter->XObject)
 			{
-				Import->XObject = StaticFindObjectFast(UObject::StaticClass(), ImportOuter->XObject, Import->ObjectName, false, true);
+				Import->XObject = StaticFindObjectFast(UObject::StaticClass(), ImportOuter->XObject, Import->ObjectName, false);
 				if (Import->XObject)
 				{
 					FName NameImportClass(Import->ClassName);
@@ -1950,7 +1950,7 @@ EAsyncPackageState::Type FAsyncPackage::SetupImports_Event()
 				UPackage* ImportPackage = Import.XObject ? CastChecked<UPackage>(Import.XObject) : nullptr;
 				if (!ImportPackage)
 				{
-					ImportPackage = FindObjectFast<UPackage>(NULL, Import.ObjectName, false, false);
+					ImportPackage = FindObjectFast<UPackage>(NULL, Import.ObjectName, false);
 					if (!ImportPackage)
 					{
 						Import.bImportFailed = true;
@@ -2002,7 +2002,7 @@ EAsyncPackageState::Type FAsyncPackage::SetupImports_Event()
 			UPackage* ImportPackage = OuterMostImport.XObject ? OuterMostImport.XObject->GetPackage() : nullptr;
 			if (!ImportPackage)
 			{
-				ImportPackage = FindObjectFast<UPackage>(nullptr, ImportPackageName, false, false);
+				ImportPackage = FindObjectFast<UPackage>(nullptr, ImportPackageName, false);
 				if (!ImportPackage)
 				{
 					Import.bImportFailed = true;
@@ -2513,10 +2513,10 @@ void FAsyncPackage::LinkImport(int32 LocalImportIndex)
 				}
 				else
 				{
-					UPackage* ClassPackage = FindObjectFast<UPackage>(NULL, Import.ClassPackage, false, false);
+					UPackage* ClassPackage = FindObjectFast<UPackage>(NULL, Import.ClassPackage, false);
 					if (ClassPackage)
 					{
-						UClass*	FindClass = FindObjectFast<UClass>(ClassPackage, Import.ClassName, false, false);
+						UClass*	FindClass = FindObjectFast<UClass>(ClassPackage, Import.ClassName, false);
 						if (FindClass)
 						{
 							UObject* Outer = ImportPackage;
diff --git a/Engine/Source/Runtime/CoreUObject/Private/Serialization/AsyncLoading2.cpp b/Engine/Source/Runtime/CoreUObject/Private/Serialization/AsyncLoading2.cpp
--- a/Engine/Source/Runtime/CoreUObject/Private/Serialization/AsyncLoading2.cpp
+++ b/Engine/Source/Runtime/CoreUObject/Private/Serialization/AsyncLoading2.cpp
@@ -3241,7 +3241,7 @@ static UObject* GFindExistingScriptImport(FPackageObjectIndex GlobalImportIndex,
 			Object = &ScriptObjects.FindChecked(GlobalImportIndex);
 			if (Outer)
 			{
-				*Object = StaticFindObjectFast(UObject::StaticClass(), Outer, MinimalNameToName(Entry->ObjectName), false, true);
+				*Object = StaticFindObjectFast(UObject::StaticClass(), Outer, MinimalNameToName(Entry->ObjectName), false);
 			}
 		}
 	}
diff --git a/Engine/Source/Runtime/CoreUObject/Private/UObject/Class.cpp b/Engine/Source/Runtime/CoreUObject/Private/UObject/Class.cpp
--- a/Engine/Source/Runtime/CoreUObject/Private/UObject/Class.cpp
+++ b/Engine/Source/Runtime/CoreUObject/Private/UObject/Class.cpp
@@ -53,6 +53,7 @@
 #include "Math/InterpCurvePoint.h"
 #include "UObject/UE5MainStreamObjectVersion.h"
 #include "UObject/TopLevelAssetPath.h"
+#include "AssetRegistry/AssetData.h"
 #include "HAL/PlatformStackWalk.h"
 
 // This flag enables some expensive class tree validation that is meant to catch mutations of 
@@ -523,14 +524,14 @@ void UField::SetMetaData(const FName& Key, const TCHAR* InValue)
 UClass* UField::GetClassMetaData(const TCHAR* Key) const
 {
 	const FString& ClassName = GetMetaData(Key);
-	UClass* const FoundObject = FindObject<UClass>(ANY_PACKAGE, *ClassName);
+	UClass* const FoundObject = UClass::TryFindTypeSlow<UClass>(ClassName);
 	return FoundObject;
 }
 
 UClass* UField::GetClassMetaData(const FName& Key) const
 {
 	const FString& ClassName = GetMetaData(Key);
-	UClass* const FoundObject = FindObject<UClass>(ANY_PACKAGE, *ClassName);
+	UClass* const FoundObject = UClass::TryFindTypeSlow<UClass>(ClassName);
 	return FoundObject;
 }
 
@@ -610,6 +611,7 @@ UStruct::UStruct(EStaticConstructor, int32 InSize, int32 InMinAlignment, EObject
 {
 #if WITH_EDITORONLY_DATA
 	FieldPathSerialNumber = GetNextFieldPathSerialNumber();
+	bHasAssetRegistrySearchableProperties = false;
 #endif // WITH_EDITORONLY_DATA
 }
 
@@ -631,6 +633,7 @@ UStruct::UStruct(UStruct* InSuperStruct, SIZE_T ParamsSize, SIZE_T Alignment)
 #endif
 #if WITH_EDITORONLY_DATA
 	FieldPathSerialNumber = GetNextFieldPathSerialNumber();
+	bHasAssetRegistrySearchableProperties = false;
 #endif // WITH_EDITORONLY_DATA
 }
 
@@ -652,6 +655,7 @@ UStruct::UStruct(const FObjectInitializer& ObjectInitializer, UStruct* InSuperSt
 #endif
 #if WITH_EDITORONLY_DATA
 	FieldPathSerialNumber = GetNextFieldPathSerialNumber();
+	bHasAssetRegistrySearchableProperties = false;
 #endif // WITH_EDITORONLY_DATA
 }
 
@@ -958,6 +962,12 @@ void UStruct::Link(FArchive& Ar, bool bRelinkExistingProperties)
 			PostConstructLinkPtr = &(*PostConstructLinkPtr)->PostConstructLinkNext;
 		}
 
+#if WITH_EDITORONLY_DATA
+		// Set the bHasAssetRegistrySearchableProperties flag.
+		// Note that we're also iterating over super class properties here so this flag is being automatically inherited
+		bHasAssetRegistrySearchableProperties |= Property->HasAnyPropertyFlags(CPF_AssetRegistrySearchable);
+#endif
+
 		*PropertyLinkPtr = Property;
 		PropertyLinkPtr = &(*PropertyLinkPtr)->PropertyLinkNext;
 	}
@@ -1640,6 +1650,7 @@ void UStruct::DestroyChildPropertiesAndResetPropertyLinks()
 	PostConstructLink = nullptr;
 #if WITH_EDITORONLY_DATA
 	FieldPathSerialNumber = GetNextFieldPathSerialNumber();
+	bHasAssetRegistrySearchableProperties = false;
 #endif // WITH_EDITORONLY_DATA
 	DestroyUnversionedSchema(this);
 }
@@ -4435,7 +4446,7 @@ void UClass::GetAssetRegistryTags(TArray<FAssetRegistryTag>& OutTags) const
 #if WITH_EDITOR
 	static const FName ParentClassFName = "ParentClass";
 	const UClass* const ParentClass = GetSuperClass();
-	OutTags.Add( FAssetRegistryTag(ParentClassFName, ((ParentClass) ? ParentClass->GetFName() : NAME_None).ToString(), FAssetRegistryTag::TT_Alphabetical) );
+	OutTags.Add( FAssetRegistryTag(ParentClassFName, ((ParentClass) ? ParentClass->GetPathName() : FString()), FAssetRegistryTag::TT_Alphabetical) );
 
 	static const FName ModuleNameFName = "ModuleName";
 	const UPackage* const ClassPackage = GetOuterUPackage();
@@ -4447,6 +4458,24 @@ void UClass::GetAssetRegistryTags(TArray<FAssetRegistryTag>& OutTags) const
 #endif
 }
 
+#if WITH_EDITOR
+void UClass::PostLoadAssetRegistryTags(const FAssetData& InAssetData, TArray<FAssetRegistryTag>& OutTagsAndValuesToUpdate) const
+{
+	Super::PostLoadAssetRegistryTags(InAssetData, OutTagsAndValuesToUpdate);
+
+	static const FName ParentClassFName(TEXT("ParentClass"));
+	FString ParentClassTagValue = InAssetData.GetTagValueRef<FString>(ParentClassFName);
+	if (!ParentClassTagValue.IsEmpty() && FPackageName::IsShortPackageName(ParentClassTagValue))
+	{
+		FTopLevelAssetPath ParentClassPathName = UClass::TryConvertShortTypeNameToPathName<UStruct>(ParentClassTagValue, ELogVerbosity::Warning, TEXT("UClass::PostLoadAssetRegistryTags"));
+		if (!ParentClassPathName.IsNull())
+		{
+			OutTagsAndValuesToUpdate.Add(FAssetRegistryTag(ParentClassFName, ParentClassPathName.ToString(), FAssetRegistryTag::TT_Alphabetical));
+		}
+	}
+}
+#endif // WITH_EDITOR
+
 void UClass::GetPreloadDependencies(TArray<UObject*>& OutDeps)
 {
 	Super::GetPreloadDependencies(OutDeps);
@@ -5326,6 +5355,7 @@ void UClass::PurgeClass(bool bRecompilingOnLoad)
 	}
 	// Update the serial number so that FFieldPaths that point to properties of this struct know they need to resolve themselves again
 	FieldPathSerialNumber = GetNextFieldPathSerialNumber();
+	bHasAssetRegistrySearchableProperties = false;
 #else
 	{
 		// Destroy all properties owned by this struct
@@ -5986,7 +6016,7 @@ UField* UClass::TryFindTypeSlow(UClass* TypeClass, const FString& InPathNameOrSh
 			ANSICHAR Buffer[1024];			
 			uint64 StackFrames[10];
 			uint32 NumStackFrames = FPlatformStackWalk::CaptureStackBackTrace(StackFrames, UE_ARRAY_COUNT(StackFrames));
-			for (uint32 Idx = 0; Idx < NumStackFrames; Idx++)
+			for (uint32 Idx = 0; Idx < NumStackFrames && Idx < UE_ARRAY_COUNT(StackFrames); Idx++)
 			{
 				Buffer[0] = '\0';
 				FPlatformStackWalk::ProgramCounterToHumanReadableString(Idx, StackFrames[Idx], Buffer, sizeof(Buffer));
@@ -6005,7 +6035,7 @@ UField* UClass::TryFindTypeSlow(UClass* TypeClass, const FString& InPathNameOrSh
 			
 			FoundType = (UField*)StaticFindFirstObject(TypeClass, *InPathNameOrShortName, InOptions | EFindFirstObjectOptions::EnsureIfAmbiguous | EFindFirstObjectOptions::NativeFirst, ELogVerbosity::Error, TEXT("TryFindType"));
 
-			UE_LOG(LogClass, Warning, TEXT("Short type name \"%s\" provided for TryFindType. Please convert it to a path name (suggested: \"%s\"). Callstack:\r\n%s"), *InPathNameOrShortName, *GetPathNameSafe(FoundType), Callstack.GetData());
+			UE_LOG(LogClass, Warning, TEXT("Short type name \"%s\" provided for TryFindType. Please convert it to a path name (suggested: \"%s\"). Callstack:\r\n\r\n%.*s"), *InPathNameOrShortName, *GetPathNameSafe(FoundType), Callstack.Len(), Callstack.GetData());
 		}
 	}
 	return FoundType;
@@ -6565,7 +6595,7 @@ UStruct* FStructUtils::FindStructureInPackageChecked(const TCHAR* StructName, co
 	}
 	else
 	{
-		return FindObjectChecked<UStruct>(ANY_PACKAGE, StructName);
+		return CastChecked<UStruct>(UClass::TryFindTypeSlow<UStruct>(StructName));
 	}
 }
 
@@ -6616,7 +6646,7 @@ static UScriptStruct* StaticGetBaseStructureInternal(FName Name)
 {
 	static UPackage* CoreUObjectPkg = FindObjectChecked<UPackage>(nullptr, TEXT("/Script/CoreUObject"));
 
-	UScriptStruct* Result = (UScriptStruct*)StaticFindObjectFastInternal(UScriptStruct::StaticClass(), CoreUObjectPkg, Name, false, false, RF_NoFlags, EInternalObjectFlags::None);
+	UScriptStruct* Result = (UScriptStruct*)StaticFindObjectFastInternal(UScriptStruct::StaticClass(), CoreUObjectPkg, Name, false, RF_NoFlags, EInternalObjectFlags::None);
 
 #if !(UE_BUILD_SHIPPING || UE_BUILD_TEST)
 	if (!Result)
diff --git a/Engine/Source/Runtime/CoreUObject/Private/UObject/Field.cpp b/Engine/Source/Runtime/CoreUObject/Private/UObject/Field.cpp
--- a/Engine/Source/Runtime/CoreUObject/Private/UObject/Field.cpp
+++ b/Engine/Source/Runtime/CoreUObject/Private/UObject/Field.cpp
@@ -911,15 +911,15 @@ void FField::SetMetaData(const FName& Key, FString&& InValue)
 UClass* FField::GetClassMetaData(const TCHAR* Key) const
 {
 	const FString& ClassName = GetMetaData(Key);
-	UClass* const FoundObject = FindObject<UClass>(ANY_PACKAGE, *ClassName);
-	return FoundObject;
+	UClass* FoundClass = UClass::TryFindTypeSlow<UClass>(ClassName);
+	return FoundClass;
 }
 
 UClass* FField::GetClassMetaData(const FName& Key) const
 {
 	const FString& ClassName = GetMetaData(Key);
-	UClass* const FoundObject = FindObject<UClass>(ANY_PACKAGE, *ClassName);
-	return FoundObject;
+	UClass* FoundClass = UClass::TryFindTypeSlow<UClass>(ClassName);;
+	return FoundClass;
 }
 
 void FField::RemoveMetaData(const TCHAR* Key)
@@ -1183,7 +1183,7 @@ FField* FindFPropertyByPath(const TCHAR* InFieldPath)
 		// And the FField part
 		InFieldPath += (LastSubobjectDelimiterIndex + 1);
 
-		UStruct* Owner = FindObject<UStruct>(ANY_PACKAGE, PathBuffer);
+		UStruct* Owner = FindObject<UStruct>(nullptr, PathBuffer);
 		if (Owner)
 		{
 #if DO_CHECK
diff --git a/Engine/Source/Runtime/CoreUObject/Private/UObject/GCObjectInfo.cpp b/Engine/Source/Runtime/CoreUObject/Private/UObject/GCObjectInfo.cpp
--- a/Engine/Source/Runtime/CoreUObject/Private/UObject/GCObjectInfo.cpp
+++ b/Engine/Source/Runtime/CoreUObject/Private/UObject/GCObjectInfo.cpp
@@ -19,7 +19,7 @@
 	if (Outer == nullptr)
 	{
 		// We are a package, those are the only objects that are allowed to have no outer
-		return StaticFindObjectFast(UObject::StaticClass(), nullptr, Name, false, false);
+		return StaticFindObjectFast(UObject::StaticClass(), nullptr, Name, false);
 	}
 
 	UClass* ResolvedClass = Cast<UClass>(Class->TryResolveObject());
@@ -30,7 +30,7 @@
 	
 	if (UObject* OuterObj = Outer->TryResolveObject())
 	{
-		return StaticFindObjectFast(ResolvedClass, OuterObj, Name, true, false);
+		return StaticFindObjectFast(ResolvedClass, OuterObj, Name, true);
 	}
 	return nullptr;
 }
diff --git a/Engine/Source/Runtime/CoreUObject/Private/UObject/LinkerLoad.cpp b/Engine/Source/Runtime/CoreUObject/Private/UObject/LinkerLoad.cpp
--- a/Engine/Source/Runtime/CoreUObject/Private/UObject/LinkerLoad.cpp
+++ b/Engine/Source/Runtime/CoreUObject/Private/UObject/LinkerLoad.cpp
@@ -2425,6 +2425,7 @@ UObject* FLinkerLoad::FindExistingExport(int32 ExportIndex)
 			}
 			else
 			{
+				// RobM: No class package so try and find any class matching the name. Sounds sketchy and we should remove it
 				TheClass = FindFirstObject<UClass>(*ImpExp(Export.ClassIndex).ObjectName.ToString(), EFindFirstObjectOptions::None, ELogVerbosity::Fatal, TEXT("finding existing export"));
 			}
 		}
@@ -2433,7 +2434,7 @@ UObject* FLinkerLoad::FindExistingExport(int32 ExportIndex)
 		if (TheClass)
 		{
 			TheClass->GetDefaultObject(); // build the CDO if it isn't already built
-			Export.Object = StaticFindObjectFast(TheClass, OuterObject, Export.ObjectName, /*bExactClass*/true, /*bAnyPackage*/false);
+			Export.Object = StaticFindObjectFast(TheClass, OuterObject, Export.ObjectName, /*bExactClass*/true);
 			
 			// if we found an object, set it's linker to us
 			if (Export.Object)
@@ -2464,7 +2465,7 @@ UObject* FLinkerLoad::FindExistingImport(int32 ImportIndex)
 	{
 		// if the import outer is null then we have a package, resolve it, potentially remapping it
 		FName ObjectName = InstancingContextRemap(Import.ObjectName);
-		return StaticFindObjectFast(UPackage::StaticClass(), nullptr, ObjectName, /*bExactClass*/true, /*bAnyPackage*/false);
+		return StaticFindObjectFast(UPackage::StaticClass(), nullptr, ObjectName, /*bExactClass*/true);
 	}
 	// if our outer is an import, recurse to find it
 	else if (Import.OuterIndex.IsImport())
@@ -2481,20 +2482,24 @@ UObject* FLinkerLoad::FindExistingImport(int32 ImportIndex)
 	{
 		// find the class of this object
 		UClass* TheClass = nullptr;
-		if (Import.ClassName == NAME_Class || Import.ClassName.IsNone())
+		if ((Import.ClassName == NAME_Class && (Import.ClassPackage == GLongCoreUObjectPackageName || Import.ClassPackage == NAME_CoreUObject)) || Import.ClassName.IsNone())
 		{
 			TheClass = UClass::StaticClass();
 		}
 		else
 		{
 			//@todo: Could we have an import that has its class as an export?
-			TheClass = (UClass*)StaticFindObjectFast(UClass::StaticClass(), nullptr, Import.ClassName, /*bExactClass*/false, /*bAnyPackage*/true);
+			UPackage* ClassPackage = FindObject<UPackage>(nullptr, *Import.ClassPackage.ToString()); // FindObject because *theoretically* this could be an old package where ClassPackage was a short package name and FindObject handles that
+			if (ClassPackage)
+			{
+				TheClass = FindObjectFast<UClass>(ClassPackage, Import.ClassName, /*bExactClass*/false);
+			}
 		}
 
 		// if the class exists, try to find the object
 		if (TheClass)
 		{
-			return StaticFindObjectFast(UClass::StaticClass(), OuterObject, Import.ObjectName, /*bExactClass*/true, /*bAnyPackage*/false);
+			return StaticFindObjectFast(UClass::StaticClass(), OuterObject, Import.ObjectName, /*bExactClass*/true);
 		}
 	}
 	return nullptr;
@@ -3459,17 +3464,15 @@ bool FLinkerLoad::VerifyImportInner(const int32 ImportIndex, FString& WarningSuf
 		}
 	}
 
-	if( (Pkg == nullptr) && ((LoadFlags & LOAD_FindIfFail) != 0) )
-	{
-		Pkg = ANY_PACKAGE;
-	}
+	// RobM: We should remove the bFindObjectByName path
+	const bool bFindObjectByName = (Pkg == nullptr) && ((LoadFlags & LOAD_FindIfFail) != 0);
 
 	// If not found in file, see if it's a public native transient class or field.
-	if( Import.SourceIndex==INDEX_NONE && Pkg!=nullptr )
+	if( Import.SourceIndex==INDEX_NONE && (Pkg!=nullptr || bFindObjectByName))
 	{
 		TStringBuilder<256> ImportClassTemp;
 		Import.ClassPackage.ToString(ImportClassTemp);
-		UObject* ClassPackage = FindObject<UPackage>( NULL, *ImportClassTemp);
+		UObject* ClassPackage = FindObject<UPackage>( nullptr, *ImportClassTemp );
 		if( ClassPackage )
 		{
 			Import.ClassName.ToString(ImportClassTemp);
@@ -3486,27 +3489,29 @@ bool FLinkerLoad::VerifyImportInner(const int32 ImportIndex, FString& WarningSuf
 					// otherwise, this import represents a field of an intrinsic class, and OuterImport's XObject should be non-NULL (the object
 					// that contains the field)
 					FObjectImport& OuterImport	= Imp(Import.OuterIndex);
-					if ( OuterImport.XObject != NULL )
+					if ( OuterImport.XObject != nullptr )
 					{
 						FindOuter = OuterImport.XObject;
 					}
 				}
 
-				bool bAnyPackage = FindOuter == ANY_PACKAGE;
-				UObject* FindObject = FindImportFast(FindClass, bAnyPackage ? nullptr : FindOuter, Import.ObjectName, bAnyPackage);
+				UObject* FindObject = FindImportFast(FindClass, bFindObjectByName ? nullptr : FindOuter, Import.ObjectName, bFindObjectByName);
 				// Reference to in memory-only package's object, native transient class or CDO of such a class.
-				bool bIsInMemoryOnlyOrNativeTransient = bCameFromMemoryOnlyPackage || (FindObject != NULL && ((FindObject->IsNative() && FindObject->HasAllFlags(RF_Public | RF_Transient)) || (FindObject->HasAnyFlags(RF_ClassDefaultObject) && FindObject->GetClass()->IsNative() && FindObject->GetClass()->HasAllFlags(RF_Public | RF_Transient))));
+				bool bIsInMemoryOnlyOrNativeTransient = bCameFromMemoryOnlyPackage || (FindObject != nullptr && ((FindObject->IsNative() && FindObject->HasAllFlags(RF_Public | RF_Transient)) || (FindObject->HasAnyFlags(RF_ClassDefaultObject) && FindObject->GetClass()->IsNative() && FindObject->GetClass()->HasAllFlags(RF_Public | RF_Transient))));
 				// Check for structs which have been moved to another header (within the same class package).
 				if (!FindObject && bIsInMemoryOnlyOrNativeTransient && FindClass == UScriptStruct::StaticClass())
 				{
-					FindObject = StaticFindObjectFast(FindClass, nullptr, Import.ObjectName, /*bExactClass*/true, /*bAnyPackage*/true);
+					FindObject = StaticFindFirstObject(FindClass, *Import.ObjectName.ToString(), 
+						EFindFirstObjectOptions::ExactClass | EFindFirstObjectOptions::NativeFirst | EFindFirstObjectOptions::EnsureIfAmbiguous,
+						ELogVerbosity::Warning, TEXT("Finding import by name"));
+
 					if (FindObject && FindOuter->GetOutermost() != FindObject->GetOutermost())
 					{
 						// Limit the results to the same package.I
-						FindObject = NULL;
+						FindObject = nullptr;
 					}
 				}
-				if (FindObject != NULL && ((LoadFlags & LOAD_FindIfFail) || bIsInMemoryOnlyOrNativeTransient))
+				if (FindObject != nullptr && ((LoadFlags & LOAD_FindIfFail) || bIsInMemoryOnlyOrNativeTransient))
 				{
 					Import.XObject = FindObject;
 					FUObjectSerializeContext* CurrentLoadContext = GetSerializeContext();
@@ -4739,7 +4744,7 @@ UObject* FLinkerLoad::CreateExport( int32 Index )
 #if WITH_EDITOR
 		if ( GIsEditor && GIsRunning && !Export.Object )
 		{
-			UObjectRedirector* Redirector = (UObjectRedirector*)StaticFindObjectFast(UObjectRedirector::StaticClass(), ThisParent, Export.ObjectName, /*bExactClass*/true, /*bAnyPackage*/false);
+			UObjectRedirector* Redirector = (UObjectRedirector*)StaticFindObjectFast(UObjectRedirector::StaticClass(), ThisParent, Export.ObjectName, /*bExactClass*/true);
 			if (Redirector && Redirector->DestinationObject && Redirector->DestinationObject->IsA(LoadClass))
 			{
 				// A redirector has been found, replace this export with it.
@@ -5035,9 +5040,9 @@ UObject* FLinkerLoad::CreateImport( int32 Index )
 		if (!GIsEditor && !IsRunningCommandlet())
 		{
 			// Try to find existing version in memory first.
-			if( UPackage* ClassPackage = FindObjectFast<UPackage>( NULL, Import.ClassPackage, false, false ) )
+			if( UPackage* ClassPackage = FindObjectFast<UPackage>( NULL, Import.ClassPackage, false ) )
 			{
-				if( UClass*	FindClass = FindObjectFast<UClass>( ClassPackage, Import.ClassName, false, false ) ) // 
+				if( UClass*	FindClass = FindObjectFast<UClass>( ClassPackage, Import.ClassName, false ) ) // 
 				{
 					// Make sure the class has been loaded and linked before creating a CDO.
 					// This is an edge case, but can happen if a blueprint package has not finished creating exports for a class
diff --git a/Engine/Source/Runtime/CoreUObject/Private/UObject/Obj.cpp b/Engine/Source/Runtime/CoreUObject/Private/UObject/Obj.cpp
--- a/Engine/Source/Runtime/CoreUObject/Private/UObject/Obj.cpp
+++ b/Engine/Source/Runtime/CoreUObject/Private/UObject/Obj.cpp
@@ -192,13 +192,20 @@ bool UObject::Rename( const TCHAR* InName, UObject* NewOuter, ERenameFlags Flags
 			*GetClass()->ClassWithin->GetName() );
 	}
 
-	UObject* NameScopeOuter = (Flags & REN_ForceGlobalUnique) ? ANY_PACKAGE : NewOuter;
-
 	// find an object with the same name and same class in the new outer
 	bool bIsCaseOnlyChange = false;
 	if (InName)
 	{
-		UObject* ExistingObject = StaticFindObject(/*Class=*/ NULL, NameScopeOuter ? NameScopeOuter : GetOuter(), InName, true);
+		UObject* ExistingObject = nullptr;
+		
+		if (!(Flags & REN_ForceGlobalUnique))
+		{
+			ExistingObject = StaticFindObject(/*Class=*/ nullptr, NewOuter ? NewOuter : GetOuter(), InName, true);
+		}
+		else
+		{
+			ExistingObject = StaticFindFirstObject(/*Class=*/ nullptr, InName, EFindFirstObjectOptions::ExactClass);
+		}
 		if (ExistingObject == this)
 		{
 			if (ExistingObject->GetName().Equals(InName, ESearchCase::CaseSensitive))
@@ -255,7 +262,7 @@ bool UObject::Rename( const TCHAR* InName, UObject* NewOuter, ERenameFlags Flags
 			}
 			else
 			{
-				NewName = MakeUniqueObjectName(NameScopeOuter ? NameScopeOuter : GetOuter(), GetClass());
+				NewName = MakeUniqueObjectName(NewOuter ? NewOuter : GetOuter(), GetClass(), FName(), !!(Flags & REN_ForceGlobalUnique) ? EUniqueObjectNameOptions::GloballyUnique : EUniqueObjectNameOptions::None);
 			}
 		}
 		else
@@ -2038,6 +2045,48 @@ const FName& UObject::SourceFileTagName()
 }
 
 #if WITH_EDITOR
+
+static void PostLoadAssetRegistryTagProperty(FProperty* Prop, const FAssetData& AssetData, TArray<UObject::FAssetRegistryTag>& OutTagsAndValuesToUpdate)
+{	
+	UObject::FAssetRegistryTag::ETagType TagType = UObject::FAssetRegistryTag::ETagType::TT_Alphabetical;
+
+	if (Prop->HasAnyPropertyFlags(CPF_AssetRegistrySearchable))
+	{
+		if (Prop->IsA<FObjectPropertyBase>() && !Prop->IsA<FSoftObjectProperty>())
+		{
+			FObjectPropertyBase* PropertyObject = CastFieldChecked<FObjectPropertyBase>(Prop);
+			FString ExportPath = AssetData.GetTagValueRef<FString>(Prop->GetFName());
+			if (!ExportPath.IsEmpty() && ExportPath[0] != '/')
+			{
+				FString ObjectPath = FPackageName::ExportTextPathToObjectPath(ExportPath);
+				ExportPath = FObjectPropertyBase::GetExportPath(PropertyObject->PropertyClass->GetClassPathName(), ObjectPath);
+				OutTagsAndValuesToUpdate.Add(UObject::FAssetRegistryTag(Prop->GetFName(), ExportPath, TagType));
+			}
+		}
+	}
+}
+
+void UObject::PostLoadAssetRegistryTags(const FAssetData& InAssetData, TArray<FAssetRegistryTag>& OutTagsAndValuesToUpdate) const
+{
+	if (GetClass()->HasAssetRegistrySearchableProperties())
+	{
+		for (TFieldIterator<FProperty> PropertyIt(GetClass()); PropertyIt; ++PropertyIt)
+		{
+			PostLoadAssetRegistryTagProperty(*PropertyIt, InAssetData, OutTagsAndValuesToUpdate);
+		}
+
+		UScriptStruct* SparseClassDataStruct = GetClass()->GetSparseClassDataStruct();
+		if (SparseClassDataStruct)
+		{
+			const void* SparseClassData = GetClass()->GetSparseClassData(EGetSparseClassDataMethod::ArchetypeIfNull);
+			for (TFieldIterator<FProperty> PropertyIt(SparseClassDataStruct); PropertyIt; ++PropertyIt)
+			{
+				PostLoadAssetRegistryTagProperty(*PropertyIt, InAssetData, OutTagsAndValuesToUpdate);
+			}
+		}
+	}
+}
+
 static TSet<FName> MetaDataTagsForAssetRegistry;
 
 TSet<FName>& UObject::GetMetaDataTagsForAssetRegistry()
@@ -2059,7 +2108,7 @@ void UObject::GetAssetRegistryTagMetadata(TMap<FName, FAssetRegistryTagMetadata>
 		.SetTooltip(NSLOCTEXT("UObject", "PrimaryAssetNameTooltip", "Logical name registered with the Asset Manager system"))
 	);
 }
-#endif
+#endif // WITH_EDITOR
 
 void UObject::GetResourceSizeEx(FResourceSizeEx& CumulativeResourceSize)
 {
@@ -3321,7 +3370,7 @@ static void PerformSetCommand( const TCHAR* Str, FOutputDevice& Ar, bool bNotify
 	TCHAR ObjectName[256], PropertyName[256];
 	if (FParse::Token(Str, ObjectName, UE_ARRAY_COUNT(ObjectName), true) && FParse::Token(Str, PropertyName, UE_ARRAY_COUNT(PropertyName), true))
 	{
-		UClass* Class = FindObject<UClass>(ANY_PACKAGE, ObjectName);
+		UClass* Class = FindFirstObject<UClass>(ObjectName, EFindFirstObjectOptions::None, ELogVerbosity::Warning, TEXT("PerformSetCommand"));
 		if (Class != NULL)
 		{
 			FProperty* Property = FindFProperty<FProperty>(Class, PropertyName);
@@ -3340,7 +3389,7 @@ static void PerformSetCommand( const TCHAR* Str, FOutputDevice& Ar, bool bNotify
 		}
 		else
 		{
-			UObject* Object = FindObject<UObject>(ANY_PACKAGE, ObjectName);
+			UObject* Object = FindFirstObject<UObject>(ObjectName, EFindFirstObjectOptions::NativeFirst, ELogVerbosity::Warning, TEXT("PerformSetCommand"));
 			if (Object != NULL)
 			{
 				FProperty* Property = FindFProperty<FProperty>(Object->GetClass(), PropertyName);
@@ -3529,11 +3578,11 @@ bool StaticExec( UWorld* InWorld, const TCHAR* Cmd, FOutputDevice& Ar )
 		FProperty* Property;
 		if
 		(	FParse::Token( Str, ClassName, UE_ARRAY_COUNT(ClassName), 1 )
-		&&	(Class=FindObject<UClass>( ANY_PACKAGE, ClassName))!=NULL )
+		&&	(Class = FindFirstObject<UClass>(ClassName, EFindFirstObjectOptions::None, ELogVerbosity::Warning, TEXT("StaticExec GET"))) != nullptr)
 		{
 			if
 			(	FParse::Token( Str, PropertyName, UE_ARRAY_COUNT(PropertyName), 1 )
-			&&	(Property=FindFProperty<FProperty>( Class, PropertyName))!=NULL )
+			&&	(Property=FindFProperty<FProperty>( Class, PropertyName)) != nullptr)
 			{
 				FString	Temp;
 				if( Class->GetDefaultsCount() )
@@ -3561,7 +3610,7 @@ bool StaticExec( UWorld* InWorld, const TCHAR* Cmd, FOutputDevice& Ar )
 		FString PropWildcard;
 
 		if ( FParse::Token(Str, ClassName, UE_ARRAY_COUNT(ClassName), 1) &&
-			(Class = FindObject<UClass>(ANY_PACKAGE, ClassName)) != NULL &&
+			(Class = FindFirstObject<UClass>(ClassName, EFindFirstObjectOptions::None, ELogVerbosity::Warning, TEXT("StaticExec LISTPROPS"))) != nullptr &&
 			FParse::Token(Str, PropWildcard, true) )
 		{
 			// split up the search string by wildcard symbols
@@ -3690,17 +3739,17 @@ bool StaticExec( UWorld* InWorld, const TCHAR* Cmd, FOutputDevice& Ar )
 		FProperty* Property;
 
 		if ( FParse::Token(Str,ClassName,UE_ARRAY_COUNT(ClassName), 1) &&
-			(Class=FindObject<UClass>( ANY_PACKAGE, ClassName)) != NULL )
+			(Class = FindFirstObject<UClass>(ClassName, EFindFirstObjectOptions::None, ELogVerbosity::Warning, TEXT("StaticExec GETALL"))) != nullptr )
 		{
 			FParse::Token(Str,PropertyName,UE_ARRAY_COUNT(PropertyName),1);
 			{
 				Property=FindFProperty<FProperty>(Class,PropertyName);
 				{
 					int32 cnt = 0;
-					UObject* LimitOuter = NULL;
+					UObject* LimitOuter = nullptr;
 
 					const bool bHasOuter = FCString::Strifind(Str,TEXT("OUTER=")) ? true : false;
-					ParseObject<UObject>(Str,TEXT("OUTER="),LimitOuter,ANY_PACKAGE);
+					ParseObject<UObject>(Str,TEXT("OUTER="),LimitOuter,nullptr);
 
 					// Check for a specific object name
 					TCHAR ObjNameStr[256];
@@ -3836,7 +3885,7 @@ bool StaticExec( UWorld* InWorld, const TCHAR* Cmd, FOutputDevice& Ar )
 		if (FParse::Token(Str, ClassName, UE_ARRAY_COUNT(ClassName), true))
 		{
 			//if ( (Property=FindFProperty<FProperty>(Class,PropertyName)) != NULL )
-			UClass* Class = FindObject<UClass>(ANY_PACKAGE, ClassName);
+			UClass* Class = FindFirstObject<UClass>(ClassName, EFindFirstObjectOptions::None, ELogVerbosity::Warning, TEXT("StaticExec LISTFUNCS"));
 
 			if (Class != NULL)
 			{
@@ -3862,7 +3911,7 @@ bool StaticExec( UWorld* InWorld, const TCHAR* Cmd, FOutputDevice& Ar )
 		TCHAR FunctionName[256];
 		if (FParse::Token(Str, ClassName, UE_ARRAY_COUNT(ClassName), true) && FParse::Token(Str, FunctionName, UE_ARRAY_COUNT(FunctionName), true))
 		{
-			UClass* Class = FindObject<UClass>(ANY_PACKAGE, ClassName);
+			UClass* Class = FindFirstObject<UClass>(ClassName, EFindFirstObjectOptions::None, ELogVerbosity::Warning, TEXT("StaticExec LISTFUNC"));
 
 			if (Class != NULL)
 			{
@@ -4151,7 +4200,7 @@ bool StaticExec( UWorld* InWorld, const TCHAR* Cmd, FOutputDevice& Ar )
 		else if( FParse::Command(&Str,TEXT("REFS")) )
 		{
 			UObject* Object;
-			if (ParseObject(Str,TEXT("NAME="),Object,ANY_PACKAGE))
+			if (ParseObject(Str,TEXT("NAME="),Object,nullptr))
 			{
 				
 				EReferenceChainSearchMode SearchModeFlags = EReferenceChainSearchMode::PrintResults;
@@ -4245,12 +4294,12 @@ bool StaticExec( UWorld* InWorld, const TCHAR* Cmd, FOutputDevice& Ar )
 			UClass* Class;
 			UClass* ReferencerClass = NULL;
 			FString ReferencerName;
-			if (ParseObject<UClass>(Str, TEXT("CLASS="), Class, ANY_PACKAGE) == false)
+			if (ParseObject<UClass>(Str, TEXT("CLASS="), Class, nullptr) == false)
 			{
 				Class = UObject::StaticClass();
 				bListClass = true;
 			}
-			if (ParseObject<UClass>(Str, TEXT("REFCLASS="), ReferencerClass, ANY_PACKAGE) == false)
+			if (ParseObject<UClass>(Str, TEXT("REFCLASS="), ReferencerClass, nullptr) == false)
 			{
 				ReferencerClass = NULL;
 			}
@@ -4413,8 +4462,8 @@ bool StaticExec( UWorld* InWorld, const TCHAR* Cmd, FOutputDevice& Ar )
 		{
 			UClass*	Class		= NULL;
 			UClass*	ListClass	= NULL;
-			ParseObject<UClass>(Str, TEXT("CLASS="		), Class,		ANY_PACKAGE );
-			ParseObject<UClass>(Str, TEXT("LISTCLASS="  ), ListClass,	ANY_PACKAGE );
+			ParseObject<UClass>(Str, TEXT("CLASS="		), Class,		nullptr );
+			ParseObject<UClass>(Str, TEXT("LISTCLASS="  ), ListClass,	nullptr );
 		
 			if( Class )
 			{
@@ -4521,7 +4570,7 @@ bool StaticExec( UWorld* InWorld, const TCHAR* Cmd, FOutputDevice& Ar )
 			UObject* Obj = NULL;
 			if ( FParse::Token(Str,ObjectName,UE_ARRAY_COUNT(ObjectName), 1) )
 			{
-				Obj = FindObject<UObject>(ANY_PACKAGE,ObjectName);
+				Obj = FindFirstObject<UObject>(ObjectName, EFindFirstObjectOptions::None, ELogVerbosity::Warning, TEXT("FLAGS command"));
 			}
 
 			if ( Obj )
@@ -4538,7 +4587,7 @@ bool StaticExec( UWorld* InWorld, const TCHAR* Cmd, FOutputDevice& Ar )
 			// Usage:  OBJ REP CLASS=PlayerController
 			UClass* Cls;
 
-			if( ParseObject<UClass>( Str, TEXT("CLASS="), Cls, ANY_PACKAGE ) )
+			if( ParseObject<UClass>( Str, TEXT("CLASS="), Cls, nullptr ) )
 			{
 				Ar.Logf(TEXT("=== Replicated properties for class: %s==="), *Cls->GetName());
 				for ( TFieldIterator<FProperty> It(Cls); It; ++It )
@@ -4574,7 +4623,7 @@ bool StaticExec( UWorld* InWorld, const TCHAR* Cmd, FOutputDevice& Ar )
 		if (FParse::Token(Str,ClassName,UE_ARRAY_COUNT(ClassName),1))
 		{
 			// Try to find a corresponding class
-			UClass* ClassToReload = FindObject<UClass>(ANY_PACKAGE,ClassName);
+			UClass* ClassToReload = FindFirstObject<UClass>(ClassName, EFindFirstObjectOptions::None, ELogVerbosity::Warning, TEXT("StaticExec RELOADCONFIG"));
 			if (ClassToReload)
 			{
 				ClassToReload->ReloadConfig();
@@ -4582,7 +4631,7 @@ bool StaticExec( UWorld* InWorld, const TCHAR* Cmd, FOutputDevice& Ar )
 			else
 			{
 				// If the class is missing, search for an object with that name
-				UObject* ObjectToReload = FindObject<UObject>(ANY_PACKAGE,ClassName);
+				UObject* ObjectToReload = FindFirstObject<UObject>(ClassName, EFindFirstObjectOptions::None, ELogVerbosity::Warning, TEXT("StaticExec RELOADCONFIG"));
 				if (ObjectToReload)
 				{
 					ObjectToReload->ReloadConfig();
diff --git a/Engine/Source/Runtime/CoreUObject/Private/UObject/ObjectRedirector.cpp b/Engine/Source/Runtime/CoreUObject/Private/UObject/ObjectRedirector.cpp
--- a/Engine/Source/Runtime/CoreUObject/Private/UObject/ObjectRedirector.cpp
+++ b/Engine/Source/Runtime/CoreUObject/Private/UObject/ObjectRedirector.cpp
@@ -61,9 +61,9 @@ bool UObjectRedirector::NeedsLoadForEditorGame() const
 void UObjectRedirector::GetAssetRegistryTags(TArray<FAssetRegistryTag>& OutTags) const
 {
 	FString DestVal;
-	if ( DestinationObject != NULL )
+	if ( DestinationObject != nullptr )
 	{
-		DestVal = FString::Printf(TEXT("%s'%s'"), *DestinationObject->GetClass()->GetName(), *DestinationObject->GetPathName());
+		DestVal = FObjectPropertyBase::GetExportPath(DestinationObject);
 	}
 	else
 	{
diff --git a/Engine/Source/Runtime/CoreUObject/Private/UObject/PropertyBaseObject.cpp b/Engine/Source/Runtime/CoreUObject/Private/UObject/PropertyBaseObject.cpp
--- a/Engine/Source/Runtime/CoreUObject/Private/UObject/PropertyBaseObject.cpp
+++ b/Engine/Source/Runtime/CoreUObject/Private/UObject/PropertyBaseObject.cpp
@@ -181,7 +181,7 @@ void FObjectPropertyBase::AddReferencedObjects(FReferenceCollector& Collector)
 
 FString FObjectPropertyBase::GetExportPath(FTopLevelAssetPath ClassPathName, const FString& ObjectPathName)
 {
-	return FString::Printf(TEXT("%s'%s'"), *ClassPathName.GetAssetName().ToString(), *ObjectPathName);
+	return FString::Printf(TEXT("%s.%s'%s'"), *ClassPathName.GetPackageName().ToString(), *ClassPathName.GetAssetName().ToString(), *ObjectPathName);
 }
 
 FString FObjectPropertyBase::GetExportPath(const UObject* Object, const UObject* Parent /*= nullptr*/, const UObject* ExportRootScope /*= nullptr*/, const uint32 PortFlags /*= PPF_None*/)
@@ -341,7 +341,7 @@ bool FObjectPropertyBase::ParseObjectPropertyValue(const FProperty* Property, UO
 				return false;
 			}
 
-			// ignore the object class, it isn't fully qualified, and searching ANY_PACKAGE might get the wrong one!
+			// ignore the object class, it isn't fully qualified, and searching globally might get the wrong one!
 			// Try the find the object.
 			out_ResolvedValue = FObjectPropertyBase::FindImportedObject(Property, OwnerObject, ObjectClass, RequiredMetaClass, Temp.ToString(), PortFlags, InSerializeContext, bAllowAnyPackage);
 		}
@@ -485,8 +485,9 @@ UObject* FObjectPropertyBase::FindImportedObject( const FProperty* Property, UOb
 
 		if (Result == nullptr && bAllowAnyPackage)
 		{
+			// RobM: We should delete this path
 			// match any object of the correct class who shares the same name regardless of package path
-			Result = StaticFindObjectSafe(ObjectClass, ANY_PACKAGE, Text);
+			Result = StaticFindFirstObject(ObjectClass, Text, EFindFirstObjectOptions::EnsureIfAmbiguous);
 			// disallow class default subobjects here while importing defaults
 			if (Result != nullptr && (PortFlags & PPF_ParsingDefaultProperties) && Result->IsTemplate(RF_ClassDefaultObject))
 			{
diff --git a/Engine/Source/Runtime/CoreUObject/Private/UObject/PropertyByte.cpp b/Engine/Source/Runtime/CoreUObject/Private/UObject/PropertyByte.cpp
--- a/Engine/Source/Runtime/CoreUObject/Private/UObject/PropertyByte.cpp
+++ b/Engine/Source/Runtime/CoreUObject/Private/UObject/PropertyByte.cpp
@@ -234,7 +234,7 @@ EConvertFromTypeResult FByteProperty::ConvertFromType(const FPropertyTag& Tag, F
 		// now copy the value into the object's address space
 		SetPropertyValue_InContainer(Data, PreviousValue, Tag.ArrayIndex);
 	}
-	else if (Tag.Type == NAME_EnumProperty && (Enum == nullptr || Tag.EnumName == Enum->GetFName()))
+	else if (Tag.Type == NAME_EnumProperty && (Enum == nullptr || Tag.EnumName == Enum->GetFName() || Tag.EnumName.ToString() == Enum->GetPathName()))
 	{
 		// an enum property became a byte
 		// attempt to find the old enum and get the byte value from the serialized enum name
diff --git a/Engine/Source/Runtime/CoreUObject/Private/UObject/PropertyHelper.h b/Engine/Source/Runtime/CoreUObject/Private/UObject/PropertyHelper.h
--- a/Engine/Source/Runtime/CoreUObject/Private/UObject/PropertyHelper.h
+++ b/Engine/Source/Runtime/CoreUObject/Private/UObject/PropertyHelper.h
@@ -234,7 +234,7 @@ namespace DelegatePropertyTools
 		}
 		if (Cls == nullptr)
 		{
-			Cls = FindObject<UClass>(ANY_PACKAGE,ObjName);
+			Cls = UClass::TryFindTypeSlow<UClass>(ObjName);
 			if (Cls != nullptr)
 			{
 				Object = Cls->GetDefaultObject();
@@ -251,7 +251,7 @@ namespace DelegatePropertyTools
 				
 				if (Object == nullptr)
 				{
-					Object = StaticFindObject(UObject::StaticClass(), ANY_PACKAGE, ObjName);
+					Object = StaticFindFirstObject(UObject::StaticClass(), ObjName, EFindFirstObjectOptions::NativeFirst | EFindFirstObjectOptions::EnsureIfAmbiguous);
 				}
 				if (Object != nullptr)
 				{
diff --git a/Engine/Source/Runtime/CoreUObject/Private/UObject/PropertyNumeric.cpp b/Engine/Source/Runtime/CoreUObject/Private/UObject/PropertyNumeric.cpp
--- a/Engine/Source/Runtime/CoreUObject/Private/UObject/PropertyNumeric.cpp
+++ b/Engine/Source/Runtime/CoreUObject/Private/UObject/PropertyNumeric.cpp
@@ -33,7 +33,8 @@ int64 FNumericProperty::ReadEnumAsInt64(FStructuredArchive::FSlot Slot, UStruct*
 	UEnum* Enum = FindUField<UEnum>(dynamic_cast<UClass*>(DefaultsStruct) ? static_cast<UClass*>(DefaultsStruct) : DefaultsStruct->GetTypedOuter<UClass>(), Tag.EnumName);
 	if (!Enum)
 	{
-		Enum = FindObject<UEnum>(ANY_PACKAGE, *Tag.EnumName.ToString());
+		// Enums (at least native) are stored as short names (for now) so find the Tag enum by name
+		Enum = FindFirstObject<UEnum>(*Tag.EnumName.ToString(), EFindFirstObjectOptions::NativeFirst);
 	}
 
 	if (!Enum)
diff --git a/Engine/Source/Runtime/CoreUObject/Private/UObject/PropertyTag.cpp b/Engine/Source/Runtime/CoreUObject/Private/UObject/PropertyTag.cpp
--- a/Engine/Source/Runtime/CoreUObject/Private/UObject/PropertyTag.cpp
+++ b/Engine/Source/Runtime/CoreUObject/Private/UObject/PropertyTag.cpp
@@ -21,6 +21,19 @@ FPropertyTag::FPropertyTag( FArchive& InSaveAr, FProperty* Property, int32 InInd
 	check(!InSaveAr.GetArchiveState().UseUnversionedPropertySerialization());
 	if (Property)
 	{
+		// RobM: Ugly hack so that we can avoid content changes in most of the packages
+		auto GetEnumName = [](UEnum* InEnum)
+		{
+			if (InEnum->GetPackage()->HasAnyPackageFlags(PKG_CompiledIn))
+			{				
+				return InEnum->GetFName();
+			}
+			else
+			{
+				return FName(*InEnum->GetPathName());
+			}
+		};
+
 		// Handle structs.
 		if (FStructProperty* StructProperty = CastField<FStructProperty>(Property))
 		{
@@ -31,14 +44,14 @@ FPropertyTag::FPropertyTag( FArchive& InSaveAr, FProperty* Property, int32 InInd
 		{
 			if (UEnum* Enum = EnumProp->GetEnum())
 			{
-				EnumName = Enum->GetFName();
+				EnumName = GetEnumName(Enum);
 			}
 		}
 		else if (FByteProperty* ByteProp = CastField<FByteProperty>(Property))
 		{
 			if (ByteProp->Enum != nullptr)
 			{
-				EnumName = ByteProp->Enum->GetFName();
+				EnumName = GetEnumName(ByteProp->Enum);
 			}
 		}
 		else if (FArrayProperty* ArrayProp = CastField<FArrayProperty>(Property))
diff --git a/Engine/Source/Runtime/CoreUObject/Private/UObject/SavePackage/SavePackageUtilities.cpp b/Engine/Source/Runtime/CoreUObject/Private/UObject/SavePackage/SavePackageUtilities.cpp
--- a/Engine/Source/Runtime/CoreUObject/Private/UObject/SavePackage/SavePackageUtilities.cpp
+++ b/Engine/Source/Runtime/CoreUObject/Private/UObject/SavePackage/SavePackageUtilities.cpp
@@ -2266,7 +2266,7 @@ namespace AssetRegistry
 
 			// Exclude the package name in the object path, we just need to know the path relative to the package we are saving
 			FString ObjectPath = Object->GetPathName(Package);
-			FString ObjectClassName = Object->GetClass()->GetName();
+			FString ObjectClassName = Object->GetClass()->GetPathName();
 
 			TArray<UObject::FAssetRegistryTag> SourceTags;
 			Object->GetAssetRegistryTags(SourceTags);
diff --git a/Engine/Source/Runtime/CoreUObject/Private/UObject/SparseDelegate.cpp b/Engine/Source/Runtime/CoreUObject/Private/UObject/SparseDelegate.cpp
--- a/Engine/Source/Runtime/CoreUObject/Private/UObject/SparseDelegate.cpp
+++ b/Engine/Source/Runtime/CoreUObject/Private/UObject/SparseDelegate.cpp
@@ -396,7 +396,7 @@ void FSparseDelegateStorage::SparseDelegateReport(const TArray<FString>& Args, U
 		{
 			FString ClassName;
 			Arg.Split(TEXT("="), nullptr, &ClassName, ESearchCase::CaseSensitive);
-			ObjectType = FindObject<UClass>(ANY_PACKAGE, *ClassName);
+			ObjectType = UClass::TryFindTypeSlow<UClass>(ClassName);
 			if (ObjectType == nullptr)
 			{
 				Ar.Logf(ELogVerbosity::Warning, TEXT("No class of specified name found."));
diff --git a/Engine/Source/Runtime/CoreUObject/Private/UObject/TopLevelAssetPath.cpp b/Engine/Source/Runtime/CoreUObject/Private/UObject/TopLevelAssetPath.cpp
--- a/Engine/Source/Runtime/CoreUObject/Private/UObject/TopLevelAssetPath.cpp
+++ b/Engine/Source/Runtime/CoreUObject/Private/UObject/TopLevelAssetPath.cpp
@@ -11,6 +11,14 @@
 #include "UObject/UnrealType.h"
 #include "UObject/UObjectThreadContext.h"
 
+// The reason behind HACK_HEADER_GENERATOR is that without it UHT is going to 'see' cppstructops for TopLevelAssetPath
+// and will not generate temp FTopLevelAssetPath struct for codegen purposes where it can access all of its members
+// which are public in the temp struct and private in the actual FTopLevelAssetPath.
+// If UHT compiles with IMPLEMENT_STRUCT(TopLevelAssetPath) the generated code will fail to compile trying to access private struct members
+#if !HACK_HEADER_GENERATOR
+IMPLEMENT_STRUCT(TopLevelAssetPath)
+#endif
+
 FTopLevelAssetPath::FTopLevelAssetPath(const UObject* InObject)
 {
 	if (InObject == nullptr)
@@ -89,13 +97,14 @@ bool FTopLevelAssetPath::TrySetPath(FWideStringView Path)
 	}
 	else
 	{
-		if (Path[0] != '/')
+		if (Path[0] != '/' || Path[Path.Len() - 1] == '\'')
 		{
 			// Possibly an ExportText path. Trim the ClassName.
 			Path = FPackageName::ExportTextPathToObjectPath(Path);
 
 			if (Path.IsEmpty() || Path[0] != '/')
 			{
+				checkf(false, TEXT("Short asset name used to create FTopLevelAssetPath: \"%.*s\""), Path.Len(), Path.GetData());
 				Reset();
 				return false;
 			}
diff --git a/Engine/Source/Runtime/CoreUObject/Private/UObject/UObjectClusters.cpp b/Engine/Source/Runtime/CoreUObject/Private/UObject/UObjectClusters.cpp
--- a/Engine/Source/Runtime/CoreUObject/Private/UObject/UObjectClusters.cpp
+++ b/Engine/Source/Runtime/CoreUObject/Private/UObject/UObjectClusters.cpp
@@ -254,7 +254,7 @@ static void ParseObjectNameArrayForClusters(TArray<int32>& OutIndexArray, const
 {
 	for (const FString& ObjectName : InNameArray)
 	{
-		UObject* Res = StaticFindObject(UObject::StaticClass(), ANY_PACKAGE, *ObjectName);
+		UObject* Res = StaticFindFirstObject(UObject::StaticClass(), *ObjectName, EFindFirstObjectOptions::NativeFirst, ELogVerbosity::Warning, TEXT("ParseObjectNameArrayForClusters"));
 		if (Res)
 		{
 			int32 ObjectIndex = GUObjectArray.ObjectToIndex(Res);
diff --git a/Engine/Source/Runtime/CoreUObject/Private/UObject/UObjectGlobals.cpp b/Engine/Source/Runtime/CoreUObject/Private/UObject/UObjectGlobals.cpp
--- a/Engine/Source/Runtime/CoreUObject/Private/UObject/UObjectGlobals.cpp
+++ b/Engine/Source/Runtime/CoreUObject/Private/UObject/UObjectGlobals.cpp
@@ -65,6 +65,7 @@
 #include "Misc/PackageAccessTracking.h"
 #include "Misc/PackageAccessTracking.h"
 #include "UObject/PropertyWithSetterAndGetter.h"
+#include "UObject/AnyPackagePrivate.h"
 
 DEFINE_LOG_CATEGORY(LogUObjectGlobals);
 
@@ -311,7 +312,30 @@
 
 	// We don't want to return any objects that are currently being background loaded unless we're using FindObject during async loading.
 	ExclusiveInternalFlags |= IsInAsyncLoadingThread() ? EInternalObjectFlags::None : EInternalObjectFlags::AsyncLoading;	
+	PRAGMA_DISABLE_DEPRECATION_WARNINGS
 	UObject* FoundObject = StaticFindObjectFastInternal(ObjectClass, ObjectPackage, ObjectName, bExactClass, bAnyPackage, ExclusiveFlags, ExclusiveInternalFlags);
+	PRAGMA_ENABLE_DEPRECATION_WARNINGS
+	if (!FoundObject)
+	{
+		FoundObject = StaticFindObjectWithChangedLegacyPath(ObjectClass, ObjectPackage, ObjectName, bExactClass);
+	}
+
+	return FoundObject;
+}
+
+//
+// Find an object, path must unqualified
+//
+UObject* StaticFindObjectFast(UClass* ObjectClass, UObject* ObjectPackage, FName ObjectName, bool bExactClass, EObjectFlags ExclusiveFlags, EInternalObjectFlags ExclusiveInternalFlags)
+{
+	if (UE::IsSavingPackage(nullptr) || IsGarbageCollectingAndLockingUObjectHashTables())
+	{
+		UE_LOG(LogUObjectGlobals, Fatal, TEXT("Illegal call to StaticFindObjectFast() while serializing object data or garbage collecting!"));
+	}
+
+	// We don't want to return any objects that are currently being background loaded unless we're using FindObject during async loading.
+	ExclusiveInternalFlags |= IsInAsyncLoadingThread() ? EInternalObjectFlags::None : EInternalObjectFlags::AsyncLoading;
+	UObject* FoundObject = StaticFindObjectFastInternal(ObjectClass, ObjectPackage, ObjectName, bExactClass, ExclusiveFlags, ExclusiveInternalFlags);
 
 	if (!FoundObject)
 	{
@@ -329,7 +353,27 @@
 	{
 		// We don't want to return any objects that are currently being background loaded unless we're using FindObject during async loading.
 		ExclusiveInternalFlags |= IsInAsyncLoadingThread() ? EInternalObjectFlags::None : EInternalObjectFlags::AsyncLoading;
+		PRAGMA_DISABLE_DEPRECATION_WARNINGS
 		FoundObject = StaticFindObjectFastInternal(ObjectClass, ObjectPackage, ObjectName, bExactClass, bAnyPackage, ExclusiveFlags, ExclusiveInternalFlags);
+		PRAGMA_ENABLE_DEPRECATION_WARNINGS
+		if (!FoundObject)
+		{
+			FoundObject = StaticFindObjectWithChangedLegacyPath(ObjectClass, ObjectPackage, ObjectName, bExactClass);
+		}
+	}
+
+	return FoundObject;
+}
+
+UObject* StaticFindObjectFastSafe(UClass* ObjectClass, UObject* ObjectPackage, FName ObjectName, bool bExactClass, EObjectFlags ExclusiveFlags, EInternalObjectFlags ExclusiveInternalFlags)
+{
+	UObject* FoundObject = nullptr;
+
+	if (!UE::IsSavingPackage(nullptr) && !IsGarbageCollectingAndLockingUObjectHashTables())
+	{
+		// We don't want to return any objects that are currently being background loaded unless we're using FindObject during async loading.
+		ExclusiveInternalFlags |= IsInAsyncLoadingThread() ? EInternalObjectFlags::None : EInternalObjectFlags::AsyncLoading;
+		FoundObject = StaticFindObjectFastInternal(ObjectClass, ObjectPackage, ObjectName, bExactClass, ExclusiveFlags, ExclusiveInternalFlags);
 		if (!FoundObject)
 		{
 			FoundObject = StaticFindObjectWithChangedLegacyPath(ObjectClass, ObjectPackage, ObjectName, bExactClass);
@@ -368,7 +412,7 @@
 	INC_DWORD_STAT(STAT_FindObject);
 
 	// Resolve the object and package name.
-	const bool bAnyPackage = InObjectPackage==ANY_PACKAGE;
+	const bool bAnyPackage = InObjectPackage == ANY_PACKAGE_DEPRECATED;
 	UObject* ObjectPackage = bAnyPackage ? nullptr : InObjectPackage;
 
 	UObject* MatchingObject = nullptr;
@@ -403,8 +447,9 @@
 
 		ObjectName = FName(*InName, FNAME_Add);
 	}
-
+	PRAGMA_DISABLE_DEPRECATION_WARNINGS
 	return StaticFindObjectFast(ObjectClass, ObjectPackage, ObjectName, bExactClass, bAnyPackage);
+	PRAGMA_ENABLE_DEPRECATION_WARNINGS
 }
 
 //
@@ -446,7 +491,7 @@
 #if !(UE_BUILD_SHIPPING || UE_BUILD_TEST)
 	if( !Result )
 	{
-		UE_LOG(LogUObjectGlobals, Fatal, TEXT("Failed to find object '%s %s.%s'"), *ObjectClass->GetName(), ObjectParent==ANY_PACKAGE ? TEXT("Any") : ObjectParent ? *ObjectParent->GetName() : TEXT("None"), InName);
+		UE_LOG(LogUObjectGlobals, Fatal, TEXT("Failed to find object '%s %s.%s'"), *ObjectClass->GetName(), ObjectParent == ANY_PACKAGE_DEPRECATED ? TEXT("Any") : ObjectParent ? *ObjectParent->GetName() : TEXT("None"), InName);
 	}
 #endif
 	return Result;
@@ -979,7 +1024,14 @@
 			UObject* NewPackage = FindObject<UPackage>( InPackage, *PartialName );
 			if( !NewPackage )
 			{
-				NewPackage = FindObject<UObject>( InPackage == NULL ? ANY_PACKAGE : InPackage, *PartialName );
+				if (InPackage)
+				{
+					NewPackage = FindObject<UObject>(InPackage, *PartialName);
+				}
+				else
+				{
+					NewPackage = FindFirstObject<UObject>(*PartialName, EFindFirstObjectOptions::NativeFirst | EFindFirstObjectOptions::EnsureIfAmbiguous);
+				}
 				if( !NewPackage )
 				{
 					return bSubobjectPath;
@@ -1021,22 +1073,30 @@
 	else if( FCString::Stricmp(TempStr,TEXT("NONE"))==0 )
 	{
 		// Match found, object explicit set to be None
-		DestRes = NULL;
+		DestRes = nullptr;
 		return 1;
 	}
 	else
 	{
+		UObject* Res = nullptr;
 		// Look this object up.
-		UObject* Res = StaticFindObject( Class, InParent, TempStr );
+		if ((!InParent || InParent == ANY_PACKAGE_DEPRECATED) && !FPackageName::IsValidObjectPath(TempStr))
+		{
+			// Try to find first object matching the provided name
+			Res = StaticFindFirstObject(Class, TempStr, EFindFirstObjectOptions::EnsureIfAmbiguous);
+		}
+		else
+		{
+			Res = StaticFindObject(Class, InParent, TempStr);
+		}
 		if( !Res )
 		{
 			if (Class->IsChildOf<UClass>())
 			{
-				const FString ObjectName = FPackageName::ObjectPathToObjectName(FString(TempStr));
-				const FName RedirectedObjectName = FLinkerLoad::FindNewNameForClass(*ObjectName, false);
-				if (!RedirectedObjectName.IsNone())
+				FString RedirectedObjectName = FLinkerLoad::FindNewPathNameForClass(TempStr, false);
+				if (!RedirectedObjectName.IsEmpty())
 				{
-					Res = StaticFindObject(Class, InParent, *RedirectedObjectName.ToString());
+					Res = StaticFindObject(Class, InParent, *RedirectedObjectName);
 				}
 			}
 
@@ -2368,7 +2428,7 @@
 
 	static FNameRangeCache GNameRangeCache;
 
-	FName MakeUniqueObjectNameReusingNumber(UObject* Parent, FName BaseName)
+	FName MakeUniqueObjectNameReusingNumber(UObject* Parent, FName BaseName, EUniqueObjectNameOptions Options)
 	{
 		if (!GTryReuseNames)
 		{
@@ -2376,7 +2436,7 @@
 		}
 
 		static const FName NamePackage(NAME_Package);
-		if (!Parent || Parent == ANY_PACKAGE || BaseName == NamePackage || FPlatformProperties::HasEditorOnlyData() || !GFastPathUniqueNameGeneration)
+		if (!Parent || Parent == ANY_PACKAGE_DEPRECATED || !!(Options & EUniqueObjectNameOptions::GloballyUnique) || BaseName == NamePackage || FPlatformProperties::HasEditorOnlyData() || !GFastPathUniqueNameGeneration)
 		{
 			return FName();
 		}
@@ -2402,7 +2462,7 @@
 }
 #endif
 
-FName MakeUniqueObjectName( UObject* Parent, const UClass* Class, FName InBaseName/*=NAME_None*/ )
+FName MakeUniqueObjectName(UObject* Parent, const UClass* Class, FName InBaseName/*=NAME_None*/, EUniqueObjectNameOptions Options /*= EUniqueObjectNameOptions::None*/)
 {
 	TRACE_CPUPROFILER_EVENT_SCOPE(MakeUniqueObjectName);
 
@@ -2411,7 +2471,7 @@
 	const FName BaseName = InBaseName.IsNone() ? Class->GetFName() : InBaseName;
 
 #if TRY_REUSE_NAMES
-	if (FName Result = NameReuse::MakeUniqueObjectNameReusingNumber(Parent, BaseName); !Result.IsNone())
+	if (FName Result = NameReuse::MakeUniqueObjectNameReusingNumber(Parent, BaseName, Options); !Result.IsNone())
 	{
 		return Result;
 	}
@@ -2443,7 +2503,7 @@
 			else
 			{
 				int32 NameNumber = 0;
-				if (Parent && (Parent != ANY_PACKAGE))
+				if (Parent && (Parent != ANY_PACKAGE_DEPRECATED) && !(Options & EUniqueObjectNameOptions::GloballyUnique))
 				{
 					if (!FPlatformProperties::HasEditorOnlyData() && GFastPathUniqueNameGeneration)
 					{
@@ -2477,13 +2537,13 @@
 				TestName = FName(BaseName, NameNumber);
 			}
 
-			if (Parent == ANY_PACKAGE)
+			if (Parent == ANY_PACKAGE_DEPRECATED || !!(Options & EUniqueObjectNameOptions::GloballyUnique))
 			{
-				ExistingObject = StaticFindObject(NULL, ANY_PACKAGE, *TestName.ToString());
+				ExistingObject = StaticFindFirstObject(nullptr, *TestName.ToString());
 			}
 			else
 			{
-				ExistingObject = StaticFindObjectFastInternal(NULL, Parent, TestName);
+				ExistingObject = StaticFindObjectFastInternal(nullptr, Parent, TestName);
 			}
 		} while (ExistingObject);
 	// InBaseName can be a name of an object from a different hierarchy (so it's still unique within given parents scope), we don't want to return the same name.
@@ -4157,7 +4217,7 @@
 
 void ConstructorHelpers::CheckFoundViaRedirect(UObject *Object, const FString& PathName, const TCHAR* ObjectToFind)
 {
-	UObjectRedirector* Redir = FindObject<UObjectRedirector>(ANY_PACKAGE, *PathName);
+	UObjectRedirector* Redir = FindObject<UObjectRedirector>(nullptr, *PathName);
 	if (Redir && Redir->DestinationObject == Object)
 	{
 		FString NewString = Object->GetFullName();
@@ -4799,7 +4859,7 @@
 
 	if (StringName.EndsWith(HEADER_GENERATED_DELEGATE_SIGNATURE_SUFFIX))
 	{
-		return FindObject<UFunction>(ANY_PACKAGE, *StringName);
+		return FindFirstObject<UFunction>(*StringName, EFindFirstObjectOptions::NativeFirst | EFindFirstObjectOptions::EnsureIfAmbiguous);
 	}
 
 	return nullptr;
@@ -5330,7 +5390,7 @@
 			return;
 		}
 
-		UObject* FoundPackage = StaticFindObjectFast(UPackage::StaticClass(), nullptr, FName(UTF8_TO_TCHAR(Params.NameUTF8)), false, false);
+		UObject* FoundPackage = StaticFindObjectFast(UPackage::StaticClass(), nullptr, FName(UTF8_TO_TCHAR(Params.NameUTF8)), false);
 
 #if USE_PER_MODULE_UOBJECT_BOOTSTRAP
 		if (!FoundPackage)
diff --git a/Engine/Source/Runtime/CoreUObject/Private/UObject/UObjectHash.cpp b/Engine/Source/Runtime/CoreUObject/Private/UObject/UObjectHash.cpp
--- a/Engine/Source/Runtime/CoreUObject/Private/UObject/UObjectHash.cpp
+++ b/Engine/Source/Runtime/CoreUObject/Private/UObject/UObjectHash.cpp
@@ -11,6 +11,7 @@
 #include "Misc/PackageName.h"
 #include "Async/ParallelFor.h"
 #include "HAL/IConsoleManager.h"
+#include "UObject/AnyPackagePrivate.h"
 
 DEFINE_LOG_CATEGORY_STATIC(LogUObjectHash, Log, All);
 
@@ -697,6 +698,11 @@ UObject* StaticFindObjectFastInternalThreadSafe(FUObjectHashTables& ThreadHash,
 {
 	ExclusiveInternalFlags |= EInternalObjectFlags::Unreachable;
 
+	if (!ObjectPackage && bAnyPackage && ObjectClass && ObjectClass->IsChildOf(UClass::StaticClass()))
+	{
+		UE_LOG(LogUObjectHash, Fatal, TEXT("Looking for classes without specifying their outer is not allowed (class name: %s)"), *ObjectName.ToString());
+		return nullptr;
+	}
 	// If they specified an outer use that during the hashing
 	UObject* Result = nullptr;
 	if (ObjectPackage != nullptr)
@@ -804,13 +810,26 @@ UObject* StaticFindObjectFastInternal(const UClass* ObjectClass, const UObject*
 {
 	INC_DWORD_STAT(STAT_FindObjectFast);
 
-	check(ObjectPackage != ANY_PACKAGE); // this could never have returned anything but nullptr
+	check(ObjectPackage != ANY_PACKAGE_DEPRECATED); // this could never have returned anything but nullptr
+
 	// If they specified an outer use that during the hashing
 	FUObjectHashTables& ThreadHash = FUObjectHashTables::Get();
 	UObject* Result = StaticFindObjectFastInternalThreadSafe(ThreadHash, ObjectClass, ObjectPackage, ObjectName, bExactClass, bAnyPackage, ExcludeFlags | RF_NewerVersionExists, ExclusiveInternalFlags);
 	return Result;
 }
 
+UObject* StaticFindObjectFastInternal(const UClass* ObjectClass, const UObject* ObjectPackage, FName ObjectName, bool bExactClass, EObjectFlags ExcludeFlags, EInternalObjectFlags ExclusiveInternalFlags)
+{
+	INC_DWORD_STAT(STAT_FindObjectFast);
+
+	check(ObjectPackage != ANY_PACKAGE_DEPRECATED); // this could never have returned anything but nullptr
+
+	// If they specified an outer use that during the hashing
+	FUObjectHashTables& ThreadHash = FUObjectHashTables::Get();
+	UObject* Result = StaticFindObjectFastInternalThreadSafe(ThreadHash, ObjectClass, ObjectPackage, ObjectName, bExactClass, /*bAnyPackage =*/ false, ExcludeFlags | RF_NewerVersionExists, ExclusiveInternalFlags);
+	return Result;
+}
+
 // Approximate search for finding unused object names
 bool DoesObjectPossiblyExist(const UObject* InOuter, FName ObjectName)
 {
@@ -1194,7 +1213,7 @@ UObjectBase* FindObjectWithOuter(const class UObjectBase* Outer, const class UCl
 
 	if( NameToLookFor != NAME_None )
 	{
-		Result = StaticFindObjectFastInternal(ClassToLookFor, static_cast<const UObject*>(Outer), NameToLookFor, false, false, RF_NoFlags, ExclusionInternalFlags);
+		Result = StaticFindObjectFastInternal(ClassToLookFor, static_cast<const UObject*>(Outer), NameToLookFor, false, RF_NoFlags, ExclusionInternalFlags);
 	}
 	else
 	{
diff --git a/Engine/Source/Runtime/CoreUObject/Public/AssetRegistry/ARFilter.h b/Engine/Source/Runtime/CoreUObject/Public/AssetRegistry/ARFilter.h
--- a/Engine/Source/Runtime/CoreUObject/Public/AssetRegistry/ARFilter.h
+++ b/Engine/Source/Runtime/CoreUObject/Public/AssetRegistry/ARFilter.h
@@ -14,6 +14,14 @@ struct FAssetData;
  */
 struct FARFilter
 {
+PRAGMA_DISABLE_DEPRECATION_WARNINGS // Compilers can complain about deprecated members in compiler generated code
+	FARFilter() = default;
+	FARFilter(FARFilter&&) = default;
+	FARFilter(const FARFilter&) = default;
+	FARFilter& operator=(FARFilter&&) = default;
+	FARFilter& operator=(const FARFilter&) = default;
+PRAGMA_ENABLE_DEPRECATION_WARNINGS
+
 	/** The filter component for package names */
 	TArray<FName> PackageNames;
 
@@ -23,19 +31,27 @@ struct FARFilter
 	/** The filter component containing specific object paths */
 	TArray<FName> ObjectPaths;
 
-	/** The filter component for class names. Instances of the specified classes, but not subclasses (by default), will be included. Derived classes will be included only if bRecursiveClasses is true. */
+	/** Deprecated. The filter component for class names. Instances of the specified classes, but not subclasses (by default), will be included. Derived classes will be included only if bRecursiveClasses is true. */
+	UE_DEPRECATED(5.1, "Class names are now represented by path names. Please use ClassPaths.")
 	TArray<FName> ClassNames;
 
+	/** The filter component for class path names. Instances of the specified classes, but not subclasses (by default), will be included. Derived classes will be included only if bRecursiveClasses is true. */
+	TArray<FTopLevelAssetPath> ClassPaths;
+
 	/** The filter component for properties marked with the AssetRegistrySearchable flag */
 	TMultiMap<FName, TOptional<FString>> TagsAndValues;
 
-	/** Only if bRecursiveClasses is true, the results will exclude classes (and subclasses) in this list */
+	/** Deprecated. Only if bRecursiveClasses is true, the results will exclude classes (and subclasses) in this list */
+	UE_DEPRECATED(5.1, "Class names are now represented by path names. Please use RecursiveClassPathsExclusionSet.")
 	TSet<FName> RecursiveClassesExclusionSet;
 
+	/** Only if bRecursiveClasses is true, the results will exclude classes (and subclasses) in this list */
+	TSet<FTopLevelAssetPath> RecursiveClassPathsExclusionSet;
+
 	/** If true, PackagePath components will be recursive */
 	bool bRecursivePaths = false;
 
-	/** If true, subclasses of ClassNames will also be included and RecursiveClassesExclusionSet will be excluded. */
+	/** If true, subclasses of ClassPaths will also be included and RecursiveClassPathsExclusionSet will be excluded. */
 	bool bRecursiveClasses = false;
 
 	/** If true, only on-disk assets will be returned. Be warned that this is rarely what you want and should only be used for performance reasons */
@@ -53,14 +69,14 @@ struct FARFilter
 		PackageNames.Append(Other.PackageNames);
 		PackagePaths.Append(Other.PackagePaths);
 		ObjectPaths.Append(Other.ObjectPaths);
-		ClassNames.Append(Other.ClassNames);
+		ClassPaths.Append(Other.ClassPaths);
 
 		for (auto TagIt = Other.TagsAndValues.CreateConstIterator(); TagIt; ++TagIt)
 		{
 			TagsAndValues.Add(TagIt.Key(), TagIt.Value());
 		}
 
-		RecursiveClassesExclusionSet.Append(Other.RecursiveClassesExclusionSet);
+		RecursiveClassPathsExclusionSet.Append(Other.RecursiveClassPathsExclusionSet);
 
 		bRecursivePaths |= Other.bRecursivePaths;
 		bRecursiveClasses |= Other.bRecursiveClasses;
@@ -72,7 +88,7 @@ struct FARFilter
 	/** Returns true if this filter has no entries */
 	bool IsEmpty() const
 	{
-		return PackageNames.Num() + PackagePaths.Num() + ObjectPaths.Num() + ClassNames.Num() + TagsAndValues.Num() + WithoutPackageFlags + WithPackageFlags == 0;
+		return PackageNames.Num() + PackagePaths.Num() + ObjectPaths.Num() + ClassPaths.Num() + TagsAndValues.Num() + WithoutPackageFlags + WithPackageFlags == 0;
 	}
 
 	/** Returns true if this filter is recursive */
@@ -87,9 +103,9 @@ struct FARFilter
 		PackageNames.Empty();
 		PackagePaths.Empty();
 		ObjectPaths.Empty();
-		ClassNames.Empty();
+		ClassPaths.Empty();
 		TagsAndValues.Empty();
-		RecursiveClassesExclusionSet.Empty();
+		RecursiveClassPathsExclusionSet.Empty();
 
 		bRecursivePaths = false;
 		bRecursiveClasses = false;
@@ -99,7 +115,19 @@ struct FARFilter
 
 		ensure(IsEmpty());
 	}
+
+	void PostSerialize(const FArchive& Ar);
+};
+
+template<>
+struct TStructOpsTypeTraits<FARFilter> : public TStructOpsTypeTraitsBase2<FARFilter>
+{
+	enum
+	{
+		WithPostSerialize = true,
+	};
 };
+
 /**
  * A struct to serve as a filter for Asset Registry queries.
  * Each component element is processed as an 'OR' operation while all the components are processed together as an 'AND' operation.
@@ -107,6 +135,14 @@ struct FARFilter
  */
 struct FARCompiledFilter
 {
+PRAGMA_DISABLE_DEPRECATION_WARNINGS // Compilers can complain about deprecated members in compiler generated code
+	FARCompiledFilter() = default;
+	FARCompiledFilter(FARCompiledFilter&&) = default;
+	FARCompiledFilter(const FARCompiledFilter&) = default;
+	FARCompiledFilter& operator=(FARCompiledFilter&&) = default;
+	FARCompiledFilter& operator=(const FARCompiledFilter&) = default;
+PRAGMA_ENABLE_DEPRECATION_WARNINGS
+
 	/** The filter component for package names */
 	TSet<FName> PackageNames;
 
@@ -116,9 +152,13 @@ struct FARCompiledFilter
 	/** The filter component containing specific object paths */
 	TSet<FName> ObjectPaths;
 
-	/** The filter component for class names. Instances of the specified classes, but not subclasses (by default), will be included. Derived classes will be included only if bRecursiveClasses is true. */
+	/** Deprecated. The filter component for class names. Instances of the specified classes, but not subclasses (by default), will be included. Derived classes will be included only if bRecursiveClasses is true. */
+	UE_DEPRECATED(5.1, "Class names are now represented by path names. Please use ClassPaths.")
 	TSet<FName> ClassNames;
 
+	/** The filter component for class names. Instances of the specified classes, but not subclasses (by default), will be included. Derived classes will be included only if bRecursiveClasses is true. */
+	TSet<FTopLevelAssetPath> ClassPaths;
+
 	/** The filter component for properties marked with the AssetRegistrySearchable flag */
 	TMultiMap<FName, TOptional<FString>> TagsAndValues;
 
@@ -134,7 +174,7 @@ struct FARCompiledFilter
 	/** Returns true if this filter has no entries */
 	bool IsEmpty() const
 	{
-		return PackageNames.Num() + PackagePaths.Num() + ObjectPaths.Num() + ClassNames.Num() + TagsAndValues.Num() == 0;
+		return PackageNames.Num() + PackagePaths.Num() + ObjectPaths.Num() + ClassPaths.Num() + TagsAndValues.Num() == 0;
 	}
 
 	/** Clears this filter of all entries */
@@ -143,7 +183,7 @@ struct FARCompiledFilter
 		PackageNames.Empty();
 		PackagePaths.Empty();
 		ObjectPaths.Empty();
-		ClassNames.Empty();
+		ClassPaths.Empty();
 		TagsAndValues.Empty();
 
 		bIncludeOnlyOnDiskAssets = false;
