Deprecating ANY_PACKAGE.

diff --git a/Engine/Source/Runtime/CoreUObject/Public/AssetRegistry/AssetData.h b/Engine/Source/Runtime/CoreUObject/Public/AssetRegistry/AssetData.h
--- a/Engine/Source/Runtime/CoreUObject/Public/AssetRegistry/AssetData.h
+++ b/Engine/Source/Runtime/CoreUObject/Public/AssetRegistry/AssetData.h
@@ -22,6 +22,7 @@
 #include "UObject/Package.h"
 #include "UObject/PrimaryAssetId.h"
 #include "UObject/LinkerInstancingContext.h"
+#include "UObject/TopLevelAssetPath.h"
 
 struct FCustomVersion;
 
@@ -55,8 +56,9 @@ struct COREUOBJECT_API FAssetRegistryVersion
 		WorkspaceDomain,					// Added Version information to AssetPackageData
 		PackageImportedClasses,				// Added ImportedClasses to AssetPackageData
 		PackageFileSummaryVersionChange,	// A new version number of UE5 was added to FPackageFileSummary
-		ObjectResourceOptionalVersionChange,// Change to linker export/import resource serializationn
+		ObjectResourceOptionalVersionChange,// Change to linker export/import resource serialization
 		AddedChunkHashes,					// Added FIoHash for each FIoChunkId in the package to the AssetPackageData.
+		ClassPaths,							// Classes are serialized as path names rather than short object names, e.g. /Script/Engine.StaticMesh
 
 		// -----<new versions can be added above this line>-------------------------------------------------
 		VersionPlusOne,
@@ -104,7 +106,10 @@ public:
 	/** The name of the asset without the package */
 	FName AssetName;
 	/** The name of the asset's class */
+	UE_DEPRECATED(5.1, "Class names are now represented by path names. Please use AssetClassPath.")
 	FName AssetClass;
+	/** The path of the asset's class, e.g. /Script/Engine.StaticMesh */
+	FTopLevelAssetPath AssetClassPath;
 	/** Asset package flags */
 	uint32 PackageFlags = 0;
 	/** The map of values for properties that were marked AssetRegistrySearchable or added by GetAssetRegistryTags */
@@ -120,14 +125,27 @@ public:
 	/** The IDs of the pakchunks this asset is located in for streaming install.  Empty if not assigned to a chunk */
 	TArray<int32, TInlineAllocator<2>> ChunkIDs;
 
-public:
-	/** Default constructor */
-	FAssetData() {}
+public:	
+PRAGMA_DISABLE_DEPRECATION_WARNINGS // Compilers can complain about deprecated members in compiler generated code
+	/** Default constructors. */
+	FAssetData()
+	{
+	}
+	FAssetData(FAssetData&&) = default;
+	FAssetData(const FAssetData&) = default;
+	FAssetData& operator=(FAssetData&&) = default;
+	FAssetData& operator=(const FAssetData&) = default;
+PRAGMA_ENABLE_DEPRECATION_WARNINGS
+
+	UE_DEPRECATED(5.1, "Class names are now represented by path names. Please use a version of FAssetData constructor that uses FTopLevelAssetPath.")
+	COREUOBJECT_API FAssetData(FName InPackageName, FName InPackagePath, FName InAssetName, FName InAssetClassName, FAssetDataTagMap InTags = FAssetDataTagMap(), TArrayView<const int32> InChunkIDs = TArrayView<const int32>(), uint32 InPackageFlags = 0);
+	UE_DEPRECATED(5.1, "Class names are now represented by path names. Please use a version of FAssetData constructor that uses FTopLevelAssetPath.")
+	COREUOBJECT_API FAssetData(const FString& InLongPackageName, const FString& InObjectPath, FName InAssetClassName, FAssetDataTagMap InTags = FAssetDataTagMap(), TArrayView<const int32> InChunkIDs = TArrayView<const int32>(), uint32 InPackageFlags = 0);
 
 	/** Constructor building the ObjectPath in the form of InPackageName.InAssetName. does not work for object outer-ed to a different package. */
-	COREUOBJECT_API FAssetData(FName InPackageName, FName InPackagePath, FName InAssetName, FName InAssetClass, FAssetDataTagMap InTags = FAssetDataTagMap(), TArrayView<const int32> InChunkIDs = TArrayView<const int32>(), uint32 InPackageFlags = 0);
+	COREUOBJECT_API FAssetData(FName InPackageName, FName InPackagePath, FName InAssetName, FTopLevelAssetPath InAssetClassPathName, FAssetDataTagMap InTags = FAssetDataTagMap(), TArrayView<const int32> InChunkIDs = TArrayView<const int32>(), uint32 InPackageFlags = 0);
 	/** Constructor with a long package name and a full object path which might not be part of the package this asset is in. */
-	COREUOBJECT_API FAssetData(const FString& InLongPackageName, const FString& InObjectPath, FName InAssetClass, FAssetDataTagMap InTags = FAssetDataTagMap(), TArrayView<const int32> InChunkIDs = TArrayView<const int32>(), uint32 InPackageFlags = 0);
+	COREUOBJECT_API FAssetData(const FString& InLongPackageName, const FString& InObjectPath, FTopLevelAssetPath InAssetClassPathName, FAssetDataTagMap InTags = FAssetDataTagMap(), TArrayView<const int32> InChunkIDs = TArrayView<const int32>(), uint32 InPackageFlags = 0);
 
 	/** Constructor taking a UObject. By default trying to create one for a blueprint class will create one for the UBlueprint instead, but this can be overridden */
 	COREUOBJECT_API FAssetData(const UObject* InAsset, FAssetData::ECreationFlags InCreationFlags = ECreationFlags::None);
@@ -225,7 +243,7 @@ public:
 	void GetFullName(FString& OutFullName) const
 	{
 		OutFullName.Reset();
-		AssetClass.AppendString(OutFullName);
+		OutFullName += AssetClassPath.ToString();
 		OutFullName.AppendChar(TEXT(' '));
 		ObjectPath.AppendString(OutFullName);
 	}
@@ -234,7 +252,7 @@ public:
 	void GetFullName(FStringBuilderBase& OutFullNameBuilder) const
 	{
 		OutFullNameBuilder.Reset();
-		AssetClass.AppendString(OutFullNameBuilder);
+		AssetClassPath.AppendString(OutFullNameBuilder);
 		OutFullNameBuilder.AppendChar(TEXT(' '));
 		ObjectPath.AppendString(OutFullNameBuilder);
 	}
@@ -251,7 +269,7 @@ public:
 	void GetExportTextName(FString& OutExportTextName) const
 	{
 		OutExportTextName.Reset();
-		AssetClass.AppendString(OutExportTextName);
+		OutExportTextName.Append(AssetClassPath.ToString());
 		OutExportTextName.AppendChar(TEXT('\''));
 		ObjectPath.AppendString(OutExportTextName);
 		OutExportTextName.AppendChar(TEXT('\''));
@@ -261,7 +279,7 @@ public:
 	void GetExportTextName(FStringBuilderBase& OutExportTextNameBuilder) const
 	{
 		OutExportTextNameBuilder.Reset();
-		AssetClass.AppendString(OutExportTextNameBuilder);
+		AssetClassPath.AppendString(OutExportTextNameBuilder);
 		OutExportTextNameBuilder.AppendChar(TEXT('\''));
 		ObjectPath.AppendString(OutExportTextNameBuilder);
 		OutExportTextNameBuilder.AppendChar(TEXT('\''));
@@ -270,33 +288,31 @@ public:
 	/** Returns true if the this asset is a redirector. */
 	bool IsRedirector() const
 	{
-		return IsRedirectorClassName(AssetClass);
+		return IsRedirectorClassName(AssetClassPath);
 	}
 
 	static bool IsRedirector(UObject* Object)
 	{
-		return Object && IsRedirectorClassName(Object->GetClass()->GetFName());
+		return Object && IsRedirectorClassName(Object->GetClass()->GetClassPathName());
 	}
 
-	/** Returns the class UClass if it is loaded. It is not possible to load the class if it is unloaded since we only have the short name. */
+	/** Returns the class UClass if it is loaded. */
 	UClass* GetClass() const
 	{
 		if ( !IsValid() )
 		{
 			// Dont even try to find the class if the objectpath isn't set
-			return NULL;
+			return nullptr;
 		}
 
-		UClass* FoundClass = FindObject<UClass>(ANY_PACKAGE, *AssetClass.ToString());
-
+		UClass* FoundClass = FindObject<UClass>(AssetClassPath);
 		if (!FoundClass)
 		{
 			// Look for class redirectors
-			FName NewPath = FLinkerLoad::FindNewNameForClass(AssetClass, false);
-
-			if (NewPath != NAME_None)
+			FString NewPath = FLinkerLoad::FindNewPathNameForClass(AssetClassPath.ToString(), false);
+			if (!NewPath.IsEmpty())
 			{
-				FoundClass = FindObject<UClass>(ANY_PACKAGE, *NewPath.ToString());
+				FoundClass = FindObject<UClass>(nullptr, *NewPath);
 			}
 		}
 		return FoundClass;
@@ -461,7 +477,7 @@ public:
 		UE_LOG(LogAssetData, Log, TEXT("        PackageName: %s"), *PackageName.ToString());
 		UE_LOG(LogAssetData, Log, TEXT("        PackagePath: %s"), *PackagePath.ToString());
 		UE_LOG(LogAssetData, Log, TEXT("        AssetName: %s"), *AssetName.ToString());
-		UE_LOG(LogAssetData, Log, TEXT("        AssetClass: %s"), *AssetClass.ToString());
+		UE_LOG(LogAssetData, Log, TEXT("        AssetClassPath: %s"), *AssetClassPath.ToString());
 		UE_LOG(LogAssetData, Log, TEXT("        TagsAndValues: %d"), TagsAndValues.Num());
 
 		for (const auto TagValue: TagsAndValues)
@@ -508,14 +524,14 @@ public:
 	 * To version this data change FAssetRegistryVersion
 	 */
 	template<class Archive>
-	void SerializeForCache(Archive&& Ar)
+	FORCEINLINE void SerializeForCache(Archive&& Ar, FAssetRegistryVersion::Type Version = FAssetRegistryVersion::LatestVersion)
 	{
-		SerializeForCacheInternal(Ar, [](FArchive& Ar, FAssetData& Ad) {
+		SerializeForCacheInternal(Ar, Version, [](FArchive& Ar, FAssetData& Ad) {
 			static_cast<Archive&>(Ar).SerializeTagsAndBundles(Ad);
 			});
 	}
 private:
-	COREUOBJECT_API void SerializeForCacheInternal(FArchive& Ar, void (*SerializeTagsAndBundles)(FArchive& , FAssetData&));
+	COREUOBJECT_API void SerializeForCacheInternal(FArchive& Ar, FAssetRegistryVersion::Type Version, void (*SerializeTagsAndBundles)(FArchive& , FAssetData&));
 
 	static bool DetectIsUAssetByNames(FStringView PackageName, FStringView ObjectPathName)
 	{
@@ -530,11 +546,15 @@ private:
 		return PackageBaseName.Equals(ObjectPathName, ESearchCase::IgnoreCase);
 	}
 
-	static bool IsRedirectorClassName(FName ClassName)
+	static bool IsRedirectorClassName(FTopLevelAssetPath ClassPathName)
 	{
-		static const FName ObjectRedirectorClassName = UObjectRedirector::StaticClass()->GetFName();
-		return ClassName == ObjectRedirectorClassName;
+		static const FTopLevelAssetPath ObjectRedirectorClassPathName = UObjectRedirector::StaticClass()->GetClassPathName();
+		return ClassPathName == ObjectRedirectorClassPathName;
 	}
+
+public:
+	/** Helper function that tries to convert short class name to path name */
+	COREUOBJECT_API static FTopLevelAssetPath TryConvertShortClassNameToPathName(FName InClassName);
 };
 
 ENUM_CLASS_FLAGS(FAssetData::ECreationFlags);
@@ -895,7 +915,7 @@ struct FAssetIdentifier
 	/** Returns primary asset id for this identifier, if valid */
 	FPrimaryAssetId GetPrimaryAssetId() const
 	{
-		if (PrimaryAssetType != NAME_None)
+		if (PrimaryAssetType.IsValid())
 		{
 			return FPrimaryAssetId(PrimaryAssetType, ObjectName);
 		}
@@ -937,7 +957,7 @@ struct FAssetIdentifier
 	/** Appends to the given builder the string version of this identifier in Package.Object::Name format */
 	void AppendString(FStringBuilderBase& Builder) const
 	{
-		if (PrimaryAssetType != NAME_None)
+		if (PrimaryAssetType.IsValid())
 		{
 			GetPrimaryAssetId().AppendString(Builder);
 		}
@@ -1016,7 +1036,7 @@ struct FAssetIdentifier
 		if (Ar.IsSaving())
 		{
 			FieldBits |= (AssetIdentifier.PackageName != NAME_None) << 0;
-			FieldBits |= (AssetIdentifier.PrimaryAssetType != NAME_None) << 1;
+			FieldBits |= (AssetIdentifier.PrimaryAssetType.IsValid()) << 1;
 			FieldBits |= (AssetIdentifier.ObjectName != NAME_None) << 2;
 			FieldBits |= (AssetIdentifier.ValueName != NAME_None) << 3;
 		}
diff --git a/Engine/Source/Runtime/CoreUObject/Public/AssetRegistry/AssetDataTagMap.h b/Engine/Source/Runtime/CoreUObject/Public/AssetRegistry/AssetDataTagMap.h
--- a/Engine/Source/Runtime/CoreUObject/Public/AssetRegistry/AssetDataTagMap.h
+++ b/Engine/Source/Runtime/CoreUObject/Public/AssetRegistry/AssetDataTagMap.h
@@ -7,7 +7,7 @@
 #include "Misc/StringBuilder.h"
 #include "Templates/RefCounting.h"
 #include "Templates/TypeCompatibleBytes.h"
-
+#include "UObject/TopLevelAssetPath.h"
 
 class FAssetTagValueRef;
 class FAssetDataTagMapSharedView;
@@ -21,10 +21,19 @@ struct FAssetRegistrySerializationOptions;
  */
 struct COREUOBJECT_API FAssetRegistryExportPath
 {
+PRAGMA_DISABLE_DEPRECATION_WARNINGS // Compilers can complain about deprecated members in compiler generated code
 	FAssetRegistryExportPath() = default;
+	FAssetRegistryExportPath(FAssetRegistryExportPath&&) = default;
+	FAssetRegistryExportPath(const FAssetRegistryExportPath&) = default;
+	FAssetRegistryExportPath& operator=(FAssetRegistryExportPath&&) = default;
+	FAssetRegistryExportPath& operator=(const FAssetRegistryExportPath&) = default;
+PRAGMA_ENABLE_DEPRECATION_WARNINGS
+
 	explicit FAssetRegistryExportPath(FWideStringView String);
 	explicit FAssetRegistryExportPath(FAnsiStringView String);
 
+	FTopLevelAssetPath ClassPath;
+	UE_DEPRECATED(5.1, "Class names are now represented by path names. Please use ClassPath member")
 	FName Class;
 	FName Package;
 	FName Object;
@@ -32,8 +41,13 @@ struct COREUOBJECT_API FAssetRegistryExportPath
 	FString ToString() const;
 	FName ToName() const;
 	void ToString(FStringBuilderBase& Out) const;
-
-	bool IsEmpty() const { return Class.IsNone() & Package.IsNone() & Object.IsNone(); } //-V792
+	FString ToPath() const;
+	void ToPath(FStringBuilderBase& Out) const;
+	FTopLevelAssetPath ToTopLevelAssetPath() const
+	{
+		return FTopLevelAssetPath(Package, Object);
+	}
+	bool IsEmpty() const { return ClassPath.IsNull() & Package.IsNone() & Object.IsNone(); } //-V792
 	explicit operator bool() const { return !IsEmpty(); }
 };
 
@@ -45,7 +59,8 @@ namespace FixedTagPrivate
 	// Compact FAssetRegistryExportPath equivalent for when all FNames are numberless
 	struct FNumberlessExportPath
 	{
-		FNameEntryId Class;
+		FNameEntryId ClassPackage;
+		FNameEntryId ClassObject;
 		FNameEntryId Package;
 		FNameEntryId Object;
 
diff --git a/Engine/Source/Runtime/CoreUObject/Public/AssetRegistry/AssetDataTagMapSerializationDetails.h b/Engine/Source/Runtime/CoreUObject/Public/AssetRegistry/AssetDataTagMapSerializationDetails.h
--- a/Engine/Source/Runtime/CoreUObject/Public/AssetRegistry/AssetDataTagMapSerializationDetails.h
+++ b/Engine/Source/Runtime/CoreUObject/Public/AssetRegistry/AssetDataTagMapSerializationDetails.h
@@ -2,6 +2,7 @@
 
 #pragma once
 
+#include "AssetRegistry/AssetData.h"
 #include "AssetRegistry/AssetDataTagMap.h"
 #include "Async/Async.h"
 
@@ -9,6 +10,14 @@ struct FAssetRegistrySerializationOptions;
 
 namespace FixedTagPrivate
 {
+	// Legacy version of FAssetRegistryExportPath (before FAssetRegistryVersion::ClassPaths)
+	struct FLegacyAssetRegistryExportPath
+	{
+		FName Class;
+		FName Package;
+		FName Object;
+	};
+
 	/// Stores a fixed set of values and all the key-values maps used for lookup
 	struct FStore
 	{
@@ -134,16 +143,22 @@ namespace FixedTagPrivate
 				return HashCombineQuick(GetTypeHash(Key.GetDisplayIndex()), Key.GetNumber());
 			}
 
-			template<class ExportPathType>
-			static bool Matches(const ExportPathType& A, const ExportPathType& B)
+			static bool Matches(const FNumberlessExportPath& A, const FNumberlessExportPath& B)
 			{
-				return Matches(A.Class, B.Class) & Matches(A.Package, B.Package) & Matches(A.Object, B.Object); //-V792
+				return Matches(A.ClassPackage, B.ClassPackage) & Matches(A.ClassObject, B.ClassObject) & Matches(A.Package, B.Package) & Matches(A.Object, B.Object); //-V792
+			}
+			static bool Matches(const FAssetRegistryExportPath& A, const FAssetRegistryExportPath& B)
+			{
+				return Matches(A.ClassPath.GetPackageName(), B.ClassPath.GetPackageName()) &  Matches(A.ClassPath.GetAssetName(), B.ClassPath.GetAssetName()) & Matches(A.Package, B.Package) & Matches(A.Object, B.Object); //-V792
 			}
 
-			template<class ExportPathType>
-			static uint32 GetKeyHash(const ExportPathType& Key)
+			static uint32 GetKeyHash(const FNumberlessExportPath& Key)
+			{
+				return HashCombineQuick(HashCombineQuick(GetKeyHash(Key.ClassPackage), GetKeyHash(Key.ClassObject)), GetKeyHash(Key.Package), GetKeyHash(Key.Object));
+			}
+			static uint32 GetKeyHash(const FAssetRegistryExportPath& Key)
 			{
-				return HashCombineQuick(GetKeyHash(Key.Class), GetKeyHash(Key.Package), GetKeyHash(Key.Object));
+				return HashCombineQuick(HashCombineQuick(GetKeyHash(Key.ClassPath.GetPackageName()), GetKeyHash(Key.ClassPath.GetAssetName())), GetKeyHash(Key.Package), GetKeyHash(Key.Object));
 			}
 		};
 
@@ -179,7 +194,7 @@ namespace FixedTagPrivate
 	enum class ELoadOrder { Member, TextFirst };
 
 	COREUOBJECT_API void SaveStore(const FStoreData& Store, FArchive& Ar);
-	COREUOBJECT_API TRefCountPtr<const FStore> LoadStore(FArchive& Ar);
+	COREUOBJECT_API TRefCountPtr<const FStore> LoadStore(FArchive& Ar, FAssetRegistryVersion::Type Version = FAssetRegistryVersion::LatestVersion);
 
 	/// Loads tag store with async creation of expensive tag values
 	///
diff --git a/Engine/Source/Runtime/CoreUObject/Public/Internationalization/PackageLocalizationCache.h b/Engine/Source/Runtime/CoreUObject/Public/Internationalization/PackageLocalizationCache.h
--- a/Engine/Source/Runtime/CoreUObject/Public/Internationalization/PackageLocalizationCache.h
+++ b/Engine/Source/Runtime/CoreUObject/Public/Internationalization/PackageLocalizationCache.h
@@ -10,6 +10,7 @@
 #include "Templates/SharedPointer.h"
 #include "Templates/Tuple.h"
 #include "UObject/NameTypes.h"
+#include "UObject/Class.h"
 
 class FPackageLocalizationCache;
 
@@ -131,7 +132,7 @@ protected:
 	 * @param InAssetGroupName	The name of the asset group packages of this type belong to.
 	 * @param InAssetClassName	The name of a class used by this asset group.
 	 */
-	virtual void FindAssetGroupPackages(const FName InAssetGroupName, const FName InAssetClassName) = 0;
+	virtual void FindAssetGroupPackages(const FName InAssetGroupName, const FTopLevelAssetPath& InAssetClassName) = 0;
 
 	/**
 	 * Try and find an existing cache for the given culture name, and create an entry for one if no such cache currently exists.
@@ -179,7 +180,7 @@ protected:
 	TArray<TTuple<FString, TSharedPtr<FPackageLocalizationCultureCache>>> AllCultureCaches;
 
 	/** Mapping between a class name, and the asset group the class belongs to (for class specific package localization). */
-	TArray<TTuple<FName, FName>> AssetClassesToAssetGroups;
+	TArray<TTuple<FTopLevelAssetPath, FName>> AssetClassesToAssetGroups;
 
 	/** Mapping between a package name, and the asset group it belongs to. */
 	TMap<FName, FName> PackageNameToAssetGroup;
diff --git a/Engine/Source/Runtime/CoreUObject/Public/Misc/AssetRegistryInterface.h b/Engine/Source/Runtime/CoreUObject/Public/Misc/AssetRegistryInterface.h
--- a/Engine/Source/Runtime/CoreUObject/Public/Misc/AssetRegistryInterface.h
+++ b/Engine/Source/Runtime/CoreUObject/Public/Misc/AssetRegistryInterface.h
@@ -196,7 +196,7 @@ namespace UE::AssetRegistry
 	{
 	public:
 		/** Return whether to filter out assets of the given class and flags from the editor's asset registry */
-		static bool ShouldSkipAsset(FName AssetClass, uint32 PackageFlags);
+		static bool ShouldSkipAsset(const FTopLevelAssetPath& AssetClass, uint32 PackageFlags);
 
 		/** Return whether to filter out the given object (assumed to be an asset) from the editor's asset registry */
 		static bool ShouldSkipAsset(const UObject* InAsset);
@@ -206,7 +206,7 @@ namespace UE::AssetRegistry
 
 #if WITH_ENGINE && WITH_EDITOR
 		/** Copy the global skip classes set from the given external sets that were already populated. */
-		static void SetSkipClasses(const TSet<FName>& InSkipUncookedClasses, const TSet<FName>& InSkipCookedClasses);
+		static void SetSkipClasses(const TSet<FTopLevelAssetPath>& InSkipUncookedClasses, const TSet<FTopLevelAssetPath>& InSkipCookedClasses);
 #endif
 	};
 
@@ -214,13 +214,13 @@ namespace UE::AssetRegistry
 namespace Utils
 {
 	/** Return whether to filter out assets of the given class and flags based on the skip classes */
-	COREUOBJECT_API bool ShouldSkipAsset(FName AssetClass, uint32 PackageFlags,
-		const TSet<FName>& InSkipUncookedClasses, const TSet<FName>& InSkipCookedClasses);
+	COREUOBJECT_API bool ShouldSkipAsset(const FTopLevelAssetPath& AssetClass, uint32 PackageFlags,
+		const TSet<FTopLevelAssetPath>& InSkipUncookedClasses, const TSet<FTopLevelAssetPath>& InSkipCookedClasses);
 	/** Return whether to filter out the given object (assumed to be an asset) based on the skip classes */
 	COREUOBJECT_API bool ShouldSkipAsset(const UObject* InAsset,
-		const TSet<FName>& InSkipUncookedClasses, const TSet<FName>& InSkipCookedClasses);
+		const TSet<FTopLevelAssetPath>& InSkipUncookedClasses, const TSet<FTopLevelAssetPath>& InSkipCookedClasses);
 	/** Run the calculation of which classes to skip and store results in the given sets. */
-	COREUOBJECT_API void PopulateSkipClasses(TSet<FName>& OutSkipUncookedClasses, TSet<FName>& OutSkipCookedClasses);
+	COREUOBJECT_API void PopulateSkipClasses(TSet<FTopLevelAssetPath>& OutSkipUncookedClasses, TSet<FTopLevelAssetPath>& OutSkipCookedClasses);
 }
 #endif
 
diff --git a/Engine/Source/Runtime/CoreUObject/Public/UObject/Class.h b/Engine/Source/Runtime/CoreUObject/Public/UObject/Class.h
--- a/Engine/Source/Runtime/CoreUObject/Public/UObject/Class.h
+++ b/Engine/Source/Runtime/CoreUObject/Public/UObject/Class.h
@@ -355,6 +355,10 @@ public:
 
 	/** Get the Schema Hash for this struct - the hash of its property names and types. */
 	const FBlake3Hash& GetSchemaHash(bool bSkipEditorOnly) const;
+
+protected:
+	/** True if this struct has Asset Registry searchable properties */
+	bool bHasAssetRegistrySearchableProperties;
 #endif
 
 public:
@@ -560,6 +564,12 @@ public:
 	 * @return pointer to the UStruct that has associated metadata, nullptr if Key is not associated with any UStruct in the hierarchy
 	 */
 	const UStruct* HasMetaDataHierarchical(const FName& Key) const;
+
+	/* Returns true if this struct has Asset Registry searchable properties */
+	FORCEINLINE bool HasAssetRegistrySearchableProperties() const
+	{
+		return bHasAssetRegistrySearchableProperties;
+	}
 #endif // WITH_EDITORONLY_DATA
 
 #if HACK_HEADER_GENERATOR
@@ -2703,6 +2713,9 @@ public:
 	virtual FRestoreForUObjectOverwrite* GetRestoreForUObjectOverwrite() override;
 	virtual FString GetDesc() override;
 	virtual void GetAssetRegistryTags(TArray<FAssetRegistryTag>& OutTags) const override;
+#if WITH_EDITOR
+	virtual void PostLoadAssetRegistryTags(const FAssetData& InAssetData, TArray<FAssetRegistryTag>& OutTagsAndValuesToUpdate) const;
+#endif // WITH_EDITOR
 	virtual bool IsAsset() const override { return false; }	
 	virtual bool IsNameStableForNetworking() const override { return true; } // For now, assume all classes have stable net names
 	virtual void GetPreloadDependencies(TArray<UObject*>& OutDeps) override;
@@ -2772,7 +2785,7 @@ public:
 	template <typename T>
 	static T* TryFindTypeSlowSafe(const FString& InShortNameOrPathName, EFindFirstObjectOptions InOptions = EFindFirstObjectOptions::None)
 	{
-		return (T*)TryFindTypeSafe(T::StaticClass(), InShortNameOrPathName, InOptions);
+		return (T*)TryFindTypeSlowSafe(T::StaticClass(), InShortNameOrPathName, InOptions);
 	}
 
 	/**
diff --git a/Engine/Source/Runtime/CoreUObject/Public/UObject/LinkerLoad.h b/Engine/Source/Runtime/CoreUObject/Public/UObject/LinkerLoad.h
--- a/Engine/Source/Runtime/CoreUObject/Public/UObject/LinkerLoad.h
+++ b/Engine/Source/Runtime/CoreUObject/Public/UObject/LinkerLoad.h
@@ -1290,7 +1290,7 @@ private:
 	/** Finds import, tries to fall back to dynamic class if the object could not be found */
 	UObject* FindImport(UClass* ImportClass, UObject* ImportOuter, const TCHAR* Name);
 	/** Finds import, tries to fall back to dynamic class if the object could not be found */
-	static UObject* FindImportFast(UClass* ImportClass, UObject* ImportOuter, FName Name, bool bAnyPackage = false);
+	static UObject* FindImportFast(UClass* ImportClass, UObject* ImportOuter, FName Name, bool bFindObjectbyName = false);
 
 #if	USE_CIRCULAR_DEPENDENCY_LOAD_DEFERRING
 	/** 
diff --git a/Engine/Source/Runtime/CoreUObject/Public/UObject/NoExportTypes.h b/Engine/Source/Runtime/CoreUObject/Public/UObject/NoExportTypes.h
--- a/Engine/Source/Runtime/CoreUObject/Public/UObject/NoExportTypes.h
+++ b/Engine/Source/Runtime/CoreUObject/Public/UObject/NoExportTypes.h
@@ -1700,11 +1700,12 @@ struct FSoftClassPath : public FSoftObjectPath
 USTRUCT(noexport, BlueprintType)
 struct FTopLevelAssetPath
 {
+private:
 	/** Name of the package containing the asset e.g. /Path/To/Package */
-	UPROPERTY(EditAnywhere, SaveGame, BlueprintReadWrite, Category = TopLevelAssetPath)
+	UPROPERTY(EditAnywhere, SaveGame, BlueprintReadWrite, Category = TopLevelAssetPath, meta = (AllowPrivateAccess = "true"))
 	FName PackageName;
 	/** Name of the asset within the package e.g. 'AssetName' */
-	UPROPERTY(EditAnywhere, SaveGame, BlueprintReadWrite, Category = TopLevelAssetPath)
+	UPROPERTY(EditAnywhere, SaveGame, BlueprintReadWrite, Category = TopLevelAssetPath, meta = (AllowPrivateAccess = "true"))
 	FName AssetName;
 };
 
@@ -2145,16 +2146,24 @@ struct FARFilter
 	UPROPERTY(EditDefaultsOnly, BlueprintReadWrite, Category = AssetRegistry)
 	TArray<FName> ObjectPaths;
 
-	/** The filter component for class names. Instances of the specified classes, but not subclasses (by default), will be included. Derived classes will be included only if bRecursiveClasses is true. */
+	/** The filter component for class names. Instances of the specified classes, but not subclasses (by default), will be included. Derived classes will be included only if bRecursiveClasses is true. Deprecated. */
 	UPROPERTY(EditDefaultsOnly, BlueprintReadWrite, Category = AssetRegistry)
 	TArray<FName> ClassNames;
 
+	/** The filter component for class path names. Instances of the specified classes, but not subclasses (by default), will be included. Derived classes will be included only if bRecursiveClasses is true. */
+	UPROPERTY(EditDefaultsOnly, BlueprintReadWrite, Category = AssetRegistry)
+	TArray<FTopLevelAssetPath> ClassPaths;
+
 	/** The filter component for properties marked with the AssetRegistrySearchable flag */
 	TMultiMap<FName, TOptional<FString>> TagsAndValues;
 
+	/** Deprecated. Only if bRecursiveClasses is true, the results will exclude classes (and subclasses) in this list */
+	UPROPERTY()
+	TSet<FName> RecursiveClassesExclusionSet;
+
 	/** Only if bRecursiveClasses is true, the results will exclude classes (and subclasses) in this list */
 	UPROPERTY(EditDefaultsOnly, BlueprintReadWrite, Category = AssetRegistry)
-	TSet<FName> RecursiveClassesExclusionSet;
+	TSet<FTopLevelAssetPath> RecursiveClassPathsExclusionSet;
 
 	/** If true, PackagePath components will be recursive */
 	UPROPERTY(EditDefaultsOnly, BlueprintReadWrite, Category = AssetRegistry)
@@ -2215,16 +2224,20 @@ struct FAssetData
 	/** The name of the asset without the package */
 	UPROPERTY(BlueprintReadOnly, Category = AssetData, transient)
 	FName AssetName;
-	/** The name of the asset's class */
+	/** Deprecated. The name of the asset's class */
 	UPROPERTY(BlueprintReadOnly, Category = AssetData, transient)
 	FName AssetClass;
+	/** The path name of the asset's class */
+	UPROPERTY(BlueprintReadOnly, Category = AssetData, transient)
+	FTopLevelAssetPath AssetClassPath;
+
+	/** Asset package flags */
+	uint32 PackageFlags = 0;
 	/** The map of values for properties that were marked AssetRegistrySearchable or added by GetAssetRegistryTags */
 	FAssetDataTagMapSharedView TagsAndValues;
 	TSharedPtr<FAssetBundleData, ESPMode::ThreadSafe> TaggedAssetBundles;
 	/** The IDs of the pakchunks this asset is located in for streaming install.  Empty if not assigned to a chunk */
 	TArray<int32, TInlineAllocator<2>> ChunkIDs;
-	/** Asset package flags */
-	uint32 PackageFlags = 0;
 };
 
 USTRUCT(noexport, IsAlwaysAccessible, HasDefaults)
diff --git a/Engine/Source/Runtime/CoreUObject/Public/UObject/Object.h b/Engine/Source/Runtime/CoreUObject/Public/UObject/Object.h
--- a/Engine/Source/Runtime/CoreUObject/Public/UObject/Object.h
+++ b/Engine/Source/Runtime/CoreUObject/Public/UObject/Object.h
@@ -844,6 +844,17 @@ public:
 	static const FName& SourceFileTagName();
 
 #if WITH_EDITOR
+
+	/**
+	 * Performs fixup on loaded asset registry data. 
+	 * This function is called from inside the AssetRegistry CriticalSection. DO NOT CALL ASSETREGISTRY FUNCTIONS FROM THIS FUNCTION, IT WILL DEADLOCK.
+	 * Note that this function is only called on Class Default Objects where the actual object instance data used to generate 
+	 * the asset data is not available.
+	 * @param InAssetData Asset data loaded from the AssetRegistry
+	 * @return Pointer to new asset data after fixup or nullptr if no fixup was required
+	 */
+	virtual void PostLoadAssetRegistryTags(const FAssetData& InAssetData, TArray<FAssetRegistryTag>& OutTagsAndValuesToUpdate) const;
+
 	/**
 	 * Additional data pertaining to asset registry tags used by the editor
 	 */
diff --git a/Engine/Source/Runtime/CoreUObject/Public/UObject/ObjectMacros.h b/Engine/Source/Runtime/CoreUObject/Public/UObject/ObjectMacros.h
--- a/Engine/Source/Runtime/CoreUObject/Public/UObject/ObjectMacros.h
+++ b/Engine/Source/Runtime/CoreUObject/Public/UObject/ObjectMacros.h
@@ -657,7 +657,7 @@ struct COREUOBJECT_API FReferencerInformationList
 ----------------------------------------------------------------------------*/
 
 // Special canonical package for FindObject, ParseObject.
-#define ANY_PACKAGE ((UPackage*)-1)
+#define ANY_PACKAGE UE_DEPRECATED_MACRO(5.1, "ANY_PACKAGE has been deprecated. Either use full path name of objects (including classes) or provide a valid Outer for *FindObject* functions") ((UPackage*)-1)
 
 // Special prefix for default objects (the UObject in a UClass containing the default values, etc)
 #define DEFAULT_OBJECT_PREFIX TEXT("Default__")
diff --git a/Engine/Source/Runtime/CoreUObject/Public/UObject/TopLevelAssetPath.h b/Engine/Source/Runtime/CoreUObject/Public/UObject/TopLevelAssetPath.h
--- a/Engine/Source/Runtime/CoreUObject/Public/UObject/TopLevelAssetPath.h
+++ b/Engine/Source/Runtime/CoreUObject/Public/UObject/TopLevelAssetPath.h
@@ -33,7 +33,7 @@ struct COREUOBJECT_API FTopLevelAssetPath
 	UE_DEPRECATED(5.0, "FNames containing full asset paths have been replaced by FTopLevelAssetPath/FSoftLevelObjectPath."
 		"This function is only for temporary use interfacing with APIs that still produce an FName."
 		"Those APIS should be updated to use FTopLevelAssetPath or FSoftLevelObjectPath.")
-	FTopLevelAssetPath(FName InPath)
+	explicit FTopLevelAssetPath(FName InPath)
 	{
 		TrySetPath(InPath.ToString());
 	}
diff --git a/Engine/Source/Runtime/CoreUObject/Public/UObject/UObjectGlobals.h b/Engine/Source/Runtime/CoreUObject/Public/UObject/UObjectGlobals.h
--- a/Engine/Source/Runtime/CoreUObject/Public/UObject/UObjectGlobals.h
+++ b/Engine/Source/Runtime/CoreUObject/Public/UObject/UObjectGlobals.h
@@ -236,7 +236,6 @@
 /** Internal function used to update the suffix to be given to the next newly-created unnamed object. */
 COREUOBJECT_API int32 UpdateSuffixForNextNewObject(UObject* Parent, const UClass* Class, TFunctionRef<void(int32&)> IndexMutator);
 
-
 /**
  * Fast version of StaticFindObject that relies on the passed in FName being the object name without any group/package qualifiers.
  * This will only find top level packages or subobjects nested directly within a passed in outer.
@@ -251,7 +250,23 @@
  *
  * @return	Returns a pointer to the found object or null if none could be found
  */
-COREUOBJECT_API UObject* StaticFindObjectFast(UClass* Class, UObject* InOuter, FName InName, bool bExactClass = false, bool bAnyPackage = false, EObjectFlags ExclusiveFlags = RF_NoFlags, EInternalObjectFlags ExclusiveInternalFlags = EInternalObjectFlags::None);
+UE_DEPRECATED(5.1, "Support for searching for objects in ANY_PACKAGE has been deprecated. Please provide the actual Outer of an object you want to find.")
+COREUOBJECT_API UObject* StaticFindObjectFast(UClass* Class, UObject* InOuter, FName InName, bool bExactClass, bool bAnyPackage, EObjectFlags ExclusiveFlags = RF_NoFlags, EInternalObjectFlags ExclusiveInternalFlags = EInternalObjectFlags::None);
+
+/**
+ * Fast version of StaticFindObject that relies on the passed in FName being the object name without any group/package qualifiers.
+ * This will only find top level packages or subobjects nested directly within a passed in outer.
+ *
+ * @param	Class			The to be found object's class
+ * @param	InOuter			Outer object to look inside, if null this will only look for top level packages
+ * @param	InName			Object name to look for relative to InOuter
+ * @param	bExactClass		Whether to require an exact match with the passed in class
+ * @param	ExclusiveFlags	Ignores objects that contain any of the specified exclusive flags
+ * @param	ExclusiveInternalFlags  Ignores objects that contain any of the specified internal exclusive flags
+ *
+ * @return	Returns a pointer to the found object or null if none could be found
+ */
+COREUOBJECT_API UObject* StaticFindObjectFast(UClass* Class, UObject* InOuter, FName InName, bool bExactClass = false, EObjectFlags ExclusiveFlags = RF_NoFlags, EInternalObjectFlags ExclusiveInternalFlags = EInternalObjectFlags::None);
 
 /**
  * Fast and safe version of StaticFindObject that relies on the passed in FName being the object name without any group/package qualifiers.
@@ -268,14 +283,31 @@
  *
  * @return	Returns a pointer to the found object or null if none could be found
  */
-COREUOBJECT_API UObject* StaticFindObjectFastSafe(UClass* Class, UObject* InOuter, FName InName, bool bExactClass = false, bool bAnyPackage = false, EObjectFlags ExclusiveFlags = RF_NoFlags, EInternalObjectFlags ExclusiveInternalFlags = EInternalObjectFlags::None);
+UE_DEPRECATED(5.1, "Support for searching for objects in ANY_PACKAGE has been deprecated. Please provide the actual Outer of an object you want to find.")
+COREUOBJECT_API UObject* StaticFindObjectFastSafe(UClass* Class, UObject* InOuter, FName InName, bool bExactClass, bool bAnyPackage, EObjectFlags ExclusiveFlags = RF_NoFlags, EInternalObjectFlags ExclusiveInternalFlags = EInternalObjectFlags::None);
 
 
 /**
+ * Fast and safe version of StaticFindObject that relies on the passed in FName being the object name without any group/package qualifiers.
+ * It will not assert on GIsSavingPackage or IsGarbageCollectingAndLockingUObjectHashTables(). If called from within package saving code or GC, will return nullptr
+ * This will only find top level packages or subobjects nested directly within a passed in outer.
+ *
+ * @param	Class			The to be found object's class
+ * @param	InOuter			Outer object to look inside, if null this will only look for top level packages
+ * @param	InName			Object name to look for relative to InOuter
+ * @param	bExactClass		Whether to require an exact match with the passed in class
+ * @param	ExclusiveFlags	Ignores objects that contain any of the specified exclusive flags
+ * @param	ExclusiveInternalFlags  Ignores objects that contain any of the specified internal exclusive flags
+ *
+ * @return	Returns a pointer to the found object or null if none could be found
+ */
+COREUOBJECT_API UObject* StaticFindObjectFastSafe(UClass* Class, UObject* InOuter, FName InName, bool bExactClass = false, EObjectFlags ExclusiveFlags = RF_NoFlags, EInternalObjectFlags ExclusiveInternalFlags = EInternalObjectFlags::None);
+
+/**
  * Tries to find an object in memory. This will handle fully qualified paths of the form /path/packagename.object:subobject and resolve references for you.
  *
  * @param	Class			The to be found object's class
- * @param	InOuter			Outer object to look inside. If this is ANY_PACKAGE it will search all in memory packages, if this is null then InName should start with a package name
+ * @param	InOuter			Outer object to look inside. If this is null then InName should start with a package name
  * @param	InName			The object path to search for an object, relative to InOuter
  * @param	ExactClass		Whether to require an exact match with the passed in class
  *
@@ -381,6 +413,7 @@
 /**
  * Tries to find the first object matching the search paramters in memory. This will handle fully qualified paths of the form /path/packagename.object:subobject and resolve references for you.
  * If multiple objects share the same name the returned object is random and not based on its time of creation unless otherwise specified in Options (see EFindFirstObjectOptions::NativeFirst)
+ * This function is slow and should not be used in performance critical situations.
  * 
  * @param	Class						The to be found object's class
  * @param	Name						The object path to search for an object, relative to InOuter
@@ -396,6 +429,7 @@
  * Tries to find the first objects matching the search paramters in memory. This will handle fully qualified paths of the form /path/packagename.object:subobject and resolve references for you.
  * This version of StaticFindFirstObject will not assert on GIsSavingPackage or IsGarbageCollecting()
  * If multiple objects share the same name the returned object is random and not based on its time of creation unless otherwise specified in Options (see EFindFirstObjectOptions::NativeFirst)
+ * This function is slow and should not be used in performance critical situations.
  * 
  * @param	Class						The to be found object's class
  * @param	Name						The object path to search for an object, relative to InOuter
@@ -714,6 +748,14 @@
  */
 COREUOBJECT_API void IncrementalPurgeGarbage( bool bUseTimeLimit, double TimeLimit = 0.002 );
 
+
+enum class EUniqueObjectNameOptions
+{
+	None = 0,
+	GloballyUnique = 1 << 1, // Whether to make the object name unique globally (across all objects that currently exist)
+};
+ENUM_CLASS_FLAGS(EUniqueObjectNameOptions);
+
 /**
  * Create a unique name by combining a base name and an arbitrary number string.
  * The object name returned is guaranteed not to exist.
@@ -721,11 +763,12 @@
  * @param	Parent		the outer for the object that needs to be named
  * @param	Class		the class for the object
  * @param	BaseName	optional base name to use when generating the unique object name; if not specified, the class's name is used
- *
+ * @param	Options		Additional options. See EUniqueObjectNameOptions.
+ * 
  * @return	name is the form BaseName_##, where ## is the number of objects of this
  *			type that have been created since the last time the class was garbage collected.
  */
-COREUOBJECT_API FName MakeUniqueObjectName( UObject* Outer, const UClass* Class, FName BaseName=NAME_None );
+COREUOBJECT_API FName MakeUniqueObjectName( UObject* Outer, const UClass* Class, FName BaseName = NAME_None, EUniqueObjectNameOptions Options = EUniqueObjectNameOptions::None);
 
 /**
  * Given a display label string, generates an FName slug that is a valid FName for that label.
@@ -1619,9 +1662,22 @@
  * @see StaticFindObjectFast()
  */
 template< class T > 
-inline T* FindObjectFast( UObject* Outer, FName Name, bool ExactClass=false, bool AnyPackage=false, EObjectFlags ExclusiveFlags=RF_NoFlags )
+UE_DEPRECATED(5.1, "Support for searching for objects in ANY_PACKAGE has been deprecated. Please provide the actual Outer of an object you want to find.")
+inline T* FindObjectFast( UObject* Outer, FName Name, bool ExactClass, bool AnyPackage, EObjectFlags ExclusiveFlags=RF_NoFlags )
 {
+	PRAGMA_DISABLE_DEPRECATION_WARNINGS
 	return (T*)StaticFindObjectFast( T::StaticClass(), Outer, Name, ExactClass, AnyPackage, ExclusiveFlags );
+	PRAGMA_ENABLE_DEPRECATION_WARNINGS
+}
+
+/**
+ * Find an optional object, relies on the name being unqualified
+ * @see StaticFindObjectFast()
+ */
+template< class T >
+inline T* FindObjectFast(UObject* Outer, FName Name, bool ExactClass = false, EObjectFlags ExclusiveFlags = RF_NoFlags)
+{
+	return (T*)StaticFindObjectFast(T::StaticClass(), Outer, Name, ExactClass, ExclusiveFlags);
 }
 
 /**
diff --git a/Engine/Source/Runtime/CoreUObject/Public/UObject/UObjectHash.h b/Engine/Source/Runtime/CoreUObject/Public/UObject/UObjectHash.h
--- a/Engine/Source/Runtime/CoreUObject/Public/UObject/UObjectHash.h
+++ b/Engine/Source/Runtime/CoreUObject/Public/UObject/UObjectHash.h
@@ -29,7 +29,21 @@ DECLARE_DWORD_COUNTER_STAT_EXTERN(TEXT("NumObjects"), STAT_Hash_NumObjects, STAT
  * @param	ExclusiveInternalFlags	Ignores objects that contain any of the specified internal exclusive flags
  * @return	Returns a pointer to the found object or NULL if none could be found
  */
-COREUOBJECT_API UObject* StaticFindObjectFastInternal(const UClass* Class, const UObject* InOuter, FName InName, bool ExactClass = false, bool AnyPackage = false, EObjectFlags ExclusiveFlags = RF_NoFlags, EInternalObjectFlags ExclusiveInternalFlags = EInternalObjectFlags::None);
+UE_DEPRECATED(5.1, "Support for searching for objects in ANY_PACKAGE has been deprecated. Please provide the actual Outer of an object you want to find.")
+COREUOBJECT_API UObject* StaticFindObjectFastInternal(const UClass* Class, const UObject* InOuter, FName InName, bool ExactClass, bool AnyPackage, EObjectFlags ExclusiveFlags = RF_NoFlags, EInternalObjectFlags ExclusiveInternalFlags = EInternalObjectFlags::None);
+
+/**
+ * Private internal version of StaticFindObjectFast that allows using 0 exclusion flags.
+ *
+ * @param	Class			The to be found object's class
+ * @param	InOuter			The to be found object's outer
+ * @param	InName			The to be found object's class
+ * @param	ExactClass		Whether to require an exact match with the passed in class
+ * @param	ExclusiveFlags	Ignores objects that contain any of the specified exclusive flags
+ * @param	ExclusiveInternalFlags	Ignores objects that contain any of the specified internal exclusive flags
+ * @return	Returns a pointer to the found object or NULL if none could be found
+ */
+COREUOBJECT_API UObject* StaticFindObjectFastInternal(const UClass* Class, const UObject* InOuter, FName InName, bool ExactClass = false, EObjectFlags ExclusiveFlags = RF_NoFlags, EInternalObjectFlags ExclusiveInternalFlags = EInternalObjectFlags::None);
 
 /**
  * Private internal version of StaticFindAllObjectsFast that allows using 0 exclusion flags.
diff --git a/Engine/Source/Runtime/Engine/Classes/Engine/AssetManager.h b/Engine/Source/Runtime/Engine/Classes/Engine/AssetManager.h
--- a/Engine/Source/Runtime/Engine/Classes/Engine/AssetManager.h
+++ b/Engine/Source/Runtime/Engine/Classes/Engine/AssetManager.h
@@ -382,13 +382,13 @@ public:
 	// GENERAL ASSET UTILITY FUNCTIONS
 
 	/** Parses AssetData to extract the primary type/name from it. This works even if it isn't in the directory */
-	virtual FPrimaryAssetId ExtractPrimaryAssetIdFromData(const FAssetData& AssetData, FPrimaryAssetType SuggestedType = NAME_None) const;
+	virtual FPrimaryAssetId ExtractPrimaryAssetIdFromData(const FAssetData& AssetData, FPrimaryAssetType SuggestedType = FPrimaryAssetType()) const;
 
 	/** Gets the FAssetData at a specific path, handles redirectors and blueprint classes correctly. Returns true if it found a valid data */
 	virtual bool GetAssetDataForPath(const FSoftObjectPath& ObjectPath, FAssetData& AssetData) const;
 
 	/** Checks to see if the given asset data is a blueprint with a base class in the ClassNameSet. This checks the parent asset tag */
-	virtual bool IsAssetDataBlueprintOfClassSet(const FAssetData& AssetData, const TSet<FName>& ClassNameSet) const;
+	virtual bool IsAssetDataBlueprintOfClassSet(const FAssetData& AssetData, const TSet<FTopLevelAssetPath>& ClassNameSet) const;
 
 	/** Turns an FAssetData into FSoftObjectPath, handles adding _C as necessary */
 	virtual FSoftObjectPath GetAssetPathForData(const FAssetData& AssetData) const;
diff --git a/Engine/Source/Runtime/Engine/Classes/Engine/BlueprintGeneratedClass.h b/Engine/Source/Runtime/Engine/Classes/Engine/BlueprintGeneratedClass.h
--- a/Engine/Source/Runtime/Engine/Classes/Engine/BlueprintGeneratedClass.h
+++ b/Engine/Source/Runtime/Engine/Classes/Engine/BlueprintGeneratedClass.h
@@ -739,6 +739,9 @@ public:
 	virtual void PostInitProperties() override;
 	virtual void GetPreloadDependencies(TArray<UObject*>& OutDeps) override;
 	virtual void GetAssetRegistryTags(TArray<FAssetRegistryTag>& OutTags) const override;
+#if WITH_EDITOR
+	virtual void PostLoadAssetRegistryTags(const FAssetData& InAssetData, TArray<FAssetRegistryTag>& OutTagsAndValuesToUpdate) const;
+#endif // WITH_EDITOR
 	virtual FPrimaryAssetId GetPrimaryAssetId() const override;
 	virtual bool NeedsLoadForServer() const override;
 	virtual bool NeedsLoadForClient() const override;
diff --git a/Engine/Source/Runtime/Engine/Classes/Engine/DataTable.h b/Engine/Source/Runtime/Engine/Classes/Engine/DataTable.h
--- a/Engine/Source/Runtime/Engine/Classes/Engine/DataTable.h
+++ b/Engine/Source/Runtime/Engine/Classes/Engine/DataTable.h
@@ -119,6 +119,7 @@ public:
 	
 #if WITH_EDITOR
 	ENGINE_API virtual void PostEditChangeProperty(FPropertyChangedEvent& PropertyChangedEvent) override;
+	ENGINE_API virtual void PostLoadAssetRegistryTags(const FAssetData& InAssetData, TArray<FAssetRegistryTag>& OutTagsAndValuesToUpdate) const override;
 #endif // WITH_EDITOR
 
 	//~ Begin UObject Interface.
@@ -130,7 +131,9 @@ public:
 	virtual bool NeedsLoadForClient() const override { return bStripFromClientBuilds ? false : Super::NeedsLoadForClient(); }
 	virtual bool NeedsLoadForEditorGame() const override { return bStripFromClientBuilds ? false : Super::NeedsLoadForEditorGame(); }
 #if WITH_EDITORONLY_DATA
+	UE_DEPRECATED(5.1, "Class names are now represented by path names. Please use GetRowStructPathName.")
 	ENGINE_API FName GetRowStructName() const;
+	ENGINE_API FTopLevelAssetPath GetRowStructPathName() const;
 	ENGINE_API virtual void GetAssetRegistryTags(TArray<FAssetRegistryTag>& OutTags) const override;
 	ENGINE_API virtual void PostInitProperties() override;
 	ENGINE_API virtual void PostLoad() override;
@@ -146,7 +149,11 @@ public:
 
 	/** The name of the RowStruct we were using when we were last saved */
 	UPROPERTY()
-	FName RowStructName;
+	FName RowStructName_DEPRECATED;
+
+	/** The name of the RowStruct we were using when we were last saved */
+	UPROPERTY()
+	FTopLevelAssetPath RowStructPathName;
 
 protected:
 	/** When RowStruct is being modified, row data is stored serialized with tags */
diff --git a/Engine/Source/Runtime/Engine/Private/ActorFolder.cpp b/Engine/Source/Runtime/Engine/Private/ActorFolder.cpp
--- a/Engine/Source/Runtime/Engine/Private/ActorFolder.cpp
+++ b/Engine/Source/Runtime/Engine/Private/ActorFolder.cpp
@@ -81,8 +81,8 @@ FActorFolderDesc UActorFolder::GetAssetRegistryInfoFromPackage(FName ActorFolder
 	if (Assets.Num() == 1)
 	{
 		const FAssetData& Asset = Assets[0];
-		static const FName NAME_ActorFolder(TEXT("ActorFolder"));
-		check(Asset.AssetClass == NAME_ActorFolder);
+		static const FTopLevelAssetPath NAME_ActorFolder(TEXT("/Script/Engine"), TEXT("ActorFolder"));
+		check(Asset.AssetClassPath == NAME_ActorFolder);
 		{
 			FString Value;
 			if (Asset.GetTagValue(ActorFolder::NAME_FolderGuid, Value))
diff --git a/Engine/Source/Runtime/Engine/Private/ActorFolderDesc.cpp b/Engine/Source/Runtime/Engine/Private/ActorFolderDesc.cpp
--- a/Engine/Source/Runtime/Engine/Private/ActorFolderDesc.cpp
+++ b/Engine/Source/Runtime/Engine/Private/ActorFolderDesc.cpp
@@ -60,7 +60,7 @@ FActorFolderDescsContext::FActorFolderDescsContext(const FActorFolderDesc& InAct
 	FARFilter Filter;
 	Filter.bRecursivePaths = true;
 	Filter.bIncludeOnlyOnDiskAssets = true;
-	Filter.ClassNames.Add(UActorFolder::StaticClass()->GetFName());
+	Filter.ClassPaths.Add(UActorFolder::StaticClass()->GetClassPathName());
 	Filter.PackagePaths.Add(*ActorFoldersRoot);
 	TArray<FAssetData> Assets;
 	AssetRegistry.GetAssets(Filter, Assets);
diff --git a/Engine/Source/Runtime/Engine/Private/Animation/AnimCompressionTypes.cpp b/Engine/Source/Runtime/Engine/Private/Animation/AnimCompressionTypes.cpp
--- a/Engine/Source/Runtime/Engine/Private/Animation/AnimCompressionTypes.cpp
+++ b/Engine/Source/Runtime/Engine/Private/Animation/AnimCompressionTypes.cpp
@@ -109,7 +109,7 @@ void FCompressibleAnimData::WriteCompressionDataToJSON(TArrayView<FName> Origina
 	{
 		FString JSONString;
 		TSharedRef<TJsonStringWriter<>> Writer = TJsonStringWriter<>::Create(&JSONString);
-		const UEnum* InterpolationEnum = FindObject<UEnum>(ANY_PACKAGE, TEXT("EAnimInterpolationType"), true);
+		const UEnum* InterpolationEnum = FindObject<UEnum>(nullptr, TEXT("/Script/Engine.EAnimInterpolationType"), true);
 
 		Writer->WriteObjectStart();
 		{
@@ -259,9 +259,9 @@ void FCompressibleAnimData::WriteCompressionDataToJSON(TArrayView<FName> Origina
 							
 						if(FloatCurve.FloatCurve.GetConstRefOfKeys().Num())
 						{
-							const UEnum* CurveInterpolationEnum = FindObject<UEnum>(ANY_PACKAGE, TEXT("ERichCurveInterpMode"), true);
-							const UEnum* TangentModeEnum = FindObject<UEnum>(ANY_PACKAGE, TEXT("ERichCurveTangentMode"), true);
-							const UEnum* TangentWeightModeEnum = FindObject<UEnum>(ANY_PACKAGE, TEXT("ERichCurveTangentWeightMode"), true);
+							const UEnum* CurveInterpolationEnum = FindObject<UEnum>(nullptr, TEXT("/Script/Engine.ERichCurveInterpMode"), true);
+							const UEnum* TangentModeEnum = FindObject<UEnum>(nullptr, TEXT("/Script/Engine.ERichCurveTangentMode"), true);
+							const UEnum* TangentWeightModeEnum = FindObject<UEnum>(nullptr, TEXT("/Script/Engine.ERichCurveTangentWeightMode"), true);
 								
 							Writer->WriteArrayStart(TEXT("key_data"));
 							for (const FRichCurveKey& Key : FloatCurve.FloatCurve.GetConstRefOfKeys())
diff --git a/Engine/Source/Runtime/Engine/Private/Animation/Skeleton.cpp b/Engine/Source/Runtime/Engine/Private/Animation/Skeleton.cpp
--- a/Engine/Source/Runtime/Engine/Private/Animation/Skeleton.cpp
+++ b/Engine/Source/Runtime/Engine/Private/Animation/Skeleton.cpp
@@ -148,7 +148,7 @@ bool USkeleton::IsCompatibleSkeletonByAssetString(const FString& SkeletonAssetSt
 	FString CompatibleName;
 	for (const auto& CompatibleSkeleton : CompatibleSkeletons)
 	{
-		CompatibleName = FString::Printf(TEXT("%s'%s'"), *GetClass()->GetName(), *CompatibleSkeleton.ToString());
+		CompatibleName = FObjectPropertyBase::GetExportPath(CompatibleSkeleton.Get());
 		if (CompatibleName == SkeletonAssetString)
 		{
 			return true;
@@ -1516,7 +1516,7 @@ void USkeleton::CollectAnimationNotifies(TArray<FName>& OutNotifies) const
 	// meanwhile if you remove this, this will miss the links
 	//AnimationNotifies.Empty();
 	TArray<FAssetData> AssetList;
-	AssetRegistryModule.Get().GetAssetsByClass(UAnimSequenceBase::StaticClass()->GetFName(), AssetList, true);
+	AssetRegistryModule.Get().GetAssetsByClass(UAnimSequenceBase::StaticClass()->GetClassPathName(), AssetList, true);
 
 	// do not clear AnimationNotifies. We can't remove old ones yet. 
 	FString CurrentSkeletonName = FAssetData(this).GetExportTextName();
@@ -1552,7 +1552,7 @@ void USkeleton::AddNewAnimationNotify(FName NewAnimNotifyName)
 USkeletalMesh* USkeleton::FindCompatibleMesh() const
 {
 	FARFilter Filter;
-	Filter.ClassNames.Add(USkeletalMesh::StaticClass()->GetFName());
+	Filter.ClassPaths.Add(USkeletalMesh::StaticClass()->GetClassPathName());
 
 	FString SkeletonString = FAssetData(this).GetExportTextName();
 	
diff --git a/Engine/Source/Runtime/Engine/Private/AssetManager.cpp b/Engine/Source/Runtime/Engine/Private/AssetManager.cpp
--- a/Engine/Source/Runtime/Engine/Private/AssetManager.cpp
+++ b/Engine/Source/Runtime/Engine/Private/AssetManager.cpp
@@ -143,11 +143,11 @@ struct FCompiledAssetManagerSearchRules : FAssetManagerSearchRules
 		// Check class first
 		if (AssetBaseClass)
 		{
-			AssetClassNames.Add(AssetBaseClass->GetFName());
+			AssetClassNames.Add(AssetBaseClass->GetClassPathName());
 
 #if WITH_EDITOR
 			// Add any old names to the list in case things haven't been resaved
-			TArray<FName> OldNames = FLinkerLoad::FindPreviousNamesForClass(AssetBaseClass->GetPathName(), false);
+			TArray<FString> OldNames = FLinkerLoad::FindPreviousPathNamesForClass(AssetBaseClass->GetPathName(), false);
 			AssetClassNames.Append(OldNames);
 #endif
 		}
@@ -203,8 +203,8 @@ struct FCompiledAssetManagerSearchRules : FAssetManagerSearchRules
 
 	TArray<FString> IncludeWildcards;
 	TArray<FString> ExcludeWildcards;
-	TArray<FName> AssetClassNames;
-	TSet<FName> DerivedClassNames;
+	TArray<FTopLevelAssetPath> AssetClassNames;
+	TSet<FTopLevelAssetPath> DerivedClassNames;
 	bool bShouldCallDelegate;
 	bool bShouldCheckWildcards;
 };
@@ -499,25 +499,26 @@ FPrimaryAssetId UAssetManager::DeterminePrimaryAssetIdForObject(const UObject* O
 	return FPrimaryAssetId();
 }
 
-bool UAssetManager::IsAssetDataBlueprintOfClassSet(const FAssetData& AssetData, const TSet<FName>& ClassNameSet) const
+bool UAssetManager::IsAssetDataBlueprintOfClassSet(const FAssetData& AssetData, const TSet<FTopLevelAssetPath>& ClassNameSet) const
 {
 	const FString ParentClassFromData = AssetData.GetTagValueRef<FString>(FBlueprintTags::ParentClassPath);
 	if (!ParentClassFromData.IsEmpty())
 	{
-		const FString ClassObjectPath = FPackageName::ExportTextPathToObjectPath(ParentClassFromData);
-		const FName ClassName = FName(*FPackageName::ObjectPathToObjectName(ClassObjectPath));
+		const FTopLevelAssetPath ClassObjectPath(FPackageName::ExportTextPathToObjectPath(ParentClassFromData));
+		const FName ClassName = ClassObjectPath.GetAssetName();
 
-		TArray<FName> ValidNames;
-		ValidNames.Add(ClassName);
+		TArray<FTopLevelAssetPath> ValidNames;
+		ValidNames.Add(ClassObjectPath);
 #if WITH_EDITOR
 		// Check for redirected name
-		FName RedirectedName = FLinkerLoad::FindNewNameForClass(ClassName, false);
-		if (RedirectedName != NAME_None && RedirectedName != ClassName)
+		FTopLevelAssetPath RedirectedName = FTopLevelAssetPath(FLinkerLoad::FindNewPathNameForClass(ClassObjectPath.ToString(), false));
+		if (!RedirectedName.IsNull())
 		{
+			// @todo: we need a full path
 			ValidNames.Add(RedirectedName);
 		}
 #endif
-		for (const FName& ValidName : ValidNames)
+		for (const FTopLevelAssetPath& ValidName : ValidNames)
 		{
 			if (ClassNameSet.Contains(ValidName))
 			{
@@ -582,7 +583,7 @@ int32 UAssetManager::SearchAssetRegistryPaths(TArray<FAssetData>& OutAssetDataLi
 		if (!Rules.bHasBlueprintClasses)
 		{
 			// Use class directly
-			ARFilter.ClassNames = CompiledRules.AssetClassNames;
+			ARFilter.ClassPaths = CompiledRules.AssetClassNames;
 			ARFilter.bRecursiveClasses = true;
 		}
 		else
@@ -592,10 +593,10 @@ int32 UAssetManager::SearchAssetRegistryPaths(TArray<FAssetData>& OutAssetDataLi
 			GetDerivedClasses(UBlueprintCore::StaticClass(), BlueprintCoreDerivedClasses);
 			for (UClass* BPCoreClass : BlueprintCoreDerivedClasses)
 			{
-				ARFilter.ClassNames.Add(BPCoreClass->GetFName());
+				ARFilter.ClassPaths.Add(BPCoreClass->GetClassPathName());
 			}
 
-			GetAssetRegistry().GetDerivedClassNames(CompiledRules.AssetClassNames, TSet<FName>(), CompiledRules.DerivedClassNames);
+			GetAssetRegistry().GetDerivedClassNames(CompiledRules.AssetClassNames, TSet<FTopLevelAssetPath>(), CompiledRules.DerivedClassNames);
 		}
 	}
 
@@ -665,11 +666,11 @@ bool UAssetManager::DoesAssetMatchSearchRules(const FAssetData& AssetData, const
 	// Check class first
 	if (Rules.AssetBaseClass)
 	{
-		GetAssetRegistry().GetDerivedClassNames(CompiledRules.AssetClassNames, TSet<FName>(), CompiledRules.DerivedClassNames);
+		GetAssetRegistry().GetDerivedClassNames(CompiledRules.AssetClassNames, TSet<FTopLevelAssetPath>(), CompiledRules.DerivedClassNames);
 
 		if (!Rules.bHasBlueprintClasses)
 		{
-			if (!CompiledRules.DerivedClassNames.Contains(AssetData.AssetClass))
+			if (!CompiledRules.DerivedClassNames.Contains(AssetData.AssetClassPath))
 			{
 				return false;
 			}
@@ -1470,7 +1471,7 @@ FPrimaryAssetId UAssetManager::ExtractPrimaryAssetIdFromData(const FAssetData& A
 {
 	FPrimaryAssetId FoundId = AssetData.GetPrimaryAssetId();
 
-	if (!FoundId.IsValid() && bShouldGuessTypeAndName && SuggestedType != NAME_None)
+	if (!FoundId.IsValid() && bShouldGuessTypeAndName && SuggestedType.IsValid())
 	{
 		const TSharedRef<FPrimaryAssetTypeData>* FoundType = AssetTypeMap.Find(SuggestedType);
 
@@ -2732,17 +2733,10 @@ bool UAssetManager::GetAssetDataForPath(const FSoftObjectPath& ObjectPath, FAsse
 	return AssetData.IsValid();
 }
 
-static bool EndsWithBlueprint(FName Name)
+static bool EndsWithBlueprint(const FTopLevelAssetPath& Name)
 {
 	// Numbered names can't end with Blueprint
-	if (Name.IsNone() || Name.GetNumber() != FName().GetNumber())
-	{
-		return false;
-	}
-
-	TCHAR Buffer[NAME_SIZE];
-	FStringView PlainName(Buffer, Name.GetPlainNameString(Buffer));
-	return PlainName.EndsWith(TEXT("Blueprint"));
+	return Name.GetAssetName().ToString().EndsWith(TEXT("Blueprint"));
 }
 
 static bool ContainsSubobjectDelimiter(FName Name)
@@ -2760,7 +2754,7 @@ bool TryToSoftObjectPath(const FAssetData& AssetData, FSoftObjectPath& OutObject
 		OutObjectPath = FSoftObjectPath();
 		return false;
 	}
-	else if (EndsWithBlueprint(AssetData.AssetClass))
+	else if (EndsWithBlueprint(AssetData.AssetClassPath))
 	{
 		TStringBuilder<256> AssetPath;
 		AssetPath << AssetData.ObjectPath << TEXT("_C");
diff --git a/Engine/Source/Runtime/Engine/Private/Audio/AudioWidgetSubsystem.cpp b/Engine/Source/Runtime/Engine/Private/Audio/AudioWidgetSubsystem.cpp
--- a/Engine/Source/Runtime/Engine/Private/Audio/AudioWidgetSubsystem.cpp
+++ b/Engine/Source/Runtime/Engine/Private/Audio/AudioWidgetSubsystem.cpp
@@ -64,8 +64,8 @@ const TArray<FAssetData> UAudioWidgetSubsystem::GetBlueprintAssetData(UClass* In
 	TArray<FAssetData> AssetData;
 
 	FARFilter Filter;
-	Filter.ClassNames.Add(UBlueprint::StaticClass()->GetFName());
-	Filter.ClassNames.Add("WidgetBlueprint");
+	Filter.ClassPaths.Add(UBlueprint::StaticClass()->GetClassPathName());
+	Filter.ClassPaths.Add(FTopLevelAssetPath(TEXT("/Script/UMGEditor"), TEXT("WidgetBlueprint")));
 	Filter.TagsAndValues.Add(FBlueprintTags::ImplementedInterfaces);
 
 	AssetRegistryModule.Get().EnumerateAssets(Filter, [FoundAssets = &AssetData, InInterfaceClass](const FAssetData& InAssetData)
diff --git a/Engine/Source/Runtime/Engine/Private/AudioDevice.cpp b/Engine/Source/Runtime/Engine/Private/AudioDevice.cpp
--- a/Engine/Source/Runtime/Engine/Private/AudioDevice.cpp
+++ b/Engine/Source/Runtime/Engine/Private/AudioDevice.cpp
@@ -1541,7 +1541,7 @@ bool FAudioDevice::HandleSoundClassFixup(const TCHAR* Cmd, FOutputDevice& Ar)
 	FAssetRegistryModule& AssetRegistryModule = FModuleManager::LoadModuleChecked<FAssetRegistryModule>(TEXT("AssetRegistry"));
 
 	TArray<FAssetData> AssetDataArray;
-	AssetRegistryModule.Get().GetAssetsByClass(USoundClass::StaticClass()->GetFName(), AssetDataArray);
+	AssetRegistryModule.Get().GetAssetsByClass(USoundClass::StaticClass()->GetClassPathName(), AssetDataArray);
 
 	static const FString EngineDir = TEXT("/Engine/");
 	FAssetToolsModule& AssetToolsModule = FModuleManager::GetModuleChecked<FAssetToolsModule>(TEXT("AssetTools"));
@@ -1681,7 +1681,7 @@ bool FAudioDevice::HandleAudioMemoryInfo(const TCHAR* Cmd, FOutputDevice& Ar)
 
 	// Get the sound wave class
 	UClass* SoundWaveClass = nullptr;
-	ParseObject<UClass>(TEXT("class=SoundWave"), TEXT("CLASS="), SoundWaveClass, ANY_PACKAGE);
+	ParseObject<UClass>(TEXT("class=SoundWave"), TEXT("CLASS="), SoundWaveClass, nullptr);
 
 	TArray<FSoundWaveInfo> SoundWaveObjects;
 	TMap<FString, FSoundWaveGroupInfo> SoundWaveGroupSizes;
diff --git a/Engine/Source/Runtime/Engine/Private/Blueprint.cpp b/Engine/Source/Runtime/Engine/Private/Blueprint.cpp
--- a/Engine/Source/Runtime/Engine/Private/Blueprint.cpp
+++ b/Engine/Source/Runtime/Engine/Private/Blueprint.cpp
@@ -905,7 +905,7 @@ void UBlueprint::GetAssetRegistryTags(TArray<FAssetRegistryTag>& OutTags) const
 	FString GeneratedClassVal;
 	if (GeneratedClass)
 	{
-		GeneratedClassVal = FString::Printf(TEXT("%s'%s'"), *GeneratedClass->GetClass()->GetName(), *GeneratedClass->GetPathName());
+		GeneratedClassVal = FObjectPropertyBase::GetExportPath(GeneratedClass);
 	}
 	else
 	{
@@ -915,7 +915,7 @@ void UBlueprint::GetAssetRegistryTags(TArray<FAssetRegistryTag>& OutTags) const
 	FString NativeParentClassName, ParentClassName;
 	if ( ParentClass )
 	{
-		ParentClassName = FString::Printf(TEXT("%s'%s'"), *ParentClass->GetClass()->GetName(), *ParentClass->GetPathName());
+		ParentClassName = FObjectPropertyBase::GetExportPath(ParentClass);
 
 		// Walk up until we find a native class (ie 'while they are BP classes')
 		UClass* NativeParentClass = ParentClass;
@@ -923,7 +923,7 @@ void UBlueprint::GetAssetRegistryTags(TArray<FAssetRegistryTag>& OutTags) const
 		{
 			NativeParentClass = NativeParentClass->GetSuperClass();
 		}
-		NativeParentClassName = FString::Printf(TEXT("%s'%s'"), *NativeParentClass->GetClass()->GetName(), *NativeParentClass->GetPathName());
+		NativeParentClassName = FObjectPropertyBase::GetExportPath(NativeParentClass);
 	}
 	else
 	{
@@ -1012,6 +1012,27 @@ void UBlueprint::GetAssetRegistryTags(TArray<FAssetRegistryTag>& OutTags) const
 	}
 }
 
+void UBlueprint::PostLoadAssetRegistryTags(const FAssetData& InAssetData, TArray<FAssetRegistryTag>& OutTagsAndValuesToUpdate) const
+{
+	Super::PostLoadAssetRegistryTags(InAssetData, OutTagsAndValuesToUpdate);
+
+	auto FixTagValueShortClassName = [&InAssetData, &OutTagsAndValuesToUpdate](FName TagName, FAssetRegistryTag::ETagType TagType)
+	{
+		FString TagValue = InAssetData.GetTagValueRef<FString>(TagName);
+		if (!TagValue.IsEmpty() && TagValue != TEXT("None"))
+		{
+			if (UClass::TryFixShortClassNameExportPath(TagValue, ELogVerbosity::Warning, TEXT("UBlueprint::PostLoadAssetRegistryTags")))
+			{
+				OutTagsAndValuesToUpdate.Add(FAssetRegistryTag(TagName, TagValue, TagType));
+			}
+		}
+	};
+
+	FixTagValueShortClassName(FBlueprintTags::GeneratedClassPath, FAssetRegistryTag::TT_Hidden);
+	FixTagValueShortClassName(FBlueprintTags::ParentClassPath, FAssetRegistryTag::TT_Alphabetical);
+	FixTagValueShortClassName(FBlueprintTags::NativeParentClassPath, FAssetRegistryTag::TT_Alphabetical);
+}
+
 FPrimaryAssetId UBlueprint::GetPrimaryAssetId() const
 {
 	// Forward to our Class, which will forward to CDO if needed
diff --git a/Engine/Source/Runtime/Engine/Private/BlueprintGeneratedClass.cpp b/Engine/Source/Runtime/Engine/Private/BlueprintGeneratedClass.cpp
--- a/Engine/Source/Runtime/Engine/Private/BlueprintGeneratedClass.cpp
+++ b/Engine/Source/Runtime/Engine/Private/BlueprintGeneratedClass.cpp
@@ -236,7 +236,7 @@ void UBlueprintGeneratedClass::GetAssetRegistryTags(TArray<FAssetRegistryTag>& O
 
 	if (UClass* ParentClass = GetSuperClass())
 	{
-		ParentClassName = FString::Printf(TEXT("%s'%s'"), *ParentClass->GetClass()->GetName(), *ParentClass->GetPathName());
+		ParentClassName = FObjectPropertyBase::GetExportPath(ParentClass);
 
 		// Walk up until we find a native class (ie 'while they are BP classes')
 		UClass* NativeParentClass = ParentClass;
@@ -244,7 +244,7 @@ void UBlueprintGeneratedClass::GetAssetRegistryTags(TArray<FAssetRegistryTag>& O
 		{
 			NativeParentClass = NativeParentClass->GetSuperClass();
 		}
-		NativeParentClassName = FString::Printf(TEXT("%s'%s'"), *NativeParentClass->GetClass()->GetName(), *NativeParentClass->GetPathName());
+		NativeParentClassName = FObjectPropertyBase::GetExportPath(NativeParentClass);
 	}
 	else
 	{
@@ -272,6 +272,28 @@ void UBlueprintGeneratedClass::GetAssetRegistryTags(TArray<FAssetRegistryTag>& O
 #endif //#if WITH_EDITORONLY_DATA
 }
 
+#if WITH_EDITOR
+void UBlueprintGeneratedClass::PostLoadAssetRegistryTags(const FAssetData& InAssetData, TArray<FAssetRegistryTag>& OutTagsAndValuesToUpdate) const
+{
+	Super::PostLoadAssetRegistryTags(InAssetData, OutTagsAndValuesToUpdate);
+
+	auto FixTagValueShortClassName = [&InAssetData, &OutTagsAndValuesToUpdate](FName TagName, FAssetRegistryTag::ETagType TagType)
+	{
+		FString TagValue = InAssetData.GetTagValueRef<FString>(TagName);
+		if (!TagValue.IsEmpty() && TagValue != TEXT("None"))
+		{
+			if (UClass::TryFixShortClassNameExportPath(TagValue, ELogVerbosity::Warning, TEXT("UBlueprintGeneratedClass::PostLoadAssetRegistryTags")))
+			{
+				OutTagsAndValuesToUpdate.Add(FAssetRegistryTag(TagName, TagValue, TagType));
+			}
+		}
+	};
+
+	FixTagValueShortClassName(FBlueprintTags::ParentClassPath, FAssetRegistryTag::TT_Alphabetical);
+	FixTagValueShortClassName(FBlueprintTags::NativeParentClassPath, FAssetRegistryTag::TT_Alphabetical);
+}
+#endif // WITH_EDITOR
+
 FPrimaryAssetId UBlueprintGeneratedClass::GetPrimaryAssetId() const
 {
 	FPrimaryAssetId AssetId;
diff --git a/Engine/Source/Runtime/Engine/Private/CheatManager.cpp b/Engine/Source/Runtime/Engine/Private/CheatManager.cpp
--- a/Engine/Source/Runtime/Engine/Private/CheatManager.cpp
+++ b/Engine/Source/Runtime/Engine/Private/CheatManager.cpp
@@ -422,15 +422,7 @@ void UCheatManager::Summon( const FString& ClassName )
 	bool bSpawnedActor = false;
 	if ( bIsValidClassName )
 	{
-		UClass* NewClass = NULL;
-		if ( FPackageName::IsShortPackageName(ClassName) )
-		{
-			NewClass = FindObject<UClass>(ANY_PACKAGE, *ClassName);
-		}
-		else
-		{
-			NewClass = FindObject<UClass>(NULL, *ClassName);
-		}
+		UClass* NewClass = UClass::TryFindTypeSlow<UClass>(ClassName);
 
 		if( NewClass )
 		{
@@ -1342,19 +1334,27 @@ void UCheatManager::DestroyServerStatReplicator()
 
 void UCheatManager::ToggleServerStatReplicatorClientOverwrite()
 {
-	AServerStatReplicator* ServerStatReplicator = FindObject<AServerStatReplicator>(ANY_PACKAGE, TEXT("ServerStatReplicatorInst"));
-	if (ServerStatReplicator != nullptr)
+	APlayerController* PlayerController = GetOuterAPlayerController();
+	if (ensure(PlayerController))
 	{
-		ServerStatReplicator->bOverwriteClientStats = !ServerStatReplicator->bOverwriteClientStats;
+		AServerStatReplicator* ServerStatReplicator = FindObject<AServerStatReplicator>(PlayerController->GetLevel(), TEXT("ServerStatReplicatorInst"));
+		if (ServerStatReplicator != nullptr)
+		{
+			ServerStatReplicator->bOverwriteClientStats = !ServerStatReplicator->bOverwriteClientStats;
+		}
 	}
 }
 
 void UCheatManager::ToggleServerStatReplicatorUpdateStatNet()
 {
-	AServerStatReplicator* ServerStatReplicator = FindObject<AServerStatReplicator>(ANY_PACKAGE, TEXT("ServerStatReplicatorInst"));
-	if (ServerStatReplicator != nullptr)
+	APlayerController* PlayerController = GetOuterAPlayerController();
+	if (ensure(PlayerController))
 	{
-		ServerStatReplicator->bUpdateStatNet = !ServerStatReplicator->bUpdateStatNet;
+		AServerStatReplicator* ServerStatReplicator = FindObject<AServerStatReplicator>(PlayerController->GetLevel(), TEXT("ServerStatReplicatorInst"));
+		if (ServerStatReplicator != nullptr)
+		{
+			ServerStatReplicator->bUpdateStatNet = !ServerStatReplicator->bUpdateStatNet;
+		}
 	}
 }
 
diff --git a/Engine/Source/Runtime/Engine/Private/DataTable.cpp b/Engine/Source/Runtime/Engine/Private/DataTable.cpp
--- a/Engine/Source/Runtime/Engine/Private/DataTable.cpp
+++ b/Engine/Source/Runtime/Engine/Private/DataTable.cpp
@@ -229,7 +229,8 @@ void UDataTable::Serialize(FStructuredArchiveRecord Record)
 	// Make sure and update RowStructName before calling the parent Serialize (which will save the properties)
 	if (BaseArchive.IsSaving() && RowStruct)
 	{
-		RowStructName = RowStruct->GetFName();
+		RowStructName_DEPRECATED = RowStruct->GetFName();
+		RowStructPathName = RowStruct->GetStructPathName();
 	}
 #endif	// WITH_EDITORONLY_DATA
 
@@ -303,7 +304,12 @@ void UDataTable::FinishDestroy()
 #if WITH_EDITORONLY_DATA
 FName UDataTable::GetRowStructName() const
 {
-	return (RowStruct) ? RowStruct->GetFName() : RowStructName;
+	return (RowStruct) ? RowStruct->GetFName() : RowStructName_DEPRECATED;
+}
+
+FTopLevelAssetPath UDataTable::GetRowStructPathName() const
+{
+	return (RowStruct) ? RowStruct->GetStructPathName() : RowStructPathName;
 }
 
 void UDataTable::GetAssetRegistryTags(TArray<FAssetRegistryTag>& OutTags) const
@@ -316,7 +322,7 @@ void UDataTable::GetAssetRegistryTags(TArray<FAssetRegistryTag>& OutTags) const
 	// Add the row structure tag
 	{
 		static const FName RowStructureTag = "RowStructure";
-		OutTags.Add( FAssetRegistryTag(RowStructureTag, GetRowStructName().ToString(), FAssetRegistryTag::TT_Alphabetical) );
+		OutTags.Add( FAssetRegistryTag(RowStructureTag, GetRowStructPathName().ToString(), FAssetRegistryTag::TT_Alphabetical) );
 	}
 
 	Super::GetAssetRegistryTags(OutTags);
@@ -341,9 +347,43 @@ void UDataTable::PostLoad()
 		Info.Insert(FAssetImportInfo::FSourceFile(ImportPath_DEPRECATED));
 		AssetImportData->SourceData = MoveTemp(Info);
 	}
+	if (!RowStructName_DEPRECATED.IsNone())
+	{
+		UStruct* SavedRowStruct = RowStruct;
+		if (!SavedRowStruct)
+		{
+			SavedRowStruct = FindFirstObjectSafe<UStruct>(*RowStructName_DEPRECATED.ToString());
+		}
+		if (SavedRowStruct)
+		{
+			RowStructPathName = SavedRowStruct->GetStructPathName();
+		}
+		else
+		{
+			UE_LOG(LogDataTable, Error, TEXT("Unable to resolved RowStruct PathName from serialized short name '%s'!"), *RowStructName_DEPRECATED.ToString());
+		}
+	}
 }
 #endif // WITH_EDITORONLY_DATA
 
+#if WITH_EDITOR
+void UDataTable::PostLoadAssetRegistryTags(const FAssetData& InAssetData, TArray<FAssetRegistryTag>& OutTagsAndValuesToUpdate) const
+{
+	Super::PostLoadAssetRegistryTags(InAssetData, OutTagsAndValuesToUpdate);
+
+	static const FName RowStructureTag(TEXT("RowStructure"));
+	FString TagValue = InAssetData.GetTagValueRef<FString>(RowStructureTag);
+	if (!TagValue.IsEmpty() && FPackageName::IsShortPackageName(TagValue))
+	{
+		FTopLevelAssetPath PathName = UClass::TryConvertShortTypeNameToPathName<UField>(TagValue, ELogVerbosity::Warning, TEXT("UDataTable::PostLoadAssetRegistryTags"));
+		if (!PathName.IsNull())
+		{
+			OutTagsAndValuesToUpdate.Add(FAssetRegistryTag(RowStructureTag, PathName.ToString(), FAssetRegistryTag::TT_Alphabetical));
+		}
+	}
+}
+#endif // WITH_EDITOR
+
 UScriptStruct& UDataTable::GetEmptyUsingStruct() const
 {
 	UScriptStruct* EmptyUsingStruct = RowStruct;
@@ -602,7 +642,8 @@ bool UDataTable::CopyImportOptions(UDataTable* SourceTable)
 
 	if (RowStruct)
 	{
-		RowStructName = RowStruct->GetFName();
+		RowStructName_DEPRECATED = RowStruct->GetFName();
+		RowStructPathName = RowStruct->GetStructPathName();
 	}
 
 	if (SourceTable->AssetImportData)
diff --git a/Engine/Source/Runtime/Engine/Private/EngineGlobals.cpp b/Engine/Source/Runtime/Engine/Private/EngineGlobals.cpp
--- a/Engine/Source/Runtime/Engine/Private/EngineGlobals.cpp
+++ b/Engine/Source/Runtime/Engine/Private/EngineGlobals.cpp
@@ -153,7 +153,7 @@ void ReattachComponents(const TArray<FString>& Args)
 	UE_LOG(LogConsoleResponse, Display, TEXT("Reattach.Components:"));
 
 	UClass* Class=NULL;
-	if( ParseObject<UClass>( *Args[0], TEXT("CLASS="), Class, ANY_PACKAGE ) &&
+	if( ParseObject<UClass>( *Args[0], TEXT("CLASS="), Class, nullptr ) &&
 		Class->IsChildOf(UActorComponent::StaticClass()) )
 	{
 		for( FThreadSafeObjectIterator It(Class); It; ++It )
diff --git a/Engine/Source/Runtime/Engine/Private/EngineUtils.cpp b/Engine/Source/Runtime/Engine/Private/EngineUtils.cpp
--- a/Engine/Source/Runtime/Engine/Private/EngineUtils.cpp
+++ b/Engine/Source/Runtime/Engine/Private/EngineUtils.cpp
@@ -86,13 +86,13 @@ bool FContentComparisonHelper::CompareClasses(const FString& InBaseClassName, co
 {
 	TMap<FString,TArray<FContentComparisonAssetInfo> > ClassToAssetsMap;
 
-	UClass* TheClass = (UClass*)StaticFindObject(UClass::StaticClass(), ANY_PACKAGE, *InBaseClassName, true);
+	UClass* TheClass = FindFirstObject<UClass>(*InBaseClassName, EFindFirstObjectOptions::ExactClass, ELogVerbosity::Warning, TEXT("FContentComparisonHelper::CompareClasses"));
 	if (TheClass != NULL)
 	{
 		TArray<UClass*> IgnoreBaseClasses;
 		for (int32 IgnoreIdx = 0; IgnoreIdx < InBaseClassesToIgnore.Num(); IgnoreIdx++)
 		{
-			UClass* IgnoreClass = (UClass*)StaticFindObject(UClass::StaticClass(), ANY_PACKAGE, *(InBaseClassesToIgnore[IgnoreIdx]), true);
+			UClass* IgnoreClass = FindFirstObject<UClass>(*(InBaseClassesToIgnore[IgnoreIdx]), EFindFirstObjectOptions::ExactClass, ELogVerbosity::Warning, TEXT("FContentComparisonHelper::CompareClasses"));
 			if (IgnoreClass != NULL)
 			{
 				IgnoreBaseClasses.Add(IgnoreClass);
diff --git a/Engine/Source/Runtime/Engine/Private/GameViewportClient.cpp b/Engine/Source/Runtime/Engine/Private/GameViewportClient.cpp
--- a/Engine/Source/Runtime/Engine/Private/GameViewportClient.cpp
+++ b/Engine/Source/Runtime/Engine/Private/GameViewportClient.cpp
@@ -3778,7 +3778,7 @@
 	if ( FParse::Token(Cmd, ObjectName, UE_ARRAY_COUNT(ObjectName), true) &&
 		FParse::Token(Cmd, PropStr, UE_ARRAY_COUNT(PropStr), true) )
 	{
-		UObject* Obj = FindObject<UObject>(ANY_PACKAGE, ObjectName);
+		UObject* Obj = FindFirstObject<UObject>(ObjectName, EFindFirstObjectOptions::NativeFirst, ELogVerbosity::Warning, TEXT("HandleDisplayCommand"));
 		if (Obj != nullptr)
 		{
 			FName PropertyName(PropStr, FNAME_Find);
@@ -3814,7 +3814,7 @@
 			if (DotIndex != INDEX_NONE)
 			{
 				// first part is within class
-				WithinClass = FindObject<UClass>(ANY_PACKAGE, *ClassStr.Left(DotIndex));
+				WithinClass = FindFirstObject<UClass>(*ClassStr.Left(DotIndex), EFindFirstObjectOptions::None, ELogVerbosity::Warning, TEXT("HandleDisplayAllCommand"));
 				if (WithinClass == nullptr)
 				{
 					Ar.Logf(TEXT("Within class not found"));
@@ -3830,7 +3830,7 @@
 		if (bValidClassToken)
 		{
 			FParse::Token(Cmd, PropStr, UE_ARRAY_COUNT(PropStr), true);
-			UClass* Cls = FindObject<UClass>(ANY_PACKAGE, ClassName);
+			UClass* Cls = FindFirstObject<UClass>(ClassName, EFindFirstObjectOptions::None, ELogVerbosity::Warning, TEXT("HandleDisplayAllCommand"));
 			if (Cls != nullptr)
 			{
 				FName PropertyName(PropStr, FNAME_Find);
@@ -3867,7 +3867,7 @@
 	TCHAR ClassName[256];
 	if (FParse::Token(Cmd, ClassName, UE_ARRAY_COUNT(ClassName), true))
 	{
-		UClass* Cls = FindObject<UClass>(ANY_PACKAGE, ClassName);
+		UClass* Cls = FindFirstObject<UClass>(ClassName, EFindFirstObjectOptions::None, ELogVerbosity::Warning, TEXT("HandleDisplayAllLocationCommand"));
 		if (Cls != nullptr)
 		{
 			// add all un-GCable things immediately as that list is static
@@ -3899,7 +3899,7 @@
 	TCHAR ClassName[256];
 	if (FParse::Token(Cmd, ClassName, UE_ARRAY_COUNT(ClassName), true))
 	{
-		UClass* Cls = FindObject<UClass>(ANY_PACKAGE, ClassName);
+		UClass* Cls = FindFirstObject<UClass>(ClassName, EFindFirstObjectOptions::None, ELogVerbosity::Warning, TEXT("HandleDisplayAllRotationCommand"));
 		if (Cls != nullptr)
 		{
 			// add all un-GCable things immediately as that list is static
@@ -3940,7 +3940,7 @@
 	UClass* Class;
 
 	if (FParse::Token(Cmd, ClassName, UE_ARRAY_COUNT(ClassName), 1) &&
-		(Class = FindObject<UClass>(ANY_PACKAGE, ClassName)) != NULL)
+		(Class = FindFirstObject<UClass>(ClassName, EFindFirstObjectOptions::None, ELogVerbosity::Warning, TEXT("HandleGetAllLocationCommand"))) != nullptr)
 	{
 		bool bShowPendingKills = FParse::Command(&Cmd, TEXT("SHOWPENDINGKILLS"));
 		int32 cnt = 0;
@@ -3968,7 +3968,7 @@
 	UClass* Class;
 
 	if (FParse::Token(Cmd, ClassName, UE_ARRAY_COUNT(ClassName), 1) &&
-		(Class = FindObject<UClass>(ANY_PACKAGE, ClassName)) != NULL)
+		(Class = FindFirstObject<UClass>(ClassName, EFindFirstObjectOptions::None, ELogVerbosity::Warning, TEXT("HandleGetAllRotationCommand"))) != nullptr)
 	{
 		bool bShowPendingKills = FParse::Command(&Cmd, TEXT("SHOWPENDINGKILLS"));
 		int32 cnt = 0;
diff --git a/Engine/Source/Runtime/Engine/Private/GameplayStatics.cpp b/Engine/Source/Runtime/Engine/Private/GameplayStatics.cpp
--- a/Engine/Source/Runtime/Engine/Private/GameplayStatics.cpp
+++ b/Engine/Source/Runtime/Engine/Private/GameplayStatics.cpp
@@ -2345,7 +2345,7 @@ USaveGame* UGameplayStatics::LoadGameFromMemory(const TArray<uint8>& InSaveData)
 	SaveHeader.Read(MemoryReader);
 
 	// Try and find it, and failing that, load it
-	UClass* SaveGameClass = FindObject<UClass>(ANY_PACKAGE, *SaveHeader.SaveGameClassName);
+	UClass* SaveGameClass = UClass::TryFindTypeSlow<UClass>(SaveHeader.SaveGameClassName);
 	if (SaveGameClass == nullptr)
 	{
 		SaveGameClass = LoadObject<UClass>(nullptr, *SaveHeader.SaveGameClassName);
diff --git a/Engine/Source/Runtime/Engine/Private/Internationalization/EnginePackageLocalizationCache.cpp b/Engine/Source/Runtime/Engine/Private/Internationalization/EnginePackageLocalizationCache.cpp
--- a/Engine/Source/Runtime/Engine/Private/Internationalization/EnginePackageLocalizationCache.cpp
+++ b/Engine/Source/Runtime/Engine/Private/Internationalization/EnginePackageLocalizationCache.cpp
@@ -84,7 +84,7 @@ void FEnginePackageLocalizationCache::FindLocalizedPackages(const TMap<FString,
 	}
 }
 
-void FEnginePackageLocalizationCache::FindAssetGroupPackages(const FName InAssetGroupName, const FName InAssetClassName)
+void FEnginePackageLocalizationCache::FindAssetGroupPackages(const FName InAssetGroupName, const FTopLevelAssetPath& InAssetClassName)
 {
 	FAssetRegistryModule& AssetRegistryModule = FModuleManager::LoadModuleChecked<FAssetRegistryModule>(TEXT("AssetRegistry"));
 	IAssetRegistry& AssetRegistry = AssetRegistryModule.Get();
@@ -113,7 +113,7 @@ void FEnginePackageLocalizationCache::FindAssetGroupPackages(const FName InAsset
 		Filter.PackagePaths.Add(*LocalizedRootPath);
 	}
 	Filter.bIncludeOnlyOnDiskAssets = false;
-	Filter.ClassNames.Add(InAssetClassName);
+	Filter.ClassPaths.Add(InAssetClassName);
 	Filter.bRecursiveClasses = false;
 
 	TArray<FAssetData> LocalizedAssetsOfClass;
diff --git a/Engine/Source/Runtime/Engine/Private/Level.cpp b/Engine/Source/Runtime/Engine/Private/Level.cpp
--- a/Engine/Source/Runtime/Engine/Private/Level.cpp
+++ b/Engine/Source/Runtime/Engine/Private/Level.cpp
@@ -2282,8 +2282,8 @@ namespace LevelAssetRegistryHelper
 
 		for (const FAssetData& Asset : LevelPackageAssets)
 		{
-			static const FName NAME_World(TEXT("World"));
-			if (Asset.AssetClass == NAME_World)
+			static const FTopLevelAssetPath WorldClassPathName(TEXT("/Script/Engine"), TEXT("World"));
+			if (Asset.AssetClassPath == WorldClassPathName)
 			{
 				return Func(Asset);
 			}
@@ -3420,7 +3420,7 @@ bool ULevel::SetUseActorFolders(bool bInEnabled, bool bInInteractiveMode)
 
 	if (GetWorldPartition())
 	{
-		UClass* WorldPartitionBuilderCommandletClass = FindObject<UClass>(ANY_PACKAGE, TEXT("WorldPartitionBuilderCommandlet"), true);
+		UClass* WorldPartitionBuilderCommandletClass = FindObject<UClass>(nullptr, TEXT("/Script/UnrealEd.WorldPartitionBuilderCommandlet"), true);
 		const bool bIsRunningWorldPartitionBuilderCommandlet = WorldPartitionBuilderCommandletClass && GetRunningCommandletClass() && GetRunningCommandletClass()->IsChildOf(WorldPartitionBuilderCommandletClass);
 		if (!bIsRunningWorldPartitionBuilderCommandlet)
 		{
diff --git a/Engine/Source/Runtime/Engine/Private/LevelStreaming.cpp b/Engine/Source/Runtime/Engine/Private/LevelStreaming.cpp
--- a/Engine/Source/Runtime/Engine/Private/LevelStreaming.cpp
+++ b/Engine/Source/Runtime/Engine/Private/LevelStreaming.cpp
@@ -1227,7 +1227,7 @@ bool ULevelStreaming::RequestLevel(UWorld* PersistentWorld, bool bAllowLevelLoad
 
 	// Try to find the [to be] loaded package.
 	UWorld* World = nullptr;
-	UPackage* LevelPackage = (UPackage*)StaticFindObjectFast(UPackage::StaticClass(), nullptr, DesiredPackageName, 0, 0, RF_NoFlags, EInternalObjectFlags::Garbage);
+	UPackage* LevelPackage = (UPackage*)StaticFindObjectFast(UPackage::StaticClass(), nullptr, DesiredPackageName, /*bExactClass=*/false, RF_NoFlags, EInternalObjectFlags::Garbage);
 	
 	if (LevelPackage)
 	{
diff --git a/Engine/Source/Runtime/Engine/Private/MemberReference.cpp b/Engine/Source/Runtime/Engine/Private/MemberReference.cpp
--- a/Engine/Source/Runtime/Engine/Private/MemberReference.cpp
+++ b/Engine/Source/Runtime/Engine/Private/MemberReference.cpp
@@ -313,9 +313,12 @@ TFieldType* FindRemappedFieldImpl(FName FieldClassOutermostName, FName FieldClas
 				{
 					// Use package if it's there
 					ClassName = FString::Printf(TEXT("%s.%s"), *NewRedirectName.PackageName.ToString(), *NewRedirectName.OuterName.ToString());
+					SearchClass = (UClass*)StaticFindObject(UClass::StaticClass(), nullptr, *ClassName);
 				}
-
-				SearchClass = (UClass*)StaticFindObject(UClass::StaticClass(), ANY_PACKAGE, *ClassName);
+				else
+				{
+					SearchClass = FindFirstObject<UClass>(*ClassName, EFindFirstObjectOptions::None, ELogVerbosity::Fatal, TEXT("FindRemappedFieldImpl"));
+				}				
 
 				if (!SearchClass)
 				{
diff --git a/Engine/Source/Runtime/Engine/Private/ObjectLibrary.cpp b/Engine/Source/Runtime/Engine/Private/ObjectLibrary.cpp
--- a/Engine/Source/Runtime/Engine/Private/ObjectLibrary.cpp
+++ b/Engine/Source/Runtime/Engine/Private/ObjectLibrary.cpp
@@ -344,12 +344,18 @@ int32 UObjectLibrary::LoadAssetDataFromPaths(const TArray<FString>& Paths, bool
 	FARFilter ARFilter;
 	if ( ObjectBaseClass )
 	{
-		ARFilter.ClassNames.Add(ObjectBaseClass->GetFName());
+		ARFilter.ClassPaths.Add(ObjectBaseClass->GetClassPathName());
 
 #if WITH_EDITOR
 		// Add any old names to the list in case things haven't been resaved
-		TArray<FName> OldNames = FLinkerLoad::FindPreviousNamesForClass(ObjectBaseClass->GetPathName(), false);
-		ARFilter.ClassNames.Append(OldNames);
+		// @todo make FLinkerLoad support long class names
+		TArray<FString> OldNames = FLinkerLoad::FindPreviousPathNamesForClass(ObjectBaseClass->GetPathName(), false);
+		for (const FString& OldName : OldNames)
+		{
+			FTopLevelAssetPath OldPathName(OldName);
+			check(!OldPathName.IsNull() || OldName.IsEmpty());
+			ARFilter.ClassPaths.Add(OldPathName);
+		}
 #endif
 
 		ARFilter.bRecursiveClasses = true;
@@ -413,7 +419,7 @@ int32 UObjectLibrary::LoadBlueprintAssetDataFromPaths(const TArray<FString>& Pat
 #endif
 
 	FARFilter ARFilter;
-	ARFilter.ClassNames.Add(UBlueprint::StaticClass()->GetFName());
+	ARFilter.ClassPaths.Add(UBlueprint::StaticClass()->GetClassPathName());
 	
 	for (int PathIndex = 0; PathIndex < Paths.Num(); PathIndex++)
 	{
@@ -445,10 +451,10 @@ int32 UObjectLibrary::LoadBlueprintAssetDataFromPaths(const TArray<FString>& Pat
 	// Filter out any blueprints found whose parent class is not derived from ObjectBaseClass
 	if (ObjectBaseClass)
 	{
-		TSet<FName> DerivedClassNames;
-		TArray<FName> ClassNames;
-		ClassNames.Add(ObjectBaseClass->GetFName());
-		AssetRegistry.GetDerivedClassNames(ClassNames, TSet<FName>(), DerivedClassNames);
+		TSet<FTopLevelAssetPath> DerivedClassNames;
+		TArray<FTopLevelAssetPath> ClassNames;
+		ClassNames.Add(ObjectBaseClass->GetClassPathName());
+		AssetRegistry.GetDerivedClassNames(ClassNames, TSet<FTopLevelAssetPath>(), DerivedClassNames);
 
 		for(int32 AssetIdx=AssetDataList.Num() - 1; AssetIdx >= 0; --AssetIdx)
 		{
@@ -459,8 +465,7 @@ int32 UObjectLibrary::LoadBlueprintAssetDataFromPaths(const TArray<FString>& Pat
 			if (!ParentClassFromData.IsEmpty())
 			{
 				const FString ClassObjectPath = FPackageName::ExportTextPathToObjectPath(ParentClassFromData);
-				const FString ClassName = FPackageName::ObjectPathToObjectName(ClassObjectPath);
-				if (DerivedClassNames.Contains(FName(*ClassName)))
+				if (DerivedClassNames.Contains(FTopLevelAssetPath(ClassObjectPath)))
 				{
 					// This asset is derived from ObjectBaseClass. Keep it.
 					bShouldRemove = false;
diff --git a/Engine/Source/Runtime/Engine/Private/Pawn.cpp b/Engine/Source/Runtime/Engine/Private/Pawn.cpp
--- a/Engine/Source/Runtime/Engine/Private/Pawn.cpp
+++ b/Engine/Source/Runtime/Engine/Private/Pawn.cpp
@@ -57,9 +57,9 @@ APawn::APawn(const FObjectInitializer& ObjectInitializer)
 #endif
 
 		FString AIControllerClassName = GetDefault<UEngine>()->AIControllerClassName.ToString();
-
+		check(FPackageName::IsValidObjectPath(AIControllerClassName));
 		// resolve the name to a UClass
-		AIControllerClass = FindObject<UClass>(ANY_PACKAGE, *AIControllerClassName);
+		AIControllerClass = FindObject<UClass>(nullptr, *AIControllerClassName);
 
 		// if we failed to resolve, and plugins are expected to be loaded, proceed with loading it
 		if (AIControllerClass == nullptr && bLoadPluginClass)
diff --git a/Engine/Source/Runtime/Engine/Private/Slate/SlateGameResources.cpp b/Engine/Source/Runtime/Engine/Private/Slate/SlateGameResources.cpp
--- a/Engine/Source/Runtime/Engine/Private/Slate/SlateGameResources.cpp
+++ b/Engine/Source/Runtime/Engine/Private/Slate/SlateGameResources.cpp
@@ -229,7 +229,7 @@ void FSlateGameResources::RemoveAsset(const FAssetData& InRemovedAssetData)
 bool FSlateGameResources::ShouldCache( const FAssetData& InAssetData )
 {
 	return InAssetData.ObjectPath.ToString().StartsWith( ContentRootDir, ESearchCase::CaseSensitive )
-		&& InAssetData.AssetClass == USlateWidgetStyleAsset::StaticClass()->GetFName();
+		&& InAssetData.AssetClassPath == USlateWidgetStyleAsset::StaticClass()->GetClassPathName();
 }
 
 void FSlateGameResources::AddAssetToCache( UObject* InStyleObject, bool bEnsureUniqueness )
diff --git a/Engine/Source/Runtime/Engine/Private/URL.cpp b/Engine/Source/Runtime/Engine/Private/URL.cpp
--- a/Engine/Source/Runtime/Engine/Private/URL.cpp
+++ b/Engine/Source/Runtime/Engine/Private/URL.cpp
@@ -478,7 +478,7 @@ FURL::FURL( FURL* Base, const TCHAR* TextURL, ETravelType Type )
 					if (!AssetRegistry.IsLoadingAssets())
 					{
 						TArray<FAssetData> MapList;
-						if (AssetRegistry.GetAssetsByClass(UWorld::StaticClass()->GetFName(), /*out*/ MapList))
+						if (AssetRegistry.GetAssetsByClass(UWorld::StaticClass()->GetClassPathName(), /*out*/ MapList))
 						{
 							FName TargetTestName(*URLStr);
 							for (const FAssetData& MapAsset : MapList)
diff --git a/Engine/Source/Runtime/Engine/Private/UnrealEngine.cpp b/Engine/Source/Runtime/Engine/Private/UnrealEngine.cpp
--- a/Engine/Source/Runtime/Engine/Private/UnrealEngine.cpp
+++ b/Engine/Source/Runtime/Engine/Private/UnrealEngine.cpp
@@ -5781,7 +5781,7 @@ bool UEngine::HandleKismetEventCommand(UWorld* InWorld, const TCHAR* Cmd, FOutpu
 	}
 	else
 	{
-		UObject* ObjectToMatch = FindObject<UObject>(ANY_PACKAGE, *ObjectName);
+		UObject* ObjectToMatch = FindFirstObject<UObject>(*ObjectName, EFindFirstObjectOptions::None, ELogVerbosity::Warning, TEXT("HandleKismetEventCommand"));
 
 		if (ObjectToMatch == nullptr)
 		{
@@ -8068,7 +8068,7 @@ static void HandlePropAnalysisCommand(const TCHAR* Cmd, FOutputDevice& Ar)
 	int32 MaxPercentToCheck = 25;
 
 	// get command parameters
-	ParseObject<UClass>(Cmd, TEXT("CLASS="), AnalyzedClass, ANY_PACKAGE);
+	ParseObject<UClass>(Cmd, TEXT("CLASS="), AnalyzedClass, nullptr);
 	if (AnalyzedClass == nullptr)
 	{
 		Ar.Log(TEXT("Usage: PROPANALYSIS CLASS=<class> [MIN=<minpercentdiff> MAX=<maxpercentdiff>]"));
@@ -8353,7 +8353,7 @@ bool UEngine::HandleObjCommand( const TCHAR* Cmd, FOutputDevice& Ar )
 	else if (FParse::Command(&Cmd,TEXT("LIST2")))
 	{			
 		UClass* ClassToCheck = NULL;
-		ParseObject<UClass>(Cmd, TEXT("CLASS="  ), ClassToCheck, ANY_PACKAGE );
+		ParseObject<UClass>(Cmd, TEXT("CLASS="  ), ClassToCheck, nullptr);
 
 		if (ClassToCheck == NULL)
 		{
@@ -8395,7 +8395,7 @@ bool UEngine::HandleObjCommand( const TCHAR* Cmd, FOutputDevice& Ar )
 		Limit *= 1024;
 
 		UClass* ClassToCheck = NULL;
-		ParseObject<UClass>(Cmd, TEXT("CLASS="  ), ClassToCheck, ANY_PACKAGE );
+		ParseObject<UClass>(Cmd, TEXT("CLASS="  ), ClassToCheck, nullptr);
 		if (ClassToCheck == NULL)
 		{
 			ClassToCheck = UObject::StaticClass();
@@ -8523,7 +8523,7 @@ bool UEngine::HandleObjCommand( const TCHAR* Cmd, FOutputDevice& Ar )
 		Limit *= 1024;
 
 		UClass* ClassToCheck = NULL;
-		ParseObject<UClass>(Cmd, TEXT("CLASS="  ), ClassToCheck, ANY_PACKAGE );
+		ParseObject<UClass>(Cmd, TEXT("CLASS="  ), ClassToCheck, nullptr);
 
 		if (ClassToCheck == NULL)
 		{
@@ -8603,8 +8603,8 @@ bool UEngine::HandleObjCommand( const TCHAR* Cmd, FOutputDevice& Ar )
 		UPackage* InsidePackage = nullptr;
 		UObject* InsideObject = nullptr;
 		UClass* InsideClass = nullptr;
-		ParseObject<UClass>(Cmd, TEXT("CLASS="  ), CheckType, ANY_PACKAGE );
-		ParseObject<UObject>(Cmd, TEXT("OUTER="), CheckOuter, ANY_PACKAGE);
+		ParseObject<UClass>(Cmd, TEXT("CLASS="  ), CheckType, nullptr);
+		ParseObject<UObject>(Cmd, TEXT("OUTER="), CheckOuter, nullptr);
 
 		ParseObject<UPackage>(Cmd, TEXT("PACKAGE="), InsidePackage, nullptr);
 		if (InsidePackage == nullptr)
@@ -8613,7 +8613,7 @@ bool UEngine::HandleObjCommand( const TCHAR* Cmd, FOutputDevice& Ar )
 		}
 		if (InsidePackage == nullptr && InsideObject == nullptr)
 		{
-			ParseObject<UClass>(Cmd, TEXT("INSIDECLASS="), InsideClass, ANY_PACKAGE);
+			ParseObject<UClass>(Cmd, TEXT("INSIDECLASS="), InsideClass, nullptr);
 		}
 		int32 Depth = -1;
 		FParse::Value(Cmd, TEXT("DEPTH="), Depth);
@@ -8631,7 +8631,9 @@ bool UEngine::HandleObjCommand( const TCHAR* Cmd, FOutputDevice& Ar )
 		// support specifying metaclasses when listing class objects
 		if ( CheckType && CheckType->IsChildOf(UClass::StaticClass()) )
 		{
-			ParseObject<UClass>  ( Cmd, TEXT("TYPE="   ), MetaClass,     ANY_PACKAGE );
+			FString ClassName;
+			FParse::Value(Cmd, TEXT("TYPE="), ClassName);
+			MetaClass = FindFirstObject<UClass>(*ClassName, EFindFirstObjectOptions::None, ELogVerbosity::Warning, TEXT("LIST command"));
 		}
 
 		const bool bAll				= FParse::Param( Cmd, TEXT("ALL") );
@@ -9099,7 +9101,7 @@ bool UEngine::HandleObjCommand( const TCHAR* Cmd, FOutputDevice& Ar )
 
 		if ( FParse::Token(Cmd,ObjectName,true) )
 		{
-			Obj = FindObject<UObject>(ANY_PACKAGE,*ObjectName);
+			Obj = FindFirstObject<UObject>(*ObjectName, EFindFirstObjectOptions::None, ELogVerbosity::Warning, TEXT("OBJ COMPONENTS"));
 
 			if ( Obj != NULL )
 			{
@@ -9132,11 +9134,11 @@ bool UEngine::HandleObjCommand( const TCHAR* Cmd, FOutputDevice& Ar )
 
 		TArray<FString> HiddenCategories, ShowingCategories;
 
-		if ( !ParseObject<UClass>( Cmd, TEXT("CLASS="), Cls, ANY_PACKAGE ) || !ParseObject(Cmd,TEXT("NAME="), Cls, Obj, ANY_PACKAGE) )
+		if ( !ParseObject<UClass>( Cmd, TEXT("CLASS="), Cls, nullptr) || !ParseObject(Cmd,TEXT("NAME="), Cls, Obj, nullptr) )
 		{
 			if ( FParse::Token(Cmd,ObjectName,UE_ARRAY_COUNT(ObjectName), 1) )
 			{
-				Obj = FindObject<UObject>(ANY_PACKAGE,ObjectName);
+				Obj = FindFirstObject<UObject>(ObjectName, EFindFirstObjectOptions::None, ELogVerbosity::Warning, TEXT("DUMP"));
 			}
 		}
 
diff --git a/Engine/Source/Runtime/Engine/Private/WorldPartition/ActorDescContainer.cpp b/Engine/Source/Runtime/Engine/Private/WorldPartition/ActorDescContainer.cpp
--- a/Engine/Source/Runtime/Engine/Private/WorldPartition/ActorDescContainer.cpp
+++ b/Engine/Source/Runtime/Engine/Private/WorldPartition/ActorDescContainer.cpp
@@ -71,7 +71,7 @@
 				FCoreRedirectObjectName NewClassName = FCoreRedirects::GetRedirectedName(ECoreRedirectFlags::Type_Class, OldClassName);
 
 				bool bIsValidClass = true;
-				UClass* ActorClass = FindObject<UClass>(ANY_PACKAGE, *NewClassName.ToString(), true);
+				UClass* ActorClass = UClass::TryFindTypeSlow<UClass>(NewClassName.ToString(), EFindFirstObjectOptions::ExactClass);
 
 				if (!ActorClass)
 				{
diff --git a/Engine/Source/Runtime/Engine/Private/WorldPartition/WorldPartitionLevelStreamingDynamic.cpp b/Engine/Source/Runtime/Engine/Private/WorldPartition/WorldPartitionLevelStreamingDynamic.cpp
--- a/Engine/Source/Runtime/Engine/Private/WorldPartition/WorldPartitionLevelStreamingDynamic.cpp
+++ b/Engine/Source/Runtime/Engine/Private/WorldPartition/WorldPartitionLevelStreamingDynamic.cpp
@@ -195,7 +195,7 @@ bool UWorldPartitionLevelStreamingDynamic::RequestLevel(UWorld* InPersistentWorl
 
 	// Try to find the package to load
 	const FName DesiredPackageName = GetWorldAssetPackageFName();
-	UPackage* LevelPackage = (UPackage*)StaticFindObjectFast(UPackage::StaticClass(), nullptr, DesiredPackageName, 0, 0, RF_NoFlags, EInternalObjectFlags::Garbage);
+	UPackage* LevelPackage = (UPackage*)StaticFindObjectFast(UPackage::StaticClass(), nullptr, DesiredPackageName, /*bExactClass =*/ false, RF_NoFlags, EInternalObjectFlags::Garbage);
 	UWorld* FoundWorld = LevelPackage ? UWorld::FindWorldInPackage(LevelPackage) : nullptr;
 	check(!FoundWorld || IsValidChecked(FoundWorld));
 	check(!FoundWorld || FoundWorld->PersistentLevel);
diff --git a/Engine/Source/Runtime/Engine/Public/ExternalPackageHelper.h b/Engine/Source/Runtime/Engine/Public/ExternalPackageHelper.h
--- a/Engine/Source/Runtime/Engine/Public/ExternalPackageHelper.h
+++ b/Engine/Source/Runtime/Engine/Public/ExternalPackageHelper.h
@@ -83,7 +83,7 @@ void FExternalPackageHelper::LoadObjectsFromExternalPackages(UObject* InOuter, T
 	FARFilter Filter;
 	Filter.bRecursivePaths = true;
 	Filter.bIncludeOnlyOnDiskAssets = true;
-	Filter.ClassNames.Add(T::StaticClass()->GetFName());
+	Filter.ClassPaths.Add(T::StaticClass()->GetClassPathName());
 	Filter.PackagePaths.Add(*ExternalObjectsPath);
 
 	TArray<FAssetData> Assets;
diff --git a/Engine/Source/Runtime/Engine/Public/Internationalization/EnginePackageLocalizationCache.h b/Engine/Source/Runtime/Engine/Public/Internationalization/EnginePackageLocalizationCache.h
--- a/Engine/Source/Runtime/Engine/Public/Internationalization/EnginePackageLocalizationCache.h
+++ b/Engine/Source/Runtime/Engine/Public/Internationalization/EnginePackageLocalizationCache.h
@@ -17,7 +17,7 @@ public:
 protected:
 	//~ FPackageLocalizationCache interface
 	virtual void FindLocalizedPackages(const TMap<FString, TArray<FString>>& NewSourceToLocalizedPaths, TMap<FName, TArray<FName>>& InOutSourcePackagesToLocalizedPackages) override;
-	virtual void FindAssetGroupPackages(const FName InAssetGroupName, const FName InAssetClassName) override;
+	virtual void FindAssetGroupPackages(const FName InAssetGroupName, const FTopLevelAssetPath& InAssetClassName) override;
 
 private:
 	/**
diff --git a/Engine/Source/Runtime/Engine/Public/VisualLogger/VisualLoggerTypes.h b/Engine/Source/Runtime/Engine/Public/VisualLogger/VisualLoggerTypes.h
--- a/Engine/Source/Runtime/Engine/Public/VisualLogger/VisualLoggerTypes.h
+++ b/Engine/Source/Runtime/Engine/Public/VisualLogger/VisualLoggerTypes.h
@@ -9,9 +9,9 @@ class AActor;
 class UCanvas;
 struct FLogEntryItem;
 
-#define DEFINE_ENUM_TO_STRING(EnumType) FString EnumToString(const EnumType Value) \
+#define DEFINE_ENUM_TO_STRING(EnumType, EnumPackage) FString EnumToString(const EnumType Value) \
 { \
-	static const UEnum* TypeEnum = FindObject<UEnum>(ANY_PACKAGE, TEXT(#EnumType)); \
+	static const UEnum* TypeEnum = FindObject<UEnum>(nullptr, TEXT(EnumPackage) TEXT(".") TEXT(#EnumType)); \
 	return TypeEnum->GetNameStringByIndex(static_cast<int32>(Value)); \
 }
 #define DECLARE_ENUM_TO_STRING(EnumType) FString EnumToString(const EnumType Value)
diff --git a/Engine/Source/Runtime/Experimental/Interchange/Core/Private/InterchangeResult.cpp b/Engine/Source/Runtime/Experimental/Interchange/Core/Private/InterchangeResult.cpp
--- a/Engine/Source/Runtime/Experimental/Interchange/Core/Private/InterchangeResult.cpp
+++ b/Engine/Source/Runtime/Experimental/Interchange/Core/Private/InterchangeResult.cpp
@@ -69,7 +69,7 @@ UInterchangeResult* UInterchangeResult::FromJson(const FString& JsonString)
 		return nullptr;
 	}
 
-	UClass* ClassToImport = FindObject<UClass>(ANY_PACKAGE, *ClassType);	// @TODO: no idea how we should do this without ANY_PACKAGE
+	UClass* ClassToImport = UClass::TryFindTypeSlow<UClass>(ClassType);
 	UInterchangeResult* Result = NewObject<UInterchangeResult>(GetTransientPackage(), ClassToImport);
 
 	for (const TSharedPtr<FJsonValue>& Property : *Properties)
diff --git a/Engine/Source/Runtime/Experimental/Interchange/Engine/Private/InterchangeManager.cpp b/Engine/Source/Runtime/Experimental/Interchange/Engine/Private/InterchangeManager.cpp
--- a/Engine/Source/Runtime/Experimental/Interchange/Engine/Private/InterchangeManager.cpp
+++ b/Engine/Source/Runtime/Experimental/Interchange/Engine/Private/InterchangeManager.cpp
@@ -1101,20 +1101,20 @@ void UInterchangeManager::FindPipelineCandidate(TArray<UClass*>& PipelineCandida
 	//TODO do we have an other alternative, this call is synchronous and will wait unitl the registry database have finish the initial scan. If there is a lot of asset it can take multiple second the first time we call it.
 	AssetRegistry.ScanPathsSynchronous(ContentPaths);
 
-	FName BaseClassName = UInterchangePipelineBase::StaticClass()->GetFName();
+	FTopLevelAssetPath BaseClassName = UInterchangePipelineBase::StaticClass()->GetClassPathName();
 
 	// Use the asset registry to get the set of all class names deriving from Base
-	TSet< FName > DerivedNames;
+	TSet< FTopLevelAssetPath > DerivedNames;
 	{
-		TArray< FName > BaseNames;
+		TArray< FTopLevelAssetPath > BaseNames;
 		BaseNames.Add(BaseClassName);
 
-		TSet< FName > Excluded;
+		TSet< FTopLevelAssetPath > Excluded;
 		AssetRegistry.GetDerivedClassNames(BaseNames, Excluded, DerivedNames);
 	}
 
 	FARFilter Filter;
-	Filter.ClassNames.Add(UBlueprint::StaticClass()->GetFName());
+	Filter.ClassPaths.Add(UBlueprint::StaticClass()->GetClassPathName());
 	Filter.bRecursiveClasses = true;
 	Filter.bRecursivePaths = true;
 
@@ -1129,11 +1129,10 @@ void UInterchangeManager::FindPipelineCandidate(TArray<UClass*>& PipelineCandida
 		if (GeneratedClassPath.IsSet())
 		{
 			// Convert path to just the name part
-			const FString ClassObjectPath = FPackageName::ExportTextPathToObjectPath(*GeneratedClassPath.GetValue());
-			const FString ClassName = FPackageName::ObjectPathToObjectName(ClassObjectPath);
+			const FTopLevelAssetPath ClassObjectPath(FPackageName::ExportTextPathToObjectPath(*GeneratedClassPath.GetValue()));
 
 			// Check if this class is in the derived set
-			if (!DerivedNames.Contains(*ClassName))
+			if (!DerivedNames.Contains(ClassObjectPath))
 			{
 				continue;
 			}
diff --git a/Engine/Source/Runtime/Experimental/Interchange/Engine/Private/Tasks/InterchangeTaskCreateAsset.cpp b/Engine/Source/Runtime/Experimental/Interchange/Engine/Private/Tasks/InterchangeTaskCreateAsset.cpp
--- a/Engine/Source/Runtime/Experimental/Interchange/Engine/Private/Tasks/InterchangeTaskCreateAsset.cpp
+++ b/Engine/Source/Runtime/Experimental/Interchange/Engine/Private/Tasks/InterchangeTaskCreateAsset.cpp
@@ -126,8 +126,15 @@ void UE::Interchange::FTaskCreatePackage::DoTask(ENamedThreads::Type CurrentThre
 	{
 		Private::InternalGetPackageName(*AsyncHelper, SourceIndex, PackageBasePath, FactoryNode, PackageName, AssetName);
 		UPackage* ResultFindPackage = FindPackage(nullptr, *PackageName);
-		UObject* FindOuter = (ResultFindPackage == nullptr ? ANY_PACKAGE : ResultFindPackage);
-		UObject* ExistingObject = FindObject<UObject>(FindOuter, *AssetName);
+		UObject* ExistingObject = nullptr;
+		if (ResultFindPackage)
+		{
+			ExistingObject = FindObject<UObject>(ResultFindPackage, *AssetName);
+		}
+		else
+		{
+			ExistingObject = FindFirstObject<UObject>(*AssetName, EFindFirstObjectOptions::NativeFirst | EFindFirstObjectOptions::EnsureIfAmbiguous);
+		}
 
 		if (ExistingObject != AsyncHelper->TaskData.ReimportObject)
 		{
@@ -255,8 +262,14 @@ void UE::Interchange::FTaskCreateAsset::DoTask(ENamedThreads::Type CurrentThread
 		//When we re-import one particular asset, if the source file contain other assets, we want to set the node= reference UObject for those asset to the existing asset
 		//The way to discover this case is to compare the re-import asset with the node asset.
 		UPackage* ResultFindPackage = FindPackage(nullptr, *PackageName);
-		UObject* FindOuter = (ResultFindPackage == nullptr ? ANY_PACKAGE : ResultFindPackage);
-		ExistingObject = FindObject<UObject>(FindOuter, *AssetName);
+		if (ResultFindPackage)
+		{
+			ExistingObject = FindObject<UObject>(ResultFindPackage, *AssetName);
+		}
+		else
+		{
+			ExistingObject = FindFirstObject<UObject>(*AssetName, EFindFirstObjectOptions::NativeFirst | EFindFirstObjectOptions::EnsureIfAmbiguous);
+		}
 
 		if(ExistingObject != AsyncHelper->TaskData.ReimportObject)
 		{
