AppendToClassSchemaContext: Hide the use of FBlake3 behind an interface so we can easily change it.
#rnx

diff --git a/Engine/Source/Runtime/CoreUObject/Private/Serialization/UnversionedPropertySerialization.cpp b/Engine/Source/Runtime/CoreUObject/Private/Serialization/UnversionedPropertySerialization.cpp
--- a/Engine/Source/Runtime/CoreUObject/Private/Serialization/UnversionedPropertySerialization.cpp
+++ b/Engine/Source/Runtime/CoreUObject/Private/Serialization/UnversionedPropertySerialization.cpp
@@ -319,7 +319,8 @@ struct FUnversionedStructSchema
 		const UClass* StructAsClass = Cast<const UClass>(Struct);
 		if (StructAsClass)
 		{
-			StructAsClass->CallAppendToClassSchema(HashBuilder);
+			FAppendToClassSchemaContext Context(&HashBuilder);
+			StructAsClass->CallAppendToClassSchema(Context);
 		}
 		Schema->SchemaHash = HashBuilder.Finalize();
 #endif
@@ -977,4 +978,11 @@ COREUOBJECT_API void DumpClassSchemas(const TCHAR* Str, FOutputDevice& Ar)
 	}
 #endif
 }
+
+void FAppendToClassSchemaContext::Update(const void* Data, uint64 Size)
+{
+	FBlake3& Blake3Hasher = *(reinterpret_cast<FBlake3*>(Hasher));
+	Blake3Hasher.Update(Data, Size);
+}
+
 #endif
diff --git a/Engine/Source/Runtime/CoreUObject/Private/UObject/Obj.cpp b/Engine/Source/Runtime/CoreUObject/Private/UObject/Obj.cpp
--- a/Engine/Source/Runtime/CoreUObject/Private/UObject/Obj.cpp
+++ b/Engine/Source/Runtime/CoreUObject/Private/UObject/Obj.cpp
@@ -1576,7 +1576,7 @@ void UObject::DeclareCustomVersions(FArchive& Ar, const UClass* SpecificSubclass
 	}
 }
 
-void UObject::AppendToClassSchema(FBlake3& Hasher)
+void UObject::AppendToClassSchema(FAppendToClassSchemaContext& Context)
 {
 }
 #endif
diff --git a/Engine/Source/Runtime/CoreUObject/Public/UObject/Class.h b/Engine/Source/Runtime/CoreUObject/Public/UObject/Class.h
--- a/Engine/Source/Runtime/CoreUObject/Public/UObject/Class.h
+++ b/Engine/Source/Runtime/CoreUObject/Public/UObject/Class.h
@@ -2526,10 +2526,10 @@ public:
 	}
 
 	/** Calls the c++ class's AppendToClassSchema static function */
-	void CallAppendToClassSchema(FBlake3& Hasher) const
+	void CallAppendToClassSchema(FAppendToClassSchemaContext& Context) const
 	{
 		check(CppClassStaticFunctions.GetAppendToClassSchema());
-		CppClassStaticFunctions.GetAppendToClassSchema()(Hasher);
+		CppClassStaticFunctions.GetAppendToClassSchema()(Context);
 	}
 #endif
 
diff --git a/Engine/Source/Runtime/CoreUObject/Public/UObject/Object.h b/Engine/Source/Runtime/CoreUObject/Public/UObject/Object.h
--- a/Engine/Source/Runtime/CoreUObject/Public/UObject/Object.h
+++ b/Engine/Source/Runtime/CoreUObject/Public/UObject/Object.h
@@ -15,7 +15,6 @@
 #include "UObject/PrimaryAssetId.h"
 
 struct FAssetData;
-class FBlake3;
 class FConfigCacheIni;
 class FCustomPropertyConditionState;
 class FEditPropertyChain;
@@ -26,6 +25,7 @@ class FObjectPreSaveRootContext;
 class ITargetPlatform;
 class ITransactionObjectAnnotation;
 class FTransactionObjectEvent;
+struct FAppendToClassSchemaContext;
 struct FFrame;
 struct FObjectInstancingGraph;
 struct FPropertyChangedChainEvent;
@@ -378,7 +378,7 @@ public:
 	 * values that determine how version upgraded are conducted. Can also append a unique guid when necessary to
 	 * invalidate previous results because serialization changed and no custom version was udpated.
 	 */
-	static void AppendToClassSchema(FBlake3& Hasher);
+	static void AppendToClassSchema(FAppendToClassSchemaContext& Context);
 #endif
 
 	/** After a critical error, perform any mission-critical cleanup, such as restoring the video mode orreleasing hardware resources. */
@@ -1696,6 +1696,21 @@ struct FInternalUObjectBaseUtilityIsValidFlagsChecker
 	}
 };
 
+#if WITH_EDITORONLY_DATA
+struct FAppendToClassSchemaContext
+{
+public:
+	explicit FAppendToClassSchemaContext(void* InHasher) // Type is void* to mask the implementation detail
+		:Hasher(InHasher)
+	{
+	}
+	COREUOBJECT_API void Update(const void* Data, uint64 Size);
+
+private:
+	void* Hasher; // Type is void* to mask the implementation detail
+};
+#endif
+
 /**
 * Test validity of object
 *
diff --git a/Engine/Source/Runtime/CoreUObject/Public/UObject/ObjectMacros.h b/Engine/Source/Runtime/CoreUObject/Public/UObject/ObjectMacros.h
--- a/Engine/Source/Runtime/CoreUObject/Public/UObject/ObjectMacros.h
+++ b/Engine/Source/Runtime/CoreUObject/Public/UObject/ObjectMacros.h
@@ -10,9 +10,9 @@
 #include "Misc/EnumClassFlags.h"
 #include "UObject/Script.h"
 
-class FBlake3;
 class FObjectInitializer;
 class FReferenceCollector;
+struct FAppendToClassSchemaContext;
 struct FFrame;
 struct FClassReloadVersionInfo;
 
@@ -1878,7 +1878,7 @@ public:
 	typedef void (*AddReferencedObjectsType)	(UObject* ThisObject, FReferenceCollector& Ar);
 #if WITH_EDITORONLY_DATA
 	typedef void (*DeclareCustomVersionsType)   (FArchive& Ar, const UClass* SpecificSubclass);
-	typedef void (*AppendToClassSchemaType)   (FBlake3& Hasher);
+	typedef void (*AppendToClassSchemaType)   (FAppendToClassSchemaContext& Context);
 #endif
 	// UObjectCppClassStaticFunctions: Extend this list of types with the type of your new static function.
 
diff --git a/Engine/Source/Runtime/Engine/Classes/Engine/Texture.h b/Engine/Source/Runtime/Engine/Classes/Engine/Texture.h
--- a/Engine/Source/Runtime/Engine/Classes/Engine/Texture.h
+++ b/Engine/Source/Runtime/Engine/Classes/Engine/Texture.h
@@ -1639,7 +1639,7 @@
 	ENGINE_API virtual void FinishDestroy() override;
 	ENGINE_API virtual void PostCDOContruct() override;
 #if WITH_EDITORONLY_DATA
-	ENGINE_API static void AppendToClassSchema(FBlake3& Hasher);
+	ENGINE_API static void AppendToClassSchema(FAppendToClassSchemaContext& Context);
 	ENGINE_API virtual void GetAssetRegistryTags(TArray<FAssetRegistryTag>& OutTags) const override;
 #endif
 	ENGINE_API virtual bool IsPostLoadThreadSafe() const override;
diff --git a/Engine/Source/Runtime/Engine/Private/Texture.cpp b/Engine/Source/Runtime/Engine/Private/Texture.cpp
--- a/Engine/Source/Runtime/Engine/Private/Texture.cpp
+++ b/Engine/Source/Runtime/Engine/Private/Texture.cpp
@@ -40,7 +40,6 @@
 #include "ImageCoreUtils.h"
 #include "ImageUtils.h"
 #include "Algo/Unique.h"
-#include "Hash/Blake3.h"
 
 #if WITH_EDITOR
 #include "DerivedDataBuildVersion.h"
@@ -886,13 +885,13 @@ void UTexture::Serialize(FArchive& Ar)
 }
 
 #if WITH_EDITORONLY_DATA
-void UTexture::AppendToClassSchema(FBlake3& Hasher)
+void UTexture::AppendToClassSchema(FAppendToClassSchemaContext& Context)
 {
-	Super::AppendToClassSchema(Hasher);
+	Super::AppendToClassSchema(Context);
 
 	// IsEnableLegacyAlphaCoverageThresholdScaling affects upgrades executed in Serialize, so include it in the ClassSchema
 	uint8 LegacyScalingBool = IsEnableLegacyAlphaCoverageThresholdScaling();
-	Hasher.Update(&LegacyScalingBool, sizeof(LegacyScalingBool));
+	Context.Update(&LegacyScalingBool, sizeof(LegacyScalingBool));
 }
 #endif
 
