Bulk data - fix bug when loading bulk data referenced from the original .uasset file

- added additional check that the bulk data is not loading from a referenced .uasset file when using the attached archive
- reverted back to set default Offset value to INDEX_NONE
- removed setting SizeOnDisk when transacting

diff --git a/Engine/Source/Runtime/CoreUObject/Private/Serialization/BulkData.cpp b/Engine/Source/Runtime/CoreUObject/Private/Serialization/BulkData.cpp
--- a/Engine/Source/Runtime/CoreUObject/Private/Serialization/BulkData.cpp
+++ b/Engine/Source/Runtime/CoreUObject/Private/Serialization/BulkData.cpp
@@ -1305,7 +1305,6 @@ void FBulkData::Serialize(FArchive& Ar, UObject* Owner, bool bAttemptFileMapping
 				Ar << ElementCount;
 
 				BulkMeta.SetSize(ElementCount * ElementSize);
-				BulkMeta.SetSizeOnDisk(BulkMeta.GetSize());
 
 				// Allocate bulk data.
 				void* DataBuffer = ReallocateData(GetBulkDataSize());
@@ -1364,7 +1363,7 @@ void FBulkData::Serialize(FArchive& Ar, UObject* Owner, bool bAttemptFileMapping
 #if !USE_RUNTIME_BULKDATA
 			check(MetaResource.ElementCount <= 0 || BulkMeta.GetSizeOnDisk() == MetaResource.SizeOnDisk);
 #endif
-			check(MetaResource.Offset <= 0 || BulkMeta.GetOffset() == MetaResource.Offset);
+			check(BulkMeta.GetOffset() == MetaResource.Offset);
 			check(BulkMeta.GetFlags() == MetaResource.Flags);
 
 			if (GIsEditor)
@@ -1946,7 +1945,9 @@ bool FBulkData::TryLoadDataIntoMemory(FIoBuffer Dest)
 #if WITH_EDITOR
 	// BulkDatas in the same package share AttachedAr that they get from the LinkerLoad of the package.
 	// To make calls to those BulkDatas threadsafe, we need to not use AttachedAr when called multithreaded.
-	if (IsInGameThread())
+	// Also don't use the attached archive when loading from the editor domain and the bulk data is referenced
+	// from the original .uasset file
+	if (IsInGameThread() && IsInExternalResource() == false)
 	{
 		if (FArchive* Ar = AttachedAr)
 		{
diff --git a/Engine/Source/Runtime/CoreUObject/Public/Serialization/BulkData.h b/Engine/Source/Runtime/CoreUObject/Public/Serialization/BulkData.h
--- a/Engine/Source/Runtime/CoreUObject/Public/Serialization/BulkData.h
+++ b/Engine/Source/Runtime/CoreUObject/Public/Serialization/BulkData.h
@@ -320,7 +320,18 @@ struct FBulkMetaResource
 class FBulkMetaData
 {
 public:
-	FBulkMetaData() = default;
+	/** 40 bits for max bulk data size. */
+	static constexpr int64 MaxSize = 0xFFffFFffFF;
+	/** 39 bits for max bulk data offset and 1 bit to indicate INDEX_NONE. */
+	static constexpr int64 MaxOffset = 0xFFffFFffFE;
+
+	FBulkMetaData()
+	{
+		SetOffset(-1);
+#if !USE_RUNTIME_BULKDATA
+		SetSizeOnDisk(-1);
+#endif
+	}
 
 	inline int64 GetSize() const
 	{
@@ -329,6 +340,7 @@ public:
 
 	inline void SetSize(int64 Size)
 	{
+		check(Size <= MaxSize);
 		WriteUInt40(Data, uint64(Size));
 	}
 
@@ -350,11 +362,19 @@ public:
 
 	inline int64 GetOffset() const
 	{
-		return static_cast<int64>(ReadUInt40(Data + 5));
+		const int64 Offset = static_cast<int64>(ReadUInt40(Data + 5));
+		return Offset > MaxOffset ? INDEX_NONE : Offset;
 	}
 
 	inline void SetOffset(int64 Offset)
 	{
+		check(Offset <= MaxOffset);
+
+		if (Offset < 0)
+		{
+			Offset = MAX_int64;
+		}
+
 		WriteUInt40(Data + 5, uint64(Offset));
 	}
 
@@ -390,10 +410,10 @@ public:
 		if (MetaResource.ElementCount > 0)
 		{
 			Meta.SetSize(MetaResource.ElementCount * ElementSize);
-			Meta.SetSizeOnDisk(MetaResource.SizeOnDisk);
 		}
 
-		Meta.SetOffset(MetaResource.Offset > 0 ? MetaResource.Offset : 0);
+		Meta.SetSizeOnDisk(MetaResource.SizeOnDisk);
+		Meta.SetOffset(MetaResource.Offset);
 		Meta.SetFlags(MetaResource.Flags);
 
 		return Meta;
@@ -1040,14 +1060,14 @@ private:
 	
 	FAllocatedPtr								DataAllocation;
 	UE::BulkData::Private::FBulkMetaData		BulkMeta;
-	UE::BulkData::Private::FBulkDataChunkId	BulkChunkId;
+	UE::BulkData::Private::FBulkDataChunkId		BulkChunkId;
 	
 protected:
 #if WITH_EDITOR
 	/** Archive associated with bulk data for serialization																*/
-	FArchive*					AttachedAr = nullptr;
+	FArchive*									AttachedAr = nullptr;
 	/** Used to make sure the linker doesn't get garbage collected at runtime for things with attached archives			*/
-	FLinkerLoad*				Linker = nullptr;
+	FLinkerLoad*								Linker = nullptr;
 #endif // WITH_EDITOR
 	   //
 #if !USE_RUNTIME_BULKDATA
