Fixing a crash when reporting ambiguous FindFirstObject results. The crash was a result of passing *Message.ToString() as ensureMsgf parameter. Also downgrading ambiguous message from an ensure to a warning (thus restoring pre-ANY_PACKAGE removal behavior) in ResolveName and when importing objects from text

diff --git a/Engine/Source/Runtime/CoreUObject/Private/UObject/PropertyBaseObject.cpp b/Engine/Source/Runtime/CoreUObject/Private/UObject/PropertyBaseObject.cpp
--- a/Engine/Source/Runtime/CoreUObject/Private/UObject/PropertyBaseObject.cpp
+++ b/Engine/Source/Runtime/CoreUObject/Private/UObject/PropertyBaseObject.cpp
@@ -487,7 +487,7 @@ UObject* FObjectPropertyBase::FindImportedObject( const FProperty* Property, UOb
 		{
 			// RobM: We should delete this path
 			// match any object of the correct class who shares the same name regardless of package path
-			Result = StaticFindFirstObject(ObjectClass, Text, EFindFirstObjectOptions::EnsureIfAmbiguous);
+			Result = StaticFindFirstObject(ObjectClass, Text, EFindFirstObjectOptions::None, ELogVerbosity::Warning, TEXT("FindImportedObject"));
 			// disallow class default subobjects here while importing defaults
 			if (Result != nullptr && (PortFlags & PPF_ParsingDefaultProperties) && Result->IsTemplate(RF_ClassDefaultObject))
 			{
diff --git a/Engine/Source/Runtime/CoreUObject/Private/UObject/UObjectGlobals.cpp b/Engine/Source/Runtime/CoreUObject/Private/UObject/UObjectGlobals.cpp
--- a/Engine/Source/Runtime/CoreUObject/Private/UObject/UObjectGlobals.cpp
+++ b/Engine/Source/Runtime/CoreUObject/Private/UObject/UObjectGlobals.cpp
@@ -613,7 +613,7 @@ UObject* StaticFindFirstObject(UClass* Class, const TCHAR* Name, EFindFirstObjec
 								break;
 							}
 						}
-						else if (FoundObject->GetClass()->GetOutermost()->HasAnyPackageFlags(PKG_CompiledIn))
+						else if (!Result && FoundObject->GetClass()->GetOutermost()->HasAnyPackageFlags(PKG_CompiledIn))
 						{
 							Result = FoundObject;
 							// Don't break yet, maybe we can find a native type (see above) which is usually what we're after anyway
@@ -638,25 +638,34 @@ UObject* StaticFindFirstObject(UClass* Class, const TCHAR* Name, EFindFirstObjec
 					Message.Append(TEXT(", will return "));
 					Message.Append(Result->GetPathName());
 					Message.Append(TEXT(" but could also be: "));
-					for (int32 ObjectIndex = 1; ObjectIndex < FoundObjects.Num(); ++ObjectIndex)
+					const int32 MaxObjectsToPrint = 1;
+					int32 PrintedObjects = 0;
+					for (int32 ObjectIndex = 0; ObjectIndex < FoundObjects.Num() && PrintedObjects < MaxObjectsToPrint; ++ObjectIndex)
 					{
-						Message.Append(FoundObjects[ObjectIndex]->GetPathName());
-						if (ObjectIndex < FoundObjects.Num() - 1)
+						if (FoundObjects[ObjectIndex] != Result)
 						{
-							Message.Append(TEXT(","));
+							if (PrintedObjects > 0)
+							{
+								Message.Append(TEXT(", "));
+							}
+							Message.Append(FoundObjects[ObjectIndex]->GetPathName());
+							PrintedObjects++;
 						}
 					}
-
+					if (FoundObjects.Num() > (MaxObjectsToPrint + 1)) // +1 because we also printed Result's PathName
+					{
+						Message.Appendf(TEXT("or %d other object(s)"), FoundObjects.Num() - (MaxObjectsToPrint + 1));
+					}
 					if (AmbiguousMessageVerbosity == ELogVerbosity::Fatal)
 					{
 						UE_LOG(LogUObjectGlobals, Fatal, TEXT("%s"), *Message.ToString());
 					}
 					else if (AmbiguousMessageVerbosity != ELogVerbosity::NoLogging)
 					{
-						GLog->CategorizedLogf(TEXT("LogUObjectGlobals"), AmbiguousMessageVerbosity, TEXT("%s"), *Message.ToString());
+						GLog->CategorizedLogf(TEXT("LogUObjectGlobals"), AmbiguousMessageVerbosity, TEXT("%.*s"), Message.Len(), Message.GetData());
 					}
 
-					ensureAlwaysMsgf(!(Options & EFindFirstObjectOptions::EnsureIfAmbiguous), TEXT("%s"), *Message.ToString());
+					ensureAlwaysMsgf(!(Options & EFindFirstObjectOptions::EnsureIfAmbiguous), TEXT("%.*s"), Message.Len(), Message.GetData());
 				}
 			}
 			else
@@ -1030,7 +1039,7 @@ bool ResolveName(UObject*& InPackage, FString& InOutName, bool Create, bool Thro
 				}
 				else
 				{
-					NewPackage = FindFirstObject<UObject>(*PartialName, EFindFirstObjectOptions::NativeFirst | EFindFirstObjectOptions::EnsureIfAmbiguous);
+					NewPackage = FindFirstObject<UObject>(*PartialName, EFindFirstObjectOptions::NativeFirst, ELogVerbosity::Warning, TEXT("ResolveName"));
 				}
 				if( !NewPackage )
 				{
