Per object config sections can now also contain class path names [ObjectName /Script/Module.ClassName] (ANY_PACKAGE removal related)

diff --git a/Engine/Source/Runtime/CoreUObject/Private/UObject/Obj.cpp b/Engine/Source/Runtime/CoreUObject/Private/UObject/Obj.cpp
--- a/Engine/Source/Runtime/CoreUObject/Private/UObject/Obj.cpp
+++ b/Engine/Source/Runtime/CoreUObject/Private/UObject/Obj.cpp
@@ -2433,6 +2433,7 @@ void UObject::LoadConfig( UClass* ConfigClass/*=NULL*/, const TCHAR* InFilename/
 
 
 	FString ClassSection;
+	FString ClassPathSection;
 	FName LongCommitName;
 
 	if (bPerObject)
@@ -2449,10 +2450,22 @@ void UObject::LoadConfig( UClass* ConfigClass/*=NULL*/, const TCHAR* InFilename/
 			GetPathName(Outermost, PathNameString);
 			LongCommitName = Outermost->GetFName();
 		}
+		
+		ClassSection = PathNameString + TEXT(" ") + GetClass()->GetName();		
 
-		ClassSection = PathNameString + TEXT(" ") + GetClass()->GetName();
-
-		OverridePerObjectConfigSection(ClassSection);
+		FString OverrideClassSection;
+		OverridePerObjectConfigSection(OverrideClassSection);
+		if (OverrideClassSection.Len() && OverrideClassSection != ClassSection)
+		{
+			// If a we got a section name override no need to perform short class name checks
+			ClassSection = MoveTemp(OverrideClassSection);
+			// Keep ClassPathSection empty so that we don't check for it when the section name has been overridden
+		}
+		else
+		{
+			// Cache both version of per object config section name
+			ClassPathSection = PathNameString + TEXT(" ") + GetClass()->GetPathName();
+		}
 	}
 
 	// If any of my properties are class variables, then LoadConfig() would also be called for each one of those classes.
@@ -2470,6 +2483,30 @@ void UObject::LoadConfig( UClass* ConfigClass/*=NULL*/, const TCHAR* InFilename/
 		UE_LOG(LogConfig, Verbose, TEXT("(%s) '%s' loading configuration for property %s from %s"), *ConfigClass->GetName(), *GetName(), *PropertyToLoad->GetName(), *Filename);
 	}
 
+	auto GetConfigValue = [&OverrideConfig, &bUseConfigOverride](const TCHAR* ClassSection, const TCHAR* Key, const TCHAR* ConfigName, FString& OutValue)
+	{
+		if (bUseConfigOverride)
+		{
+			return OverrideConfig.GetString(ClassSection, Key, OutValue);
+		}
+		else
+		{
+			return GConfig->GetString(ClassSection, Key, OutValue, ConfigName);
+		}
+	};
+
+	auto GetConfigSection = [&bUseConfigOverride, &OverrideConfig](const TCHAR* SectionName, const TCHAR* ConfigFilename)
+	{
+		if (bUseConfigOverride)
+		{
+			return OverrideConfig.Find(SectionName);
+		}
+		else
+		{
+			return GConfig->GetSectionPrivate(SectionName, false, true, ConfigFilename);
+		}
+	};
+
 	for ( FProperty* Property = ConfigClass->PropertyLink; Property; Property = Property->PropertyLinkNext )
 	{
 #if WITH_EDITOR
@@ -2535,14 +2572,11 @@ void UObject::LoadConfig( UClass* ConfigClass/*=NULL*/, const TCHAR* InFilename/
 				}
 
 				FString Value;
-				bool bFoundValue;
-				if (bUseConfigOverride)
-				{
-					bFoundValue = OverrideConfig.GetString(*ClassSection, *Key, Value);
-				}
-				else
+				bool bFoundValue = GetConfigValue(*ClassSection, *Key, *PropFileName, Value);
+				if (!bFoundValue && bPerObject && ClassPathSection.Len())
 				{
-					bFoundValue = GConfig->GetString(*ClassSection, *Key, Value, *PropFileName);
+					// Try to get the value from POC config section with class path name
+					bFoundValue = GetConfigValue(*ClassPathSection, *Key, *PropFileName, Value);
 				}
 
 				if (bFoundValue)
@@ -2564,30 +2598,17 @@ void UObject::LoadConfig( UClass* ConfigClass/*=NULL*/, const TCHAR* InFilename/
 		}
 		else
 		{
-			FConfigSection* Sec;
-			if (bUseConfigOverride)
-			{
-				Sec = OverrideConfig.Find(*ClassSection);
-			}
-			else
+			FConfigSection* Sec = GetConfigSection(*ClassSection, *PropFileName);
+			if (!Sec && bPerObject && ClassPathSection.Len())
 			{
-				Sec = GConfig->GetSectionPrivate(*ClassSection, false, true, *PropFileName);
+				Sec = GetConfigSection(*ClassPathSection, *PropFileName);
 			}
-
-			FConfigSection* AltSec = NULL;
-			//@Package name transition
 			if( Sec )
 			{
 				TArray<FConfigValue> List;
 				const FName KeyName(*Key, FNAME_Find);
 				Sec->MultiFind(KeyName,List);
 
-				// If we didn't find anything in the first section, try the alternate
-				if ((List.Num() == 0) && AltSec)
-				{
-					AltSec->MultiFind(KeyName,List);
-				}
-
 				FScriptArrayHelper_InContainer ArrayHelper(Array, this);
 				const int32 Size = Array->Inner->ElementSize;
 				// Only override default properties if there is something to override them with.
@@ -2679,6 +2700,7 @@ void UObject::SaveConfig( uint64 Flags, const TCHAR* InFilename, FConfigCacheIni
 			GetPathName(Outermost, PathNameString);
 		}
 
+		//RobM: we need to update this to use GetClass()->GetPathName() after we fix all places that format section names
 		Section = PathNameString + TEXT(" ") + GetClass()->GetName();
 
 		OverridePerObjectConfigSection(Section);
