SavePackage: allow generated packages to reference private exports from their source package.
This change is required so we can enable streamed objects having hard references to always loaded actors.

#rnx

diff --git a/Engine/Source/Runtime/CoreUObject/Private/UObject/SavePackage2.cpp b/Engine/Source/Runtime/CoreUObject/Private/UObject/SavePackage2.cpp
--- a/Engine/Source/Runtime/CoreUObject/Private/UObject/SavePackage2.cpp
+++ b/Engine/Source/Runtime/CoreUObject/Private/UObject/SavePackage2.cpp
@@ -673,6 +673,22 @@ ESavePackageResult ValidateImports(FSaveContext& SaveContext)
 		return false;
 
 	};
+	auto IsSourcePackageReferenceAllowed = [](UPackage* InSourcePackage, UPackage* InImportPackage) -> bool
+	{
+#if WITH_EDITORONLY_DATA
+		// Generated packages must have the same persistent GUID as their source package
+		if (InSourcePackage->GetPersistentGuid() == InImportPackage->GetPersistentGuid())
+		{
+			// Generated packages can reference into their source package, or into other generated packages of the same source
+			if (InSourcePackage->HasAnyPackageFlags(PKG_CookGenerated))
+			{
+				return true;
+			}
+		}
+#endif
+		return false;
+	};
+
 	auto IsMapReferenceAllowed = [&SaveContext](UObject* InImport) -> bool
 	{
 		if (!InImport->HasAnyFlags(RF_Public))
@@ -752,15 +768,19 @@ ESavePackageResult ValidateImports(FSaveContext& SaveContext)
 			continue;
 		}
 
-		// See whether the object we are referencing is in another map package and if it is allowed (i.e. from redirector).
-		if (ImportPackage->ContainsMap() && !IsMapReferenceAllowed(Import))
+		// Allow private imports for split packages into their source package
+		if (!IsSourcePackageReferenceAllowed(Package, ImportPackage))
 		{
-			ObjectsInOtherMaps.Add(Import);
-		}
+			// See whether the object we are referencing is in another map package and if it is allowed (i.e. from redirector).
+			if (ImportPackage->ContainsMap() && !IsMapReferenceAllowed(Import))
+			{
+				ObjectsInOtherMaps.Add(Import);
+			}
 
-		if (!Import->HasAnyFlags(RF_Public) && (!SaveContext.IsCooking() || !ImportPackage->HasAnyPackageFlags(PKG_CompiledIn)))
-		{
-			PrivateObjects.Add(Import);
+			if (!Import->HasAnyFlags(RF_Public) && (!SaveContext.IsCooking() || !ImportPackage->HasAnyPackageFlags(PKG_CompiledIn)))
+			{
+				PrivateObjects.Add(Import);
+			}
 		}
 
 		// Enforce that Private content can only be referenced by something within the same Mount Point
