Turning off logging when failing to convert short type names when running diff asset registries commandlet

diff --git a/Engine/Source/Runtime/CoreUObject/Private/AssetRegistry/AssetData.cpp b/Engine/Source/Runtime/CoreUObject/Private/AssetRegistry/AssetData.cpp
--- a/Engine/Source/Runtime/CoreUObject/Private/AssetRegistry/AssetData.cpp
+++ b/Engine/Source/Runtime/CoreUObject/Private/AssetRegistry/AssetData.cpp
@@ -242,7 +242,7 @@ PRAGMA_ENABLE_DEPRECATION_WARNINGS
 	Ar << PackageFlags;
 }
 
-FTopLevelAssetPath FAssetData::TryConvertShortClassNameToPathName(FName InClassName)
+FTopLevelAssetPath FAssetData::TryConvertShortClassNameToPathName(FName InClassName, ELogVerbosity::Type FailureMessageVerbosity /*= ELogVerbosity::Warning*/)
 {
 	FTopLevelAssetPath ClassPath;
 	if (!InClassName.IsNone())
@@ -266,7 +266,12 @@ FTopLevelAssetPath FAssetData::TryConvertShortClassNameToPathName(FName InClassN
 			{
 				// Fallback to a fake name but at least the class name will be preserved
 				ClassPath = FTopLevelAssetPath(TEXT("/Unknown"), InClassName);
-				UE_LOG(LogAssetData, Warning, TEXT("Failed to convert deprecated short class name \"%s\" to path name. Using \"%s\""), *InClassName.ToString(), *ClassPath.ToString());
+#if !NO_LOGGING
+				if (FailureMessageVerbosity != ELogVerbosity::NoLogging)
+				{
+					FMsg::Logf(__FILE__, __LINE__, LogAssetData.GetCategoryName(), FailureMessageVerbosity, TEXT("Failed to convert deprecated short class name \"%s\" to path name. Using \"%s\""), *InClassName.ToString(), *ClassPath.ToString());
+				}
+#endif
 			}
 		}
 	}
diff --git a/Engine/Source/Runtime/CoreUObject/Private/AssetRegistry/AssetDataTagMap.cpp b/Engine/Source/Runtime/CoreUObject/Private/AssetRegistry/AssetDataTagMap.cpp
--- a/Engine/Source/Runtime/CoreUObject/Private/AssetRegistry/AssetDataTagMap.cpp
+++ b/Engine/Source/Runtime/CoreUObject/Private/AssetRegistry/AssetDataTagMap.cpp
@@ -275,7 +275,7 @@ namespace FixedTagPrivate
 	static FAssetRegistryExportPath ConvertLegacyPath(const FLegacyAssetRegistryExportPath& Path)
 	{
 		FAssetRegistryExportPath Out;
-		Out.ClassPath = FAssetData::TryConvertShortClassNameToPathName(Path.Class);
+		Out.ClassPath = FAssetData::TryConvertShortClassNameToPathName(Path.Class, ELogVerbosity::NoLogging);
 		Out.Object = Path.Object;
 		Out.Package = Path.Package;
 		return Out;
diff --git a/Engine/Source/Runtime/CoreUObject/Public/AssetRegistry/AssetData.h b/Engine/Source/Runtime/CoreUObject/Public/AssetRegistry/AssetData.h
--- a/Engine/Source/Runtime/CoreUObject/Public/AssetRegistry/AssetData.h
+++ b/Engine/Source/Runtime/CoreUObject/Public/AssetRegistry/AssetData.h
@@ -554,7 +554,7 @@ private:
 
 public:
 	/** Helper function that tries to convert short class name to path name */
-	COREUOBJECT_API static FTopLevelAssetPath TryConvertShortClassNameToPathName(FName InClassName);
+	COREUOBJECT_API static FTopLevelAssetPath TryConvertShortClassNameToPathName(FName InClassName, ELogVerbosity::Type FailureMessageVerbosity = ELogVerbosity::Warning);
 };
 
 ENUM_CLASS_FLAGS(FAssetData::ECreationFlags);
