Fix some trivial casting warnings

#rnx

- Might aswell just bump NumberOfDifferencesInPackages to int64 as ADD_COOK_STAT will turn it into a string and print it to the log as a string anyway.

diff --git a/Engine/Source/Runtime/CoreUObject/Private/UObject/SavePackage/SavePackageUtilities.cpp b/Engine/Source/Runtime/CoreUObject/Private/UObject/SavePackage/SavePackageUtilities.cpp
--- a/Engine/Source/Runtime/CoreUObject/Private/UObject/SavePackage/SavePackageUtilities.cpp
+++ b/Engine/Source/Runtime/CoreUObject/Private/UObject/SavePackage/SavePackageUtilities.cpp
@@ -128,9 +128,9 @@ void FSavePackageStats::AddSavePackageStats(FCookStatsManager::AddStatFuncRef Ad
 		DiffSize += PackageStat.Value.DiffSize;
 	}
 
-	double DifferentPackagesSizeMB = (double)NewFileTotalSize / 1024.0 / 1024.0;
-	int32  NumberOfDifferencesInPackages = NumDiffs;
-	double PackageDifferencesSizeMB = (double)DiffSize / 1024.0 / 1024.0;
+	const double DifferentPackagesSizeMB = (double)NewFileTotalSize / 1024.0 / 1024.0;
+	const int64  NumberOfDifferencesInPackages = NumDiffs;
+	const double PackageDifferencesSizeMB = (double)DiffSize / 1024.0 / 1024.0;
 
 	StatsList.Empty(15);
 	ADD_COOK_STAT(NumberOfDifferentPackages);
