Fix a trivial cast warning

#rnx

- FScopedSlowTask expects a float for the amount of work.

diff --git a/Engine/Source/Runtime/CoreUObject/Private/UObject/SavePackage2.cpp b/Engine/Source/Runtime/CoreUObject/Private/UObject/SavePackage2.cpp
--- a/Engine/Source/Runtime/CoreUObject/Private/UObject/SavePackage2.cpp
+++ b/Engine/Source/Runtime/CoreUObject/Private/UObject/SavePackage2.cpp
@@ -1751,7 +1751,7 @@ ESavePackageResult WriteExports(FStructuredArchive::FRecord& StructuredArchiveRo
 	//COOK_STAT(FScopedDurationTimer SaveTimer(FSavePackageStats::SerializeExportsTimeSec));
 	SCOPED_SAVETIMER(UPackage_Save_SaveExports);
 	FLinkerSave* Linker = SaveContext.GetLinker();
-	FScopedSlowTask SlowTask(Linker->ExportMap.Num(), FText(), SaveContext.IsUsingSlowTask());
+	FScopedSlowTask SlowTask((float)Linker->ExportMap.Num(), FText(), SaveContext.IsUsingSlowTask());
 
 	FStructuredArchive::FRecord ExportsRecord = StructuredArchiveRoot.EnterRecord(TEXT("Exports"));
 
