Work around apparently spurious static analysis warning.

diff --git a/Engine/Source/Runtime/CoreUObject/Private/UObject/UObjectGlobals.cpp b/Engine/Source/Runtime/CoreUObject/Private/UObject/UObjectGlobals.cpp
--- a/Engine/Source/Runtime/CoreUObject/Private/UObject/UObjectGlobals.cpp
+++ b/Engine/Source/Runtime/CoreUObject/Private/UObject/UObjectGlobals.cpp
@@ -604,20 +604,23 @@ UObject* StaticFindFirstObject(UClass* Class, const TCHAR* Name, EFindFirstObjec
 					// Prioritize native class instances or native type objects
 					for (UObject* FoundObject : FoundObjects)
 					{
-						if (FoundObject->IsA<UField>())
+						if (FoundObject)
 						{
-							// If we were looking for a 'type' (UEnum / UClass / UScriptStruct) object priotize native types
-							if (FoundObject->GetOutermost()->HasAnyPackageFlags(PKG_CompiledIn))
+							if (FoundObject->IsA<UField>())
+							{
+								// If we were looking for a 'type' (UEnum / UClass / UScriptStruct) object priotize native types
+								if (FoundObject->GetOutermost()->HasAnyPackageFlags(PKG_CompiledIn))
+								{
+									Result = FoundObject;
+									break;
+								}
+							}
+							else if (!Result && FoundObject->GetClass()->GetOutermost()->HasAnyPackageFlags(PKG_CompiledIn))
 							{
 								Result = FoundObject;
-								break;
+								// Don't break yet, maybe we can find a native type (see above) which is usually what we're after anyway
 							}
 						}
-						else if (!Result && FoundObject->GetClass()->GetOutermost()->HasAnyPackageFlags(PKG_CompiledIn))
-						{
-							Result = FoundObject;
-							// Don't break yet, maybe we can find a native type (see above) which is usually what we're after anyway
-						}
 					}
 				}
 				if (!Result)
