Bulk Data - remove checking the global GEventDrivenLoader and add support for loading from optional packages

diff --git a/Engine/Source/Runtime/CoreUObject/Private/Serialization/AsyncLoading2.cpp b/Engine/Source/Runtime/CoreUObject/Private/Serialization/AsyncLoading2.cpp
--- a/Engine/Source/Runtime/CoreUObject/Private/Serialization/AsyncLoading2.cpp
+++ b/Engine/Source/Runtime/CoreUObject/Private/Serialization/AsyncLoading2.cpp
@@ -6372,11 +6372,6 @@ IAsyncPackageLoader* MakeAsyncPackageLoader2(FIoDispatcher& InIoDispatcher, IAsy
 
 #endif //WITH_ASYNCLOADING2
 
-bool IsPackageLoadingFromIoDispatcher(const UPackage* Package, const FArchive& Ar)
-{
-	return Package && Package->GetPackageId().IsValid() && Ar.IsLoadingFromCookedPackage();
-}
-
 #if UE_BUILD_DEVELOPMENT || UE_BUILD_DEBUG
 PRAGMA_ENABLE_OPTIMIZATION
 #endif
diff --git a/Engine/Source/Runtime/CoreUObject/Private/Serialization/BulkData.cpp b/Engine/Source/Runtime/CoreUObject/Private/Serialization/BulkData.cpp
--- a/Engine/Source/Runtime/CoreUObject/Private/Serialization/BulkData.cpp
+++ b/Engine/Source/Runtime/CoreUObject/Private/Serialization/BulkData.cpp
@@ -27,9 +27,6 @@
 // the new loader to keep the results consistent.
 #define UE_KEEP_INLINE_RELOADING_CONSISTENT 0
 
-// Implemented in AsyncLoading2.cpp
-bool IsPackageLoadingFromIoDispatcher(const UPackage* Package, const FArchive& Ar);
-
 //////////////////////////////////////////////////////////////////////////////
 
 const FIoFilenameHash FALLBACK_IO_FILENAME_HASH = INVALID_IO_FILENAME_HASH - 1;
@@ -80,24 +77,21 @@ FIoFilenameHash MakeIoFilenameHash(const FIoChunkId& ChunkID)
 namespace UE::BulkData::Private
 {
 
-/** Returns an I/O chunk ID to be used when loading from I/O store. */
-FIoChunkId CreateBulkDataIoChunkId(const FPackageId& PackageId, EBulkDataFlags BulkDataFlags);
-
 /** Open the bulk data chunk fo reading. */
 bool OpenReadBulkData(
+	const FBulkMetaData& BulkMeta,
 	const FBulkDataChunkId& BulkChunkId,
-	EBulkDataFlags BulkDataFlags,
 	int64 Offset,
 	int64 Size,
 	TFunction<void(FArchive& Ar)>&& Read);
 
 /** Open async read file handle for the specified bulk data chunk ID. */
-TUniquePtr<IAsyncReadFileHandle> OpenAsyncReadBulkData(const FBulkDataChunkId& BulkChunkId, EBulkDataFlags BulkDataFlags);
+TUniquePtr<IAsyncReadFileHandle> OpenAsyncReadBulkData(const FBulkMetaData& BulkMeta, const FBulkDataChunkId& BulkChunkId);
 
 /** Create bulk data streaming request. */
 TUniquePtr<IBulkDataIORequest> CreateStreamingRequest(
+	const FBulkMetaData& BulkMeta,
 	const FBulkDataChunkId& BulkChunkId,
-	EBulkDataFlags BulkDataFlags,
 	int64 Offset,
 	int64 Size,
 	EAsyncIOPriorityAndFlags Priority,
@@ -105,12 +99,12 @@ TUniquePtr<IBulkDataIORequest> CreateStreamingRequest(
 	uint8* UserSuppliedMemory);
 
 /** Returns whether the bulk data chunk exist or not. */
-bool DoesBulkDataExist(const FBulkDataChunkId& BulkChunkId, EBulkDataFlags BulkDataFlags);
+bool DoesBulkDataExist(const FBulkMetaData& BulkMeta, const FBulkDataChunkId& BulkChunkId);
 
 /** Try memory map the chunk specified by the bulk data ID. */
 bool TryMemoryMapBulkData(
+	const FBulkMetaData& BulkMeta,
 	const FBulkDataChunkId& BulkChunkId,
-	EBulkDataFlags BulkDataFlags,
 	int64 Offset,
 	int64 Size,
 	FIoMappedRegion& OutRegion);
@@ -118,8 +112,8 @@ bool TryMemoryMapBulkData(
 /** Start load the internal bulk data payload. */ 
 bool StartAsyncLoad(
 	FBulkData* Owner,
+	const FBulkMetaData& BulkMeta,
 	const FBulkDataChunkId& BulkChunkId,
-	EBulkDataFlags BulkDataFlags,
 	int64 Offset,
 	int64 Size, 
 	TFunction<void(TIoStatusOr<FIoBuffer>)>&& Callback);
@@ -304,7 +298,7 @@ FIoFilenameHash FBulkDataChunkId::GetIoFilenameHash(EBulkDataFlags BulkDataFlags
 	{
 		if (const FPackageId* Id = ImplPtr->PathOrId.TryGet<FPackageId>())
 		{
-			return MakeIoFilenameHash(CreateBulkDataIoChunkId(*Id, BulkDataFlags));
+			return MakeIoFilenameHash(CreateBulkDataIoChunkId(UE::BulkData::Private::FBulkMetaData(BulkDataFlags), *Id));
 		}
 
 		if (const FPackagePath* PackagePath = ImplPtr->PathOrId.TryGet<FPackagePath>())
@@ -627,7 +621,7 @@ FBulkData& FBulkData::operator=( const FBulkData& Other )
 	else
 	{
 		FIoMappedRegion MappedRegion;
-		if (UE::BulkData::Private::TryMemoryMapBulkData(BulkChunkId, BulkMeta.GetFlags(), BulkMeta.GetOffset(), BulkMeta.GetSize(), MappedRegion))
+		if (UE::BulkData::Private::TryMemoryMapBulkData(BulkMeta, BulkChunkId, BulkMeta.GetOffset(), BulkMeta.GetSize(), MappedRegion))
 		{
 			DataAllocation.SetMemoryMappedData(this, MappedRegion.MappedFileHandle, MappedRegion.MappedFileRegion);
 		}
@@ -798,7 +792,7 @@ bool FBulkData::DoesExist() const
 	}
 #endif
 
-	return UE::BulkData::Private::DoesBulkDataExist(BulkChunkId, BulkMeta.GetFlags());
+	return UE::BulkData::Private::DoesBulkDataExist(BulkMeta, BulkChunkId);
 }
 
 /**
@@ -1085,8 +1079,8 @@ bool FBulkData::LoadBulkDataWithFileReader()
 		const int64 BulkDataSizeOnDisk = GetBulkDataSizeOnDisk();
 
 		const bool bOk = UE::BulkData::Private::OpenReadBulkData(
+			BulkMeta,
 			BulkChunkId,
-			BulkMeta.GetFlags(),
 			BulkDataOffset,
 			BulkDataSizeOnDisk,
 			[this](FArchive& Ar)
@@ -1160,8 +1154,8 @@ bool FBulkData::StartAsyncLoading()
 
 	return UE::BulkData::Private::StartAsyncLoad(
 		this,
+		BulkMeta,
 		BulkChunkId,
-		BulkMeta.GetFlags(),
 		BulkMeta.GetOffset(), 
 		BulkMeta.GetSize(),
 		[this](TIoStatusOr<FIoBuffer> Result)
@@ -1240,7 +1234,7 @@ void FBulkData::ClearBulkDataFlags( uint32 BulkDataFlagsToClear )
 
 FIoChunkId FBulkData::CreateChunkId() const
 {
-	return UE::BulkData::Private::CreateBulkDataIoChunkId(BulkChunkId.GetPackageId(), BulkMeta.GetFlags());
+	return UE::BulkData::Private::CreateBulkDataIoChunkId(BulkMeta, BulkChunkId.GetPackageId());
 }
 
 /*-----------------------------------------------------------------------------
@@ -1352,13 +1346,19 @@ void FBulkData::Serialize(FArchive& Ar, UObject* Owner, bool bAttemptFileMapping
 			UPackage* Package					= Owner ? Owner->GetPackage() : nullptr;
 			FArchive* CacheableAr				= Ar.GetCacheableArchive();
 			const bool bCanLazyLoad				= CacheableAr != nullptr && Ar.IsAllowingLazyLoading();
-			const bool bIsUsingIoDispatcher		= IsPackageLoadingFromIoDispatcher(Package, Ar);
 			FLinkerLoad* LinkerLoad				= nullptr;
 
 			FBulkMetaResource MetaResource;
 			Ar << MetaResource;
 			BulkMeta = FBulkMetaData::FromSerialized(MetaResource, ElementSize);
 
+			bool bIsUsingIoDispatcher = false;
+			if (Ar.IsLoadingFromCookedPackage())
+			{
+				BulkMeta.SetMetaFlags(FBulkMetaData::EMetaFlags::CookedPackage);
+				bIsUsingIoDispatcher = Package && Package->GetPackageId().IsValid();
+			}
+
 			check(MetaResource.ElementCount <= 0 || BulkMeta.GetSize() == MetaResource.ElementCount * ElementSize);
 #if !USE_RUNTIME_BULKDATA
 			check(MetaResource.ElementCount <= 0 || BulkMeta.GetSizeOnDisk() == MetaResource.SizeOnDisk);
@@ -1431,7 +1431,7 @@ void FBulkData::Serialize(FArchive& Ar, UObject* Owner, bool bAttemptFileMapping
 
 					if (IsDuplicateNonOptional())
 					{
-						const bool bOptionalDataExist = DoesBulkDataExist(BulkChunkId, BULKDATA_OptionalPayload);
+						const bool bOptionalDataExist = DoesBulkDataExist(FBulkMetaData(BULKDATA_OptionalPayload), BulkChunkId);
 
 						if (bOptionalDataExist)
 						{
@@ -1457,7 +1457,7 @@ void FBulkData::Serialize(FArchive& Ar, UObject* Owner, bool bAttemptFileMapping
 					if (bAttemptFileMapping && (bCanLazyLoad || bIsUsingIoDispatcher))
 					{
 						FIoMappedRegion MappedRegion;
-						if (TryMemoryMapBulkData(BulkChunkId, BulkMeta.GetFlags(), BulkMeta.GetOffset(), BulkMeta.GetSize(), MappedRegion))
+						if (TryMemoryMapBulkData(BulkMeta, BulkChunkId, BulkMeta.GetOffset(), BulkMeta.GetSize(), MappedRegion))
 						{
 							DataAllocation.SetMemoryMappedData(this, MappedRegion.MappedFileHandle, MappedRegion.MappedFileRegion);
 						}
@@ -1837,7 +1837,7 @@ void FBulkData::SerializeBulkData(FArchive& Ar, void* Data, int64 DataSize, EBul
 
 IAsyncReadFileHandle* FBulkData::OpenAsyncReadHandle() const
 {
-	return UE::BulkData::Private::OpenAsyncReadBulkData(BulkChunkId, BulkMeta.GetFlags()).Release();
+	return UE::BulkData::Private::OpenAsyncReadBulkData(BulkMeta, BulkChunkId).Release();
 }
 
 IBulkDataIORequest* FBulkData::CreateStreamingRequest(EAsyncIOPriorityAndFlags Priority, FBulkDataIORequestCallBack* CompleteCallback, uint8* UserSuppliedMemory) const
@@ -1850,8 +1850,8 @@ IBulkDataIORequest* FBulkData::CreateStreamingRequest(EAsyncIOPriorityAndFlags P
 IBulkDataIORequest* FBulkData::CreateStreamingRequest(int64 OffsetInBulkData, int64 BytesToRead, EAsyncIOPriorityAndFlags Priority, FBulkDataIORequestCallBack* CompleteCallback, uint8* UserSuppliedMemory) const
 {
 	return UE::BulkData::Private::CreateStreamingRequest(
+		BulkMeta,
 		BulkChunkId,
-		BulkMeta.GetFlags(),
 		BulkMeta.GetOffset() + OffsetInBulkData,
 		BytesToRead,
 		Priority,
@@ -1881,8 +1881,8 @@ IBulkDataIORequest* FBulkData::CreateStreamingRequestForRange(const BulkDataRang
 		*Start.BulkChunkId.ToDebugString(), *End.BulkChunkId.ToDebugString());
 
 	return UE::BulkData::Private::CreateStreamingRequest(
+		Start.BulkMeta,
 		Start.BulkChunkId,
-		Start.BulkMeta.GetFlags(),
 		ReadOffset,
 		ReadSize,
 		Priority,
@@ -1977,8 +1977,8 @@ bool FBulkData::TryLoadDataIntoMemory(FIoBuffer Dest)
 	}
 
 	const bool bOk = UE::BulkData::Private::OpenReadBulkData(
+		BulkMeta,
 		BulkChunkId,
-		BulkMeta.GetFlags(),
 		BulkDataOffset,
 		BulkDataSizeOnDisk,
 		[this, BulkDataSize, &Dest](FArchive& Ar)
diff --git a/Engine/Source/Runtime/CoreUObject/Private/Serialization/BulkDataStreaming.cpp b/Engine/Source/Runtime/CoreUObject/Private/Serialization/BulkDataStreaming.cpp
--- a/Engine/Source/Runtime/CoreUObject/Private/Serialization/BulkDataStreaming.cpp
+++ b/Engine/Source/Runtime/CoreUObject/Private/Serialization/BulkDataStreaming.cpp
@@ -87,38 +87,43 @@ namespace UE::BulkData::Private
 
 //////////////////////////////////////////////////////////////////////////////
 
-FIoChunkId CreateBulkDataIoChunkId(const FPackageId& PackageId, EBulkDataFlags BulkDataFlags)
+FIoChunkId CreateBulkDataIoChunkId(const FBulkMetaData& BulkMeta, const FPackageId& PackageId)
 {
 	if (PackageId.IsValid() == false)
 	{
 		return FIoChunkId();
 	}
 
+	const EBulkDataFlags BulkDataFlags = BulkMeta.GetFlags();
+
 	const EIoChunkType ChunkType = BulkDataFlags & BULKDATA_OptionalPayload
 		? EIoChunkType::OptionalBulkData
 		: BulkDataFlags & BULKDATA_MemoryMappedPayload
 			? EIoChunkType::MemoryMappedBulkData
 			: EIoChunkType::BulkData;
 
-	return CreateIoChunkId(PackageId.Value(), 0, ChunkType);
+	const uint16 ChunkIndex = EnumHasAnyFlags(BulkMeta.GetMetaFlags(), FBulkMetaData::EMetaFlags::OptionalPackage) ? 1 : 0;
+	return CreateIoChunkId(PackageId.Value(), ChunkIndex, ChunkType);
 }
 
 //////////////////////////////////////////////////////////////////////////////
 
-EPackageSegment GetPackageSegmentFromFlags(EBulkDataFlags BulkDataFlags)
+EPackageSegment GetPackageSegmentFromFlags(const FBulkMetaData& BulkMeta)
 {
+	const EBulkDataFlags BulkDataFlags = BulkMeta.GetFlags();
+
 	if ((BulkDataFlags & BULKDATA_PayloadInSeperateFile) == 0)
 	{
-		if (GEventDrivenLoaderEnabled)
+		const bool bLoadingFromCookedPackage = EnumHasAnyFlags(BulkMeta.GetMetaFlags(), FBulkMetaData::EMetaFlags::CookedPackage);
+		if (bLoadingFromCookedPackage)
 		{
-			// With EDL data, packages are split into EPackageSegment::Header (summary and linker tables) and
+			// Cooked packages are split into EPackageSegment::Header (summary and linker tables) and
 			// EPackageSegment::Exports (serialized UObject bytes and the bulk data section)
 			// Inline and end-of-file bulk data is in the Exports section
 			return EPackageSegment::Exports;
 		}
 		else
 		{
-			// Without EDL, there is only a single combined file for Summary,linkertables,exports,and end-of-file bulk data
 			return EPackageSegment::Header;
 		}
 	}
@@ -449,8 +454,8 @@ TUniquePtr<IBulkDataIORequest> CreateBulkDataIoDispatcherRequest(
 //////////////////////////////////////////////////////////////////////////////
 
 bool OpenReadBulkData(
+	const FBulkMetaData& BulkMeta,
 	const FBulkDataChunkId& BulkChunkId,
-	EBulkDataFlags BulkDataFlags,
 	int64 Offset,
 	int64 Size,
 	TFunction<void(FArchive& Ar)>&& Read)
@@ -462,7 +467,7 @@ bool OpenReadBulkData(
 
 	if (FPackageId PackageId = BulkChunkId.GetPackageId(); PackageId.IsValid())
 	{
-		const FIoChunkId ChunkId = CreateBulkDataIoChunkId(PackageId, BulkDataFlags);
+		const FIoChunkId ChunkId = CreateBulkDataIoChunkId(BulkMeta, PackageId);
 		FIoBatch Batch = FIoDispatcher::Get().NewBatch();
 
 		FIoRequest Request = Batch.Read(ChunkId, FIoReadOptions(Offset, Size), IoDispatcherPriority_Medium);
@@ -482,11 +487,8 @@ bool OpenReadBulkData(
 	{
 		IPackageResourceManager& ResourceMgr = IPackageResourceManager::Get();
 		
-		constexpr EBulkDataFlags ExternalResourceFlags = static_cast<EBulkDataFlags>(BULKDATA_PayloadInSeperateFile | BULKDATA_WorkspaceDomainPayload);
-		const bool bExternalResource = (BulkDataFlags & ExternalResourceFlags) == ExternalResourceFlags;
-
+		const bool bExternalResource = BulkMeta.HasAllFlags(static_cast<EBulkDataFlags>(BULKDATA_PayloadInSeperateFile | BULKDATA_WorkspaceDomainPayload));
 		const FPackagePath& Path = BulkChunkId.GetPackagePath();
-		const EPackageSegment Segment = GetPackageSegmentFromFlags(BulkDataFlags);
 
 		TUniquePtr<FArchive> Ar;
 		if (bExternalResource)
@@ -495,6 +497,7 @@ bool OpenReadBulkData(
 		}
 		else
 		{
+			const EPackageSegment Segment = GetPackageSegmentFromFlags(BulkMeta);
 			Ar = ResourceMgr.OpenReadPackage(Path, Segment).Archive;
 		}
 
@@ -512,7 +515,7 @@ bool OpenReadBulkData(
 
 //////////////////////////////////////////////////////////////////////////////
 
-TUniquePtr<IAsyncReadFileHandle> OpenAsyncReadBulkData(const FBulkDataChunkId& BulkChunkId, EBulkDataFlags BulkDataFlags)
+TUniquePtr<IAsyncReadFileHandle> OpenAsyncReadBulkData(const FBulkMetaData& BulkMeta, const FBulkDataChunkId& BulkChunkId)
 {
 	if (BulkChunkId.IsValid() == false)
 	{
@@ -521,17 +524,14 @@ TUniquePtr<IAsyncReadFileHandle> OpenAsyncReadBulkData(const FBulkDataChunkId& B
 
 	if (FPackageId PackageId = BulkChunkId.GetPackageId(); PackageId.IsValid())
 	{
-		return MakeUnique<FChunkReadFileHandle>(CreateBulkDataIoChunkId(PackageId, BulkDataFlags));
+		return MakeUnique<FChunkReadFileHandle>(CreateBulkDataIoChunkId(BulkMeta, PackageId));
 	}
 	else
 	{
 		IPackageResourceManager& ResourceMgr = IPackageResourceManager::Get();
 		
-		constexpr EBulkDataFlags ExternalResourceFlags = static_cast<EBulkDataFlags>(BULKDATA_PayloadInSeperateFile | BULKDATA_WorkspaceDomainPayload);
-		const bool bExternalResource = (BulkDataFlags & ExternalResourceFlags) == ExternalResourceFlags;
-
+		const bool bExternalResource = BulkMeta.HasAllFlags(static_cast<EBulkDataFlags>(BULKDATA_PayloadInSeperateFile | BULKDATA_WorkspaceDomainPayload));
 		const FPackagePath& Path = BulkChunkId.GetPackagePath();
-		const EPackageSegment Segment = GetPackageSegmentFromFlags(BulkDataFlags);
 
 		if (bExternalResource)
 		{
@@ -539,6 +539,7 @@ TUniquePtr<IAsyncReadFileHandle> OpenAsyncReadBulkData(const FBulkDataChunkId& B
 		}
 		else
 		{
+			const EPackageSegment Segment = GetPackageSegmentFromFlags(BulkMeta);
 			return ResourceMgr.OpenAsyncReadPackage(Path, Segment).Handle;
 		}
 	}
@@ -547,8 +548,8 @@ TUniquePtr<IAsyncReadFileHandle> OpenAsyncReadBulkData(const FBulkDataChunkId& B
 //////////////////////////////////////////////////////////////////////////////
 
 TUniquePtr<IBulkDataIORequest> CreateStreamingRequest(
+	const FBulkMetaData& BulkMeta,
 	const FBulkDataChunkId& BulkChunkId,
-	EBulkDataFlags BulkDataFlags,
 	int64 Offset,
 	int64 Size,
 	EAsyncIOPriorityAndFlags Priority,
@@ -563,14 +564,14 @@ TUniquePtr<IBulkDataIORequest> CreateStreamingRequest(
 	if (FPackageId PackageId = BulkChunkId.GetPackageId(); PackageId.IsValid())
 	{
 		FIoBuffer Buffer = UserSuppliedMemory ? FIoBuffer(FIoBuffer::Wrap, UserSuppliedMemory, Size) : FIoBuffer(Size);
-		const FIoChunkId ChunkId = CreateBulkDataIoChunkId(PackageId, BulkDataFlags);
+		const FIoChunkId ChunkId = CreateBulkDataIoChunkId(BulkMeta, PackageId);
 
 		FChunkBulkDataRequest* Request = new FChunkBulkDataRequest(Callback, MoveTemp(Buffer));
 		Request->Issue(ChunkId, FIoReadOptions(Offset, Size), Priority);
 		
 		return TUniquePtr<IBulkDataIORequest>(Request);
 	}
-	else if (TUniquePtr<IAsyncReadFileHandle> FileHandle = OpenAsyncReadBulkData(BulkChunkId, BulkDataFlags); FileHandle.IsValid())
+	else if (TUniquePtr<IAsyncReadFileHandle> FileHandle = OpenAsyncReadBulkData(BulkMeta, BulkChunkId); FileHandle.IsValid())
 	{
 		TUniquePtr<FBulkDataIORequest> Request = MakeUnique<FBulkDataIORequest>(FileHandle.Release());
 
@@ -585,7 +586,7 @@ TUniquePtr<IBulkDataIORequest> CreateStreamingRequest(
 
 //////////////////////////////////////////////////////////////////////////////
 
-bool DoesBulkDataExist(const FBulkDataChunkId& BulkChunkId, EBulkDataFlags BulkDataFlags)
+bool DoesBulkDataExist(const FBulkMetaData& BulkMeta, const FBulkDataChunkId& BulkChunkId)
 {
 	if (BulkChunkId.IsValid() == false)
 	{
@@ -594,17 +595,14 @@ bool DoesBulkDataExist(const FBulkDataChunkId& BulkChunkId, EBulkDataFlags BulkD
 
 	if (FPackageId Id = BulkChunkId.GetPackageId(); Id.IsValid())
 	{
-		const FIoChunkId ChunkId = CreateBulkDataIoChunkId(Id, BulkDataFlags);
+		const FIoChunkId ChunkId = CreateBulkDataIoChunkId(BulkMeta, Id);
 		return FIoDispatcher::Get().DoesChunkExist(ChunkId);
 	}
 	else
 	{
 		IPackageResourceManager& ResourceMgr = IPackageResourceManager::Get();
 		const FPackagePath& PackagePath = BulkChunkId.GetPackagePath();
-		const EPackageSegment PackageSegment = GetPackageSegmentFromFlags(BulkDataFlags);
-
-		constexpr EBulkDataFlags ExternalResourceFlags = static_cast<EBulkDataFlags>(BULKDATA_PayloadInSeperateFile | BULKDATA_WorkspaceDomainPayload);
-		const bool bExternalResource = (BulkDataFlags & ExternalResourceFlags) == ExternalResourceFlags;
+		const bool bExternalResource = BulkMeta.HasAllFlags(static_cast<EBulkDataFlags>(BULKDATA_PayloadInSeperateFile | BULKDATA_WorkspaceDomainPayload));
 
 		if (bExternalResource)
 		{
@@ -612,6 +610,7 @@ bool DoesBulkDataExist(const FBulkDataChunkId& BulkChunkId, EBulkDataFlags BulkD
 				EPackageExternalResource::WorkspaceDomainFile, PackagePath.GetPackageName());
 		}
 
+		const EPackageSegment PackageSegment = GetPackageSegmentFromFlags(BulkMeta);
 		return ResourceMgr.DoesPackageExist(PackagePath, PackageSegment);
 	}
 }
@@ -619,15 +618,15 @@ bool DoesBulkDataExist(const FBulkDataChunkId& BulkChunkId, EBulkDataFlags BulkD
 //////////////////////////////////////////////////////////////////////////////
 
 bool TryMemoryMapBulkData(
+	const FBulkMetaData& BulkMeta,
 	const FBulkDataChunkId& BulkChunkId,
-	EBulkDataFlags BulkDataFlags,
 	int64 Offset,
 	int64 Size,
 	FIoMappedRegion& OutRegion)
 {
 	if (FPackageId Id = BulkChunkId.GetPackageId(); Id.IsValid())
 	{
-		const FIoChunkId ChunkId = CreateBulkDataIoChunkId(Id, BulkDataFlags);
+		const FIoChunkId ChunkId = CreateBulkDataIoChunkId(BulkMeta, Id);
 		TIoStatusOr<FIoMappedRegion> Status = FIoDispatcher::Get().OpenMapped(ChunkId, FIoReadOptions(Offset, Size));
 
 		if (Status.IsOk())
@@ -709,13 +708,13 @@ private:
 
 bool StartAsyncLoad(
 	FBulkData* Owner,
+	const FBulkMetaData& BulkMeta,
 	const FBulkDataChunkId& BulkChunkId,
-	EBulkDataFlags BulkDataFlags,
 	int64 Offset,
 	int64 Size,
 	TFunction<void(TIoStatusOr<FIoBuffer>)>&& Callback)
 {
-	if (TUniquePtr<IAsyncReadFileHandle> FileHandle = OpenAsyncReadBulkData(BulkChunkId, BulkDataFlags))
+	if (TUniquePtr<IAsyncReadFileHandle> FileHandle = OpenAsyncReadBulkData(BulkMeta, BulkChunkId))
 	{
 		FAsyncFileCallBack FileReadCallback = [Callback = MoveTemp(Callback), Size = Size]
 			(bool bCanceled, IAsyncReadRequest* Request)
diff --git a/Engine/Source/Runtime/CoreUObject/Public/Serialization/BulkData.h b/Engine/Source/Runtime/CoreUObject/Public/Serialization/BulkData.h
--- a/Engine/Source/Runtime/CoreUObject/Public/Serialization/BulkData.h
+++ b/Engine/Source/Runtime/CoreUObject/Public/Serialization/BulkData.h
@@ -279,8 +279,6 @@ private:
 namespace UE::BulkData::Private
 {
 
-COREUOBJECT_API EPackageSegment GetPackageSegmentFromFlags(EBulkDataFlags BulkDataFlags);
-
 /**
  * Serialized bulk meta data.
  */
@@ -315,11 +313,22 @@ struct FBulkMetaResource
  *
  * Uses 5 bytes for size and offset.
  * [0 - 4][5 - 9][  10	][    11	][12    -   15]([16   -  23])
- * [Size][Offset][Unused][LockFlags][BulkDataFlags]([SizeOnDisk])
+ * [Size][Offset][MetaFlags][LockFlags][BulkDataFlags]([SizeOnDisk])
  */
 class FBulkMetaData
 {
 public:
+	enum class EMetaFlags : uint8
+	{
+		/** No additional bulk data flags. */
+		None				= 0,
+		/** Loading from a cooked package. */
+		CookedPackage		= (1 << 0),
+		/** Loading optional package data. */
+		OptionalPackage		= (1 << 1)
+	};
+	FRIEND_ENUM_CLASS_FLAGS(EMetaFlags);
+
 	/** 40 bits for max bulk data size. */
 	static constexpr int64 MaxSize = 0xFFffFFffFF;
 	/** 39 bits for max bulk data offset and 1 bit to indicate INDEX_NONE. */
@@ -333,6 +342,12 @@ public:
 #endif
 	}
 
+	explicit FBulkMetaData(EBulkDataFlags Flags)
+		: FBulkMetaData()
+	{
+		SetFlags(Flags);
+	}
+
 	inline int64 GetSize() const
 	{
 		return static_cast<int64>(ReadUInt40(Data));
@@ -402,6 +417,21 @@ public:
 	{
 		return (GetFlags() & Flags) != 0;
 	}
+
+	inline bool HasAllFlags(EBulkDataFlags Flags) const
+	{
+		return (GetFlags() & Flags) == Flags;
+	}
+
+	void SetMetaFlags(EMetaFlags MetaFlags)
+	{
+		Data[10] = static_cast<uint8>(MetaFlags);
+	}
+	
+	EMetaFlags GetMetaFlags() const
+	{
+		return static_cast<EMetaFlags>(Data[10]);
+	}
 	
 	static FBulkMetaData FromSerialized(const FBulkMetaResource& MetaResource, int64 ElementSize)
 	{
@@ -477,6 +507,12 @@ private:
 	TUniquePtr<FImpl> ImplPtr;
 };
 
+/** Returns an I/O chunk ID to be used when loading from I/O store. */
+FIoChunkId CreateBulkDataIoChunkId(const FBulkMetaData& BulkMeta, const FPackageId& PackageId);
+
+/** Returns the package segment when loading from the package resource manager. */
+COREUOBJECT_API EPackageSegment GetPackageSegmentFromFlags(const FBulkMetaData& BulkMeta);
+
 } // namespace UE::BulkData::Private
 
 /**
@@ -757,7 +793,7 @@ public:
 	const FPackagePath& GetPackagePath() const { return BulkChunkId.GetPackagePath(); }
 
 	/** Returns which segment of its PackagePath this bulkdata resides in */
-	EPackageSegment GetPackageSegment() const { return UE::BulkData::Private::GetPackageSegmentFromFlags(BulkMeta.GetFlags()); }
+	EPackageSegment GetPackageSegment() const { return UE::BulkData::Private::GetPackageSegmentFromFlags(BulkMeta); }
 
 	/** 
 	 * Returns the io filename hash associated with this bulk data.
