TRACE_LOG_MESSAGE inside of UE_LOG does not support extra sizing arguments for strings.

#rnx

diff --git a/Engine/Source/Runtime/CoreUObject/Private/UObject/Class.cpp b/Engine/Source/Runtime/CoreUObject/Private/UObject/Class.cpp
--- a/Engine/Source/Runtime/CoreUObject/Private/UObject/Class.cpp
+++ b/Engine/Source/Runtime/CoreUObject/Private/UObject/Class.cpp
@@ -6036,7 +6036,7 @@ UField* UClass::TryFindTypeSlow(UClass* TypeClass, const FString& InPathNameOrSh
 			
 			FoundType = (UField*)StaticFindFirstObject(TypeClass, *InPathNameOrShortName, InOptions | EFindFirstObjectOptions::EnsureIfAmbiguous | EFindFirstObjectOptions::NativeFirst, ELogVerbosity::Error, TEXT("TryFindType"));
 
-			UE_LOG(LogClass, Warning, TEXT("Short type name \"%s\" provided for TryFindType. Please convert it to a path name (suggested: \"%s\"). Callstack:\r\n\r\n%.*s"), *InPathNameOrShortName, *GetPathNameSafe(FoundType), Callstack.Len(), Callstack.GetData());
+			UE_LOG(LogClass, Warning, TEXT("Short type name \"%s\" provided for TryFindType. Please convert it to a path name (suggested: \"%s\"). Callstack:\r\n\r\n%s"), *InPathNameOrShortName, *GetPathNameSafe(FoundType), Callstack.ToString());
 		}
 	}
 	return FoundType;
