DerivedData: Fixed the namespace of GetTypeHash

#rnx

diff --git a/Engine/Source/Runtime/CoreUObject/Public/Serialization/DerivedData.h b/Engine/Source/Runtime/CoreUObject/Public/Serialization/DerivedData.h
--- a/Engine/Source/Runtime/CoreUObject/Public/Serialization/DerivedData.h
+++ b/Engine/Source/Runtime/CoreUObject/Public/Serialization/DerivedData.h
@@ -166,7 +166,6 @@ struct FDerivedDataIoPriority
 	friend constexpr bool operator==(FDerivedDataIoPriority A, FDerivedDataIoPriority B) { return A.Value == B.Value; }
 	friend constexpr bool operator!=(FDerivedDataIoPriority A, FDerivedDataIoPriority B) { return A.Value != B.Value; }
 	friend constexpr bool operator<(FDerivedDataIoPriority A, FDerivedDataIoPriority B) { return A.Value < B.Value; }
-	friend constexpr uint32 GetTypeHash(const FDerivedDataIoPriority& Priority) { return uint32(Priority.Value); }
 
 	/** Priority. Lower values have lower priority and 0 is equivalent to Normal(). */
 	int32 Value = 0;
@@ -377,6 +376,11 @@ private:
 
 } // UE
 
+constexpr uint32 GetTypeHash(const UE::FDerivedDataIoPriority& Priority)
+{
+	return uint32(Priority.Value);
+}
+
 namespace UE::DerivedData::IoStore
 {
 
