Removing sizing arguments from format strings in logging macros (TRACE_LOG_MESSAGE inside of UE_LOG does not support them)

diff --git a/Engine/Source/Runtime/CoreUObject/Private/UObject/UObjectGlobals.cpp b/Engine/Source/Runtime/CoreUObject/Private/UObject/UObjectGlobals.cpp
--- a/Engine/Source/Runtime/CoreUObject/Private/UObject/UObjectGlobals.cpp
+++ b/Engine/Source/Runtime/CoreUObject/Private/UObject/UObjectGlobals.cpp
@@ -661,14 +661,14 @@ UObject* StaticFindFirstObject(UClass* Class, const TCHAR* Name, EFindFirstObjec
 					}
 					if (AmbiguousMessageVerbosity == ELogVerbosity::Fatal)
 					{
-						UE_LOG(LogUObjectGlobals, Fatal, TEXT("%s"), *Message.ToString());
+						UE_LOG(LogUObjectGlobals, Fatal, TEXT("%s"), Message.ToString());
 					}
 					else if (AmbiguousMessageVerbosity != ELogVerbosity::NoLogging)
 					{
-						GLog->CategorizedLogf(TEXT("LogUObjectGlobals"), AmbiguousMessageVerbosity, TEXT("%.*s"), Message.Len(), Message.GetData());
+						GLog->CategorizedLogf(TEXT("LogUObjectGlobals"), AmbiguousMessageVerbosity, TEXT("%s"), Message.ToString());
 					}
 
-					ensureAlwaysMsgf(!(Options & EFindFirstObjectOptions::EnsureIfAmbiguous), TEXT("%.*s"), Message.Len(), Message.GetData());
+					ensureAlwaysMsgf(!(Options & EFindFirstObjectOptions::EnsureIfAmbiguous), TEXT("%s"), Message.ToString());
 				}
 			}
 			else
