Bulk Data - fix bug when calling IBulkDataRequest::GetReadResult() before calling IBulkDataRequest::GetSize()

#rnx

diff --git a/Engine/Source/Runtime/CoreUObject/Private/Serialization/BulkDataStreaming.cpp b/Engine/Source/Runtime/CoreUObject/Private/Serialization/BulkDataStreaming.cpp
--- a/Engine/Source/Runtime/CoreUObject/Private/Serialization/BulkDataStreaming.cpp
+++ b/Engine/Source/Runtime/CoreUObject/Private/Serialization/BulkDataStreaming.cpp
@@ -160,7 +160,7 @@ public:
 
 	inline EChunkRequestStatus GetStatus() const
 	{
-		return static_cast<EChunkRequestStatus>(Status.load(std::memory_order_relaxed));
+		return static_cast<EChunkRequestStatus>(Status.load(std::memory_order_consume));
 	}
 
 protected:
@@ -169,6 +169,7 @@ protected:
 	virtual void HandleChunkResult(TIoStatusOr<FIoBuffer>&& Result) = 0;
 	bool WaitForChunkRequest(float TimeLimitSeconds = 0.0f);
 	void CancelChunkRequest();
+	int64 GetSizeResult() const { return SizeResult; }
 	
 	FIoBuffer			Buffer;
 
@@ -176,12 +177,14 @@ private:
 
 	UE::FLazyEvent		DoneEvent;
 	FIoRequest			Request;
+	int64				SizeResult;
 	std::atomic<uint32>	Status;
 };
 
 FChunkRequest::FChunkRequest(FIoBuffer&& InBuffer)
 	: Buffer(MoveTemp(InBuffer))
 	, DoneEvent(EEventMode ::ManualReset)
+	, SizeResult(-1)
 	, Status{0}
 {
 }
@@ -201,7 +204,16 @@ void FChunkRequest::Issue(FIoChunkId ChunkId, FIoReadOptions Options, int32 Prio
 	FIoBatch IoBatch = FIoDispatcher::Get().NewBatch();
 	Request = IoBatch.ReadWithCallback(ChunkId, Options, Priority, [this](TIoStatusOr<FIoBuffer> Result)
 	{
-		Status.store(uint32(Result.IsOk() ? EChunkRequestStatus::Ok : EChunkRequestStatus::Canceled), std::memory_order_relaxed); 
+		if (Result.IsOk())
+		{
+			SizeResult = Result.ValueOrDie().GetSize();
+			Status.store(uint32(EChunkRequestStatus::Ok), std::memory_order_release); 
+		}
+		else
+		{
+			Status.store(uint32(EChunkRequestStatus::Canceled), std::memory_order_release); 
+		}
+
 		HandleChunkResult(MoveTemp(Result));
 		DoneEvent.Trigger();
 	});
@@ -383,7 +395,8 @@ public:
 	
 	inline virtual int64 GetSize() const override
 	{
-		return GetStatus() == EChunkRequestStatus::Ok ? Buffer.GetSize() : -1;
+		checkf(GetStatus() != EChunkRequestStatus::None, TEXT("The request must be issued before polling for size"));
+		return GetStatus() == EChunkRequestStatus::Ok ? GetSizeResult() : -1;
 	}
 
 	virtual void Cancel() override
