Fixing an assert related to null FGCObjectInfo::Class caused by improperly initialized FGCObjectInfo instance

diff --git a/Engine/Source/Runtime/CoreUObject/Private/UObject/ReferenceChainSearch.cpp b/Engine/Source/Runtime/CoreUObject/Private/UObject/ReferenceChainSearch.cpp
--- a/Engine/Source/Runtime/CoreUObject/Private/UObject/ReferenceChainSearch.cpp
+++ b/Engine/Source/Runtime/CoreUObject/Private/UObject/ReferenceChainSearch.cpp
@@ -827,10 +827,11 @@ int32 FReferenceChainSearch::PrintResults(TFunctionRef<bool(FCallbackParams& Par
 	{
 		if (TargetObject)
 		{
-			FGCObjectInfo ObjInfo(TargetObject);
+			FGCObjectInfo* ObjInfo = FGCObjectInfo::FindOrAddInfoHelper(TargetObject, *const_cast<TMap<UObject*, FGCObjectInfo*>*>(&ObjectToInfoMap));
+			check(ObjInfo);
 			UE_LOG(LogReferenceChain, Log, TEXT("%s%s is not currently reachable."),
-				*GetObjectFlags(&ObjInfo),
-				*ObjInfo.GetFullName()
+				*GetObjectFlags(ObjInfo),
+				*ObjInfo->GetFullName()
 			);
 		}
 		else
