Removing an assert added with ANY_PACKAGE removal that guards against searching for UClasses with ANY_PACKAGE Outer as it's no longer necessary and would break backwards compatibility for licensees.

diff --git a/Engine/Source/Runtime/CoreUObject/Private/UObject/UObjectHash.cpp b/Engine/Source/Runtime/CoreUObject/Private/UObject/UObjectHash.cpp
--- a/Engine/Source/Runtime/CoreUObject/Private/UObject/UObjectHash.cpp
+++ b/Engine/Source/Runtime/CoreUObject/Private/UObject/UObjectHash.cpp
@@ -698,11 +698,6 @@ UObject* StaticFindObjectFastInternalThreadSafe(FUObjectHashTables& ThreadHash,
 {
 	ExclusiveInternalFlags |= EInternalObjectFlags::Unreachable;
 
-	if (!ObjectPackage && bAnyPackage && ObjectClass && ObjectClass->IsChildOf(UClass::StaticClass()))
-	{
-		UE_LOG(LogUObjectHash, Fatal, TEXT("Looking for classes without specifying their outer is not allowed (class name: %s)"), *ObjectName.ToString());
-		return nullptr;
-	}
 	// If they specified an outer use that during the hashing
 	UObject* Result = nullptr;
 	if (ObjectPackage != nullptr)
