Add a mode for garbage collection to report a sample of garbage references rather than all of them. Can be enabled by setting gc.GarbageReferenceTrackingEnabled to 2.

diff --git a/Engine/Source/Runtime/CoreUObject/Private/UObject/GarbageCollection.cpp b/Engine/Source/Runtime/CoreUObject/Private/UObject/GarbageCollection.cpp
--- a/Engine/Source/Runtime/CoreUObject/Private/UObject/GarbageCollection.cpp
+++ b/Engine/Source/Runtime/CoreUObject/Private/UObject/GarbageCollection.cpp
@@ -263,7 +263,7 @@ int32 GGarbageReferenceTrackingEnabled = 0;
 static FAutoConsoleVariableRef CGarbageReferenceTrackingEnabled(
 	TEXT("gc.GarbageReferenceTrackingEnabled"),
 	GGarbageReferenceTrackingEnabled,
-	TEXT("If true, Garbage Collector will track and log unreleased garbage objects"),
+	TEXT("Causes the Garbage Collector to track and log unreleased garbage objects. If 1, will dump every reference. If 2, will dump a sample of the references to highlight problematic properties."),
 	ECVF_Default
 );
 
@@ -2475,34 +2475,73 @@ FString FGarbageReferenceInfo::GetReferencingObjectInfo() const
 }
 
 void FGCArrayPool::DumpGarbageReferencers(TArray<FGCArrayStruct*>& AllArrays)
-{	
-	if (GGarbageReferenceTrackingEnabled)
+{		
+	if (GGarbageReferenceTrackingEnabled == 0)
 	{
-		const bool bIsEngineExitRequested = IsEngineExitRequested();
-		int32 NumGarbageReferences = 0;
+		return;
+	}
+
+	// We don't care about leaks when engine exit was requested since the final GC pass will destroy everything anyway
+	// We still want to clear the GarbageReferences array though
+	if (IsEngineExitRequested())
+	{
+		for (FGCArrayStruct* ArrayStruct : AllArrays)
+		{
+			ArrayStruct->GarbageReferences.Reset();
+		}
+		return;
+	}
 
+	double StartTime = FPlatformTime::Seconds();
+	// Dump all garbage references in detail
+	if (GGarbageReferenceTrackingEnabled == 1)
+	{
+		int32 NumGarbageReferences = 0;
+		for (FGCArrayStruct* ArrayStruct : AllArrays)
+		{
+			for (FGarbageReferenceInfo& GarbageReference : ArrayStruct->GarbageReferences)
+			{
+				// No need to spam with references that would never get released anyway (even if we still had PendingKill) because they were referenced by a persistent reference 
+				if (!GarbageReference.GarbageObject->HasAnyInternalFlags(EInternalObjectFlags::PersistentGarbage))
+				{
+					UE_LOG(LogGarbage, Error, TEXT("Reachable garbage object: %s"), *GarbageReference.GarbageObject->GetFullName());
+					UE_LOG(LogGarbage, Error, TEXT("Referenced by:            %s"), *GarbageReference.GetReferencingObjectInfo());
+					UE_LOG(LogGarbage, Error, TEXT(""));
+					NumGarbageReferences++;
+				}
+			}
+			ArrayStruct->GarbageReferences.Reset();
+		}
+		UE_CLOG(NumGarbageReferences > 0, LogGarbage, Log, TEXT("Found %d garbage references in %f ms."), NumGarbageReferences, (FPlatformTime::Seconds() - StartTime) * 1000);
+	}
+	// Dump one reference for each source property
+	else if (GGarbageReferenceTrackingEnabled == 2)
+	{
+		typedef TPair<void*, FName> FKey;
+		TSet<FKey> Seen;
+		int32 TotalGarbageReferences = 0, ReportedGarbageReferences = 0;
 		for (FGCArrayStruct* ArrayStruct : AllArrays)
 		{
-			// We don't care about leaks when engine exit was requested since the final GC pass will destroy everything anyway
-			// We still want to clear the GarbageReferences array though
-			if (!bIsEngineExitRequested)
+			for (FGarbageReferenceInfo& GarbageReference : ArrayStruct->GarbageReferences)
 			{
-				for (FGarbageReferenceInfo& GarbageReference : ArrayStruct->GarbageReferences)
+				// No need to spam with references that would never get released anyway (even if we still had PendingKill) because they were referenced by a persistent reference 
+				if (!GarbageReference.GarbageObject->HasAnyInternalFlags(EInternalObjectFlags::PersistentGarbage))
 				{
-					// No need to spam with references that would never get released anyway (even if we still had PendingKill) because they were referenced by a persistent reference 
-					if (!GarbageReference.GarbageObject->HasAnyInternalFlags(EInternalObjectFlags::PersistentGarbage))
+					FKey Key(GarbageReference.bReferencerUObject ? (void*)GarbageReference.Referencer.Object->GetClass() : (void*)  GarbageReference.Referencer.GCObject, GarbageReference.PropertyName);
+					if (Seen.Contains(Key) == false)
 					{
 						UE_LOG(LogGarbage, Error, TEXT("Reachable garbage object: %s"), *GarbageReference.GarbageObject->GetFullName());
 						UE_LOG(LogGarbage, Error, TEXT("Referenced by:            %s"), *GarbageReference.GetReferencingObjectInfo());
 						UE_LOG(LogGarbage, Error, TEXT(""));
-						NumGarbageReferences++;
+						Seen.Add(Key);
+						ReportedGarbageReferences++;
 					}
+					TotalGarbageReferences++;
 				}
 			}
 			ArrayStruct->GarbageReferences.Reset();
 		}
-
-		UE_CLOG(NumGarbageReferences > 0, LogGarbage, Log, TEXT("Found %d garbage references."), NumGarbageReferences);
+		UE_CLOG(TotalGarbageReferences > 0, LogGarbage, Log, TEXT("Reported %d/%d garbage references in %f ms."), ReportedGarbageReferences, TotalGarbageReferences, (FPlatformTime::Seconds() - StartTime) * 1000);
 	}
 }
 
