Improved assert in FLinkerLoad to state the object being preloaded

#rnx

diff --git a/Engine/Source/Runtime/CoreUObject/Private/UObject/LinkerLoad.cpp b/Engine/Source/Runtime/CoreUObject/Private/UObject/LinkerLoad.cpp
--- a/Engine/Source/Runtime/CoreUObject/Private/UObject/LinkerLoad.cpp
+++ b/Engine/Source/Runtime/CoreUObject/Private/UObject/LinkerLoad.cpp
@@ -3994,7 +3994,7 @@ void FLinkerLoad::Preload( UObject* Object )
 			TRACE_CPUPROFILER_EVENT_SCOPE_TEXT_ON_CHANNEL(*WriteToString<256>(TEXT("FLinkerLoad::Preload "), Object->GetFName()), AssetLoadTimeChannel);
 
 			UClass* Cls = Cast<UClass>(Object);
-			check(!GEventDrivenLoaderEnabled || !bLockoutLegacyOperations || !EVENT_DRIVEN_ASYNC_LOAD_ACTIVE_AT_RUNTIME);
+			checkf(!GEventDrivenLoaderEnabled || !bLockoutLegacyOperations || !EVENT_DRIVEN_ASYNC_LOAD_ACTIVE_AT_RUNTIME, TEXT("Invalid call to FLinkerLoad::Preload while using the EDL. '%d' should have been reported via GetPreloadDependencies instead."), *Object->GetPathName());
 #if USE_CIRCULAR_DEPENDENCY_LOAD_DEFERRING
 			bool const bIsNonNativeObject = !Object->GetOutermost()->HasAnyPackageFlags(PKG_CompiledIn);
 			// we can determine that this is a blueprint class/struct by checking if it 
