AssetDataTagMap: operator== performance and readability improvement.
#rnx

diff --git a/Engine/Source/Runtime/CoreUObject/Private/AssetRegistry/AssetDataTagMap.cpp b/Engine/Source/Runtime/CoreUObject/Private/AssetRegistry/AssetDataTagMap.cpp
--- a/Engine/Source/Runtime/CoreUObject/Private/AssetRegistry/AssetDataTagMap.cpp
+++ b/Engine/Source/Runtime/CoreUObject/Private/AssetRegistry/AssetDataTagMap.cpp
@@ -1455,6 +1455,39 @@ bool operator==(const FAssetDataTagMapSharedView& A, const FAssetDataTagMap& B)
 	return A.Num() == 0 || *A.Loose == B;
 }
 
+namespace FixedTagPrivate
+{
+	static bool MapEqualsHelper(FixedTagPrivate::FMapHandle Fixed, FAssetDataTagMap* Loose)
+	{
+		checkSlow(Fixed.Num == Loose->Num());
+		// Since Num is the same and the maps are unique, test only whether all keys in Fixed exist with equal value
+		// in Loose. Once we've done that, we don't have to test whether any keys in Loose are missing from Fixed.
+		if (Fixed.HasNumberlessKeys != 0)
+		{
+			for (FNumberlessPair Pair : Fixed.GetNumberlessView())
+			{
+				FString* LooseValue = Loose->Find(FName::CreateFromDisplayId(Pair.Key, 0));
+				if (!LooseValue || FAssetTagValueRef(Fixed.StoreIndex, Pair.Value) != *LooseValue)
+				{
+					return false;
+				}
+			}
+		}
+		else
+		{
+			for (FNumberedPair Pair : Fixed.GetNumberedView())
+			{
+				FString* LooseValue = Loose->Find(Pair.Key);
+				if (!LooseValue || FAssetTagValueRef(Fixed.StoreIndex, Pair.Value) != *LooseValue)
+				{
+					return false;
+				}
+			}
+		}
+		return true;
+	}
+}
+
 bool operator==(const FAssetDataTagMapSharedView& A, const FAssetDataTagMapSharedView& B)
 {
 	if (A.Num() != B.Num())
@@ -1467,32 +1500,8 @@ bool operator==(const FAssetDataTagMapSharedView& A, const FAssetDataTagMapShare
 	}
 	else if (A.IsFixed() != B.IsFixed())
 	{
-		FixedTagPrivate::FMapHandle Fixed;
-		FAssetDataTagMap* Loose;
-		if (A.IsFixed())
-		{
-			Fixed = A.Fixed;
-			Loose = B.Loose;
-		}
-		else
-		{
-			Fixed = B.Fixed;
-			Loose = A.Loose;
-		}
-		// Since the num is the same and these are unique maps, we only have to test whether all keys in Fixed exist with equal value in Loose
-		// Once we've done that, we don't have to test whether any keys in Loose are missing from Fixed.
-		bool bEqual = true;
-		Fixed.ForEachPair([&](FixedTagPrivate::FNumberedPair FixedPair)
-			{
-				FString* LooseValue = Loose->Find(FixedPair.Key);
-				if (!LooseValue)
-				{
-					bEqual = false;
-					return;
-				}
-				bEqual &= (FAssetTagValueRef(Fixed.StoreIndex, FixedPair.Value) == *LooseValue);
-			});
-		return bEqual;
+		return A.IsFixed() ? FixedTagPrivate::MapEqualsHelper(A.Fixed, B.Loose) :
+			FixedTagPrivate::MapEqualsHelper(B.Fixed, A.Loose);
 	}
 	else if (A.IsFixed())
 	{
