Turning off logging when failing to convert short type names when running diff asset registries commandlet (second attempt)

diff --git a/Engine/Source/Runtime/CoreUObject/Private/AssetRegistry/AssetData.cpp b/Engine/Source/Runtime/CoreUObject/Private/AssetRegistry/AssetData.cpp
--- a/Engine/Source/Runtime/CoreUObject/Private/AssetRegistry/AssetData.cpp
+++ b/Engine/Source/Runtime/CoreUObject/Private/AssetRegistry/AssetData.cpp
@@ -214,7 +214,7 @@ void FAssetData::SerializeForCacheInternal(FArchive& Ar, FAssetRegistryVersion::
 	{
 PRAGMA_DISABLE_DEPRECATION_WARNINGS
 		Ar << AssetClass;
-		AssetClassPath = FAssetData::TryConvertShortClassNameToPathName(AssetClass);
+		AssetClassPath = FAssetData::TryConvertShortClassNameToPathName(AssetClass, ELogVerbosity::NoLogging);
 PRAGMA_ENABLE_DEPRECATION_WARNINGS
 	}
 	else
