Downgrading an assert to an ensure when initializing FTopLevelAssetPath with a short asset (class) name

diff --git a/Engine/Source/Runtime/CoreUObject/Private/UObject/TopLevelAssetPath.cpp b/Engine/Source/Runtime/CoreUObject/Private/UObject/TopLevelAssetPath.cpp
--- a/Engine/Source/Runtime/CoreUObject/Private/UObject/TopLevelAssetPath.cpp
+++ b/Engine/Source/Runtime/CoreUObject/Private/UObject/TopLevelAssetPath.cpp
@@ -104,7 +104,7 @@ bool FTopLevelAssetPath::TrySetPath(FWideStringView Path)
 
 			if (Path.IsEmpty() || Path[0] != '/')
 			{
-				checkf(false, TEXT("Short asset name used to create FTopLevelAssetPath: \"%.*s\""), Path.Len(), Path.GetData());
+				ensureAlwaysMsgf(false, TEXT("Short asset name used to create FTopLevelAssetPath: \"%.*s\""), Path.Len(), Path.GetData());
 				Reset();
 				return false;
 			}
