C++20 fixes to FOpenPackageResult and FOpenAsyncPackageResult and braced-init.

diff --git a/Engine/Source/Runtime/CoreUObject/Private/UObject/PackageResourceManager.cpp b/Engine/Source/Runtime/CoreUObject/Private/UObject/PackageResourceManager.cpp
--- a/Engine/Source/Runtime/CoreUObject/Private/UObject/PackageResourceManager.cpp
+++ b/Engine/Source/Runtime/CoreUObject/Private/UObject/PackageResourceManager.cpp
@@ -257,6 +257,13 @@
 	}
 }
 
+FOpenAsyncPackageResult::FOpenAsyncPackageResult(TUniquePtr<IAsyncReadFileHandle>&& InHandle, EPackageFormat InFormat, bool bInNeedsEngineVersionChecks)
+	: Handle(MoveTemp(InHandle))
+	, Format(InFormat)
+	, bNeedsEngineVersionChecks(bInNeedsEngineVersionChecks)
+{
+}
+
 FOpenAsyncPackageResult::~FOpenAsyncPackageResult()
 {
 	// Defined in CPP so we can forward declare IAsyncReadFileHandle
diff --git a/Engine/Source/Runtime/CoreUObject/Public/UObject/PackageResourceManager.h b/Engine/Source/Runtime/CoreUObject/Public/UObject/PackageResourceManager.h
--- a/Engine/Source/Runtime/CoreUObject/Public/UObject/PackageResourceManager.h
+++ b/Engine/Source/Runtime/CoreUObject/Public/UObject/PackageResourceManager.h
@@ -55,11 +55,6 @@ struct FOpenPackageResult
 	 */
 	bool bNeedsEngineVersionChecks = true;
 
-	FOpenPackageResult() = default;
-	FOpenPackageResult(const FOpenPackageResult&) = delete;
-	FOpenPackageResult(FOpenPackageResult&&) = default;
-	FOpenPackageResult& operator=(const FOpenPackageResult&) = delete;
-	FOpenPackageResult& operator=(FOpenPackageResult&&) = default;
 	void CopyMetaData(const FOpenPackageResult& Other)
 	{
 		Format = Other.Format;
@@ -88,7 +83,10 @@ struct FOpenAsyncPackageResult
 	FOpenAsyncPackageResult(FOpenAsyncPackageResult&&) = default;
 	FOpenAsyncPackageResult& operator=(const FOpenAsyncPackageResult&) = delete;
 	FOpenAsyncPackageResult& operator=(FOpenAsyncPackageResult&&) = default;
+
+	COREUOBJECT_API FOpenAsyncPackageResult(TUniquePtr<IAsyncReadFileHandle>&& InHandle, EPackageFormat InFormat, bool bInNeedsEngineVersionChecks = true);
 	COREUOBJECT_API ~FOpenAsyncPackageResult();
+
 	void CopyMetaData(const FOpenPackageResult& Other)
 	{
 		Format = Other.Format;
