EditorBulkdata now longer stores a EPackageSegment member since it is always set to EPackageSegment::Header

#rnx

- Storing the segment type as a member made the code would imply to the reader that it could be changed, even though in all paths it was being set to  EPackageSegment::Header. Removing the member and just hard coding the segment use should make the code easier to follow.
- Editor bulkdata always stores it's payload in the package trailer, which for now counts as the package header.
- The only deviation for this is the experimental feature that allows us to store the payloads in a sidecar file. This code path is not currently supported, but has been updated to use hardcoded EPackageSegment::PayloadSidecar.
- Given that PackageSegment was always set to EPackageSegment::Header we can easily remove any assert of if statement checking for that.

diff --git a/Engine/Source/Runtime/CoreUObject/Private/Serialization/EditorBulkData.cpp b/Engine/Source/Runtime/CoreUObject/Private/Serialization/EditorBulkData.cpp
--- a/Engine/Source/Runtime/CoreUObject/Private/Serialization/EditorBulkData.cpp
+++ b/Engine/Source/Runtime/CoreUObject/Private/Serialization/EditorBulkData.cpp
@@ -187,10 +187,8 @@ static const FPackageTrailer* GetTrailerFromOwner(UObject* Owner)
 }
 
 /** Utility for finding the package path associated with a given UObject */
-static FPackagePath GetPackagePathFromOwner(UObject* Owner, EPackageSegment& OutPackageSegment)
+static FPackagePath GetPackagePathFromOwner(UObject* Owner)
 {
-	OutPackageSegment = EPackageSegment::Header;
-
 	const FLinkerLoad* Linker = GetLinkerLoadFromOwner(Owner);
 
 	if (Linker != nullptr)
@@ -386,7 +384,6 @@ FEditorBulkData& FEditorBulkData::operator=(FEditorBulkData&& Other)
 	PayloadSize = MoveTemp(Other.PayloadSize);
 	OffsetInFile = MoveTemp(Other.OffsetInFile);
 	PackagePath = MoveTemp(Other.PackagePath);
-	PackageSegment = MoveTemp(Other.PackageSegment);
 	Flags = MoveTemp(Other.Flags);
 	CompressionSettings = MoveTemp(Other.CompressionSettings);
 
@@ -433,7 +430,6 @@ FEditorBulkData& FEditorBulkData::operator=(const FEditorBulkData& Other)
 	PayloadSize = Other.PayloadSize;
 	OffsetInFile = Other.OffsetInFile;
 	PackagePath = Other.PackagePath;
-	PackageSegment = Other.PackageSegment;
 	Flags = Other.Flags;
 	CompressionSettings = Other.CompressionSettings;
 
@@ -695,7 +691,6 @@ void FEditorBulkData::CreateFromBulkData(FBulkData& InBulkData, const FGuid& InG
 	PayloadSize = InBulkData.GetBulkDataSize();
 	
 	PackagePath = InBulkData.GetPackagePath();
-	PackageSegment = InBulkData.GetPackageSegment();
 	
 	OffsetInFile = InBulkData.GetBulkDataOffsetInFile();
 
@@ -757,7 +752,6 @@ void FEditorBulkData::Serialize(FArchive& Ar, UObject* Owner, bool bAllowRegiste
 			Ar << PayloadContentId;
 			Ar << PayloadSize;
 			Ar << PackagePath;
-			Ar << PackageSegment;
 			Ar << OffsetInFile;
 
 			// TODO: We could consider compressing the payload so it takes up less space in the 
@@ -802,7 +796,7 @@ void FEditorBulkData::Serialize(FArchive& Ar, UObject* Owner, bool bAllowRegiste
 			LinkerSave = Cast<FLinkerSave>(Ar.GetLinker());
 			// If we're doing a save that can refer to bulk data by reference, and our legacy data format supports it,
 			// keep any legacy data we have referenced rather than stored, to save space and avoid spending time loading it.
-			bKeepFileDataByReference = LinkerSave && LinkerSave->bProceduralSave && PackageSegment == EPackageSegment::Header;
+			bKeepFileDataByReference = LinkerSave != nullptr && LinkerSave->bProceduralSave;
 			if (!bKeepFileDataByReference)
 			{
 				UpdateKeyIfNeeded();
@@ -863,7 +857,6 @@ void FEditorBulkData::Serialize(FArchive& Ar, UObject* Owner, bool bAllowRegiste
 			bool bWriteOutPayload = true;
 			if (IsReferencingByPackagePath(UpdatedFlags))
 			{
-				check(PackageSegment == EPackageSegment::Header); // This should have been checked before setting bKeepFileDataByReference=true
 				if (!IsStoredInPackageTrailer(UpdatedFlags))
 				{
 					Ar << OffsetInFile;
@@ -972,7 +965,6 @@ void FEditorBulkData::Serialize(FArchive& Ar, UObject* Owner, bool bAllowRegiste
 
 			OffsetInFile = INDEX_NONE;
 			PackagePath.Empty();
-			PackageSegment = EPackageSegment::Header;
 
 			const FPackageTrailer* Trailer = GetTrailerFromOwner(Owner);
 				
@@ -1022,12 +1014,11 @@ void FEditorBulkData::Serialize(FArchive& Ar, UObject* Owner, bool bAllowRegiste
 				FArchive* CacheableArchive = Ar.GetCacheableArchive();
 				if (Ar.IsAllowingLazyLoading() && CacheableArchive != nullptr)
 				{
-					PackagePath = GetPackagePathFromOwner(Owner, PackageSegment);
+					PackagePath = GetPackagePathFromOwner(Owner);
 				}
 				else
 				{
 					PackagePath.Empty();
-					PackageSegment = EPackageSegment::Header;
 				}
 					
 				if (!PackagePath.IsEmpty() && CacheableArchive != nullptr)
@@ -1090,7 +1081,6 @@ void FEditorBulkData::SerializeForRegistry(FArchive& Ar)
 	if (Ar.IsSaving())
 	{
 		FString PackageName = PackagePath.GetPackageName();
-		check(PackageName.IsEmpty() || PackageSegment == EPackageSegment::Header);
 		Ar << PackageName;
 	}
 	else
@@ -1105,7 +1095,6 @@ void FEditorBulkData::SerializeForRegistry(FArchive& Ar)
 		{
 			ensure(FPackagePath::TryFromPackageName(PackageName, PackagePath));
 		}
-		PackageSegment = EPackageSegment::Header;
 	}
 	Ar << OffsetInFile;
 }
@@ -1113,8 +1102,7 @@ void FEditorBulkData::SerializeForRegistry(FArchive& Ar)
 bool FEditorBulkData::CanSaveForRegistry() const
 {
 	return BulkDataId.IsValid() && PayloadSize > 0 && !IsMemoryOnlyPayload()
-		&& EnumHasAnyFlags(Flags, EFlags::IsTornOff) && !EnumHasAnyFlags(Flags, EFlags::HasRegistered)
-		&& (PackagePath.IsEmpty() || PackageSegment == EPackageSegment::Header);
+		&& EnumHasAnyFlags(Flags, EFlags::IsTornOff) && !EnumHasAnyFlags(Flags, EFlags::HasRegistered);
 }
 
 
@@ -1171,13 +1159,13 @@ FCompressedBuffer FEditorBulkData::LoadFromPackageFile() const
 {
 	TRACE_CPUPROFILER_EVENT_SCOPE(FEditorBulkData::LoadFromPackageFile);
 
-	UE_LOG(LogSerialization, Verbose, TEXT("Attempting to load payload from the package file: '%s'"), *PackagePath.GetLocalFullPath(PackageSegment));
+	UE_LOG(LogSerialization, Verbose, TEXT("Attempting to load payload from the package file: '%s'"), *PackagePath.GetLocalFullPath(EPackageSegment::Header));
 
 	// Open a reader to the file
 	TUniquePtr<FArchive> BulkArchive;
-	if (!IsReferencingByPackagePath() || PackageSegment != EPackageSegment::Header)
+	if (!IsReferencingByPackagePath())
 	{
-		FOpenPackageResult Result = IPackageResourceManager::Get().OpenReadPackage(PackagePath, PackageSegment);
+		FOpenPackageResult Result = IPackageResourceManager::Get().OpenReadPackage(PackagePath, EPackageSegment::Header);
 		if (Result.Format == EPackageFormat::Binary)
 		{
 			BulkArchive = MoveTemp(Result.Archive);
@@ -1195,11 +1183,11 @@ FCompressedBuffer FEditorBulkData::LoadFromPackageFile() const
 
 	if (!BulkArchive.IsValid())
 	{
-		LogPackageOpenFailureMessage(PackagePath, PackageSegment);
+		LogPackageOpenFailureMessage(PackagePath, EPackageSegment::Header);
 		return FCompressedBuffer();
 	}
 
-	checkf(OffsetInFile != INDEX_NONE, TEXT("Attempting to load '%s' from disk with an invalid OffsetInFile!"), *PackagePath.GetDebugNameWithExtension(PackageSegment));
+	checkf(OffsetInFile != INDEX_NONE, TEXT("Attempting to load '%s' from disk with an invalid OffsetInFile!"), *PackagePath.GetDebugNameWithExtension(EPackageSegment::Header));
 	// Move the correct location of the data in the file
 	BulkArchive->Seek(OffsetInFile);
 
@@ -1214,15 +1202,15 @@ FCompressedBuffer FEditorBulkData::LoadFromPackageTrailer() const
 {
 	TRACE_CPUPROFILER_EVENT_SCOPE(FEditorBulkData::LoadFromPackageTrailer);
 
-	UE_LOG(LogSerialization, Verbose, TEXT("Attempting to load payload from the package trailer: '%s'"), *PackagePath.GetLocalFullPath(PackageSegment));
+	UE_LOG(LogSerialization, Verbose, TEXT("Attempting to load payload from the package trailer: '%s'"), *PackagePath.GetLocalFullPath(EPackageSegment::Header));
 
 	// TODO: Could just get the trailer from the owning FLinkerLoad if still attached
 
 	// Open a reader to the file
 	TUniquePtr<FArchive> BulkArchive;
-	if (!IsReferencingByPackagePath() || PackageSegment != EPackageSegment::Header)
+	if (!IsReferencingByPackagePath())
 	{
-		FOpenPackageResult Result = IPackageResourceManager::Get().OpenReadPackage(PackagePath, PackageSegment);
+		FOpenPackageResult Result = IPackageResourceManager::Get().OpenReadPackage(PackagePath, EPackageSegment::Header);
 		if (Result.Format == EPackageFormat::Binary)
 		{
 			BulkArchive = MoveTemp(Result.Archive);
@@ -1240,7 +1228,7 @@ FCompressedBuffer FEditorBulkData::LoadFromPackageTrailer() const
 
 	if (!BulkArchive.IsValid())
 	{
-		LogPackageOpenFailureMessage(PackagePath, PackageSegment);
+		LogPackageOpenFailureMessage(PackagePath, EPackageSegment::Header);
 		return FCompressedBuffer();
 	}
 
@@ -1425,7 +1413,6 @@ void FEditorBulkData::PushData(const FPackagePath& InPackagePath)
 			// Clear members associated with non-virtualized data and release the in-memory
 			// buffer.
 			PackagePath.Empty();
-			PackageSegment = EPackageSegment::Header;
 			OffsetInFile = INDEX_NONE;
 
 			// Update our information in the registry
@@ -1477,7 +1464,6 @@ void FEditorBulkData::Reset()
 	PayloadSize = 0;
 	OffsetInFile = INDEX_NONE;
 	PackagePath.Empty();
-	PackageSegment = EPackageSegment::Header;
 	Flags = EFlags::None;
 
 	CompressionSettings.Reset();
@@ -1508,7 +1494,6 @@ void FEditorBulkData::DetachFromDisk(FArchive* Ar, bool bEnsurePayloadIsLoaded)
 		}
 
 		PackagePath.Empty();
-		PackageSegment = EPackageSegment::Header;
 		OffsetInFile = INDEX_NONE;
 
 		EnumRemoveFlags(Flags, EFlags::ReferencesLegacyFile | EFlags::ReferencesWorkspaceDomain | EFlags::LegacyFileIsCompressed);
@@ -1706,7 +1691,6 @@ void FEditorBulkData::UpdatePayloadImpl(FSharedBuffer&& InPayload, FIoHash&& InP
 							EFlags::LegacyKeyWasGuidDerived);
 
 	PackagePath.Empty();
-	PackageSegment = EPackageSegment::Header;
 	OffsetInFile = INDEX_NONE;
 
 	if (PayloadSize > 0 && !BulkDataId.IsValid())
diff --git a/Engine/Source/Runtime/CoreUObject/Public/Serialization/EditorBulkData.h b/Engine/Source/Runtime/CoreUObject/Public/Serialization/EditorBulkData.h
--- a/Engine/Source/Runtime/CoreUObject/Public/Serialization/EditorBulkData.h
+++ b/Engine/Source/Runtime/CoreUObject/Public/Serialization/EditorBulkData.h
@@ -497,9 +497,6 @@ private:
 	/** PackagePath containing the payload (this will be empty if the payload does not come from PackageResourceManager)*/
 	FPackagePath PackagePath;
 
-	/** PackageSegment to load with the packagepath (unused if the payload does not come from PackageResourceManager) */
-	EPackageSegment PackageSegment= EPackageSegment::Header;
-
 	/** A 32bit bitfield of flags */
 	EFlags Flags = EFlags::None;
 
