Fixed UObject::GetPreloadDependencies not gathering from C++ classes

CLASS_Intrinsic is apparently set on all C++ classes, even those parsed by UHT, despite what its comment suggests...

diff --git a/Engine/Source/Runtime/CoreUObject/Private/UObject/Obj.cpp b/Engine/Source/Runtime/CoreUObject/Private/UObject/Obj.cpp
--- a/Engine/Source/Runtime/CoreUObject/Private/UObject/Obj.cpp
+++ b/Engine/Source/Runtime/CoreUObject/Private/UObject/Obj.cpp
@@ -1306,16 +1306,16 @@ void UObject::GetPreloadDependencies(TArray<UObject*>& OutDeps)
 		{
 			OutDeps.Add(ObjClass->GetDefaultObject());
 		}
+	}
 
-		// The iterator will recursively loop through all structs in structs/containers too.
-		for (TPropertyValueIterator<FStructProperty> It(ObjClass, this); It; ++It)
+	// The iterator will recursively loop through all structs in structs/containers too.
+	for (TPropertyValueIterator<FStructProperty> It(ObjClass, this); It; ++It)
+	{
+		const UScriptStruct* StructType = It.Key()->Struct;
+		if (UScriptStruct::ICppStructOps* CppStructOps = StructType->GetCppStructOps())
 		{
-			const UScriptStruct* StructType = It.Key()->Struct;
-			if (UScriptStruct::ICppStructOps* CppStructOps = StructType->GetCppStructOps())
-			{
-				void* StructDataPtr = const_cast<void*>(It.Value());
-				CppStructOps->GetPreloadDependencies(StructDataPtr, OutDeps);
-			}
+			void* StructDataPtr = const_cast<void*>(It.Value());
+			CppStructOps->GetPreloadDependencies(StructDataPtr, OutDeps);
 		}
 	}
 }
