Remove the validation option when loading payloads from sidecar files as we now validate all loaded payloads at a higher level.

#rnx

diff --git a/Engine/Source/Runtime/CoreUObject/Private/Serialization/EditorBulkData.cpp b/Engine/Source/Runtime/CoreUObject/Private/Serialization/EditorBulkData.cpp
--- a/Engine/Source/Runtime/CoreUObject/Private/Serialization/EditorBulkData.cpp
+++ b/Engine/Source/Runtime/CoreUObject/Private/Serialization/EditorBulkData.cpp
@@ -59,11 +59,6 @@ static TAutoConsoleVariable<bool> CVarShouldLoadFromTrailer(
 	false,
 	TEXT("When true FEditorBulkData will load payloads via the package trailer rather than the package itself"));
 
-static TAutoConsoleVariable<bool> CVarShouldValidatePayload(
-	TEXT("Serialization.ValidatePayloads"),
-	false,
-	TEXT("When true FEditorBulkData validate any payload loaded from the sidecar file"));
-
 static TAutoConsoleVariable<bool> CVarShouldAllowSidecarSyncing(
 	TEXT("Serialization.AllowSidecarSyncing"),
 	false,
@@ -1119,28 +1114,7 @@ FCompressedBuffer FEditorBulkData::LoadFromDisk() const
 
 	if (HasPayloadSidecarFile() && CVarShouldLoadFromSidecar.GetValueOnAnyThread())
 	{
-		// Note that this code path is purely for debugging and not expected to be enabled by default
-		if (CVarShouldValidatePayload.GetValueOnAnyThread())
-		{
-			UE_LOG(LogSerialization, Verbose, TEXT("Validating payload loaded from sidecar file: '%s'"), *PackagePath.GetLocalFullPath(EPackageSegment::PayloadSidecar));
-
-			// Load both payloads then generate a FPayloadId from them, since this identifier is a hash of the buffers content
-			// we only need to verify them against PayloadContentId to be sure that the data is correct.
-			FCompressedBuffer SidecarBuffer = LoadFromSidecarFile();
-			FCompressedBuffer AssetBuffer = LoadFromPackageFile();
-
-			const FIoHash SidecarId = HashPayload(SidecarBuffer.Decompress());
-			const FIoHash AssetId = HashPayload(AssetBuffer.Decompress());
-
-			UE_CLOG(SidecarId != PayloadContentId, LogSerialization, Error, TEXT("Sidecar content did not hash correctly! Found '%s' Expected '%s'"), *LexToString(SidecarId), *LexToString(PayloadContentId));
-			UE_CLOG(AssetId != PayloadContentId, LogSerialization, Error, TEXT("Asset content did not hash correctly! Found '%s' Expected '%s'"), *LexToString(AssetId), *LexToString(PayloadContentId))
-
-			return SidecarBuffer;
-		}
-		else
-		{
-			return LoadFromSidecarFile();
-		}
+		return LoadFromSidecarFile();
 	}
 	else
 	{
