AsyncLoading2: Also store public export objects without the RF_WasLoaded in the import store when loading in the editor (BP compilation can clear this flag)

diff --git a/Engine/Source/Runtime/CoreUObject/Private/Serialization/AsyncLoading2.cpp b/Engine/Source/Runtime/CoreUObject/Private/Serialization/AsyncLoading2.cpp
--- a/Engine/Source/Runtime/CoreUObject/Private/Serialization/AsyncLoading2.cpp
+++ b/Engine/Source/Runtime/CoreUObject/Private/Serialization/AsyncLoading2.cpp
@@ -195,6 +195,11 @@ static FString GAsyncLoading2_DebugPackageNamesString;
 static TSet<FPackageId> GAsyncLoading2_VerbosePackageIds;
 static FString GAsyncLoading2_VerbosePackageNamesString;
 static int32 GAsyncLoading2_VerboseLogFilter = 2; //None=0,Filter=1,All=2
+#if WITH_EDITOR
+static constexpr EObjectFlags GlobalImportObjectConditionalFlags = RF_Public;
+#else
+static constexpr EObjectFlags GlobalImportObjectConditionalFlags = RF_Public | RF_WasLoaded; // Assumes that we don't do any renaming of objects clearing the RF_WasLoaded flag
+#endif
 #if !UE_BUILD_SHIPPING
 static void ParsePackageNames(const FString& PackageNamesString, TSet<FPackageId>& PackageIds)
 {
@@ -506,7 +511,7 @@ struct FGlobalImportStore
 				{
 					UObject* GCObject = Item.Value;
 
-					checkf(GCObject->HasAllFlags(RF_WasLoaded | RF_Public),
+					checkf(GCObject->HasAllFlags(GlobalImportObjectConditionalFlags),
 						TEXT("The serialized GC Object '%s' with (ObjectFlags=%x, InternalObjectFlags=%x)and id 0x%llX:0x%llX is currently missing RF_WasLoaded or RF_Public. ")
 						TEXT("The flags must have been altered incorrectly by a higher level system after the object was loaded. ")
 						TEXT("This may cause memory corruption."),
@@ -972,7 +977,7 @@ private:
 		}
 		ForEachObjectWithOuter(ImportedPackage, [](UObject* Object)
 		{
-			if (Object->HasAllFlags(RF_Public | RF_WasLoaded))
+			if (Object->HasAllFlags(GlobalImportObjectConditionalFlags))
 			{
 				checkf(!Object->HasAnyInternalFlags(EInternalObjectFlags::LoaderImport), TEXT("%s"), *Object->GetFullName());
 				Object->SetInternalFlags(EInternalObjectFlags::LoaderImport);
@@ -992,7 +997,7 @@ private:
 		}
 		ForEachObjectWithOuter(ImportedPackage, [](UObject* Object)
 		{
-			if (Object->HasAllFlags(RF_Public | RF_WasLoaded))
+			if (Object->HasAllFlags(GlobalImportObjectConditionalFlags))
 			{
 				// RF_Public can be added to objects after they've been created and in that case they might not have the LoaderImport flag set
 				//checkf(Object->HasAnyInternalFlags(EInternalObjectFlags::LoaderImport), TEXT("%s"), *Object->GetFullName());
@@ -3339,7 +3344,7 @@ void FAsyncPackage2::ImportPackagesRecursiveInner(FIoBatch& IoBatch, FPackageSto
 
 					ForEachObjectWithOuter(UncookedPackage, [this, ImportedPackageId](UObject* Object)
 					{
-						if (Object->HasAllFlags(RF_Public | RF_WasLoaded))
+						if (Object->HasAllFlags(GlobalImportObjectConditionalFlags))
 						{
 							checkf(!Object->HasAnyInternalFlags(EInternalObjectFlags::LoaderImport), TEXT("%s"), *Object->GetFullName());
 							Object->SetInternalFlags(EInternalObjectFlags::LoaderImport);
@@ -5628,7 +5633,7 @@ FORCENOINLINE static void FilterUnreachableObjects(
 		UObject* Object = static_cast<UObject*>(ObjectItem->Object);
 		// Including all objects is slow,
 		// but allows RemovePublicExports to check for serialized public exports that have unintentionally lost their flags
-		if (DO_CHECK || Object->HasAllFlags(RF_WasLoaded | RF_Public))
+		if (DO_CHECK || Object->HasAllFlags(GlobalImportObjectConditionalFlags))
 		{
 			if (Object->GetOuter())
 			{
