Marked the FAssetData::AssetClass, FARFilter::ClassNames and FARFilter::RecursiveClassesExeclusionSet as 'Deprecated Property'.
  - Print a deprecated warning message for Python users trying to access the property.

diff --git a/Engine/Source/Runtime/CoreUObject/Public/UObject/NoExportTypes.h b/Engine/Source/Runtime/CoreUObject/Public/UObject/NoExportTypes.h
--- a/Engine/Source/Runtime/CoreUObject/Public/UObject/NoExportTypes.h
+++ b/Engine/Source/Runtime/CoreUObject/Public/UObject/NoExportTypes.h
@@ -2147,7 +2147,7 @@ struct FARFilter
 	TArray<FName> ObjectPaths;
 
 	/** [DEPRECATED] - Class names are now represented by path names. Please use ClassPaths instead. */
-	UPROPERTY(BlueprintReadWrite, Category = AssetRegistry)
+	UPROPERTY(BlueprintReadWrite, Category = AssetRegistry, meta=(DeprecatedProperty, DeprecationMessage="Short asset class names must be converted to full asset pathnames. Use ClassPaths instead."))
 	TArray<FName> ClassNames;
 
 	/** The filter component for class path names. Instances of the specified classes, but not subclasses (by default), will be included. Derived classes will be included only if bRecursiveClasses is true. */
@@ -2158,7 +2158,7 @@ struct FARFilter
 	TMultiMap<FName, TOptional<FString>> TagsAndValues;
 
 	/** [DEPRECATED] - Class names are now represented by path names. Please use RecursiveClassPathsExclusionSet instead. */
-	UPROPERTY(BlueprintReadWrite, Category = AssetRegistry)
+	UPROPERTY(BlueprintReadWrite, Category = AssetRegistry, meta=(DeprecatedProperty, DeprecationMessage="Short asset class names must be converted to full asset pathnames. Use RecursiveClassPathsExclusionSet instead."))
 	TSet<FName> RecursiveClassesExclusionSet;
 
 	/** Only if bRecursiveClasses is true, the results will exclude classes (and subclasses) in this list */
@@ -2224,8 +2224,8 @@ struct FAssetData
 	/** The name of the asset without the package */
 	UPROPERTY(BlueprintReadOnly, Category = AssetData, transient)
 	FName AssetName;
-	/** Deprecated. The name of the asset's class */
-	UPROPERTY(BlueprintReadOnly, Category = AssetData, transient)
+	/** The name of the asset's class */
+	UPROPERTY(BlueprintReadOnly, Category = AssetData, transient, meta=(DeprecatedProperty, DeprecationMessage="Short asset class name must be converted to full asset pathname. Use AssetClassPath instead."))
 	FName AssetClass;
 	/** The path name of the asset's class */
 	UPROPERTY(BlueprintReadOnly, Category = AssetData, transient)
