Changing CppStructOps registration code use struct path names instead of short names to allow identically named native structs to be registered

diff --git a/Engine/Source/Editor/Sequencer/Private/SequencerKeyStructGenerator.cpp b/Engine/Source/Editor/Sequencer/Private/SequencerKeyStructGenerator.cpp
--- a/Engine/Source/Editor/Sequencer/Private/SequencerKeyStructGenerator.cpp
+++ b/Engine/Source/Editor/Sequencer/Private/SequencerKeyStructGenerator.cpp
@@ -100,7 +100,7 @@ void FSequencerKeyStructGenerator::FinalizeNewKeyStruct(UMovieSceneKeyStructType
 	InStruct->Bind();
 	InStruct->StaticLink(true);
 
-	UMovieSceneKeyStructType::DeferCppStructOps(InStruct->GetFName(), new UScriptStruct::TCppStructOps<FGeneratedMovieSceneKeyStruct>);
+	UMovieSceneKeyStructType::DeferCppStructOps(InStruct->GetStructPathName(), new UScriptStruct::TCppStructOps<FGeneratedMovieSceneKeyStruct>);
 
 	check(InStruct->IsComplete());
 }
diff --git a/Engine/Source/Editor/UnrealEd/Private/Editor.cpp b/Engine/Source/Editor/UnrealEd/Private/Editor.cpp
--- a/Engine/Source/Editor/UnrealEd/Private/Editor.cpp
+++ b/Engine/Source/Editor/UnrealEd/Private/Editor.cpp
@@ -167,7 +167,7 @@ PRAGMA_ENABLE_DEPRECATION_WARNINGS
 	Globals.
 -----------------------------------------------------------------------------*/
 
-IMPLEMENT_STRUCT(SlatePlayInEditorInfo);
+UE_IMPLEMENT_STRUCT("/Script/UnrealEd", SlatePlayInEditorInfo);
 
 //////////////////////////////////////////////////////////////////////////
 // FReimportManager
diff --git a/Engine/Source/Programs/UnrealHeaderTool/Private/PropertyTypes.cpp b/Engine/Source/Programs/UnrealHeaderTool/Private/PropertyTypes.cpp
--- a/Engine/Source/Programs/UnrealHeaderTool/Private/PropertyTypes.cpp
+++ b/Engine/Source/Programs/UnrealHeaderTool/Private/PropertyTypes.cpp
@@ -2254,7 +2254,7 @@ struct FPropertyTypeTraitsStruct : public FPropertyTypeTraitsBase
 	static FString GetCPPTypeForwardDeclaration(const FUnrealPropertyDefinitionInfo& PropDef)
 	{
 		// Core type structs don't need to forward declare in UHT as every generated.h indirectly includes CoreMinimal.h
-		if (UScriptStruct::ICppStructOps* CppStructOps = UScriptStruct::FindDeferredCppStructOps(PropDef.GetPropertyBase().ScriptStructDef->GetFName()); CppStructOps != nullptr && CppStructOps->IsUECoreType())
+		if (UScriptStruct::ICppStructOps* CppStructOps = UScriptStruct::FindDeferredCppStructOps(PropDef.GetPropertyBase().ScriptStructDef->GetStructPathName()); CppStructOps != nullptr && CppStructOps->IsUECoreType())
 		{
 			return FString();
 		}
diff --git a/Engine/Source/Programs/UnrealHeaderTool/Private/UnrealTypeDefinitionInfo.cpp b/Engine/Source/Programs/UnrealHeaderTool/Private/UnrealTypeDefinitionInfo.cpp
--- a/Engine/Source/Programs/UnrealHeaderTool/Private/UnrealTypeDefinitionInfo.cpp
+++ b/Engine/Source/Programs/UnrealHeaderTool/Private/UnrealTypeDefinitionInfo.cpp
@@ -1166,6 +1166,11 @@ void FUnrealObjectDefinitionInfo::GetPathName(FUnrealObjectDefinitionInfo* StopO
 		ResultString << TEXT("None");
 	}
 }
+FTopLevelAssetPath FUnrealObjectDefinitionInfo::GetStructPathName() const
+{
+	checkf(GetOuter() && !GetOuter()->GetOuter(), TEXT("GetStructPathName can only be used with top level objects (current: \"%s\")"), *GetPathName());
+	return FTopLevelAssetPath(GetOuter()->GetFName(), GetFName());
+}
 
 FString FUnrealObjectDefinitionInfo::GetFullGroupName(bool bStartWithOuter) const
 {
@@ -1991,7 +1996,7 @@ void FUnrealScriptStructDefinitionInfo::SetHasInstancedReference()
 
 bool FUnrealScriptStructDefinitionInfo::HasNoOpConstructor() const
 {
-	UScriptStruct::ICppStructOps* CppStructOps = UScriptStruct::FindDeferredCppStructOps(GetFName());
+	UScriptStruct::ICppStructOps* CppStructOps = UScriptStruct::FindDeferredCppStructOps(GetStructPathName());
 	return CppStructOps ? CppStructOps->HasNoopConstructor() : false;
 }
 
diff --git a/Engine/Source/Programs/UnrealHeaderTool/Private/UnrealTypeDefinitionInfo.h b/Engine/Source/Programs/UnrealHeaderTool/Private/UnrealTypeDefinitionInfo.h
--- a/Engine/Source/Programs/UnrealHeaderTool/Private/UnrealTypeDefinitionInfo.h
+++ b/Engine/Source/Programs/UnrealHeaderTool/Private/UnrealTypeDefinitionInfo.h
@@ -461,6 +461,11 @@ public:
 	virtual FString GetPathName(FUnrealObjectDefinitionInfo* StopOuter = nullptr) const = 0;
 	virtual void GetPathName(FUnrealObjectDefinitionInfo* StopOuter, FStringBuilderBase& ResultString) const = 0;
 
+	/**
+	 * Returns the fully qualified pathname for this object as an FTopLevelAssetPath
+	 */
+	virtual FTopLevelAssetPath GetStructPathName() const = 0;
+
 	/**
 	 * Walks up the chain of packages until it reaches the top level, which it ignores.
 	 *
@@ -716,6 +721,15 @@ public:
 	virtual FString GetPathName(FUnrealObjectDefinitionInfo* StopOuter = nullptr) const override;
 	virtual void GetPathName(FUnrealObjectDefinitionInfo* StopOuter, FStringBuilderBase& ResultString) const override;
 
+	/**
+	 * Returns the fully qualified pathname for this object as an FTopLevelAssetPath
+	 */
+	virtual FTopLevelAssetPath GetStructPathName() const override
+	{
+		checkf(false, TEXT("Proeprties do not support GetStructPathName()"));
+		return FTopLevelAssetPath();
+	}
+
 	/**
 	 * Walks up the chain of packages until it reaches the top level, which it ignores.
 	 *
@@ -1349,6 +1363,11 @@ public:
 	virtual FString GetPathName(FUnrealObjectDefinitionInfo* StopOuter = nullptr) const override;
 	virtual void GetPathName(FUnrealObjectDefinitionInfo* StopOuter, FStringBuilderBase& ResultString) const override;
 
+	/**
+	 * Returns the fully qualified pathname for this object as an FTopLevelAssetPath
+	 */
+	virtual FTopLevelAssetPath GetStructPathName() const override;
+
 	/**
 	 * Walks up the chain of packages until it reaches the top level, which it ignores.
 	 *
@@ -2480,7 +2499,7 @@ public:
 	 */
 	bool HasDefaults() const
 	{
-		return UScriptStruct::FindDeferredCppStructOps(GetFName()) != nullptr;
+		return UScriptStruct::FindDeferredCppStructOps(GetStructPathName()) != nullptr;
 	}
 
 
diff --git a/Engine/Source/Runtime/CoreUObject/Private/AssetRegistry/AssetBundleData.cpp b/Engine/Source/Runtime/CoreUObject/Private/AssetRegistry/AssetBundleData.cpp
--- a/Engine/Source/Runtime/CoreUObject/Private/AssetRegistry/AssetBundleData.cpp
+++ b/Engine/Source/Runtime/CoreUObject/Private/AssetRegistry/AssetBundleData.cpp
@@ -4,7 +4,7 @@
 #include "AssetRegistry/AssetData.h"
 #include "UObject/PropertyPortFlags.h"
 
-IMPLEMENT_STRUCT(AssetBundleData);
+UE_IMPLEMENT_STRUCT("/Script/CoreUObject", AssetBundleData);
 
 
 namespace UE { namespace AssetBundleEntry { namespace Private
diff --git a/Engine/Source/Runtime/CoreUObject/Private/AssetRegistry/AssetData.cpp b/Engine/Source/Runtime/CoreUObject/Private/AssetRegistry/AssetData.cpp
--- a/Engine/Source/Runtime/CoreUObject/Private/AssetRegistry/AssetData.cpp
+++ b/Engine/Source/Runtime/CoreUObject/Private/AssetRegistry/AssetData.cpp
@@ -14,8 +14,8 @@
 
 DEFINE_LOG_CATEGORY(LogAssetData);
 
-IMPLEMENT_STRUCT(ARFilter);
-IMPLEMENT_STRUCT(AssetData);
+UE_IMPLEMENT_STRUCT("/Script/CoreUObject", ARFilter);
+UE_IMPLEMENT_STRUCT("/Script/CoreUObject", AssetData);
 
 // Register Asset Registry version
 const FGuid FAssetRegistryVersion::GUID(0x717F9EE7, 0xE9B0493A, 0x88B39132, 0x1B388107);
diff --git a/Engine/Source/Runtime/CoreUObject/Private/UObject/Class.cpp b/Engine/Source/Runtime/CoreUObject/Private/UObject/Class.cpp
--- a/Engine/Source/Runtime/CoreUObject/Private/UObject/Class.cpp
+++ b/Engine/Source/Runtime/CoreUObject/Private/UObject/Class.cpp
@@ -65,7 +65,7 @@
 DEFINE_LOG_CATEGORY(LogScriptSerialization);
 DEFINE_LOG_CATEGORY(LogClass);
 
-IMPLEMENT_STRUCT(TestUninitializedScriptStructMembersTest);
+UE_IMPLEMENT_STRUCT("/Script/CoreUObject", TestUninitializedScriptStructMembersTest);
 
 #if defined(_MSC_VER) && _MSC_VER == 1900
 	#ifdef PRAGMA_DISABLE_SHADOW_VARIABLE_WARNINGS
@@ -2376,9 +2376,9 @@ struct TStructOpsTypeTraits<FTestStruct> : public TStructOpsTypeTraitsBase2<FTes
 /** Used to hold virtual methods to construct, destruct, etc native structs in a generic and dynamic fashion 
  * singleton-style to avoid issues with static constructor order
 **/
-static TMap<FName,UScriptStruct::ICppStructOps*>& GetDeferredCppStructOps()
+static TMap<FTopLevelAssetPath,UScriptStruct::ICppStructOps*>& GetDeferredCppStructOps()
 {
-	static struct TMapWithAutoCleanup : public TMap<FName, UScriptStruct::ICppStructOps*>
+	static struct TMapWithAutoCleanup : public TMap<FTopLevelAssetPath, UScriptStruct::ICppStructOps*>
 	{
 		~TMapWithAutoCleanup()
 		{
@@ -2518,9 +2518,9 @@ UScriptStruct::UScriptStruct(const FObjectInitializer& ObjectInitializer)
 	* @param Target Name of the struct 
 	* @param InCppStructOps Cpp ops for this struct
 **/
-void UScriptStruct::DeferCppStructOps(FName Target, ICppStructOps* InCppStructOps)
+void UScriptStruct::DeferCppStructOps(FTopLevelAssetPath Target, ICppStructOps* InCppStructOps)
 {
-	TMap<FName,UScriptStruct::ICppStructOps*>& DeferredStructOps = GetDeferredCppStructOps();
+	TMap<FTopLevelAssetPath, UScriptStruct::ICppStructOps*>& DeferredStructOps = GetDeferredCppStructOps();
 
 	if (UScriptStruct::ICppStructOps* ExistingOps = DeferredStructOps.FindRef(Target))
 	{
@@ -2537,11 +2537,11 @@ void UScriptStruct::DeferCppStructOps(FName Target, ICppStructOps* InCppStructOp
 		}
 		// in reload, we will just leak these...they may be in use.
 	}
-	DeferredStructOps.Add(Target,InCppStructOps);
+	DeferredStructOps.Add(Target, InCppStructOps);
 }
 
 #if HACK_HEADER_GENERATOR
-UScriptStruct::ICppStructOps* UScriptStruct::FindDeferredCppStructOps(FName StructName)
+UScriptStruct::ICppStructOps* UScriptStruct::FindDeferredCppStructOps(FTopLevelAssetPath StructName)
 {
 	return GetDeferredCppStructOps().FindRef(StructName);
 }
@@ -2555,13 +2555,18 @@ void UScriptStruct::PrepareCppStructOps()
 		return;
 	}
 	if (!CppStructOps)
-	{
-		CppStructOps = GetDeferredCppStructOps().FindRef(GetFName());
+	{	
+		// RobM: the use of GetFlattenedStructPathName() is a bit of a hack to make AnimBPs work as they create nested structs.
+		// Theoretically we could just wrap the line below in if (StructFlags&STRUCT_Native) but the native flag is set after
+		// CppStructOps struct is actually found (see below).
+		// It's also confusing why PrepareCppStructOps is being called from UScriptStruct::Serialize when loading since serialized
+		// structs are not native and they should not implement CppStructOps
+		CppStructOps = GetDeferredCppStructOps().FindRef(GetFlattenedStructPathName());
 		if (!CppStructOps)
 		{
 			if (!GIsUCCMakeStandaloneHeaderGenerator && (StructFlags&STRUCT_Native))
 			{
-				UE_LOG(LogClass, Fatal,TEXT("Couldn't bind to native struct %s. Headers need to be rebuilt, or a noexport class is missing a IMPLEMENT_STRUCT."),*GetName());
+				UE_LOG(LogClass, Fatal, TEXT("Couldn't bind to native struct %s. Headers need to be rebuilt, or a noexport class is missing a IMPLEMENT_STRUCT."),*GetName());
 			}
 			check(!bPrepareCppStructOpsCompleted); // recursion is unacceptable
 			bPrepareCppStructOpsCompleted = true;
diff --git a/Engine/Source/Runtime/CoreUObject/Private/UObject/Property.cpp b/Engine/Source/Runtime/CoreUObject/Private/UObject/Property.cpp
--- a/Engine/Source/Runtime/CoreUObject/Private/UObject/Property.cpp
+++ b/Engine/Source/Runtime/CoreUObject/Private/UObject/Property.cpp
@@ -43,9 +43,9 @@ struct TVector3StructOpsTypeTraits : public TStructOpsTypeTraitsBase2<T>
 };
 template<> struct TStructOpsTypeTraits<FVector3f> : public TVector3StructOpsTypeTraits<FVector3f> {};
 template<> struct TStructOpsTypeTraits<FVector3d> : public TVector3StructOpsTypeTraits<FVector3d> {};
-IMPLEMENT_STRUCT(Vector3f);
-IMPLEMENT_STRUCT(Vector3d);
-IMPLEMENT_STRUCT(Vector);	// Aliased
+UE_IMPLEMENT_STRUCT("/Script/CoreUObject", Vector3f);
+UE_IMPLEMENT_STRUCT("/Script/CoreUObject", Vector3d);
+UE_IMPLEMENT_STRUCT("/Script/CoreUObject", Vector);	// Aliased
 
 template<>
 struct TStructOpsTypeTraits<FIntPoint> : public TStructOpsTypeTraitsBase2<FIntPoint>
@@ -58,7 +58,7 @@ struct TStructOpsTypeTraits<FIntPoint> : public TStructOpsTypeTraitsBase2<FIntPo
 		WithSerializer = true,
 	};
 };
-IMPLEMENT_STRUCT(IntPoint);
+UE_IMPLEMENT_STRUCT("/Script/CoreUObject", IntPoint);
 
 template<>
 struct TStructOpsTypeTraits<FIntVector> : public TStructOpsTypeTraitsBase2<FIntVector>
@@ -71,7 +71,7 @@ struct TStructOpsTypeTraits<FIntVector> : public TStructOpsTypeTraitsBase2<FIntV
 		WithSerializer = true,
 	};
 };
-IMPLEMENT_STRUCT(IntVector);
+UE_IMPLEMENT_STRUCT("/Script/CoreUObject", IntVector);
 
 template<>
 struct TStructOpsTypeTraits<FIntVector4> : public TStructOpsTypeTraitsBase2<FIntVector4>
@@ -102,9 +102,9 @@ struct TVector2StructOpsTypeTraits : public TStructOpsTypeTraitsBase2<T>
 };
 template<> struct TStructOpsTypeTraits<FVector2f> : public TVector2StructOpsTypeTraits<FVector2f> {};
 template<> struct TStructOpsTypeTraits<FVector2d> : public TVector2StructOpsTypeTraits<FVector2d> {};
-IMPLEMENT_STRUCT(Vector2f);
-IMPLEMENT_STRUCT(Vector2d);
-IMPLEMENT_STRUCT(Vector2D);
+UE_IMPLEMENT_STRUCT("/Script/CoreUObject", Vector2f);
+UE_IMPLEMENT_STRUCT("/Script/CoreUObject", Vector2d);
+UE_IMPLEMENT_STRUCT("/Script/CoreUObject", Vector2D);
 
 template<typename T>
 struct TVector4StructOpsTypeTraits : public TStructOpsTypeTraitsBase2<T>
@@ -120,9 +120,9 @@ struct TVector4StructOpsTypeTraits : public TStructOpsTypeTraitsBase2<T>
 };
 template<> struct TStructOpsTypeTraits<FVector4f> : public TVector4StructOpsTypeTraits<FVector4f> {};
 template<> struct TStructOpsTypeTraits<FVector4d> : public TVector4StructOpsTypeTraits<FVector4d> {};
-IMPLEMENT_STRUCT(Vector4f);
-IMPLEMENT_STRUCT(Vector4d);
-IMPLEMENT_STRUCT(Vector4);	// Aliased
+UE_IMPLEMENT_STRUCT("/Script/CoreUObject", Vector4f);
+UE_IMPLEMENT_STRUCT("/Script/CoreUObject", Vector4d);
+UE_IMPLEMENT_STRUCT("/Script/CoreUObject", Vector4);	// Aliased
 
 template<typename T>
 struct TPlaneStructOpsTypeTraits : public TStructOpsTypeTraitsBase2<T>
@@ -140,9 +140,9 @@ struct TPlaneStructOpsTypeTraits : public TStructOpsTypeTraitsBase2<T>
 };
 template<> struct TStructOpsTypeTraits<FPlane4f> : public TPlaneStructOpsTypeTraits<FPlane4f> {};
 template<> struct TStructOpsTypeTraits<FPlane4d> : public TPlaneStructOpsTypeTraits<FPlane4d> {};
-IMPLEMENT_STRUCT(Plane4f);
-IMPLEMENT_STRUCT(Plane4d);
-IMPLEMENT_STRUCT(Plane);	// Aliased
+UE_IMPLEMENT_STRUCT("/Script/CoreUObject", Plane4f);
+UE_IMPLEMENT_STRUCT("/Script/CoreUObject", Plane4d);
+UE_IMPLEMENT_STRUCT("/Script/CoreUObject", Plane);	// Aliased
 
 template<typename T>
 struct TRotatorStructOpsTypeTraits : public TStructOpsTypeTraitsBase2<T>
@@ -161,9 +161,9 @@ struct TRotatorStructOpsTypeTraits : public TStructOpsTypeTraitsBase2<T>
 
 template<> struct TStructOpsTypeTraits<FRotator3f> : public TRotatorStructOpsTypeTraits<FRotator3f> {};
 template<> struct TStructOpsTypeTraits<FRotator3d> : public TRotatorStructOpsTypeTraits<FRotator3d> {};
-IMPLEMENT_STRUCT(Rotator3f);
-IMPLEMENT_STRUCT(Rotator3d);
-IMPLEMENT_STRUCT(Rotator);	// Aliased
+UE_IMPLEMENT_STRUCT("/Script/CoreUObject", Rotator3f);
+UE_IMPLEMENT_STRUCT("/Script/CoreUObject", Rotator3d);
+UE_IMPLEMENT_STRUCT("/Script/CoreUObject", Rotator);	// Aliased
 
 template<typename T>
 struct TBox3StructOpsTypeTraits : public TStructOpsTypeTraitsBase2<T>
@@ -179,9 +179,9 @@ struct TBox3StructOpsTypeTraits : public TStructOpsTypeTraitsBase2<T>
 };
 template<> struct TStructOpsTypeTraits<FBox3f> : public TBox3StructOpsTypeTraits<FBox3f> {};
 template<> struct TStructOpsTypeTraits<FBox3d> : public TBox3StructOpsTypeTraits<FBox3d> {};
-IMPLEMENT_STRUCT(Box3f);
-IMPLEMENT_STRUCT(Box3d);
-IMPLEMENT_STRUCT(Box);		// Aliased
+UE_IMPLEMENT_STRUCT("/Script/CoreUObject", Box3f);
+UE_IMPLEMENT_STRUCT("/Script/CoreUObject", Box3d);
+UE_IMPLEMENT_STRUCT("/Script/CoreUObject", Box);		// Aliased
 
 template<typename T>
 struct TBox2StructOpsTypeTraits : public TStructOpsTypeTraitsBase2<T>
@@ -196,9 +196,9 @@ struct TBox2StructOpsTypeTraits : public TStructOpsTypeTraitsBase2<T>
 };
 template<> struct TStructOpsTypeTraits<FBox2f> : public TBox2StructOpsTypeTraits<FBox2f> {};
 template<> struct TStructOpsTypeTraits<FBox2d> : public TBox2StructOpsTypeTraits<FBox2d> {};
-IMPLEMENT_STRUCT(Box2f);
-IMPLEMENT_STRUCT(Box2d);
-IMPLEMENT_STRUCT(Box2D);
+UE_IMPLEMENT_STRUCT("/Script/CoreUObject", Box2f);
+UE_IMPLEMENT_STRUCT("/Script/CoreUObject", Box2d);
+UE_IMPLEMENT_STRUCT("/Script/CoreUObject", Box2D);
 
 template<typename T>
 struct TMatrixStructOpsTypeTraits : public TStructOpsTypeTraitsBase2<T>
@@ -215,9 +215,9 @@ struct TMatrixStructOpsTypeTraits : public TStructOpsTypeTraitsBase2<T>
 
 template<> struct TStructOpsTypeTraits<FMatrix44f> : public TMatrixStructOpsTypeTraits<FMatrix44f> {};
 template<> struct TStructOpsTypeTraits<FMatrix44d> : public TMatrixStructOpsTypeTraits<FMatrix44d> {};
-IMPLEMENT_STRUCT(Matrix44f);
-IMPLEMENT_STRUCT(Matrix44d);
-IMPLEMENT_STRUCT(Matrix);	// Aliased
+UE_IMPLEMENT_STRUCT("/Script/CoreUObject", Matrix44f);
+UE_IMPLEMENT_STRUCT("/Script/CoreUObject", Matrix44d);
+UE_IMPLEMENT_STRUCT("/Script/CoreUObject", Matrix);	// Aliased
 
 template<typename T>
 struct TBoxSphereBoundsStructOpsTypeTraits : public TStructOpsTypeTraitsBase2<T>
@@ -232,15 +232,15 @@ struct TBoxSphereBoundsStructOpsTypeTraits : public TStructOpsTypeTraitsBase2<T>
 };
 template<> struct TStructOpsTypeTraits<FBoxSphereBounds3f> : public TBoxSphereBoundsStructOpsTypeTraits<FBoxSphereBounds3f> {};
 template<> struct TStructOpsTypeTraits<FBoxSphereBounds3d> : public TBoxSphereBoundsStructOpsTypeTraits<FBoxSphereBounds3d> {};
-IMPLEMENT_STRUCT(BoxSphereBounds3f);
-IMPLEMENT_STRUCT(BoxSphereBounds3d);
-IMPLEMENT_STRUCT(BoxSphereBounds);	// Aliased
+UE_IMPLEMENT_STRUCT("/Script/CoreUObject", BoxSphereBounds3f);
+UE_IMPLEMENT_STRUCT("/Script/CoreUObject", BoxSphereBounds3d);
+UE_IMPLEMENT_STRUCT("/Script/CoreUObject", BoxSphereBounds);	// Aliased
 
 template<>
 struct TStructOpsTypeTraits<FOrientedBox> : public TStructOpsTypeTraitsBase2<FOrientedBox>
 {
 };
-IMPLEMENT_STRUCT(OrientedBox);
+UE_IMPLEMENT_STRUCT("/Script/CoreUObject", OrientedBox);
 
 template<>
 struct TStructOpsTypeTraits<FLinearColor> : public TStructOpsTypeTraitsBase2<FLinearColor>
@@ -253,7 +253,7 @@ struct TStructOpsTypeTraits<FLinearColor> : public TStructOpsTypeTraitsBase2<FLi
 		WithStructuredSerializer = true,
 	};
 };
-IMPLEMENT_STRUCT(LinearColor);
+UE_IMPLEMENT_STRUCT("/Script/CoreUObject", LinearColor);
 
 template<>
 struct TStructOpsTypeTraits<FColor> : public TStructOpsTypeTraitsBase2<FColor>
@@ -266,7 +266,7 @@ struct TStructOpsTypeTraits<FColor> : public TStructOpsTypeTraitsBase2<FColor>
 		WithSerializer = true,
 	};
 };
-IMPLEMENT_STRUCT(Color);
+UE_IMPLEMENT_STRUCT("/Script/CoreUObject", Color);
 
 
 template<typename T>
@@ -285,9 +285,9 @@ struct TQuatStructOpsTypeTraits : public TStructOpsTypeTraitsBase2<T>
 };
 template<> struct TStructOpsTypeTraits<FQuat4f> : public TQuatStructOpsTypeTraits<FQuat4f> {};
 template<> struct TStructOpsTypeTraits<FQuat4d> : public TQuatStructOpsTypeTraits<FQuat4d> {};
-IMPLEMENT_STRUCT(Quat4f);
-IMPLEMENT_STRUCT(Quat4d);
-IMPLEMENT_STRUCT(Quat);		// Aliased to one of FQuat4f/FQuat4d
+UE_IMPLEMENT_STRUCT("/Script/CoreUObject", Quat4f);
+UE_IMPLEMENT_STRUCT("/Script/CoreUObject", Quat4d);
+UE_IMPLEMENT_STRUCT("/Script/CoreUObject", Quat);		// Aliased to one of FQuat4f/FQuat4d
 
 template<>
 struct TStructOpsTypeTraits<FTwoVectors> : public TStructOpsTypeTraitsBase2<FTwoVectors>
@@ -300,7 +300,7 @@ struct TStructOpsTypeTraits<FTwoVectors> : public TStructOpsTypeTraitsBase2<FTwo
 		WithNoDestructor = true,
 	};
 };
-IMPLEMENT_STRUCT(TwoVectors);
+UE_IMPLEMENT_STRUCT("/Script/CoreUObject", TwoVectors);
 
 template<>
 struct TStructOpsTypeTraits<FInterpCurvePointFloat> : public TStructOpsTypeTraitsBase2<FInterpCurvePointFloat>
@@ -310,7 +310,7 @@ struct TStructOpsTypeTraits<FInterpCurvePointFloat> : public TStructOpsTypeTrait
 		WithNoInitConstructor = true,
 	};
 };
-IMPLEMENT_STRUCT(InterpCurvePointFloat);
+UE_IMPLEMENT_STRUCT("/Script/CoreUObject", InterpCurvePointFloat);
 
 template<>
 struct TStructOpsTypeTraits<FInterpCurvePointVector2D> : public TStructOpsTypeTraitsBase2<FInterpCurvePointVector2D>
@@ -320,7 +320,7 @@ struct TStructOpsTypeTraits<FInterpCurvePointVector2D> : public TStructOpsTypeTr
 		WithNoInitConstructor = true,
 	};
 };
-IMPLEMENT_STRUCT(InterpCurvePointVector2D);
+UE_IMPLEMENT_STRUCT("/Script/CoreUObject", InterpCurvePointVector2D);
 
 template<>
 struct TStructOpsTypeTraits<FInterpCurvePointVector> : public TStructOpsTypeTraitsBase2<FInterpCurvePointVector>
@@ -330,7 +330,7 @@ struct TStructOpsTypeTraits<FInterpCurvePointVector> : public TStructOpsTypeTrai
 		WithNoInitConstructor = true,
 	};
 };
-IMPLEMENT_STRUCT(InterpCurvePointVector);
+UE_IMPLEMENT_STRUCT("/Script/CoreUObject", InterpCurvePointVector);
 
 template<>
 struct TStructOpsTypeTraits<FInterpCurvePointQuat> : public TStructOpsTypeTraitsBase2<FInterpCurvePointQuat>
@@ -340,7 +340,7 @@ struct TStructOpsTypeTraits<FInterpCurvePointQuat> : public TStructOpsTypeTraits
 		WithNoInitConstructor = true,
 	};
 };
-IMPLEMENT_STRUCT(InterpCurvePointQuat);
+UE_IMPLEMENT_STRUCT("/Script/CoreUObject", InterpCurvePointQuat);
 
 template<>
 struct TStructOpsTypeTraits<FInterpCurvePointTwoVectors> : public TStructOpsTypeTraitsBase2<FInterpCurvePointTwoVectors>
@@ -350,7 +350,7 @@ struct TStructOpsTypeTraits<FInterpCurvePointTwoVectors> : public TStructOpsType
 		WithNoInitConstructor = true,
 	};
 };
-IMPLEMENT_STRUCT(InterpCurvePointTwoVectors);
+UE_IMPLEMENT_STRUCT("/Script/CoreUObject", InterpCurvePointTwoVectors);
 
 template<>
 struct TStructOpsTypeTraits<FInterpCurvePointLinearColor> : public TStructOpsTypeTraitsBase2<FInterpCurvePointLinearColor>
@@ -360,7 +360,7 @@ struct TStructOpsTypeTraits<FInterpCurvePointLinearColor> : public TStructOpsTyp
 		WithNoInitConstructor = true,
 	};
 };
-IMPLEMENT_STRUCT(InterpCurvePointLinearColor);
+UE_IMPLEMENT_STRUCT("/Script/CoreUObject", InterpCurvePointLinearColor);
 
 template<>
 struct TStructOpsTypeTraits<FGuid> : public TStructOpsTypeTraitsBase2<FGuid>
@@ -375,7 +375,7 @@ struct TStructOpsTypeTraits<FGuid> : public TStructOpsTypeTraitsBase2<FGuid>
 		WithStructuredSerializer = true,
 	};
 };
-IMPLEMENT_STRUCT(Guid);
+UE_IMPLEMENT_STRUCT("/Script/CoreUObject", Guid);
 
 template<typename T>
 struct TTransformStructOpsTypeTraits : public TStructOpsTypeTraitsBase2<T>
@@ -389,9 +389,9 @@ struct TTransformStructOpsTypeTraits : public TStructOpsTypeTraitsBase2<T>
 };
 template<> struct TStructOpsTypeTraits<FTransform3f> : public TTransformStructOpsTypeTraits<FTransform3f> {};
 template<> struct TStructOpsTypeTraits<FTransform3d> : public TTransformStructOpsTypeTraits<FTransform3d> {};
-IMPLEMENT_STRUCT(Transform3f);
-IMPLEMENT_STRUCT(Transform3d);
-IMPLEMENT_STRUCT(Transform); // Aliased
+UE_IMPLEMENT_STRUCT("/Script/CoreUObject", Transform3f);
+UE_IMPLEMENT_STRUCT("/Script/CoreUObject", Transform3d);
+UE_IMPLEMENT_STRUCT("/Script/CoreUObject", Transform); // Aliased
 
 template<>
 struct TStructOpsTypeTraits<FRandomStream> : public TStructOpsTypeTraitsBase2<FRandomStream>
@@ -403,7 +403,7 @@ struct TStructOpsTypeTraits<FRandomStream> : public TStructOpsTypeTraitsBase2<FR
 		WithZeroConstructor = true,
 	};
 };
-IMPLEMENT_STRUCT(RandomStream);
+UE_IMPLEMENT_STRUCT("/Script/CoreUObject", RandomStream);
 
 template<>
 struct TStructOpsTypeTraits<FDateTime> : public TStructOpsTypeTraitsBase2<FDateTime>
@@ -419,7 +419,7 @@ struct TStructOpsTypeTraits<FDateTime> : public TStructOpsTypeTraitsBase2<FDateT
 		WithIdenticalViaEquality = true,
 	};
 };
-IMPLEMENT_STRUCT(DateTime);
+UE_IMPLEMENT_STRUCT("/Script/CoreUObject", DateTime);
 
 template<>
 struct TStructOpsTypeTraits<FTimespan> : public TStructOpsTypeTraitsBase2<FTimespan>
@@ -436,7 +436,7 @@ struct TStructOpsTypeTraits<FTimespan> : public TStructOpsTypeTraitsBase2<FTimes
 		WithIdenticalViaEquality = true,
 	};
 };
-IMPLEMENT_STRUCT(Timespan);
+UE_IMPLEMENT_STRUCT("/Script/CoreUObject", Timespan);
 
 template<>
 struct TStructOpsTypeTraits<FFrameNumber> : public TStructOpsTypeTraitsBase2<FFrameNumber>
@@ -447,7 +447,7 @@ struct TStructOpsTypeTraits<FFrameNumber> : public TStructOpsTypeTraitsBase2<FFr
 		WithIdenticalViaEquality = true
 	};
 };
-IMPLEMENT_STRUCT(FrameNumber);
+UE_IMPLEMENT_STRUCT("/Script/CoreUObject", FrameNumber);
 
 template<>
 struct TStructOpsTypeTraits<FSoftObjectPath> : public TStructOpsTypeTraitsBase2<FSoftObjectPath>
@@ -463,7 +463,7 @@ struct TStructOpsTypeTraits<FSoftObjectPath> : public TStructOpsTypeTraitsBase2<
 		WithStructuredSerializeFromMismatchedTag = true,
 	};
 };
-IMPLEMENT_STRUCT(SoftObjectPath);
+UE_IMPLEMENT_STRUCT("/Script/CoreUObject", SoftObjectPath);
 
 template<>
 struct TStructOpsTypeTraits<FSoftClassPath> : public TStructOpsTypeTraitsBase2<FSoftClassPath>
@@ -479,7 +479,7 @@ struct TStructOpsTypeTraits<FSoftClassPath> : public TStructOpsTypeTraitsBase2<F
 		WithStructuredSerializeFromMismatchedTag = true,
 	};
 };
-IMPLEMENT_STRUCT(SoftClassPath);
+UE_IMPLEMENT_STRUCT("/Script/CoreUObject", SoftClassPath);
 
 template<>
 struct TStructOpsTypeTraits<FPrimaryAssetType> : public TStructOpsTypeTraitsBase2<FPrimaryAssetType>
@@ -494,7 +494,7 @@ struct TStructOpsTypeTraits<FPrimaryAssetType> : public TStructOpsTypeTraitsBase
 		WithStructuredSerializeFromMismatchedTag = true,
 	};
 };
-IMPLEMENT_STRUCT(PrimaryAssetType);
+UE_IMPLEMENT_STRUCT("/Script/CoreUObject", PrimaryAssetType);
 
 template<>
 struct TStructOpsTypeTraits<FPrimaryAssetId> : public TStructOpsTypeTraitsBase2<FPrimaryAssetId>
@@ -509,13 +509,13 @@ struct TStructOpsTypeTraits<FPrimaryAssetId> : public TStructOpsTypeTraitsBase2<
 		WithStructuredSerializeFromMismatchedTag = true,
 	};
 };
-IMPLEMENT_STRUCT(PrimaryAssetId);
+UE_IMPLEMENT_STRUCT("/Script/CoreUObject", PrimaryAssetId);
 
 template<>
 struct TStructOpsTypeTraits<FFallbackStruct> : public TStructOpsTypeTraitsBase2<FFallbackStruct>
 {
 };
-IMPLEMENT_STRUCT(FallbackStruct);
+UE_IMPLEMENT_STRUCT("/Script/CoreUObject", FallbackStruct);
 
 /*-----------------------------------------------------------------------------
 	Helpers.
diff --git a/Engine/Source/Runtime/CoreUObject/Private/UObject/TopLevelAssetPath.cpp b/Engine/Source/Runtime/CoreUObject/Private/UObject/TopLevelAssetPath.cpp
--- a/Engine/Source/Runtime/CoreUObject/Private/UObject/TopLevelAssetPath.cpp
+++ b/Engine/Source/Runtime/CoreUObject/Private/UObject/TopLevelAssetPath.cpp
@@ -14,9 +14,9 @@
 // The reason behind HACK_HEADER_GENERATOR is that without it UHT is going to 'see' cppstructops for TopLevelAssetPath
 // and will not generate temp FTopLevelAssetPath struct for codegen purposes where it can access all of its members
 // which are public in the temp struct and private in the actual FTopLevelAssetPath.
-// If UHT compiles with IMPLEMENT_STRUCT(TopLevelAssetPath) the generated code will fail to compile trying to access private struct members
+// If UHT compiles with UE_IMPLEMENT_STRUCT("/Script/CoreUObject", TopLevelAssetPath) the generated code will fail to compile trying to access private struct members
 #if !HACK_HEADER_GENERATOR
-IMPLEMENT_STRUCT(TopLevelAssetPath)
+UE_IMPLEMENT_STRUCT("/Script/CoreUObject", TopLevelAssetPath)
 #endif
 
 FTopLevelAssetPath::FTopLevelAssetPath(const UObject* InObject)
diff --git a/Engine/Source/Runtime/CoreUObject/Private/UObject/UObjectBase.cpp b/Engine/Source/Runtime/CoreUObject/Private/UObject/UObjectBase.cpp
--- a/Engine/Source/Runtime/CoreUObject/Private/UObject/UObjectBase.cpp
+++ b/Engine/Source/Runtime/CoreUObject/Private/UObject/UObjectBase.cpp
@@ -711,7 +711,7 @@ void RegisterCompiledInInfo(const TCHAR* PackageName, const FClassRegisterCompil
 		RegisterCompiledInInfo(Info.OuterRegister, PackageName, Info.Name, *Info.Info, Info.VersionInfo);
 		if (Info.CreateCppStructOps != nullptr)
 		{
-			UScriptStruct::DeferCppStructOps(FName(Info.Name), (UScriptStruct::ICppStructOps*)Info.CreateCppStructOps());
+			UScriptStruct::DeferCppStructOps(FTopLevelAssetPath(FName(PackageName), FName(Info.Name)), (UScriptStruct::ICppStructOps*)Info.CreateCppStructOps());
 		}
 	}
 
diff --git a/Engine/Source/Runtime/CoreUObject/Public/UObject/Class.h b/Engine/Source/Runtime/CoreUObject/Public/UObject/Class.h
--- a/Engine/Source/Runtime/CoreUObject/Public/UObject/Class.h
+++ b/Engine/Source/Runtime/CoreUObject/Public/UObject/Class.h
@@ -1356,13 +1356,16 @@ public:
 	template<class CPPSTRUCT>
 	struct TAutoCppStructOps
 	{
-		TAutoCppStructOps(FName InName)
+		TAutoCppStructOps(FTopLevelAssetPath InName)
 		{
 			DeferCppStructOps(InName,new TCppStructOps<CPPSTRUCT>);
 		}
 	};
 	#define IMPLEMENT_STRUCT(BaseName) \
-		static UScriptStruct::TAutoCppStructOps<F##BaseName> BaseName##_Ops(TEXT(#BaseName)); 
+		DEPRECATED_MACRO(5.1, "IMPLEMENT_STRUCT has been deprecated. Use UE_IMPLEMENT_STRUCT and provide struct package name as well as struct name") static UScriptStruct::TAutoCppStructOps<F##BaseName> BaseName##_Ops(FTopLevelAssetPath(TEXT("/Script/CoreUObject"), TEXT(#BaseName))); 
+
+	#define UE_IMPLEMENT_STRUCT(PackageNameText, BaseName) \
+		static UScriptStruct::TAutoCppStructOps<F##BaseName> BaseName##_Ops(FTopLevelAssetPath(TEXT(PackageNameText), TEXT(#BaseName))); 
 
 	DECLARE_CASTED_CLASS_INTRINSIC_NO_CTOR(UScriptStruct, UStruct, CLASS_MatchedSerializers, TEXT("/Script/CoreUObject"), CASTCLASS_UScriptStruct, COREUOBJECT_API)
 
@@ -1380,7 +1383,7 @@ public:
 	// Required by UHT makefiles for internal data serialization.
 	friend struct FScriptStructArchiveProxy;
 
-	static COREUOBJECT_API ICppStructOps* FindDeferredCppStructOps(FName StructName);
+	static COREUOBJECT_API ICppStructOps* FindDeferredCppStructOps(FTopLevelAssetPath StructName);
 #endif
 
 protected:
@@ -1388,6 +1391,17 @@ protected:
 	bool bPrepareCppStructOpsCompleted;
 	/** Holds the Cpp ctors and dtors, sizeof, etc. Is not owned by this and is not released. **/
 	ICppStructOps* CppStructOps;
+
+	/** 
+	* Similar to GetStructPathName() but works with nested structs by using just the package name and struct name
+	* so a struct path name /Package/Name.Object:Struct will be flattened to /Package/Name.Struct.
+	* This function is used only for generating keys for DeferredCppStructOps
+	*/
+	FTopLevelAssetPath GetFlattenedStructPathName() const
+	{
+		return FTopLevelAssetPath(GetOutermost()->GetFName(), GetFName());
+	}
+
 public:
 
 	// UObject Interface
@@ -1409,15 +1423,15 @@ public:
 	 * @param Target Name of the struct 
 	 * @param InCppStructOps Cpp ops for this struct
 	 */
-	static COREUOBJECT_API void DeferCppStructOps(FName Target, ICppStructOps* InCppStructOps);
+	static COREUOBJECT_API void DeferCppStructOps(FTopLevelAssetPath Target, ICppStructOps* InCppStructOps);
 
 	template<class CPPSTRUCT>
-	static typename TEnableIf<!DISABLE_ABSTRACT_CONSTRUCT>::Type DeferCppStructOps(FName Target)
+	static typename TEnableIf<!DISABLE_ABSTRACT_CONSTRUCT>::Type DeferCppStructOps(FTopLevelAssetPath Target)
 	{
 		DeferCppStructOps(Target, new UScriptStruct::TCppStructOps<CPPSTRUCT>);
 	}
 	template<class CPPSTRUCT>
-	static typename TEnableIf<DISABLE_ABSTRACT_CONSTRUCT>::Type DeferCppStructOps(FName Target)
+	static typename TEnableIf<DISABLE_ABSTRACT_CONSTRUCT>::Type DeferCppStructOps(FTopLevelAssetPath Target)
 	{
 		DeferCppStructOps(Target, nullptr);
 	}
diff --git a/Engine/Source/Runtime/Engine/Private/Materials/MaterialShared.cpp b/Engine/Source/Runtime/Engine/Private/Materials/MaterialShared.cpp
--- a/Engine/Source/Runtime/Engine/Private/Materials/MaterialShared.cpp
+++ b/Engine/Source/Runtime/Engine/Private/Materials/MaterialShared.cpp
@@ -248,12 +248,12 @@ int32 FMaterialCompiler::VectorParameter(FName ParameterName, const FLinearColor
 	return NumericParameter(EMaterialParameterType::Vector, ParameterName, DefaultValue);
 }
 
-IMPLEMENT_STRUCT(ExpressionInput);
-IMPLEMENT_STRUCT(ColorMaterialInput);
-IMPLEMENT_STRUCT(ScalarMaterialInput);
-IMPLEMENT_STRUCT(VectorMaterialInput);
-IMPLEMENT_STRUCT(Vector2MaterialInput);
-IMPLEMENT_STRUCT(MaterialAttributesInput);
+UE_IMPLEMENT_STRUCT("/Script/Engine", ExpressionInput);
+UE_IMPLEMENT_STRUCT("/Script/Engine", ColorMaterialInput);
+UE_IMPLEMENT_STRUCT("/Script/Engine", ScalarMaterialInput);
+UE_IMPLEMENT_STRUCT("/Script/Engine", VectorMaterialInput);
+UE_IMPLEMENT_STRUCT("/Script/Engine", Vector2MaterialInput);
+UE_IMPLEMENT_STRUCT("/Script/Engine", MaterialAttributesInput);
 
 #if WITH_EDITOR
 int32 FExpressionInput::Compile(class FMaterialCompiler* Compiler)
diff --git a/Engine/Source/Runtime/Engine/Private/UnrealClient.cpp b/Engine/Source/Runtime/Engine/Private/UnrealClient.cpp
--- a/Engine/Source/Runtime/Engine/Private/UnrealClient.cpp
+++ b/Engine/Source/Runtime/Engine/Private/UnrealClient.cpp
@@ -33,7 +33,7 @@
 
 DEFINE_LOG_CATEGORY_STATIC(LogClient, Log, All);
 
-IMPLEMENT_STRUCT(PostProcessSettings);
+UE_IMPLEMENT_STRUCT("/Script/Engine", PostProcessSettings);
 
 bool FViewport::bIsGameRenderingEnabled = true;
 int32 FViewport::PresentAndStopMovieDelay = 0;
