Fixing deprecation warnings about deprecated DEPRECATED_MACRO macro

diff --git a/Engine/Source/Runtime/CoreUObject/Private/UObject/Property.cpp b/Engine/Source/Runtime/CoreUObject/Private/UObject/Property.cpp
--- a/Engine/Source/Runtime/CoreUObject/Private/UObject/Property.cpp
+++ b/Engine/Source/Runtime/CoreUObject/Private/UObject/Property.cpp
@@ -84,7 +84,7 @@ struct TStructOpsTypeTraits<FIntVector4> : public TStructOpsTypeTraitsBase2<FInt
 		WithSerializer = true,
 	};
 };
-IMPLEMENT_STRUCT(IntVector4);
+UE_IMPLEMENT_STRUCT("/Script/CoreUObject", IntVector4);
 
 template<typename T>
 struct TVector2StructOpsTypeTraits : public TStructOpsTypeTraitsBase2<T>
diff --git a/Engine/Source/Runtime/CoreUObject/Public/UObject/Class.h b/Engine/Source/Runtime/CoreUObject/Public/UObject/Class.h
--- a/Engine/Source/Runtime/CoreUObject/Public/UObject/Class.h
+++ b/Engine/Source/Runtime/CoreUObject/Public/UObject/Class.h
@@ -1362,7 +1362,7 @@ public:
 		}
 	};
 	#define IMPLEMENT_STRUCT(BaseName) \
-		DEPRECATED_MACRO(5.1, "IMPLEMENT_STRUCT has been deprecated. Use UE_IMPLEMENT_STRUCT and provide struct package name as well as struct name") static UScriptStruct::TAutoCppStructOps<F##BaseName> BaseName##_Ops(FTopLevelAssetPath(TEXT("/Script/CoreUObject"), TEXT(#BaseName))); 
+		UE_DEPRECATED_MACRO(5.1, "IMPLEMENT_STRUCT has been deprecated. Use UE_IMPLEMENT_STRUCT and provide struct package name as well as struct name") static UScriptStruct::TAutoCppStructOps<F##BaseName> BaseName##_Ops(FTopLevelAssetPath(TEXT("/Script/CoreUObject"), TEXT(#BaseName))); 
 
 	#define UE_IMPLEMENT_STRUCT(PackageNameText, BaseName) \
 		static UScriptStruct::TAutoCppStructOps<F##BaseName> BaseName##_Ops(FTopLevelAssetPath(TEXT(PackageNameText), TEXT(#BaseName))); 
