bool-to-integer property conversions added.

diff --git a/Engine/Source/Runtime/CoreUObject/Private/UObject/EnumProperty.cpp b/Engine/Source/Runtime/CoreUObject/Private/UObject/EnumProperty.cpp
--- a/Engine/Source/Runtime/CoreUObject/Private/UObject/EnumProperty.cpp
+++ b/Engine/Source/Runtime/CoreUObject/Private/UObject/EnumProperty.cpp
@@ -15,11 +15,8 @@
 namespace UEEnumProperty_Private
 {
 	template <typename OldIntType>
-	void ConvertIntToEnumProperty(FStructuredArchive::FSlot Slot, FEnumProperty* EnumProp, FNumericProperty* UnderlyingProp, UEnum* Enum, void* Obj)
+	void ConvertIntValueToEnumProperty(OldIntType OldValue, FEnumProperty* EnumProp, FNumericProperty* UnderlyingProp, UEnum* Enum, void* Obj)
 	{
-		OldIntType OldValue;
-		Slot << OldValue;
-
 		using LargeIntType = typename TChooseClass<TIsSigned<OldIntType>::Value, int64, uint64>::Result;
 
 		LargeIntType NewValue = OldValue;
@@ -40,6 +37,15 @@ namespace UEEnumProperty_Private
 
 		UnderlyingProp->SetIntPropertyValue(Obj, NewValue);
 	}
+
+	template <typename OldIntType>
+	void ConvertIntToEnumProperty(FStructuredArchive::FSlot Slot, FEnumProperty* EnumProp, FNumericProperty* UnderlyingProp, UEnum* Enum, void* Obj)
+	{
+		OldIntType OldValue;
+		Slot << OldValue;
+
+		ConvertIntValueToEnumProperty(OldValue, EnumProp, UnderlyingProp, Enum, Obj);
+	}
 }
 
 IMPLEMENT_FIELD(FEnumProperty)
@@ -495,6 +501,10 @@ EConvertFromTypeResult FEnumProperty::ConvertFromType(const FPropertyTag& Tag, F
 	{
 		UEEnumProperty_Private::ConvertIntToEnumProperty<uint64>(Slot, this, UnderlyingProp, Enum, ContainerPtrToValuePtr<void>(Data, Tag.ArrayIndex));
 	}
+	else if (Tag.Type == NAME_BoolProperty)
+	{
+		UEEnumProperty_Private::ConvertIntValueToEnumProperty<uint8>(Tag.BoolVal, this, UnderlyingProp, Enum, ContainerPtrToValuePtr<void>(Data, Tag.ArrayIndex));
+	}
 	else
 	{
 		return EConvertFromTypeResult::UseSerializeItem;
diff --git a/Engine/Source/Runtime/CoreUObject/Private/UObject/PropertyByte.cpp b/Engine/Source/Runtime/CoreUObject/Private/UObject/PropertyByte.cpp
--- a/Engine/Source/Runtime/CoreUObject/Private/UObject/PropertyByte.cpp
+++ b/Engine/Source/Runtime/CoreUObject/Private/UObject/PropertyByte.cpp
@@ -186,6 +186,11 @@ struct TConvertIntToEnumProperty
 		OldIntType OldValue;
 		Slot << OldValue;
 
+		ConvertValue(OldValue, Property, Enum, Obj, Tag);
+	}
+
+	static void ConvertValue(OldIntType OldValue, FByteProperty* Property, UEnum* Enum, void* Obj, const FPropertyTag& Tag)
+	{
 		uint8 NewValue = (uint8)OldValue;
 		if (OldValue > (OldIntType)TNumericLimits<uint8>::Max() || !Enum->IsValidEnumValue(NewValue))
 		{
@@ -320,6 +325,17 @@ EConvertFromTypeResult FByteProperty::ConvertFromType(const FPropertyTag& Tag, F
 			ConvertFromArithmeticValue<uint64>(Slot, Data, Tag);
 		}
 	}
+	else if (Tag.Type == NAME_BoolProperty)
+	{
+		if (Enum)
+		{
+			TConvertIntToEnumProperty<uint64>::ConvertValue(Tag.BoolVal, this, Enum, Data, Tag);
+		}
+		else
+		{
+			SetPropertyValue_InContainer(Data, Tag.BoolVal, Tag.ArrayIndex);
+		}
+	}
 	else
 	{
 		return EConvertFromTypeResult::UseSerializeItem;
diff --git a/Engine/Source/Runtime/CoreUObject/Public/UObject/UnrealType.h b/Engine/Source/Runtime/CoreUObject/Public/UObject/UnrealType.h
--- a/Engine/Source/Runtime/CoreUObject/Public/UObject/UnrealType.h
+++ b/Engine/Source/Runtime/CoreUObject/Public/UObject/UnrealType.h
@@ -1818,6 +1818,10 @@ public:
 				ConvertFromArithmeticValue<double>(Slot, Data, Tag);
 				return EConvertFromTypeResult::Converted;
 
+			case NAME_BoolProperty:
+				this->SetPropertyValue_InContainer(Data, (TCppType)Tag.BoolVal, Tag.ArrayIndex);
+				return EConvertFromTypeResult::Converted;
+
 			default:
 				// We didn't convert it
 				break;
