Trivial casting fix for DynamicallyTypedValue

#rnx

- MinAlignmentLogTwo is being stored as uint8 and the only place in code passing in that parameter does so as a uint8 so we can just change the parameter data type to match.

diff --git a/Engine/Source/Runtime/CoreUObject/Public/UObject/DynamicallyTypedValue.h b/Engine/Source/Runtime/CoreUObject/Public/UObject/DynamicallyTypedValue.h
--- a/Engine/Source/Runtime/CoreUObject/Public/UObject/DynamicallyTypedValue.h
+++ b/Engine/Source/Runtime/CoreUObject/Public/UObject/DynamicallyTypedValue.h
@@ -18,7 +18,7 @@ namespace UE
 			Maybe,
 		};
 
-		constexpr FDynamicallyTypedValueType(SIZE_T InNumBytes, SIZE_T InMinAlignmentLogTwo, EContainsReferences InContainsReferences)
+		constexpr FDynamicallyTypedValueType(SIZE_T InNumBytes, uint8 InMinAlignmentLogTwo, EContainsReferences InContainsReferences)
 			: NumBytes(InNumBytes)
 			, MinAlignmentLogTwo(InMinAlignmentLogTwo)
 			, ContainsReferences(InContainsReferences)
