Bulk Data - correctly convert async file I/O priority to I/O dispatcher priority

#prefligh t62a32338d708f949d6486971

diff --git a/Engine/Source/Runtime/CoreUObject/Private/Serialization/BulkData.cpp b/Engine/Source/Runtime/CoreUObject/Private/Serialization/BulkData.cpp
--- a/Engine/Source/Runtime/CoreUObject/Private/Serialization/BulkData.cpp
+++ b/Engine/Source/Runtime/CoreUObject/Private/Serialization/BulkData.cpp
@@ -83,6 +83,7 @@ bool OpenReadBulkData(
 	const FBulkDataChunkId& BulkChunkId,
 	int64 Offset,
 	int64 Size,
+	EAsyncIOPriorityAndFlags Priority,
 	TFunction<void(FArchive& Ar)>&& Read);
 
 /** Open async read file handle for the specified bulk data chunk ID. */
@@ -115,7 +116,8 @@ bool StartAsyncLoad(
 	const FBulkMetaData& BulkMeta,
 	const FBulkDataChunkId& BulkChunkId,
 	int64 Offset,
-	int64 Size, 
+	int64 Size,
+	EAsyncIOPriorityAndFlags Priority,
 	TFunction<void(TIoStatusOr<FIoBuffer>)>&& Callback);
 
 /** Flush pending async load. */
@@ -1083,6 +1085,7 @@ bool FBulkData::LoadBulkDataWithFileReader()
 			BulkChunkId,
 			BulkDataOffset,
 			BulkDataSizeOnDisk,
+			AIOP_High,
 			[this](FArchive& Ar)
 			{
 				const int64 BulkDataSize = GetBulkDataSize();
@@ -1158,6 +1161,7 @@ bool FBulkData::StartAsyncLoading()
 		BulkChunkId,
 		BulkMeta.GetOffset(), 
 		BulkMeta.GetSize(),
+		AIOP_Low,
 		[this](TIoStatusOr<FIoBuffer> Result)
 	{
 		if (Result.IsOk())
@@ -1981,6 +1985,7 @@ bool FBulkData::TryLoadDataIntoMemory(FIoBuffer Dest)
 		BulkChunkId,
 		BulkDataOffset,
 		BulkDataSizeOnDisk,
+		AIOP_High,
 		[this, BulkDataSize, &Dest](FArchive& Ar)
 		{
 			SerializeBulkData(Ar, Dest.GetData(), BulkDataSize, BulkMeta.GetFlags());
diff --git a/Engine/Source/Runtime/CoreUObject/Private/Serialization/BulkDataStreaming.cpp b/Engine/Source/Runtime/CoreUObject/Private/Serialization/BulkDataStreaming.cpp
--- a/Engine/Source/Runtime/CoreUObject/Private/Serialization/BulkDataStreaming.cpp
+++ b/Engine/Source/Runtime/CoreUObject/Private/Serialization/BulkDataStreaming.cpp
@@ -471,6 +471,7 @@ bool OpenReadBulkData(
 	const FBulkDataChunkId& BulkChunkId,
 	int64 Offset,
 	int64 Size,
+	EAsyncIOPriorityAndFlags Priority,
 	TFunction<void(FArchive& Ar)>&& Read)
 {
 	if (BulkChunkId.IsValid() == false)
@@ -483,7 +484,7 @@ bool OpenReadBulkData(
 		const FIoChunkId ChunkId = CreateBulkDataIoChunkId(BulkMeta, PackageId);
 		FIoBatch Batch = FIoDispatcher::Get().NewBatch();
 
-		FIoRequest Request = Batch.Read(ChunkId, FIoReadOptions(Offset, Size), IoDispatcherPriority_Medium);
+		FIoRequest Request = Batch.Read(ChunkId, FIoReadOptions(Offset, Size), ConvertToIoDispatcherPriority(Priority));
 		FEventRef Event;
 		Batch.IssueAndTriggerEvent(Event.Get());
 		Event->Wait();
@@ -580,7 +581,7 @@ TUniquePtr<IBulkDataIORequest> CreateStreamingRequest(
 		const FIoChunkId ChunkId = CreateBulkDataIoChunkId(BulkMeta, PackageId);
 
 		FChunkBulkDataRequest* Request = new FChunkBulkDataRequest(Callback, MoveTemp(Buffer));
-		Request->Issue(ChunkId, FIoReadOptions(Offset, Size), Priority);
+		Request->Issue(ChunkId, FIoReadOptions(Offset, Size), ConvertToIoDispatcherPriority(Priority));
 		
 		return TUniquePtr<IBulkDataIORequest>(Request);
 	}
@@ -725,6 +726,7 @@ bool StartAsyncLoad(
 	const FBulkDataChunkId& BulkChunkId,
 	int64 Offset,
 	int64 Size,
+	EAsyncIOPriorityAndFlags Priority,
 	TFunction<void(TIoStatusOr<FIoBuffer>)>&& Callback)
 {
 	if (TUniquePtr<IAsyncReadFileHandle> FileHandle = OpenAsyncReadBulkData(BulkMeta, BulkChunkId))
@@ -746,7 +748,7 @@ bool StartAsyncLoad(
 				}
 			};
 		
-		if (IAsyncReadRequest* Request = FileHandle->ReadRequest(Offset, Size, AIOP_Normal, &FileReadCallback))
+		if (IAsyncReadRequest* Request = FileHandle->ReadRequest(Offset, Size, Priority, &FileReadCallback))
 		{
 			FAsyncBulkDataRequests::Get().AddPendingRequest(Owner, MoveTemp(FileHandle), TUniquePtr<IAsyncReadRequest>(Request));
 			return true;
