Bulk Data - fix bug when PollCompletion returns true before the callback is triggered

diff --git a/Engine/Source/Runtime/CoreUObject/Private/Serialization/BulkDataStreaming.cpp b/Engine/Source/Runtime/CoreUObject/Private/Serialization/BulkDataStreaming.cpp
--- a/Engine/Source/Runtime/CoreUObject/Private/Serialization/BulkDataStreaming.cpp
+++ b/Engine/Source/Runtime/CoreUObject/Private/Serialization/BulkDataStreaming.cpp
@@ -145,11 +145,13 @@ EPackageSegment GetPackageSegmentFromFlags(const FBulkMetaData& BulkMeta)
 
 enum class EChunkRequestStatus : uint32
 {
-	None		= 0,
-	Ok			= 1 << 0,
-	Pending		= 1 << 1,
-	Canceled	= 1 << 2,
+	None				= 0,
+	Pending				= 1 << 0,
+	Canceled			= 1 << 1,
+	DataReady			= 1 << 2,
+	CallbackTriggered	= 1 << 3,
 };
+ENUM_CLASS_FLAGS(EChunkRequestStatus);
 
 class FChunkRequest
 {
@@ -158,14 +160,14 @@ public:
 	
 	void Issue(FIoChunkId ChunkId, FIoReadOptions Options, int32 Priority);
 
+protected:
+	FChunkRequest(FIoBuffer&& InBuffer);
+	
 	inline EChunkRequestStatus GetStatus() const
 	{
 		return static_cast<EChunkRequestStatus>(Status.load(std::memory_order_consume));
 	}
 
-protected:
-	FChunkRequest(FIoBuffer&& InBuffer);
-
 	virtual void HandleChunkResult(TIoStatusOr<FIoBuffer>&& Result) = 0;
 	bool WaitForChunkRequest(float TimeLimitSeconds = 0.0f);
 	void CancelChunkRequest();
@@ -183,9 +185,9 @@ private:
 
 FChunkRequest::FChunkRequest(FIoBuffer&& InBuffer)
 	: Buffer(MoveTemp(InBuffer))
-	, DoneEvent(EEventMode ::ManualReset)
+	, DoneEvent(EEventMode::ManualReset)
 	, SizeResult(-1)
-	, Status{0}
+	, Status{uint32(EChunkRequestStatus::None)}
 {
 }
 
@@ -196,7 +198,7 @@ FChunkRequest::~FChunkRequest()
 
 void FChunkRequest::Issue(FIoChunkId ChunkId, FIoReadOptions Options, int32 Priority)
 {
-	Status.store(uint32(EChunkRequestStatus::Pending), std::memory_order_relaxed); 
+	Status.store(uint32(EChunkRequestStatus::Pending), std::memory_order_release); 
 
 	check(Options.GetSize() == Buffer.GetSize());
 	Options.SetTargetVa(Buffer.GetData());
@@ -204,17 +206,18 @@ void FChunkRequest::Issue(FIoChunkId ChunkId, FIoReadOptions Options, int32 Prio
 	FIoBatch IoBatch = FIoDispatcher::Get().NewBatch();
 	Request = IoBatch.ReadWithCallback(ChunkId, Options, Priority, [this](TIoStatusOr<FIoBuffer> Result)
 	{
+		EChunkRequestStatus ReadyOrCanceled = EChunkRequestStatus::Canceled;
+
 		if (Result.IsOk())
 		{
 			SizeResult = Result.ValueOrDie().GetSize();
-			Status.store(uint32(EChunkRequestStatus::Ok), std::memory_order_release); 
-		}
-		else
-		{
-			Status.store(uint32(EChunkRequestStatus::Canceled), std::memory_order_release); 
+			ReadyOrCanceled = EChunkRequestStatus::DataReady;
 		}
 
+		Status.store(uint32(ReadyOrCanceled), std::memory_order_release); 
 		HandleChunkResult(MoveTemp(Result));
+		Status.store(uint32(ReadyOrCanceled | EChunkRequestStatus::CallbackTriggered), std::memory_order_release); 
+
 		DoneEvent.Trigger();
 	});
 
@@ -382,7 +385,7 @@ public:
 	inline virtual bool PollCompletion() const override
 	{
 		checkf(GetStatus() != EChunkRequestStatus::None, TEXT("The request must be issued before polling for completion"));
-		return GetStatus() != EChunkRequestStatus::Pending;
+		return EnumHasAnyFlags(GetStatus(), EChunkRequestStatus::CallbackTriggered);
 	}
 
 	virtual bool WaitCompletion(float TimeLimitSeconds = 0.0f) override
@@ -396,7 +399,7 @@ public:
 	inline virtual int64 GetSize() const override
 	{
 		checkf(GetStatus() != EChunkRequestStatus::None, TEXT("The request must be issued before polling for size"));
-		return GetStatus() == EChunkRequestStatus::Ok ? GetSizeResult() : -1;
+		return EnumHasAnyFlags(GetStatus(), EChunkRequestStatus::DataReady) ? GetSizeResult() : -1;
 	}
 
 	virtual void Cancel() override
@@ -431,7 +434,7 @@ uint8* FChunkBulkDataRequest::GetReadResults()
 {
 	uint8* ReadResult = nullptr;
 
-	if (GetStatus() == EChunkRequestStatus::Ok)
+	if (EnumHasAnyFlags(GetStatus(), EChunkRequestStatus::DataReady))
 	{
 		if (Buffer.IsMemoryOwned())
 		{
