Improved ensure error message during package save when a name isn't mapped.

diff --git a/Engine/Source/Runtime/CoreUObject/Private/UObject/LinkerSave.cpp b/Engine/Source/Runtime/CoreUObject/Private/UObject/LinkerSave.cpp
--- a/Engine/Source/Runtime/CoreUObject/Private/UObject/LinkerSave.cpp
+++ b/Engine/Source/Runtime/CoreUObject/Private/UObject/LinkerSave.cpp
@@ -337,7 +337,7 @@ FArchive& FLinkerSave::operator<<( FName& InName )
 	{
 		// Set an error on the archive and record the error on the log output if one is set.
 		SetCriticalError();
-		FString ErrorMessage = FString::Printf(TEXT("Name \"%s\" is not mapped when saving %s (object: %s, property: %s)."),
+		FString ErrorMessage = FString::Printf(TEXT("Name \"%s\" is not mapped when saving %s (object: %s, property: %s). This can mean that this object serialize function is not deterministic between reference harvesting and serialization."),
 			*InName.ToString(),
 			*GetArchiveName(),
 			*GetSerializeContext()->SerializedObject->GetFullName(),
