Fix for aggregates being initialized with braced syntax but are no longer aggregates in C++20.

diff --git a/Engine/Source/Runtime/CoreUObject/Public/UObject/SavePackage.h b/Engine/Source/Runtime/CoreUObject/Public/UObject/SavePackage.h
--- a/Engine/Source/Runtime/CoreUObject/Public/UObject/SavePackage.h
+++ b/Engine/Source/Runtime/CoreUObject/Public/UObject/SavePackage.h
@@ -84,6 +84,31 @@ struct FSavePackageArgs
 	FSavePackageArgs(FSavePackageArgs&&) = default;
 	FSavePackageArgs& operator=(const FSavePackageArgs&) = default;
 	FSavePackageArgs& operator=(FSavePackageArgs&&) = default;
+
+	FSavePackageArgs(
+		const ITargetPlatform* InTargetPlatform,
+		FArchiveCookData* InArchiveCookData,
+		EObjectFlags InTopLevelFlags,
+		uint32 InSaveFlags,
+		bool bInForceByteSwapping,
+		bool bInWarnOfLongFilename,
+		bool bInSlowTask,
+		FDateTime InFinalTimeStamp,
+		FOutputDevice* InError,
+		FSavePackageContext* InSavePackageContext = nullptr
+	)
+	: TargetPlatform(InTargetPlatform)
+	, ArchiveCookData(InArchiveCookData)
+	, TopLevelFlags(InTopLevelFlags)
+	, SaveFlags(InSaveFlags)
+	, bForceByteSwapping(bInForceByteSwapping)
+	, bWarnOfLongFilename(bInWarnOfLongFilename)
+	, bSlowTask(bInSlowTask)
+	, FinalTimeStamp(InFinalTimeStamp)
+	, Error(InError)
+	, SavePackageContext(InSavePackageContext)
+	{
+	}
 	PRAGMA_ENABLE_DEPRECATION_WARNINGS;
 };
 
diff --git a/Engine/Source/Runtime/ImageWriteQueue/Private/ImageWriteQueue.cpp b/Engine/Source/Runtime/ImageWriteQueue/Private/ImageWriteQueue.cpp
--- a/Engine/Source/Runtime/ImageWriteQueue/Private/ImageWriteQueue.cpp
+++ b/Engine/Source/Runtime/ImageWriteQueue/Private/ImageWriteQueue.cpp
@@ -33,6 +33,14 @@ static TAutoConsoleVariable<int32> CVarImageWriteQueueMaxQueueSize(
  */
 struct FImageWriteFence
 {
+	FImageWriteFence(uint32 InID, uint32 InCount, TPromise<void>&& InCompleted, TFunction<void()> InOnCompleted)
+		: ID(InID)
+		, Count(InCount)
+		, Completed(MoveTemp(InCompleted))
+		, OnCompleted(MoveTemp(InOnCompleted))
+	{
+	}
+
 	FImageWriteFence(FImageWriteFence&&) = default;
 	FImageWriteFence(const FImageWriteFence&) = delete;
 
@@ -389,7 +397,7 @@ TFuture<void> FImageWriteQueue::CreateFence(const TFunction<void()>& InOnFenceRe
 	else
 	{
 		// Move the promise into the write fence
-		PendingFences.Add(FImageWriteFence{CurrentFenceID, CurrentFenceCount, MoveTemp(Promise), InOnFenceReached});
+		PendingFences.Emplace(CurrentFenceID, CurrentFenceCount, MoveTemp(Promise), CopyTemp(InOnFenceReached));
 
 		// Reset the current fence context
 		++CurrentFenceID;
