Make FObjectPtr/TObjectPtr default initialize to nullptr.  Also added an ENoInit constructor if a caller wants them to be uninitialized.

diff --git a/Engine/Source/Runtime/CoreUObject/Public/UObject/ObjectPtr.h b/Engine/Source/Runtime/CoreUObject/Public/UObject/ObjectPtr.h
--- a/Engine/Source/Runtime/CoreUObject/Public/UObject/ObjectPtr.h
+++ b/Engine/Source/Runtime/CoreUObject/Public/UObject/ObjectPtr.h
@@ -45,7 +45,14 @@ struct TObjectPtr;
 struct FObjectPtr
 {
 public:
-	FObjectPtr() = default;
+	FObjectPtr()
+		: Handle(MakeObjectHandle(nullptr))
+	{
+	}
+
+	explicit FObjectPtr(ENoInit)
+	{
+	}
 
 	FORCEINLINE FObjectPtr(TYPE_OF_NULLPTR)
 		: Handle(MakeObjectHandle(nullptr))
@@ -201,6 +208,11 @@ public:
 	TObjectPtr(TObjectPtr<T>&&) = default;
 	TObjectPtr(const TObjectPtr<T>&) = default;
 
+	explicit FORCEINLINE TObjectPtr(ENoInit)
+		: FObjectPtr(NoInit)
+	{
+	}
+
 	FORCEINLINE TObjectPtr(TYPE_OF_NULLPTR)
 		: FObjectPtr(nullptr)
 	{
diff --git a/Engine/Source/Runtime/CoreUObject/Tests/ObjectPtrTest.cpp b/Engine/Source/Runtime/CoreUObject/Tests/ObjectPtrTest.cpp
--- a/Engine/Source/Runtime/CoreUObject/Tests/ObjectPtrTest.cpp
+++ b/Engine/Source/Runtime/CoreUObject/Tests/ObjectPtrTest.cpp
@@ -41,7 +41,6 @@ static_assert(std::is_trivially_move_constructible<FMutableObjectPtr>::value, "T
 static_assert(std::is_trivially_copy_assignable<FMutableObjectPtr>::value, "TObjectPtr must be trivially copy assignable");
 static_assert(std::is_trivially_move_assignable<FMutableObjectPtr>::value, "TObjectPtr must be trivially move assignable");
 static_assert(std::is_trivially_destructible<FMutableObjectPtr>::value, "TObjectPtr must be trivially destructible");
-static_assert(std::is_trivially_default_constructible<FMutableObjectPtr>::value, "TObjectPtr must be trivially default constructible");
 
 // Ensure that raw pointers can be used to construct wrapped object pointers and that const-ness isn't stripped when constructing or converting with raw pointers
 static_assert(std::is_constructible<FMutableObjectPtr, UObject*>::value, "TObjectPtr<UObject> must be constructible from a raw UObject*");
