ANY_PACKAGE Warnings: Turn off ambigious package warnings in addition to missing class warnings when loading old AssetData in DiffAssetRegistriesCommandlet.
Update DiffAssetRegistriesCommandlet to including searching for files specified by full path on any project.
#rnx

diff --git a/Engine/Source/Editor/UnrealEd/Private/Commandlets/DiffAssetRegistriesCommandlet.cpp b/Engine/Source/Editor/UnrealEd/Private/Commandlets/DiffAssetRegistriesCommandlet.cpp
--- a/Engine/Source/Editor/UnrealEd/Private/Commandlets/DiffAssetRegistriesCommandlet.cpp
+++ b/Engine/Source/Editor/UnrealEd/Private/Commandlets/DiffAssetRegistriesCommandlet.cpp
@@ -171,8 +171,11 @@ int32 UDiffAssetRegistriesCommandlet::Main(const FString& FullCommandLine)
 		UE_LOG(LogDiffAssets, Error, TEXT("No platform specified on the commandline use \"-platform=<platform>\"."));
 	}
 
+	TArray<FString> LocalSearchPaths = AssetRegistrySearchPath;
+	LocalSearchPaths.AddUnique(TEXT("[buildversion]"));
+
 	auto FindAssetRegistryPath = [&](const FString& PathVal, FString& OutPath) {
-			for (const FString& SearchPath : AssetRegistrySearchPath)
+			for (const FString& SearchPath : LocalSearchPaths)
 			{
 				FString FinalSearchPath = SearchPath;
 				FinalSearchPath.ReplaceInline(TEXT("[buildversion]"), *PathVal);
diff --git a/Engine/Source/Runtime/CoreUObject/Private/AssetRegistry/AssetData.cpp b/Engine/Source/Runtime/CoreUObject/Private/AssetRegistry/AssetData.cpp
--- a/Engine/Source/Runtime/CoreUObject/Private/AssetRegistry/AssetData.cpp
+++ b/Engine/Source/Runtime/CoreUObject/Private/AssetRegistry/AssetData.cpp
@@ -248,7 +248,9 @@ FTopLevelAssetPath FAssetData::TryConvertShortClassNameToPathName(FName InClassN
 	if (!InClassName.IsNone())
 	{
 		FString ClassNameString(InClassName.ToString());
-		ClassPath = UClass::TryConvertShortTypeNameToPathName<UStruct>(ClassNameString, ELogVerbosity::Warning, TEXT("AssetRegistry trying to convert short name to path name"));
+		ELogVerbosity::Type AmbiguousMessageVerbosity = (FailureMessageVerbosity == ELogVerbosity::NoLogging || FailureMessageVerbosity > ELogVerbosity::Warning) ?
+			FailureMessageVerbosity : ELogVerbosity::Warning;
+		ClassPath = UClass::TryConvertShortTypeNameToPathName<UStruct>(ClassNameString, AmbiguousMessageVerbosity, TEXT("AssetRegistry trying to convert short name to path name"));
 		if (ClassPath.IsNull())
 		{
 			// In some cases the class name stored in asset registry tags have been redirected with ini class redirects
@@ -259,7 +261,7 @@ FTopLevelAssetPath FAssetData::TryConvertShortClassNameToPathName(FName InClassN
 			}
 			else
 			{
-				ClassPath = UClass::TryConvertShortTypeNameToPathName<UStruct>(RedirectedName, ELogVerbosity::Warning, TEXT("AssetRegistry trying to convert redirected short name to path name"));
+				ClassPath = UClass::TryConvertShortTypeNameToPathName<UStruct>(RedirectedName, AmbiguousMessageVerbosity, TEXT("AssetRegistry trying to convert redirected short name to path name"));
 			}
 
 			if (ClassPath.IsNull())
