AsyncLoading 1&2: Refactor initial load to remove zen loader dependency on script object blob in order to support editor loading.

NotifyRegistrationEvent:
Remove all NRP_Started events since they add no value.
Move the NRP_Finished event for type NRT_Class from InitializePrivateStaticClass (Class.cpp) to UObjectLoadAllCompiledInDefaultProperties (UObjectBase.cpp) since this is where the outer registration function for the class is actually executed.
Add a new NRT_NoExportObject event in ConstructUPackage for top level noexport objects like UScriptStruct, UDelegateFunction and USparseDelegateFunction that are just part of the package and has no registration info.

AsyncLoader:
Move EDLBootNotificationManager from AsyncPackageLoader to AsyncLoading.
Add FEarlyRegistrationEventsRecorder to replay early events that occur before the loader is created.

AsyncLoading:
Ignore the new NRT_NoExportObject event.
Allow for objects to skip the NRP_Added and NRP_Started phases.

AsyncLoading2:
Remove LazyInitializeFromLoadPackage and the reading and parsing of the ScriptObjects binary blob.
Remove FindScriptImportObjectFromIndex/GFindExistingScriptImport that was dependent on the ScriptObject blob, instead all script objects are added to the global import store directly from the notification system callbacks.
Change GetUnresolvedCDOs logic to check for imported classes with missing or not-yet-initialized CDOs.
Add FAsyncPackageDesc2::ReferencerRequestId to let each package track its most recent load package request referencer.
Change ProcessPendingCDOs logic to avoid recursive reentrancy and to process the most relevant CDOs first, based on the new ReferencerRequestId.

Tests:
Lyra: Editor, Runtime EDL and Zen loader on Win64 and Console
ShooterGame: Editor, Runtime EDL and Zen loader on Win64 and Console
AncientGame: Editor, Runtime EDL and Zen loader on Win64
Other: Editor, CookedEditor, Runtime EDL and Zen loader on Win64 (client and server)

#rnx

diff --git a/Engine/Source/Runtime/CoreUObject/Private/Serialization/AsyncLoading.cpp b/Engine/Source/Runtime/CoreUObject/Private/Serialization/AsyncLoading.cpp
--- a/Engine/Source/Runtime/CoreUObject/Private/Serialization/AsyncLoading.cpp
+++ b/Engine/Source/Runtime/CoreUObject/Private/Serialization/AsyncLoading.cpp
@@ -48,6 +48,7 @@
 #include "UObject/GarbageCollectionInternal.h"
 #include "ProfilingDebugging/MiscTrace.h"
 #include "Serialization/LoadTimeTracePrivate.h"
+#include "Misc/PathViews.h"
 
 #define FIND_MEMORY_STOMPS (1 && (PLATFORM_WINDOWS || PLATFORM_UNIX) && !WITH_EDITORONLY_DATA)
 
@@ -86,7 +87,398 @@ DECLARE_FLOAT_ACCUMULATOR_STAT(TEXT("Total PostLoadObjects time GT"), STAT_FAsyn
 
 DECLARE_FLOAT_ACCUMULATOR_STAT( TEXT( "Async loading block time" ), STAT_AsyncIO_AsyncLoadingBlockingTime, STATGROUP_AsyncIO );
 DECLARE_FLOAT_ACCUMULATOR_STAT( TEXT( "Async package precache wait time" ), STAT_AsyncIO_AsyncPackagePrecacheWaitTime, STATGROUP_AsyncIO );
-	
+
+struct FEDLBootObjectState
+{
+	ENotifyRegistrationType NotifyRegistrationType;
+	ENotifyRegistrationPhase LastNotifyRegistrationPhase;
+	UObject* (*Register)();
+	bool bDynamic;
+};
+
+struct FEDLBootWaitingPackage
+{
+	void* Package;
+	FPackageIndex Import;
+};
+
+struct FEDLBootNotificationManager
+{
+	TMap<FName, FEDLBootObjectState> PathToState;
+	TMultiMap<FName, FEDLBootWaitingPackage> PathToWaitingPackageNodes;
+	TArray<FName> PathsToFire;
+	TArray<UClass*> CDORecursiveStack;
+	TArray<UClass*> CDORecursives;
+	FCriticalSection EDLBootNotificationManagerLock;
+	bool bEnabled = true;
+
+	void Disable()
+	{
+		PathToState.Empty();
+		PathsToFire.Empty();
+		bEnabled = false;
+	}
+
+	// return true if you are waiting for this compiled in object
+	bool AddWaitingPackage(void* Pkg, FName PackageName, FName ObjectName, FPackageIndex Import, bool bIgnoreMissingPackage)
+	{
+		if (PackageName == GLongCoreUObjectPackageName)
+		{
+			return false; // We assume nothing in coreuobject ever loads assets in a constructor
+		}
+		FScopeLock Lock(&EDLBootNotificationManagerLock);
+		check(GIsInitialLoad);
+		check(Import.IsImport()); // compiled in exports make no sense
+		FString ObjectNameString = ObjectName.ToString();
+		FName LongFName(*(PackageName.ToString() / ObjectNameString));
+		check(LongFName != NAME_None);
+		FName WaitName = LongFName;
+		FEDLBootObjectState* ExistingState = PathToState.Find(LongFName);
+		if (!ExistingState)
+		{
+			//if (ObjectName.ToString().EndsWith(HEADER_GENERATED_DELEGATE_SIGNATURE_SUFFIX))
+			// there are also some arg structs and other things that are just part of the package with no registration
+			{
+				ExistingState = PathToState.Find(PackageName);
+				WaitName = PackageName;
+			}
+			if (!ExistingState)
+			{
+				UE_CLOG(!bIgnoreMissingPackage, LogStreaming, Fatal, TEXT("Compiled in export %s not found; it was never registered."), *LongFName.ToString());
+				return false;
+			}
+		}
+		if (ExistingState->LastNotifyRegistrationPhase == ENotifyRegistrationPhase::NRP_Finished)
+		{
+			return false;
+		}
+		FEDLBootWaitingPackage WaitingPackage;
+		WaitingPackage.Package = Pkg;
+		WaitingPackage.Import = Import;
+
+		PathToWaitingPackageNodes.Add(WaitName, WaitingPackage);
+
+		return true;
+	}
+
+	void NotifyRegistrationEvent(const TCHAR* PackageName, const TCHAR* Name, ENotifyRegistrationType NotifyRegistrationType, ENotifyRegistrationPhase NotifyRegistrationPhase, UObject* (*InRegister)(), bool InbDynamic)
+	{
+		if (!bEnabled || !GIsInitialLoad)
+		{
+			return;
+		}
+		static FName LongCoreUObjectPackageName(TEXT("/Script/CoreUObject")); // can't use the global here because it may not be initialized yet
+		FName PackageFName(PackageName);
+		if (PackageFName == LongCoreUObjectPackageName)
+		{
+			return; // We assume nothing in coreuobject ever loads assets in a constructor
+		}
+
+		if (NotifyRegistrationType == ENotifyRegistrationType::NRT_NoExportObject)
+		{
+			return; // We don't care about these, they will always be found using StaticFindObject
+		}
+
+		TStringBuilder<256> LongNameBuilder;
+		LongNameBuilder << PackageName;
+		FPathViews::Append(LongNameBuilder, Name);
+		FName LongFName(LongNameBuilder.ToView());
+
+		FScopeLock Lock(&EDLBootNotificationManagerLock);
+
+		//FPlatformMisc::LowLevelOutputDebugStringf(TEXT("NotifyRegistrationEvent %s %d %d\r\n"), *LongFName.ToString(), int32(NotifyRegistrationType), int32(NotifyRegistrationPhase));
+
+		// some things, like delegate signatures, are not registered; rather they are part of the package singleton, so we track the package state as being the max of any member of that package
+		FEDLBootObjectState* ExistingPackageState = PathToState.Find(PackageFName);
+		FEDLBootObjectState* ExistingState = PathToState.Find(LongFName);
+
+		if (!ExistingState)
+		{
+			FEDLBootObjectState NewState;
+			NewState.LastNotifyRegistrationPhase = NotifyRegistrationPhase;
+			NewState.NotifyRegistrationType = NotifyRegistrationType;
+			NewState.Register = InRegister;
+			NewState.bDynamic = InbDynamic;
+			PathToState.Add(LongFName, NewState);
+
+			if (!ExistingPackageState)
+			{
+				NewState.NotifyRegistrationType = ENotifyRegistrationType::NRT_Package;
+				PathToState.Add(PackageFName, NewState);
+			}
+		}
+		else
+		{
+			if (int32(ExistingState->LastNotifyRegistrationPhase) >= int32(NotifyRegistrationPhase))
+			{
+				UE_CLOG(GEventDrivenLoaderEnabled, LogStreaming, Fatal, TEXT("Invalid state transition %d %d with %s when it has already been processed."), int32(ExistingState->LastNotifyRegistrationPhase), int32(NotifyRegistrationPhase), *LongFName.ToString());
+			}
+			if (ExistingState->NotifyRegistrationType != (NotifyRegistrationType))
+			{
+				UE_CLOG(GEventDrivenLoaderEnabled, LogStreaming, Fatal, TEXT("Multiple types %d %d with %s when it has already been processed."), int32(ExistingState->NotifyRegistrationType), int32(NotifyRegistrationType), *LongFName.ToString());
+			}
+			ExistingState->LastNotifyRegistrationPhase = NotifyRegistrationPhase;
+			if (NotifyRegistrationPhase == ENotifyRegistrationPhase::NRP_Finished)
+			{
+				ExistingState->Register = nullptr; // we don't need to do this in ConstructWaitingBootObjects
+				PathsToFire.Add(LongFName);
+			}
+			check(ExistingPackageState); // if we have an existing state for the thing, we should also have a 
+			if (ExistingPackageState && int32(NotifyRegistrationPhase) > int32(ExistingPackageState->LastNotifyRegistrationPhase))
+			{
+				ExistingPackageState->LastNotifyRegistrationPhase = NotifyRegistrationPhase;
+				if (NotifyRegistrationPhase == ENotifyRegistrationPhase::NRP_Finished)
+				{
+					//FPlatformMisc::LowLevelOutputDebugStringf(TEXT("Fired package %s %d %d\r\n"), *PackageFName.ToString(), int32(NotifyRegistrationType), int32(NotifyRegistrationPhase));
+					PathsToFire.Add(PackageFName);
+				}
+			}
+		}
+	}
+
+	void NotifyRegistrationComplete(FAsyncLoadingThread& AsyncLoadingThread)
+	{
+		if (!bEnabled)
+		{
+			return;
+		}
+#if USE_EVENT_DRIVEN_ASYNC_LOAD_AT_BOOT_TIME
+		FireCompletedCompiledInImports(AsyncLoadingThread, true);
+		FlushAsyncLoading();
+#endif
+#if !HACK_HEADER_GENERATOR
+		check(!GIsInitialLoad && IsInGameThread());
+		FScopeLock Lock(&EDLBootNotificationManagerLock);
+		for (auto& Pair : PathToState)
+		{
+			if (Pair.Value.LastNotifyRegistrationPhase != ENotifyRegistrationPhase::NRP_Finished && !Pair.Value.bDynamic)
+			{
+#if USE_EVENT_DRIVEN_ASYNC_LOAD_AT_BOOT_TIME
+				UE_CLOG(GEventDrivenLoaderEnabled, LogStreaming, Fatal, TEXT("%s (%d) was not complete (%d) after registration was complete."), *Pair.Key.ToString(), int32(Pair.Value.NotifyRegistrationType), int32(Pair.Value.LastNotifyRegistrationPhase));
+#else
+				UE_LOG(LogStreaming, Warning, TEXT("%s was not complete (%d) after registration was complete."), *Pair.Key.ToString(), int32(Pair.Value.LastNotifyRegistrationPhase));
+#endif
+			}
+		}
+		if (PathToWaitingPackageNodes.Num())
+		{
+			UE_LOG(LogStreaming, Fatal, TEXT("Initial load is complete, but we still have %d waiting packages."), PathToWaitingPackageNodes.Num());
+		}
+		if (GEventDrivenLoaderEnabled && PathsToFire.Num() && USE_EVENT_DRIVEN_ASYNC_LOAD_AT_BOOT_TIME)
+		{
+			for (FName Path : PathsToFire)
+			{
+				UE_LOG(LogStreaming, Error, TEXT("%s was not fired."), *Path.ToString());
+			}
+			UE_LOG(LogStreaming, Fatal, TEXT("Initial load is complete, but we still have %d imports to fire (listed above)."), PathsToFire.Num());
+		}
+#endif
+		Disable();
+	}
+
+	bool ConstructWaitingBootObjects()
+	{
+		static struct FFixedBootOrder
+		{
+			TArray<FName> Array;
+			FFixedBootOrder()
+			{
+				// look for any packages that we want to force preload at startup
+				FConfigSection* BootObjects = GConfig->GetSectionPrivate(TEXT("/Script/Engine.StreamingSettings"), false, true, GEngineIni);
+				if (BootObjects)
+				{
+					// go through list and add to the array
+					for (FConfigSectionMap::TIterator It(*BootObjects); It; ++It)
+					{
+						if (It.Key() == TEXT("FixedBootOrder"))
+						{
+							// add this package to the list to be fully loaded later
+							Array.Add(FName(*It.Value().GetValue()));
+						}
+					}
+				}
+			}
+
+		} FixedBootOrder;
+
+		check(GIsInitialLoad && IsInGameThread());
+		UObject* (*BootObjectRegister)() = nullptr;
+		UObject* (*BootPackageObjectRegister)() = nullptr;
+		FName WaitingPackage;
+		bool bIsCDO = false;
+
+		while (FixedBootOrder.Array.Num())
+		{
+			FName ThisItem = FixedBootOrder.Array.Pop();
+			FScopeLock Lock(&EDLBootNotificationManagerLock);
+			FEDLBootObjectState* ExistingState = PathToState.Find(ThisItem);
+
+			if (!ExistingState)
+			{
+				UE_LOG(LogStreaming, Fatal, TEXT("%s was listed as a fixed load order but was not found,"), *ThisItem.ToString());
+			}
+			else if (!ExistingState->Register)
+			{
+				UE_LOG(LogStreaming, Log, TEXT("%s was listed as a fixed load order but was already processed"), *ThisItem.ToString());
+			}
+			else
+			{
+				//FPlatformMisc::LowLevelOutputDebugStringf(TEXT("Booting Fixed %s %d\r\n"), *ThisItem.ToString(), int32(ExistingState->NotifyRegistrationType));
+				BootObjectRegister = ExistingState->Register;
+				ExistingState->Register = nullptr; // we don't need to do this more than once
+				bIsCDO = ExistingState->NotifyRegistrationType == ENotifyRegistrationType::NRT_ClassCDO;
+				break;
+			}
+		}
+
+		if (!BootObjectRegister)
+		{
+			FScopeLock Lock(&EDLBootNotificationManagerLock);
+			for (auto& Pair : PathToWaitingPackageNodes)
+			{
+				FEDLBootObjectState* ExistingState = PathToState.Find(Pair.Key);
+				if (ExistingState)
+				{
+					if (ExistingState->Register)
+					{
+						//FPlatformMisc::LowLevelOutputDebugStringf(TEXT("Booting %s %d\r\n"), *Pair.Key.ToString(), int32(ExistingState->NotifyRegistrationType));
+						BootObjectRegister = ExistingState->Register;
+						ExistingState->Register = nullptr; // we don't need to do this more than once
+						bIsCDO = ExistingState->NotifyRegistrationType == ENotifyRegistrationType::NRT_ClassCDO;
+						break;
+					}
+				}
+			}
+		}
+		if (BootObjectRegister)
+		{
+			UObject* BootObject = BootObjectRegister();
+			check(BootObject);
+			UObjectForceRegistration(BootObject);
+			if (bIsCDO)
+			{
+				UClass* Class = CastChecked<UClass>(BootObject);
+				bool bAnyParentOnStack = false;
+				UClass* Super = Class;
+				while (Super)
+				{
+					if (CDORecursiveStack.Contains(Super))
+					{
+						bAnyParentOnStack = true;
+						break;
+					}
+					Super = Super->GetSuperClass();
+				}
+
+				if (!bAnyParentOnStack)
+				{
+					CDORecursiveStack.Push(Class);
+					//FPlatformMisc::LowLevelOutputDebugStringf(TEXT("Create CDO %s\r\n"), *BootObject->GetName());
+					Class->GetDefaultObject();
+					verify(CDORecursiveStack.Pop() == Class);
+				}
+				else
+				{
+					//FPlatformMisc::LowLevelOutputDebugStringf(TEXT("Recursive Deferred %s\r\n"), *BootObject->GetName());
+					CDORecursives.Add(Class);
+				}
+			}
+			return true;
+		}
+		if (CDORecursives.Num())
+		{
+			UClass* OkToRun = nullptr;
+			for (UClass* Class : CDORecursives)
+			{
+				bool bAnyParentOnStack = false;
+				UClass* Super = Class;
+				while (Super)
+				{
+					if (CDORecursiveStack.Contains(Super))
+					{
+						bAnyParentOnStack = true;
+						break;
+					}
+					Super = Super->GetSuperClass();
+				}
+				if (!bAnyParentOnStack)
+				{
+					OkToRun = Class;
+					break;
+				}
+			}
+			if (OkToRun)
+			{
+				//FPlatformMisc::LowLevelOutputDebugStringf(TEXT("CDORecursives %s\r\n"), *OkToRun->GetName());
+				CDORecursives.Remove(OkToRun);
+				CDORecursiveStack.Push(OkToRun);
+				OkToRun->GetDefaultObject();
+				verify(CDORecursiveStack.Pop() == OkToRun);
+			}
+			else
+			{
+				FPlatformProcess::Sleep(.001f);
+			}
+			return true; // even if we didn't do anything we need to return true to avoid checking for cycles
+		}
+		return false;
+	}
+
+	bool IsWaitingForSomething()
+	{
+		FScopeLock Lock(&EDLBootNotificationManagerLock);
+		return PathToWaitingPackageNodes.Num() > 0;
+	}
+
+	bool IsObjComplete(UObject* Obj)
+	{
+		static FName LongCoreUObjectPackageName(TEXT("/Script/CoreUObject")); // can't use the global here because it may not be initialized yet
+		FName PackageName = Obj->GetOutermost()->GetFName();
+		if (PackageName == LongCoreUObjectPackageName)
+		{
+			return true; // We assume nothing in coreuobject ever loads assets in a constructor, therefore it can be considered complete
+		}
+		FScopeLock Lock(&EDLBootNotificationManagerLock);
+		FName LongFName(*(PackageName.ToString() / Obj->GetName()));
+
+		FEDLBootObjectState* ExistingState = PathToState.Find(LongFName);
+
+		if (!ExistingState || ExistingState->LastNotifyRegistrationPhase == ENotifyRegistrationPhase::NRP_Finished)
+		{
+			return true;
+		}
+		return false;
+	}
+
+	bool FireCompletedCompiledInImports(FAsyncLoadingThread& AsyncLoadingThread, bool bFinalRun = false)
+	{
+#if USE_EVENT_DRIVEN_ASYNC_LOAD_AT_BOOT_TIME
+		FScopeLock Lock(&EDLBootNotificationManagerLock);
+		check(bFinalRun || GIsInitialLoad);
+		bool bResult = !!PathsToFire.Num();
+		for (FName LongName : PathsToFire)
+		{
+			for (auto It = PathToWaitingPackageNodes.CreateKeyIterator(LongName); It; ++It)
+			{
+				FEDLBootWaitingPackage& WaitingPackage = It.Value();
+				AsyncLoadingThread.FireCompletedCompiledInImport(WaitingPackage.Package, WaitingPackage.Import);
+			}
+			PathToWaitingPackageNodes.Remove(LongName);
+		}
+		PathsToFire.Empty();
+		return bResult;
+#else
+		return false;
+#endif
+	}
+};
+
+static FEDLBootNotificationManager& GetGEDLBootNotificationManager()
+{
+	static FEDLBootNotificationManager Singleton;
+	return Singleton;
+}
+
 /** Helper function for profiling load times */
 static FName StaticGetNativeClassName(UClass* InClass)
 {
@@ -276,6 +668,29 @@ FORCEINLINE bool FAsyncPackage::IsTimeLimitExceeded()
 	return AsyncLoadingThread.IsAsyncLoadingSuspendedInternal() || ::IsTimeLimitExceeded(TickStartTime, bUseTimeLimit, TimeLimit, LastTypeOfWorkPerformed, LastObjectWorkWasPerformedOn);
 }
 
+bool FAsyncPackage::IsFullyLoadedObj(UObject* Obj)
+{
+	if (!Obj)
+	{
+		return false;
+	}
+	if (Obj->HasAllFlags(RF_WasLoaded | RF_LoadCompleted)
+		|| Obj->IsA(UPackage::StaticClass())) // packages are never really loaded, so if it exists, it is loaded
+	{
+		return true;
+	}
+	if (Obj->HasAnyFlags(RF_WasLoaded | RF_NeedLoad | RF_WillBeLoaded))
+	{
+		return false;
+	}
+	if (GIsInitialLoad && Obj->GetOutermost()->HasAnyPackageFlags(PKG_CompiledIn))
+	{
+		return EDLBootNotificationManager.IsObjComplete(Obj);
+	}
+
+	return true;
+}
+
 FORCENOINLINE static bool CheckForFilePackageOpenLogCommandLine()
 {
 	return FParse::Param(FCommandLine::Get(), TEXT("FilePackageOpenLog"));
@@ -4569,7 +4984,7 @@ EAsyncPackageState::Type FAsyncLoadingThread::TickAsyncLoading(bool bUseTimeLimi
 #if USE_EVENT_DRIVEN_ASYNC_LOAD_AT_BOOT_TIME
 					if (GIsInitialLoad)
 					{
-						bDidSomething = EDLBootNotificationManager.FireCompletedCompiledInImports(); // no ASL, first try to fire any completed boot objects, and if there are none, then create some boot objects
+						bDidSomething = EDLBootNotificationManager.FireCompletedCompiledInImports(*this); // no ASL, first try to fire any completed boot objects, and if there are none, then create some boot objects
 						IsTimeLimitExceeded(TickStartTime, bUseTimeLimit, TimeLimit, TEXT("FireCompletedCompiledInImports"));
 						if (!bDidSomething)
 						{
@@ -4620,8 +5035,8 @@ void FMaxPackageSummarySize::Init()
 
 bool FAsyncLoadingThread::bThreadStarted = false;
 
-FAsyncLoadingThread::FAsyncLoadingThread(int32 InThreadIndex, IEDLBootNotificationManager& InEDLBootNotificationManager)
-	: EDLBootNotificationManager(InEDLBootNotificationManager)
+FAsyncLoadingThread::FAsyncLoadingThread(int32 InThreadIndex)
+	: EDLBootNotificationManager(GetGEDLBootNotificationManager())
 	, Thread(nullptr)
 {
 	if (!Instance)
@@ -4798,7 +5213,7 @@ EAsyncPackageState::Type FAsyncLoadingThread::TickAsyncThread(bool bUseTimeLimit
 		{
 			if (GIsInitialLoad && GEventDrivenLoaderEnabled)
 			{
-				EDLBootNotificationManager.FireCompletedCompiledInImports();
+				EDLBootNotificationManager.FireCompletedCompiledInImports(*this);
 			}
 			{
 				//@note FH: pretty sure the GCScopeGuard isn't needed here, nor is it useful, it is also retaken inside CreateAsyncPackagesFromQueue which isn't necessary
@@ -5095,7 +5510,17 @@ void FAsyncLoadingThread::NotifyConstructedDuringAsyncLoading(UObject* Object, b
 			AsyncPackage->MarkNewObjectForLoadIfItIsAnExport(Object);
 		}
 	}
-	}
+}
+
+void FAsyncLoadingThread::NotifyRegistrationEvent(const TCHAR* PackageName, const TCHAR* Name, ENotifyRegistrationType NotifyRegistrationType, ENotifyRegistrationPhase NotifyRegistrationPhase, UObject* (*InRegister)(), bool InbDynamic)
+{
+	EDLBootNotificationManager.NotifyRegistrationEvent(PackageName, Name, NotifyRegistrationType, NotifyRegistrationPhase, InRegister, InbDynamic);
+}
+
+void FAsyncLoadingThread::NotifyRegistrationComplete()
+{
+	EDLBootNotificationManager.NotifyRegistrationComplete(*this);
+}
 
 void FAsyncLoadingThread::FireCompletedCompiledInImport(void* AsyncPackage, FPackageIndex Import)
 {
@@ -5113,7 +5538,7 @@ void FAsyncLoadingThread::FireCompletedCompiledInImport(void* AsyncPackage, FPac
 /**
 * Constructor
 */
-FAsyncPackage::FAsyncPackage(FAsyncLoadingThread& InThread, const FAsyncPackageDesc& InDesc, IEDLBootNotificationManager& InEDLBootNotificationManager)
+FAsyncPackage::FAsyncPackage(FAsyncLoadingThread& InThread, const FAsyncPackageDesc& InDesc, FEDLBootNotificationManager& InEDLBootNotificationManager)
 : Desc(InDesc)
 , Linker(nullptr)
 , LinkerRoot(nullptr)
diff --git a/Engine/Source/Runtime/CoreUObject/Private/Serialization/AsyncLoading.h b/Engine/Source/Runtime/CoreUObject/Private/Serialization/AsyncLoading.h
--- a/Engine/Source/Runtime/CoreUObject/Private/Serialization/AsyncLoading.h
+++ b/Engine/Source/Runtime/CoreUObject/Private/Serialization/AsyncLoading.h
@@ -343,7 +343,7 @@ struct FAsyncPackage : public FGCObject
 	/**
 	 * Constructor
 	 */
-	FAsyncPackage(FAsyncLoadingThread& InThread, const FAsyncPackageDesc& InDesc, IEDLBootNotificationManager& InEDLBootNotificationManager);
+	FAsyncPackage(FAsyncLoadingThread& InThread, const FAsyncPackageDesc& InDesc, struct FEDLBootNotificationManager& InEDLBootNotificationManager);
 	~FAsyncPackage();
 
 	/**
@@ -549,6 +549,8 @@ private:
 	/** Add the list of passed in Requests to the list of requests that are dependent on this package. */
 	void AddDependentRequests(TArrayView<int32 const> Requests);
 
+	bool IsFullyLoadedObj(UObject* Obj);
+
 	/** Basic information associated with this package */
 	FAsyncPackageDesc Desc;
 	/** Linker which is going to have its exports and imports loaded									*/
@@ -643,7 +645,7 @@ private:
 	FCriticalSection ReferencedObjectsCritical;
 	/** Cached async loading thread object this package was created by */
 	FAsyncLoadingThread& AsyncLoadingThread;
-	class IEDLBootNotificationManager& EDLBootNotificationManager;
+	FEDLBootNotificationManager& EDLBootNotificationManager;
 	/** Packages that have been imported by this async package */
 	TSet<UPackage*> ImportedPackages;
 
diff --git a/Engine/Source/Runtime/CoreUObject/Private/Serialization/AsyncLoading2.cpp b/Engine/Source/Runtime/CoreUObject/Private/Serialization/AsyncLoading2.cpp
--- a/Engine/Source/Runtime/CoreUObject/Private/Serialization/AsyncLoading2.cpp
+++ b/Engine/Source/Runtime/CoreUObject/Private/Serialization/AsyncLoading2.cpp
@@ -63,6 +63,7 @@
 #include "HAL/IPlatformFileOpenLogWrapper.h"
 #include "Modules/ModuleManager.h"
 #include "Containers/SpscQueue.h"
+#include "Misc/PathViews.h"
 
 #if UE_BUILD_DEVELOPMENT || UE_BUILD_DEBUG
 PRAGMA_DISABLE_OPTIMIZATION
@@ -383,6 +384,8 @@ struct FAsyncPackageDesc2
 {
 	// A unique request id for each external call to LoadPackage
 	int32 RequestID;
+	// The request id of the referencer propagated down the import chain from the most recent load request
+	int32 ReferencerRequestId;
 	// Package priority
 	int32 Priority;
 	// The package id of the UPackage being loaded
@@ -409,6 +412,7 @@ struct FAsyncPackageDesc2
 	{
 		return FAsyncPackageDesc2
 		{
+			RequestID,
 			RequestID,
 			Priority,
 			FPackageId::FromName(UPackageName),
@@ -420,6 +424,7 @@ struct FAsyncPackageDesc2
 	}
 
 	static FAsyncPackageDesc2 FromPackageImport(
+		int32 ReferencerRequestId,
 		int32 Priority,
 		FPackageId ImportedPackageId,
 		FPackageId PackageIdToLoad,
@@ -428,6 +433,7 @@ struct FAsyncPackageDesc2
 		return FAsyncPackageDesc2
 		{
 			INDEX_NONE,
+			ReferencerRequestId,
 			Priority,
 			ImportedPackageId,
 			PackageIdToLoad,
@@ -449,10 +455,6 @@ struct FGlobalImportStore
 	TMap<FPackageObjectIndex, UObject*> ScriptObjects;
 	TMap<FPublicExportKey, int32> PublicExportToObjectIndex;
 	TMap<int32, FPublicExportKey> ObjectIndexToPublicExport;
-	// Temporary initial load data
-	TArray<FScriptObjectEntry> ScriptObjectEntries;
-	TMap<FPackageObjectIndex, FScriptObjectEntry*> ScriptObjectEntriesMap;
-	bool bHasInitializedScriptObjects = false;
 
 	FGlobalImportStore()
 	{
@@ -460,39 +462,6 @@ struct FGlobalImportStore
 		ObjectIndexToPublicExport.Reserve(32768);
 	}
 
-	void Initialize(FIoDispatcher& IoDispatcher)
-	{
-		{
-			TRACE_CPUPROFILER_EVENT_SCOPE(SetupInitialLoadData);
-
-			FEvent* InitialLoadEvent = FPlatformProcess::GetSynchEventFromPool();
-
-			FIoBatch IoBatch = IoDispatcher.NewBatch();
-			FIoRequest IoRequest = IoBatch.Read(CreateIoChunkId(0, 0, EIoChunkType::ScriptObjects), FIoReadOptions(), IoDispatcherPriority_High);
-			IoBatch.IssueAndTriggerEvent(InitialLoadEvent);
-
-			InitialLoadEvent->Wait();
-			FPlatformProcess::ReturnSynchEventToPool(InitialLoadEvent);
-
-			const FIoBuffer& InitialLoadIoBuffer = IoRequest.GetResultOrDie();
-			FLargeMemoryReader InitialLoadArchive(InitialLoadIoBuffer.Data(), InitialLoadIoBuffer.DataSize());
-			FNameMap NameMap;
-			NameMap.Load(InitialLoadArchive, FMappedName::EType::Global);
-			int32 NumScriptObjects = 0;
-			InitialLoadArchive << NumScriptObjects;
-			ScriptObjectEntries = MakeArrayView(reinterpret_cast<const FScriptObjectEntry*>(InitialLoadIoBuffer.Data() + InitialLoadArchive.Tell()), NumScriptObjects);
-
-			ScriptObjectEntriesMap.Reserve(ScriptObjectEntries.Num());
-			for (FScriptObjectEntry& ScriptObjectEntry : ScriptObjectEntries)
-			{
-				check(ScriptObjectEntry.Mapped.IsGlobal());
-				ScriptObjectEntry.ObjectName = NameMap.GetMinimalName(ScriptObjectEntry.Mapped);
-
-				ScriptObjectEntriesMap.Add(ScriptObjectEntry.GlobalIndex, &ScriptObjectEntry);
-			}
-		}
-	}
-
 	TArray<FPackageId> RemovePublicExports(const FUnreachablePublicExports& PublicExports)
 	{
 		TArray<FPackageId> PackageIds;
@@ -580,20 +549,11 @@ struct FGlobalImportStore
 		return Object;
 	}
 
-	UObject* FindScriptImportObjectFromIndex(FPackageObjectIndex ScriptImportIndex);
-
 	inline UObject* FindScriptImportObject(FPackageObjectIndex GlobalIndex)
 	{
 		check(GlobalIndex.IsScriptImport());
 		UObject* Object = nullptr;
-		if (!bHasInitializedScriptObjects)
-		{
-			Object = FindScriptImportObjectFromIndex(GlobalIndex);
-		}
-		else
-		{
-			Object = ScriptObjects.FindRef(GlobalIndex);
-		}
+		Object = ScriptObjects.FindRef(GlobalIndex);
 		return Object;
 	}
 
@@ -635,6 +595,48 @@ struct FGlobalImportStore
 	}
 
 	void FindAllScriptObjects();
+
+	void AddScriptObject(FStringView PackageName, FStringView Name)
+	{
+		TStringBuilder<FName::StringBufferSize> FullName;
+		FPathViews::Append(FullName, PackageName);
+		FPathViews::Append(FullName, Name);
+		FPackageObjectIndex GlobalImportIndex = FPackageObjectIndex::FromScriptPath(FullName);
+
+		UObject* Package = StaticFindObjectFastInternal(UPackage::StaticClass(), nullptr, FName(PackageName), true);
+		// check(Package);
+		if (!ensureMsgf(Package, TEXT("Failed to find package when registering script object '%.*s/%.*s'"),
+			PackageName.Len(), PackageName.GetData(), Name.Len(), Name.GetData()))
+		{
+			return;
+		}
+
+		UObject* Object = StaticFindObjectFastInternal(nullptr, Package, FName(Name));
+		// check(Object);
+		if (!ensureMsgf(Object, TEXT("Failed to find object when registering script object '%.*s/%.*s'"),
+			PackageName.Len(), PackageName.GetData(), Name.Len(), Name.GetData()))
+		{
+			return;
+		}
+
+#if WITH_EDITOR
+		FPackageObjectIndex PackageGlobalImportIndex = FPackageObjectIndex::FromScriptPath(PackageName);
+		ScriptObjects.Add(PackageGlobalImportIndex, Package);
+#endif
+		ScriptObjects.Add(GlobalImportIndex, Object);
+
+		TStringBuilder<FName::StringBufferSize> SubObjectName;
+		ForEachObjectWithOuter(Object, [this,&SubObjectName](UObject* SubObject)
+		{
+			if (SubObject->HasAnyFlags(RF_Public))
+			{
+				SubObjectName.Reset();
+				SubObject->GetPathName(nullptr, SubObjectName);
+				FPackageObjectIndex SubObjectGlobalImportIndex = FPackageObjectIndex::FromScriptPath(SubObjectName);
+				ScriptObjects.Add(SubObjectGlobalImportIndex, SubObject);
+			}
+		}, /* bIncludeNestedObjects*/ true);
+	}
 };
 
 class FLoadedPackageRef
@@ -940,23 +942,32 @@ struct FPackageImportStore
 				continue;
 			}
 
-			UObject* Object = GlobalImportStore.FindScriptImportObjectFromIndex(Index);
-			if (Object)
+			UObject* Object = GlobalImportStore.FindScriptImportObject(Index);
+			if (!Object)
 			{
 				continue;
 			}
 
-			const FScriptObjectEntry* Entry = GlobalImportStore.ScriptObjectEntriesMap.FindRef(Index);
-			check(Entry);
-			const FPackageObjectIndex& CDOClassIndex = Entry->CDOClassIndex;
-			if (CDOClassIndex.IsScriptImport())
+			UClass* Class = Cast<UClass>(Object);
+			if (!Class)
 			{
-				UObject* CDOClassObject = GlobalImportStore.FindScriptImportObjectFromIndex(CDOClassIndex);
-				if (CDOClassObject)
-				{
-					UClass* CDOClass = static_cast<UClass*>(CDOClassObject);
-					Classes.AddUnique(CDOClass);
-				}
+				continue;
+			}
+
+			// Filter out CDOs that are themselves classes,
+			// like Default__BlueprintGeneratedClass of type UBlueprintGeneratedClass
+			if (Class->HasAnyFlags(RF_ClassDefaultObject))
+			{
+				continue;
+			}
+
+			// Add dependency on any script CDO that has not been created and initialized yet
+			UObject* CDO = Class->GetDefaultObject(/*bCreateIfNeeded*/ false);
+			if (!CDO || CDO->HasAnyFlags(RF_NeedInitialization))
+			{
+				UE_LOG(LogStreaming, Log, TEXT("Package %s has a dependency on pending script CDO for '%s' (0x%llX)"),
+					*Header.PackageName.ToString(), *Class->GetFullName(), Index.Value());
+					Classes.AddUnique(Class);
 			}
 		}
 	}
@@ -1393,6 +1404,8 @@ public:
 		NodeState.Store(static_cast<uint8>(InNodeState));
 	}
 
+	int32 ReferencerRequestId() const;
+
 private:
 	void ProcessDependencies(FAsyncLoadingThreadState2& ThreadState);
 	void Fire(FAsyncLoadingThreadState2* ThreadState = nullptr);
@@ -1660,6 +1673,11 @@ struct FAsyncPackage2
 	/** Adds new request ID to the existing package */
 	void AddRequestID(int32 Id);
 
+	int32 ReferencerRequestId() const
+	{
+		return Desc.ReferencerRequestId;
+	}
+
 	/**
 	* Cancel loading this package.
 	*/
@@ -2019,12 +2037,10 @@ private:
 	TArray<FAsyncLoadingThreadWorker> Workers;
 	TAtomic<int32> ActiveWorkersCount { 0 };
 	bool bWorkersSuspended = false;
-
+	bool bHasRegisteredAllScriptObjects = false;
 	/** [ASYNC/GAME THREAD] true if the async thread is actually started. We don't start it until after we boot because the boot process on the game thread can create objects that are also being created by the loader */
 	bool bThreadStarted = false;
 
-	mutable bool bLazyInitializedFromLoadPackage = false;
-
 #if ALT2_VERIFY_RECURSIVE_LOADS
 	int32 LoadRecursionLevel = 0;
 #endif
@@ -2100,6 +2116,7 @@ private:
 
 	/** Initial load pending CDOs */
 	TMap<UClass*, TArray<FEventLoadNode2*>> PendingCDOs;
+	TArray<UClass*> PendingCDOsRecursiveStack;
 
 	uint32 ConditionalBeginPostLoadTick = 0;
 	uint32 ConditionalFinishLoadingTick = 0;
@@ -2196,7 +2213,9 @@ public:
 
 	virtual void NotifyUnreachableObjects(const TArrayView<FUObjectItem*>& UnreachableObjects) override;
 
-	virtual void FireCompletedCompiledInImport(void* AsyncPackage, FPackageIndex Import) override {}
+	virtual void NotifyRegistrationEvent(const TCHAR* PackageName, const TCHAR* Name, ENotifyRegistrationType NotifyRegistrationType, ENotifyRegistrationPhase NotifyRegistrationPhase, UObject* (*InRegister)(), bool InbDynamic) override;
+
+	virtual void NotifyRegistrationComplete() override;
 
 	FORCEINLINE FAsyncPackage2* FindAsyncPackage(FPackageId PackageId)
 	{
@@ -2352,35 +2371,81 @@ private:
 	void SuspendWorkers();
 	void ResumeWorkers();
 
-	void LazyInitializeFromLoadPackage();
 	void FinalizeInitialLoad();
 
 	void RemoveUnreachableObjects(const FUnreachablePublicExports& PublicExports, const FUnreachablePackages& Packages);
 
-	bool ProcessPendingCDOs()
+	void ProcessPendingCDOs()
 	{
-		if (PendingCDOs.Num() > 0)
+		TRACE_CPUPROFILER_EVENT_SCOPE(ProcessPendingCDOs);
+
+		UClass* Class = nullptr;
+		int32 MaxRequestId = -1;
+		for (TMap<UClass*, TArray<FEventLoadNode2*>>::TIterator It = PendingCDOs.CreateIterator(); It; ++It)
 		{
-			TRACE_CPUPROFILER_EVENT_SCOPE(ProcessPendingCDOs);
+			UClass* CurrentClass = It.Key();
+			if (PendingCDOsRecursiveStack.Num() > 0)
+			{
+				bool bAnyParentOnStack = false;
+				UClass* Super = CurrentClass;
+				while (Super)
+				{
+					if (PendingCDOsRecursiveStack.Contains(Super))
+					{
+						bAnyParentOnStack = true;
+						break;
+					}
+					Super = Super->GetSuperClass();
+				}
+				if (bAnyParentOnStack)
+				{
+					continue;
+				}
+			}
 
-			auto It = PendingCDOs.CreateIterator();
-			UClass* Class = It.Key();
-			TArray<FEventLoadNode2*> Nodes = MoveTemp(It.Value());
-			It.RemoveCurrent();
+			const TArray<FEventLoadNode2*>& Nodes = It.Value();
+			for (const FEventLoadNode2* Node : Nodes)
+			{
+				int32 RequestId = Node->ReferencerRequestId();
+				if (RequestId > MaxRequestId)
+				{
+					MaxRequestId = RequestId;
+					Class = CurrentClass;
+				}
+			}
+		}
 
-			UE_LOG(LogStreaming, Verbose, TEXT("ProcessPendingCDOs: Creating CDO for %s. %d entries remaining."), *Class->GetFullName(), PendingCDOs.Num());
-			UObject* CDO = Class->GetDefaultObject();
+		if (ensure(Class))
+		{
+			TArray<FEventLoadNode2*> Nodes;
+			PendingCDOs.RemoveAndCopyValue(Class, Nodes);
 
-			ensureMsgf(CDO, TEXT("Failed to create CDO for %s"), *Class->GetFullName());
-			UE_LOG(LogStreaming, Verbose, TEXT("ProcessPendingCDOs: Created CDO for %s."), *Class->GetFullName());
+			UE_LOG(LogStreaming, Log,
+				TEXT("ProcessPendingCDOs: Creating CDO for '%s' for request id %d, releasing %d nodes. %d CDOs remaining."),
+				*Class->GetFullName(), MaxRequestId, Nodes.Num(), PendingCDOs.Num());
+			PendingCDOsRecursiveStack.Push(Class);
+			UObject* CDO = Class->GetDefaultObject(/*bCreateIfNeeded*/ true);
+			verify(PendingCDOsRecursiveStack.Pop() == Class);
+
+			ensureAlwaysMsgf(CDO, TEXT("Failed to create CDO for %s"), *Class->GetFullName());
+			UE_LOG(LogStreaming, Verbose, TEXT("ProcessPendingCDOs: Created CDO for '%s'."), *Class->GetFullName());
 
 			for (FEventLoadNode2* Node : Nodes)
 			{
 				Node->ReleaseBarrier();
 			}
-			return true;
 		}
-		return false;
+		else
+		{
+			for (TMap<UClass*, TArray<FEventLoadNode2*>>::TIterator It = PendingCDOs.CreateIterator(); It; ++It)
+			{
+				UClass* CurrentClass = It.Key();
+				const TArray<FEventLoadNode2*>& Nodes = It.Value();
+				UE_LOG(LogStreaming, Warning,
+					TEXT("ProcessPendingCDOs: '%s' with %d nodes could not be processed from this stack."),
+					*Class->GetFullName(), Nodes.Num());
+			}
+		}
 	}
 
 	/**
@@ -2647,6 +2712,7 @@ FAsyncPackage2* FAsyncLoadingThread2::FindOrInsertPackage(FAsyncPackageDesc2& De
 			{
 				UpdatePackagePriority(Package, Desc.Priority);
 			}
+			Package->Desc.ReferencerRequestId = Desc.ReferencerRequestId;
 		}
 		if (PackageLoadedDelegate.IsValid())
 		{
@@ -2990,6 +3056,11 @@ void FEventLoadNode2::ProcessDependencies(FAsyncLoadingThreadState2& ThreadState
 	}
 }
 
+int32 FEventLoadNode2::ReferencerRequestId() const
+{
+	return Package->ReferencerRequestId();
+}
+
 FAsyncLoadEventQueue2::FAsyncLoadEventQueue2()
 {
 }
@@ -3227,37 +3298,6 @@ void FAsyncPackage2::SetupScriptDependencies()
 	}
 }
 
-static UObject* GFindExistingScriptImport(FPackageObjectIndex GlobalImportIndex,
-	TMap<FPackageObjectIndex, UObject*>& ScriptObjects,
-	const TMap<FPackageObjectIndex, FScriptObjectEntry*>& ScriptObjectEntriesMap)
-{
-	UObject** Object = &ScriptObjects.FindOrAdd(GlobalImportIndex);
-	if (!*Object)
-	{
-		const FScriptObjectEntry* Entry = ScriptObjectEntriesMap.FindRef(GlobalImportIndex);
-		check(Entry);
-		if (Entry->OuterIndex.IsNull())
-		{
-			*Object = StaticFindObjectFast(UPackage::StaticClass(), nullptr, MinimalNameToName(Entry->ObjectName), true);
-		}
-		else
-		{
-			UObject* Outer = GFindExistingScriptImport(Entry->OuterIndex, ScriptObjects, ScriptObjectEntriesMap);
-			Object = &ScriptObjects.FindChecked(GlobalImportIndex);
-			if (Outer)
-			{
-				*Object = StaticFindObjectFast(UObject::StaticClass(), Outer, MinimalNameToName(Entry->ObjectName), false);
-			}
-		}
-	}
-	return *Object;
-}
-
-UObject* FGlobalImportStore::FindScriptImportObjectFromIndex(FPackageObjectIndex GlobalImportIndex)
-{
-	check(ScriptObjectEntries.Num() > 0);
-	return GFindExistingScriptImport(GlobalImportIndex, ScriptObjects, ScriptObjectEntriesMap);
-}
 
 void FGlobalImportStore::FindAllScriptObjects()
 {
@@ -3269,6 +3309,13 @@ void FGlobalImportStore::FindAllScriptObjects()
 
 	for (UPackage* Package : ScriptPackages)
 	{
+#if WITH_EDITOR
+		Name.Reset();
+		Package->GetPathName(nullptr, Name);
+		FPackageObjectIndex PackageGlobalImportIndex = FPackageObjectIndex::FromScriptPath(Name);
+		ScriptObjects.Add(PackageGlobalImportIndex, Package);
+#endif
+
 		Objects.Reset();
 		GetObjectsWithOuter(Package, Objects, /*bIncludeNestedObjects*/true);
 		for (UObject* Object : Objects)
@@ -3278,19 +3325,18 @@ void FGlobalImportStore::FindAllScriptObjects()
 				Name.Reset();
 				Object->GetPathName(nullptr, Name);
 				FPackageObjectIndex GlobalImportIndex = FPackageObjectIndex::FromScriptPath(Name);
-				ScriptObjects.Add(GlobalImportIndex, Object);
+				if (!ScriptObjects.Contains(GlobalImportIndex))
+				{
+					ScriptObjects.Add(GlobalImportIndex, Object);
+#if !WITH_EDITOR
+					UE_LOG(LogStreaming, Warning, TEXT("Found additional public script object 0x%016llX: %s\r\n"),
+						GlobalImportIndex.Value(),
+						*Object->GetFullName());
+#endif
+				}
 			}
 		}
 	}
-
-	ScriptObjectEntriesMap.Empty();
-	ScriptObjectEntries.Empty();
-	ScriptObjects.Shrink();
-
-	bHasInitializedScriptObjects = true;
-
-	UE_LOG(LogStreaming, Display, TEXT("AsyncLoading2 - InitialLoad Finalized: %d script object entries in %.2f KB"),
-		ScriptObjects.Num(), (float)ScriptObjects.GetAllocatedSize() / 1024.f);
 }
 
 void FAsyncPackage2::ImportPackagesRecursiveInner(FIoBatch& IoBatch, FPackageStore& PackageStore, const TArrayView<const FPackageId>& ImportedPackageIds, int32& ImportedPackageIndex)
@@ -3394,7 +3440,7 @@ void FAsyncPackage2::ImportPackagesRecursiveInner(FIoBatch& IoBatch, FPackageSto
 
 		FAsyncPackage2* ImportedPackage = nullptr;
 		bool bInserted = false;
-		FAsyncPackageDesc2 PackageDesc = FAsyncPackageDesc2::FromPackageImport(Desc.Priority, ImportedPackageId, ImportedPackageIdToLoad, ImportedPackageUPackageName);
+		FAsyncPackageDesc2 PackageDesc = FAsyncPackageDesc2::FromPackageImport(Desc.ReferencerRequestId, Desc.Priority, ImportedPackageId, ImportedPackageIdToLoad, ImportedPackageUPackageName);
 		if (PackageRef.AreAllPublicExportsLoaded())
 		{
 			ImportedPackage = AsyncLoadingThread.FindAsyncPackage(ImportedPackageId);
@@ -3698,7 +3744,7 @@ EAsyncPackageState::Type FAsyncPackage2::Event_SetupDependencies(FAsyncLoadingTh
 	
 	if (!Package->bLoadHasFailed)
 	{
-		if (GIsInitialLoad)
+		if (!Package->AsyncLoadingThread.bHasRegisteredAllScriptObjects)
 		{
 			Package->SetupScriptDependencies();
 		}
@@ -3918,32 +3964,29 @@ UObject* FAsyncPackage2::EventDrivenIndexToObject(const FAsyncPackageHeaderData&
 	if (Index.IsExport())
 	{
 		Result = Exports[Index.ToExport()].Object;
+		UE_CLOG(!Result, LogStreaming, Warning, TEXT("Missing Dependency, missing export 0x%llX in package %s"),
+			Index.Value(),
+			*Desc.PackagePathToLoad.GetPackageFName().ToString());
 	}
 	else if (Index.IsImport())
 	{
 		Result = ImportStore.FindOrGetImportObject(Header, Index);
-		UE_CLOG(!Result, LogStreaming, Warning, TEXT("Missing %s import 0x%llX for package %s"),
+		UE_CLOG(!Result, LogStreaming, Warning, TEXT("Missing Dependency, missing %s import 0x%llX for package %s"),
 			Index.IsScriptImport() ? TEXT("script") : TEXT("package"),
 			Index.Value(),
 			*Desc.PackagePathToLoad.GetPackageFName().ToString());
 	}
 #if DO_CHECK
-	if (bCheckSerialized && !IsFullyLoadedObj(Result))
+	if (Result && bCheckSerialized)
 	{
-		/*FEventLoadNode2* MyDependentNode = GetExportNode(EEventLoadNode2::Export_Serialize, Index.ToExport());
-		if (!Result)
+		bool bIsSerialized = Index.IsScriptImport() || Result->IsA(UPackage::StaticClass()) || Result->HasAllFlags(RF_WasLoaded | RF_LoadCompleted);
+		if (!bIsSerialized)
 		{
-			UE_LOG(LogStreaming, Error, TEXT("Missing Dependency, request for %s but it hasn't been created yet."), *Linker->GetPathName(Index));
+			UE_LOG(LogStreaming, Warning, TEXT("Missing Dependency, '%s' (0x%llX) for package %s has not been serialized yet."),
+				*Result->GetFullName(),
+				Index.Value(),
+				*Desc.PackagePathToLoad.GetPackageFName().ToString());
 		}
-		else if (!MyDependentNode || MyDependentNode->GetBarrierCount() > 0)
-		{
-			UE_LOG(LogStreaming, Fatal, TEXT("Missing Dependency, request for %s but it was still waiting for serialization."), *Linker->GetPathName(Index));
-		}
-		else
-		{
-			UE_LOG(LogStreaming, Fatal, TEXT("Missing Dependency, request for %s but it was still has RF_NeedLoad."), *Linker->GetPathName(Index));
-		}*/
-		UE_LOG(LogStreaming, Warning, TEXT("Missing Dependency"));
 	}
 	if (Result)
 	{
@@ -5155,16 +5198,15 @@ EAsyncPackageState::Type FAsyncLoadingThread2::TickAsyncLoadingFromGameThread(FA
 
 		if (Result != EAsyncPackageState::TimeOut)
 		{
-			// Flush deferred messages
-			if (!IsAsyncLoadingPackages())
+			if (!bDidSomething && PendingCDOs.Num() > 0)
 			{
-				bDidSomething = true;
-				FDeferredMessageLog::Flush();
+				ProcessPendingCDOs();
 			}
 
-			if (GIsInitialLoad && !bDidSomething)
+			// Flush deferred messages
+			if (!IsAsyncLoadingPackages())
 			{
-				bDidSomething = ProcessPendingCDOs();
+				FDeferredMessageLog::Flush();
 			}
 		}
 
@@ -5300,27 +5342,21 @@ void FAsyncLoadingThread2::ResumeWorkers()
 	bWorkersSuspended = false;
 }
 
-void FAsyncLoadingThread2::LazyInitializeFromLoadPackage()
-{
-	if (bLazyInitializedFromLoadPackage)
-	{
-		return;	
-	}
-	bLazyInitializedFromLoadPackage = true;
-
-	TRACE_CPUPROFILER_EVENT_SCOPE(LazyInitializeFromLoadPackage);
-	if (GIsInitialLoad)
-	{
-		GlobalImportStore.Initialize(IoDispatcher);
-	}
-}
-
 void FAsyncLoadingThread2::FinalizeInitialLoad()
 {
 	TRACE_CPUPROFILER_EVENT_SCOPE(FinalizeInitialLoad);
-	GlobalImportStore.FindAllScriptObjects();
+	GlobalImportStore.FindAllScriptObjects(); // for verification only
+	GlobalImportStore.ScriptObjects.Shrink();
+	bHasRegisteredAllScriptObjects = true;
+
 	check(PendingCDOs.Num() == 0);
 	PendingCDOs.Empty();
+	check(PendingCDOsRecursiveStack.Num() == 0);
+	PendingCDOsRecursiveStack.Empty();
+
+	UE_LOG(LogStreaming, Display,
+		TEXT("AsyncLoading2 - InitialLoad Finalized: Registered %d public script object entries (%.2f KB)"),
+		GlobalImportStore.ScriptObjects.Num(), (float)GlobalImportStore.ScriptObjects.GetAllocatedSize() / 1024.f);
 }
 
 uint32 FAsyncLoadingThread2::Run()
@@ -5755,6 +5791,25 @@ void FAsyncLoadingThread2::NotifyConstructedDuringAsyncLoading(UObject* Object,
 	AsyncPackage2->AddConstructedObject(Object, bSubObjectThatAlreadyExists);
 }
 
+void FAsyncLoadingThread2::NotifyRegistrationEvent(
+	const TCHAR* PackageName,
+	const TCHAR* Name,
+	ENotifyRegistrationType NotifyRegistrationType,
+	ENotifyRegistrationPhase NotifyRegistrationPhase,
+	UObject* (*InRegister)(),
+	bool InbDynamic)
+{
+	if (NotifyRegistrationPhase == ENotifyRegistrationPhase::NRP_Finished)
+	{
+		GlobalImportStore.AddScriptObject(PackageName, Name);
+	}
+}
+
+void FAsyncLoadingThread2::NotifyRegistrationComplete()
+{
+
+}
+
 /*-----------------------------------------------------------------------------
 	FAsyncPackage implementation.
 -----------------------------------------------------------------------------*/
@@ -6195,8 +6250,6 @@ int32 FAsyncLoadingThread2::LoadPackage(const FPackagePath& InPackagePath, FName
 {
 	TRACE_CPUPROFILER_EVENT_SCOPE(LoadPackage);
 
-	LazyInitializeFromLoadPackage();
-	
 	const FName PackageNameToLoad = InPackagePath.GetPackageFName();
 	if (InCustomName == PackageNameToLoad)
 	{
diff --git a/Engine/Source/Runtime/CoreUObject/Private/Serialization/AsyncLoadingThread.h b/Engine/Source/Runtime/CoreUObject/Private/Serialization/AsyncLoadingThread.h
--- a/Engine/Source/Runtime/CoreUObject/Private/Serialization/AsyncLoadingThread.h
+++ b/Engine/Source/Runtime/CoreUObject/Private/Serialization/AsyncLoadingThread.h
@@ -177,7 +177,7 @@ class FAsyncLoadingThread final : public FRunnable, public IAsyncPackageLoader
 {
 	friend struct FAsyncPackage;
 
-	IEDLBootNotificationManager& EDLBootNotificationManager;
+	struct FEDLBootNotificationManager& EDLBootNotificationManager;
 
 	/** Thread to run the worker FRunnable on */
 	FRunnableThread* Thread;
@@ -288,7 +288,7 @@ private:
 
 public:
 
-	FAsyncLoadingThread(int32 InThreadIndex, IEDLBootNotificationManager& InEDLBootNotificationManager);
+	FAsyncLoadingThread(int32 InThreadIndex);
 	virtual ~FAsyncLoadingThread();
 
 	IAsyncPackageLoader* GetIoStorePackageLoader() const
@@ -338,7 +338,11 @@ public:
 
 	void NotifyUnreachableObjects(const TArrayView<FUObjectItem*>& UnreachableObjects) override {};
 
-	void FireCompletedCompiledInImport(void* AsyncPacakge, FPackageIndex Import) override;
+	void FireCompletedCompiledInImport(void* AsyncPacakge, FPackageIndex Import);
+
+	void NotifyRegistrationEvent(const TCHAR* PackageName, const TCHAR* Name, ENotifyRegistrationType NotifyRegistrationType, ENotifyRegistrationPhase NotifyRegistrationPhase, UObject* (*InRegister)(), bool InbDynamic) override;
+
+	void NotifyRegistrationComplete() override;
 
 	/** [EDL] Event queue */
 	FAsyncLoadEventQueue EventQueue;
diff --git a/Engine/Source/Runtime/CoreUObject/Private/Serialization/AsyncPackageLoader.cpp b/Engine/Source/Runtime/CoreUObject/Private/Serialization/AsyncPackageLoader.cpp
--- a/Engine/Source/Runtime/CoreUObject/Private/Serialization/AsyncPackageLoader.cpp
+++ b/Engine/Source/Runtime/CoreUObject/Private/Serialization/AsyncPackageLoader.cpp
@@ -33,6 +33,49 @@ int32 IAsyncPackageLoader::GetNextRequestId()
 	 return NextPackageRequestId.Increment();
 }
 
+class FEarlyRegistrationEventsRecorder
+{
+public:
+	FEarlyRegistrationEventsRecorder()
+	{
+		RecordedEvents.Reserve(16388);
+	}
+
+	void NotifyRegistrationEvent(const TCHAR* PackageName, const TCHAR* Name, ENotifyRegistrationType NotifyRegistrationType, ENotifyRegistrationPhase NotifyRegistrationPhase, UObject* (*RegisterFunc)(), bool bDynamic)
+	{
+		RecordedEvents.Add({ PackageName, Name, NotifyRegistrationType, NotifyRegistrationPhase, RegisterFunc, bDynamic });
+	}
+
+	void Replay(IAsyncPackageLoader& PackageLoader)
+	{
+		UE_LOG(LogStreaming, Display, TEXT("NotifyRegistrationEvent: Replay %d entries"), RecordedEvents.Num());
+		for (const FRecordedEvent& Event : RecordedEvents)
+		{
+			PackageLoader.NotifyRegistrationEvent(*Event.PackageName, *Event.Name, Event.NotifyRegistrationType, Event.NotifyRegistrationPhase, Event.RegisterFunc, Event.bDynamic);
+		}
+		RecordedEvents.Empty();
+	}
+
+private:
+	struct FRecordedEvent
+	{
+		FString PackageName;
+		FString Name;
+		ENotifyRegistrationType NotifyRegistrationType;
+		ENotifyRegistrationPhase NotifyRegistrationPhase;
+		UObject* (*RegisterFunc)();
+		bool bDynamic;
+	};
+
+	TArray<FRecordedEvent> RecordedEvents;
+};
+
+static FEarlyRegistrationEventsRecorder& GetEarlyRegistrationEventsRecorder()
+{
+	static FEarlyRegistrationEventsRecorder Singleton;
+	return Singleton;
+}
+
 #if !UE_BUILD_SHIPPING
 static void LoadPackageCommand(const TArray<FString>& Args)
 {
@@ -73,397 +116,6 @@ static FAutoConsoleCommand CVar_LoadPackageAsyncCommand(
 
 const FName PrestreamPackageClassNameLoad = FName("PrestreamPackage");
 
-struct FEDLBootObjectState
-{
-	ENotifyRegistrationType NotifyRegistrationType;
-	ENotifyRegistrationPhase LastNotifyRegistrationPhase;
-	UObject *(*Register)();
-	bool bDynamic;
-};
-
-struct FEDLBootWaitingPackage
-{
-	void* Package;
-	FPackageIndex Import;
-};
-
-struct FEDLBootNotificationManager
-	: public IEDLBootNotificationManager
-{
-	TMap<FName, FEDLBootObjectState> PathToState;
-	TMultiMap<FName, FEDLBootWaitingPackage> PathToWaitingPackageNodes;
-	TArray<FName> PathsToFire;
-	TArray<UClass*> CDORecursiveStack;
-	TArray<UClass*> CDORecursives;
-	FCriticalSection EDLBootNotificationManagerLock;
-	bool bEnabled = true;
-
-	void Disable()
-	{
-		PathToState.Empty();
-		PathsToFire.Empty();
-		bEnabled = false;
-	}
-
-	// return true if you are waiting for this compiled in object
-	bool AddWaitingPackage(void* Pkg, FName PackageName, FName ObjectName, FPackageIndex Import, bool bIgnoreMissingPackage) override
-	{
-		if (PackageName == GLongCoreUObjectPackageName)
-		{
-			return false; // We assume nothing in coreuobject ever loads assets in a constructor
-		}
-		FScopeLock Lock(&EDLBootNotificationManagerLock);
-		check(GIsInitialLoad);
-		check(Import.IsImport()); // compiled in exports make no sense
-		FString ObjectNameString = ObjectName.ToString();
-		FName LongFName(*(PackageName.ToString() / ObjectNameString));
-		check(LongFName != NAME_None);
-		FName WaitName = LongFName;
-		FEDLBootObjectState* ExistingState = PathToState.Find(LongFName);
-		if (!ExistingState)
-		{
-			//if (ObjectName.ToString().EndsWith(HEADER_GENERATED_DELEGATE_SIGNATURE_SUFFIX))
-			// there are also some arg structs and other things that are just part of the package with no registration
-			{
-				ExistingState = PathToState.Find(PackageName);
-				WaitName = PackageName;
-			}
-			if (!ExistingState)
-			{
-				UE_CLOG(!bIgnoreMissingPackage, LogStreaming, Fatal, TEXT("Compiled in export %s not found; it was never registered."), *LongFName.ToString());
-				return false;
-			}
-		}
-		if (ExistingState->LastNotifyRegistrationPhase == ENotifyRegistrationPhase::NRP_Finished)
-		{
-			return false;
-		}
-		FEDLBootWaitingPackage WaitingPackage;
-		WaitingPackage.Package = Pkg;
-		WaitingPackage.Import = Import;
-
-		PathToWaitingPackageNodes.Add(WaitName, WaitingPackage);
-
-		return true;
-	}
-
-	void NotifyRegistrationEvent(const TCHAR* PackageName, const TCHAR* Name, ENotifyRegistrationType NotifyRegistrationType, ENotifyRegistrationPhase NotifyRegistrationPhase, UObject *(*InRegister)(), bool InbDynamic)
-	{
-		if (!bEnabled || !GIsInitialLoad)
-		{
-			return;
-		}
-		static FName LongCoreUObjectPackageName(TEXT("/Script/CoreUObject")); // can't use the global here because it may not be initialized yet
-		FName PackageFName(PackageName);
-		if (PackageFName == LongCoreUObjectPackageName)
-		{
-			return; // We assume nothing in coreuobject ever loads assets in a constructor
-		}
-
-		TStringBuilder<256> LongNameBuilder;
-		LongNameBuilder << PackageName;
-		FPathViews::Append(LongNameBuilder, Name);
-		FName LongFName(LongNameBuilder.ToView());
-
-		FScopeLock Lock(&EDLBootNotificationManagerLock);
-
-		//FPlatformMisc::LowLevelOutputDebugStringf(TEXT("NotifyRegistrationEvent %s %d %d\r\n"), *LongFName.ToString(), int32(NotifyRegistrationType), int32(NotifyRegistrationPhase));
-
-		// some things, like delegate signatures, are not registered; rather they are part of the package singleton, so we track the package state as being the max of any member of that package
-		FEDLBootObjectState* ExistingPackageState = PathToState.Find(PackageFName);
-		FEDLBootObjectState* ExistingState = PathToState.Find(LongFName);
-
-		if (!ExistingState)
-		{
-			if (NotifyRegistrationPhase != ENotifyRegistrationPhase::NRP_Added)
-			{
-				UE_LOG(LogStreaming, Fatal, TEXT("Attempt to process %s before it has been added."), *LongFName.ToString());
-			}
-			FEDLBootObjectState NewState;
-			NewState.LastNotifyRegistrationPhase = NotifyRegistrationPhase;
-			NewState.NotifyRegistrationType = NotifyRegistrationType;
-			NewState.Register = InRegister;
-			NewState.bDynamic = InbDynamic;
-			PathToState.Add(LongFName, NewState);
-
-			if (!ExistingPackageState)
-			{
-				NewState.NotifyRegistrationType = ENotifyRegistrationType::NRT_Package;
-				PathToState.Add(PackageFName, NewState);
-			}
-		}
-		else
-		{
-			if (int32(ExistingState->LastNotifyRegistrationPhase) + 1 != int32(NotifyRegistrationPhase))
-			{
-				UE_CLOG(GEventDrivenLoaderEnabled, LogStreaming, Fatal, TEXT("Invalid state transition %d %d with %s when it has already been processed."), int32(ExistingState->LastNotifyRegistrationPhase), int32(NotifyRegistrationPhase), *LongFName.ToString());
-			}
-			if (ExistingState->NotifyRegistrationType != (NotifyRegistrationType))
-			{
-				UE_CLOG(GEventDrivenLoaderEnabled, LogStreaming, Fatal, TEXT("Multiple types %d %d with %s when it has already been processed."), int32(ExistingState->NotifyRegistrationType), int32(NotifyRegistrationType), *LongFName.ToString());
-			}
-			ExistingState->LastNotifyRegistrationPhase = NotifyRegistrationPhase;
-			if (NotifyRegistrationPhase == ENotifyRegistrationPhase::NRP_Finished)
-			{
-				ExistingState->Register = nullptr; // we don't need to do this in ConstructWaitingBootObjects
-				PathsToFire.Add(LongFName);
-			}
-			check(ExistingPackageState); // if we have an existing state for the thing, we should also have a 
-			if (ExistingPackageState && int32(NotifyRegistrationPhase) > int32(ExistingPackageState->LastNotifyRegistrationPhase))
-			{
-				ExistingPackageState->LastNotifyRegistrationPhase = NotifyRegistrationPhase;
-				if (NotifyRegistrationPhase == ENotifyRegistrationPhase::NRP_Finished)
-				{
-					//FPlatformMisc::LowLevelOutputDebugStringf(TEXT("Fired package %s %d %d\r\n"), *PackageFName.ToString(), int32(NotifyRegistrationType), int32(NotifyRegistrationPhase));
-					PathsToFire.Add(PackageFName);
-				}
-			}
-		}
-	}
-
-	void NotifyRegistrationComplete()
-	{
-		if (!bEnabled)
-		{
-			return;
-		}
-#if USE_EVENT_DRIVEN_ASYNC_LOAD_AT_BOOT_TIME
-		FireCompletedCompiledInImports(true);
-		FlushAsyncLoading();
-#endif
-#if !HACK_HEADER_GENERATOR
-		check(!GIsInitialLoad && IsInGameThread());
-		FScopeLock Lock(&EDLBootNotificationManagerLock);
-		for (auto& Pair : PathToState)
-		{
-			if (Pair.Value.LastNotifyRegistrationPhase != ENotifyRegistrationPhase::NRP_Finished && !Pair.Value.bDynamic)
-			{
-#if USE_EVENT_DRIVEN_ASYNC_LOAD_AT_BOOT_TIME
-				UE_CLOG(GEventDrivenLoaderEnabled, LogStreaming, Fatal, TEXT("%s (%d) was not complete (%d) after registration was complete."), *Pair.Key.ToString(), int32(Pair.Value.NotifyRegistrationType), int32(Pair.Value.LastNotifyRegistrationPhase));
-#else
-				UE_LOG(LogStreaming, Warning, TEXT("%s was not complete (%d) after registration was complete."), *Pair.Key.ToString(), int32(Pair.Value.LastNotifyRegistrationPhase));
-#endif
-			}
-		}
-		if (PathToWaitingPackageNodes.Num())
-		{
-			UE_LOG(LogStreaming, Fatal, TEXT("Initial load is complete, but we still have %d waiting packages."), PathToWaitingPackageNodes.Num());
-		}
-		if (GEventDrivenLoaderEnabled && PathsToFire.Num() && USE_EVENT_DRIVEN_ASYNC_LOAD_AT_BOOT_TIME)
-		{
-			for (FName Path : PathsToFire)
-			{
-				UE_LOG(LogStreaming, Error, TEXT("%s was not fired."), *Path.ToString());
-			}
-			UE_LOG(LogStreaming, Fatal, TEXT("Initial load is complete, but we still have %d imports to fire (listed above)."), PathsToFire.Num());
-		}
-#endif
-		Disable();
-	}
-
-	bool ConstructWaitingBootObjects() override
-	{
-		static struct FFixedBootOrder
-		{
-			TArray<FName> Array;
-			FFixedBootOrder()
-			{
-				// look for any packages that we want to force preload at startup
-				FConfigSection* BootObjects = GConfig->GetSectionPrivate(TEXT("/Script/Engine.StreamingSettings"), false, true, GEngineIni);
-				if (BootObjects)
-				{
-					// go through list and add to the array
-					for (FConfigSectionMap::TIterator It(*BootObjects); It; ++It)
-					{
-						if (It.Key() == TEXT("FixedBootOrder"))
-						{
-							// add this package to the list to be fully loaded later
-							Array.Add(FName(*It.Value().GetValue()));
-						}
-					}
-				}
-			}
-
-		} FixedBootOrder;
-
-		check(GIsInitialLoad && IsInGameThread());
-		UObject *(*BootObjectRegister)() = nullptr;
-		UObject *(*BootPackageObjectRegister)() = nullptr;
-		FName WaitingPackage;
-		bool bIsCDO = false;
-
-		while (FixedBootOrder.Array.Num())
-		{
-			FName ThisItem = FixedBootOrder.Array.Pop();
-			FScopeLock Lock(&EDLBootNotificationManagerLock);
-			FEDLBootObjectState* ExistingState = PathToState.Find(ThisItem);
-
-			if (!ExistingState)
-			{
-				UE_LOG(LogStreaming, Fatal, TEXT("%s was listed as a fixed load order but was not found,"), *ThisItem.ToString());
-			}
-			else if (!ExistingState->Register)
-			{
-				UE_LOG(LogStreaming, Log, TEXT("%s was listed as a fixed load order but was already processed"), *ThisItem.ToString());
-			}
-			else
-			{
-				//FPlatformMisc::LowLevelOutputDebugStringf(TEXT("Booting Fixed %s %d\r\n"), *ThisItem.ToString(), int32(ExistingState->NotifyRegistrationType));
-				BootObjectRegister = ExistingState->Register;
-				ExistingState->Register = nullptr; // we don't need to do this more than once
-				bIsCDO = ExistingState->NotifyRegistrationType == ENotifyRegistrationType::NRT_ClassCDO;
-				break;
-			}
-		}
-
-		if (!BootObjectRegister)
-		{
-			FScopeLock Lock(&EDLBootNotificationManagerLock);
-			for (auto& Pair : PathToWaitingPackageNodes)
-			{
-				FEDLBootObjectState* ExistingState = PathToState.Find(Pair.Key);
-				if (ExistingState)
-				{
-					if (ExistingState->Register)
-					{
-						//FPlatformMisc::LowLevelOutputDebugStringf(TEXT("Booting %s %d\r\n"), *Pair.Key.ToString(), int32(ExistingState->NotifyRegistrationType));
-						BootObjectRegister = ExistingState->Register;
-						ExistingState->Register = nullptr; // we don't need to do this more than once
-						bIsCDO = ExistingState->NotifyRegistrationType == ENotifyRegistrationType::NRT_ClassCDO;
-						break;
-					}
-				}
-			}
-		}
-		if (BootObjectRegister)
-		{
-			UObject* BootObject = BootObjectRegister();
-			check(BootObject);
-			UObjectForceRegistration(BootObject);
-			if (bIsCDO)
-			{
-				UClass* Class = CastChecked<UClass>(BootObject);
-				bool bAnyParentOnStack = false;
-				UClass* Super = Class;
-				while (Super)
-				{
-					if (CDORecursiveStack.Contains(Super))
-					{
-						bAnyParentOnStack = true;
-						break;
-					}
-					Super = Super->GetSuperClass();
-				}
-
-				if (!bAnyParentOnStack)
-				{
-					CDORecursiveStack.Push(Class);
-					//FPlatformMisc::LowLevelOutputDebugStringf(TEXT("Create CDO %s\r\n"), *BootObject->GetName());
-					Class->GetDefaultObject();
-					verify(CDORecursiveStack.Pop() == Class);
-				}
-				else
-				{
-					//FPlatformMisc::LowLevelOutputDebugStringf(TEXT("Recursive Deferred %s\r\n"), *BootObject->GetName());
-					CDORecursives.Add(Class);
-				}
-			}
-			return true;
-		}
-		if (CDORecursives.Num())
-		{
-			UClass* OkToRun = nullptr;
-			for (UClass* Class : CDORecursives)
-			{
-				bool bAnyParentOnStack = false;
-				UClass* Super = Class;
-				while (Super)
-				{
-					if (CDORecursiveStack.Contains(Super))
-					{
-						bAnyParentOnStack = true;
-						break;
-					}
-					Super = Super->GetSuperClass();
-				}
-				if (!bAnyParentOnStack)
-				{
-					OkToRun = Class;
-					break;
-				}
-			}
-			if (OkToRun)
-			{
-				//FPlatformMisc::LowLevelOutputDebugStringf(TEXT("CDORecursives %s\r\n"), *OkToRun->GetName());
-				CDORecursives.Remove(OkToRun);
-				CDORecursiveStack.Push(OkToRun);
-				OkToRun->GetDefaultObject();
-				verify(CDORecursiveStack.Pop() == OkToRun);
-			}
-			else
-			{
-				FPlatformProcess::Sleep(.001f);
-			}
-			return true; // even if we didn't do anything we need to return true to avoid checking for cycles
-		}
-		return false;
-	}
-
-	bool IsWaitingForSomething() override
-	{
-		FScopeLock Lock(&EDLBootNotificationManagerLock);
-		return PathToWaitingPackageNodes.Num() > 0;
-	}
-
-	bool IsObjComplete(UObject* Obj)
-	{
-		static FName LongCoreUObjectPackageName(TEXT("/Script/CoreUObject")); // can't use the global here because it may not be initialized yet
-		FName PackageName = Obj->GetOutermost()->GetFName();
-		if (PackageName == LongCoreUObjectPackageName)
-		{
-			return true; // We assume nothing in coreuobject ever loads assets in a constructor, therefore it can be considered complete
-		}
-		FScopeLock Lock(&EDLBootNotificationManagerLock);
-		FName LongFName(*(PackageName.ToString() / Obj->GetName()));
-
-		FEDLBootObjectState* ExistingState = PathToState.Find(LongFName);
-
-		if (!ExistingState || ExistingState->LastNotifyRegistrationPhase == ENotifyRegistrationPhase::NRP_Finished)
-		{
-			return true;
-		}
-		return false;
-	}
-
-	bool FireCompletedCompiledInImports(bool bFinalRun) override
-	{
-#if USE_EVENT_DRIVEN_ASYNC_LOAD_AT_BOOT_TIME
-		FScopeLock Lock(&EDLBootNotificationManagerLock);
-		check(bFinalRun || GIsInitialLoad);
-		bool bResult = !!PathsToFire.Num();
-		for (FName LongName : PathsToFire)
-		{
-			for (auto It = PathToWaitingPackageNodes.CreateKeyIterator(LongName); It; ++It)
-			{
-				FEDLBootWaitingPackage& WaitingPackage = It.Value();
-				GPackageLoader->FireCompletedCompiledInImport(WaitingPackage.Package, WaitingPackage.Import);
-			}
-			PathToWaitingPackageNodes.Remove(LongName);
-		}
-		PathsToFire.Empty();
-		return bResult;
-#else
-		return false;
-#endif
-	}
-};
-
-static FEDLBootNotificationManager& GetGEDLBootNotificationManager()
-{
-	static FEDLBootNotificationManager Singleton;
-	return Singleton;
-}
-
 FAsyncLoadingThreadSettings::FAsyncLoadingThreadSettings()
 	: bAsyncLoadingThreadEnabled(false)
 	, bAsyncPostLoadEnabled(false)
@@ -497,29 +149,6 @@ FAsyncLoadingThreadSettings& FAsyncLoadingThreadSettings::Get()
 	return Settings;
 }
 
-bool IsFullyLoadedObj(UObject* Obj)
-{
-	if (!Obj)
-	{
-		return false;
-	}
-	if (Obj->HasAllFlags(RF_WasLoaded | RF_LoadCompleted)
-		|| Obj->IsA(UPackage::StaticClass())) // packages are never really loaded, so if it exists, it is loaded
-	{
-		return true;
-	}
-	if (Obj->HasAnyFlags(RF_WasLoaded | RF_NeedLoad | RF_WillBeLoaded))
-	{
-		return false;
-	}
-	if (GIsInitialLoad && Obj->GetOutermost()->HasAnyPackageFlags(PKG_CompiledIn))
-	{
-		return GetGEDLBootNotificationManager().IsObjComplete(Obj);
-	}
-	
-	return true;
-}
-
 bool IsNativeCodePackage(UPackage* Package)
 {
 	return (Package && Package->HasAnyPackageFlags(PKG_CompiledIn));
@@ -563,9 +192,8 @@ void InitAsyncThread()
 		bool bHasUseIoStoreParamInEditor = WITH_EDITOR && FParse::Param(FCommandLine::Get(), TEXT("UseIoStore"));
 		if (bHasScriptObjectsChunk || bHasUseIoStoreParamInEditor)
 		{
-			GetGEDLBootNotificationManager().Disable();
 #if WITH_EDITOR
-			GPackageLoader = MakeEditorPackageLoader(IoDispatcher, GetGEDLBootNotificationManager());
+			GPackageLoader = MakeEditorPackageLoader(IoDispatcher);
 #else
 			GPackageLoader.Reset(MakeAsyncPackageLoader2(IoDispatcher));
 #endif
@@ -574,13 +202,14 @@ void InitAsyncThread()
 #endif
 	if (!GPackageLoader.IsValid())
 	{
-		GPackageLoader = MakeUnique<FAsyncLoadingThread>(/** ThreadIndex = */ 0, GetGEDLBootNotificationManager());
+		GPackageLoader = MakeUnique<FAsyncLoadingThread>(/** ThreadIndex = */ 0);
 	}
 
 	FPlatformAtomics::InterlockedIncrement(&GIsLoaderCreated);
 
 	FCoreDelegates::OnSyncLoadPackage.AddStatic([](const FString&) { GSyncLoadCount++; });
  
+	GetEarlyRegistrationEventsRecorder().Replay(*GPackageLoader);
 	GPackageLoader->InitializeLoading();
 }
 
@@ -606,7 +235,7 @@ bool IsInAsyncLoadingThreadCoreUObjectInternal()
 	}
 }
 
-void FlushAsyncLoading(int32 PackageID /* = INDEX_NONE */)
+void FlushAsyncLoading(int32 RequestId /* = INDEX_NONE */)
 {
 	TRACE_CPUPROFILER_EVENT_SCOPE(FlushAsyncLoading);
 
@@ -625,16 +254,16 @@ void FlushAsyncLoading(int32 PackageID /* = INDEX_NONE */)
 			static uint64 LastFrameNumber = -1;
 			if (LastFrameNumber != GFrameNumber)
 			{
-				UE_LOG(LogStreaming, Display, TEXT("FlushAsyncLoading: %d QueuedPackages, %d AsyncPackages"), GPackageLoader->GetNumQueuedPackages(), GPackageLoader->GetNumAsyncPackages());
+				UE_LOG(LogStreaming, Display, TEXT("FlushAsyncLoading(%d): %d QueuedPackages, %d AsyncPackages"), RequestId, GPackageLoader->GetNumQueuedPackages(), GPackageLoader->GetNumAsyncPackages());
 			}
 			else
 			{
-				UE_LOG(LogStreaming, Log, TEXT("FlushAsyncLoading: %d QueuedPackages, %d AsyncPackages"), GPackageLoader->GetNumQueuedPackages(), GPackageLoader->GetNumAsyncPackages());
+				UE_LOG(LogStreaming, Log, TEXT("FlushAsyncLoading(%d): %d QueuedPackages, %d AsyncPackages"), RequestId, GPackageLoader->GetNumQueuedPackages(), GPackageLoader->GetNumAsyncPackages());
 			}
 			LastFrameNumber = GFrameNumber;
 		}
 #endif
-		GPackageLoader->FlushLoading(PackageID);
+		GPackageLoader->FlushLoading(RequestId);
 	}
 }
 
@@ -1013,13 +642,21 @@ float GetAsyncLoadPercentage(const FName& PackageName)
 void NotifyRegistrationEvent(const TCHAR* PackageName, const TCHAR* Name, ENotifyRegistrationType NotifyRegistrationType, ENotifyRegistrationPhase NotifyRegistrationPhase, UObject *(*InRegister)(), bool InbDynamic)
 {
 	LLM_SCOPE(ELLMTag::AsyncLoading);
-	GetGEDLBootNotificationManager().NotifyRegistrationEvent(PackageName, Name, NotifyRegistrationType, NotifyRegistrationPhase, InRegister, InbDynamic);
+
+	if (GPackageLoader)
+	{
+		GPackageLoader->NotifyRegistrationEvent(PackageName, Name, NotifyRegistrationType, NotifyRegistrationPhase, InRegister, InbDynamic);
+	}
+	else
+	{
+		GetEarlyRegistrationEventsRecorder().NotifyRegistrationEvent(PackageName, Name, NotifyRegistrationType, NotifyRegistrationPhase, InRegister, InbDynamic);
+	}
 }
 
 void NotifyRegistrationComplete()
 {
 	LLM_SCOPE(ELLMTag::AsyncLoading);
-	GetGEDLBootNotificationManager().NotifyRegistrationComplete();
+	GPackageLoader->NotifyRegistrationComplete();
 	FlushAsyncLoading();
 	GPackageLoader->StartThread();
 }
diff --git a/Engine/Source/Runtime/CoreUObject/Private/Serialization/EditorPackageLoader.cpp b/Engine/Source/Runtime/CoreUObject/Private/Serialization/EditorPackageLoader.cpp
--- a/Engine/Source/Runtime/CoreUObject/Private/Serialization/EditorPackageLoader.cpp
+++ b/Engine/Source/Runtime/CoreUObject/Private/Serialization/EditorPackageLoader.cpp
@@ -16,9 +16,9 @@ class FEditorPackageLoader final
 	: public IAsyncPackageLoader
 {
 public:
-	FEditorPackageLoader(FIoDispatcher& InIoDispatcher, IEDLBootNotificationManager& InEDLBootNotificationManager)
+	FEditorPackageLoader(FIoDispatcher& InIoDispatcher)
 	{
-		UncookedPackageLoader.Reset(new FAsyncLoadingThread(/** ThreadIndex = */ 0, InEDLBootNotificationManager));
+		UncookedPackageLoader.Reset(new FAsyncLoadingThread(/** ThreadIndex = */ 0));
 		CookedPackageLoader.Reset(MakeAsyncPackageLoader2(InIoDispatcher, UncookedPackageLoader.Get()));
 		UncookedPackageLoader->SetIoStorePackageLoader(CookedPackageLoader.Get());
 	}
@@ -193,9 +193,34 @@ public:
 		CookedPackageLoader->NotifyUnreachableObjects(UnreachableObjects);
 	}
 
-	virtual void FireCompletedCompiledInImport(void* AsyncPackage, FPackageIndex Import) override
+	virtual void NotifyRegistrationEvent(
+		const TCHAR* PackageName,
+		const TCHAR* Name,
+		ENotifyRegistrationType NotifyRegistrationType,
+		ENotifyRegistrationPhase NotifyRegistrationPhase,
+		UObject* (*InRegister)(),
+		bool InbDynamic) override
 	{
-		// Only used in the old EDL loader which is not enabled in editor builds
+		if (UncookedPackageLoader)
+		{
+			UncookedPackageLoader->NotifyRegistrationEvent(PackageName, Name, NotifyRegistrationType, NotifyRegistrationPhase, InRegister, InbDynamic);
+		}
+		if (CookedPackageLoader)
+		{
+			CookedPackageLoader->NotifyRegistrationEvent(PackageName, Name, NotifyRegistrationType, NotifyRegistrationPhase, InRegister, InbDynamic);
+		}
+	}
+
+	virtual void NotifyRegistrationComplete() override
+	{
+		if (UncookedPackageLoader)
+		{
+			UncookedPackageLoader->NotifyRegistrationComplete();
+		}
+		if (CookedPackageLoader)
+		{
+			CookedPackageLoader->NotifyRegistrationComplete();
+		}
 	}
 
 private:
@@ -204,9 +229,9 @@ private:
 	bool bHasInitializedCookedPackageLoader = false;
 };
 
-TUniquePtr<IAsyncPackageLoader> MakeEditorPackageLoader(FIoDispatcher& InIoDispatcher, IEDLBootNotificationManager& InEDLBootNotificationManager)
+TUniquePtr<IAsyncPackageLoader> MakeEditorPackageLoader(FIoDispatcher& InIoDispatcher)
 {
-	return TUniquePtr<IAsyncPackageLoader>(new FEditorPackageLoader(InIoDispatcher, InEDLBootNotificationManager));
+	return TUniquePtr<IAsyncPackageLoader>(new FEditorPackageLoader(InIoDispatcher));
 }
 
 #endif // WITH_IOSTORE_IN_EDITOR
diff --git a/Engine/Source/Runtime/CoreUObject/Private/Serialization/EditorPackageLoader.h b/Engine/Source/Runtime/CoreUObject/Private/Serialization/EditorPackageLoader.h
--- a/Engine/Source/Runtime/CoreUObject/Private/Serialization/EditorPackageLoader.h
+++ b/Engine/Source/Runtime/CoreUObject/Private/Serialization/EditorPackageLoader.h
@@ -6,8 +6,7 @@
 
 class IAsyncPackageLoader;
 class FIoDispatcher;
-class IEDLBootNotificationManager;
 
 #if WITH_IOSTORE_IN_EDITOR
-TUniquePtr<IAsyncPackageLoader> MakeEditorPackageLoader(FIoDispatcher& InIoDispatcher, IEDLBootNotificationManager& InEDLBootNotificationManager);
+TUniquePtr<IAsyncPackageLoader> MakeEditorPackageLoader(FIoDispatcher& InIoDispatcher);
 #endif
diff --git a/Engine/Source/Runtime/CoreUObject/Private/UObject/Class.cpp b/Engine/Source/Runtime/CoreUObject/Private/UObject/Class.cpp
--- a/Engine/Source/Runtime/CoreUObject/Private/UObject/Class.cpp
+++ b/Engine/Source/Runtime/CoreUObject/Private/UObject/Class.cpp
@@ -97,7 +97,6 @@ COREUOBJECT_API void InitializePrivateStaticClass(
 	)
 {
 	TRACE_LOADTIME_CLASS_INFO(TClass_PrivateStaticClass, Name);
-	NotifyRegistrationEvent(PackageName, Name, ENotifyRegistrationType::NRT_Class, ENotifyRegistrationPhase::NRP_Started);
 
 	/* No recursive ::StaticClass calls allowed. Setup extras. */
 	if (TClass_Super_StaticClass != TClass_PrivateStaticClass)
@@ -116,8 +115,6 @@ COREUOBJECT_API void InitializePrivateStaticClass(
 		// Defer
 		TClass_PrivateStaticClass->Register(PackageName, Name);
 	}
-	
-	NotifyRegistrationEvent(PackageName, Name, ENotifyRegistrationType::NRT_Class, ENotifyRegistrationPhase::NRP_Finished);
 }
 
 void FNativeFunctionRegistrar::RegisterFunction(class UClass* Class, const ANSICHAR* InName, FNativeFuncPtr InPointer)
@@ -4119,17 +4116,6 @@ UObject* UClass::CreateDefaultObject()
 			// NULL (so we don't invalidate one that has already been setup)
 			if (ClassDefaultObject == NULL)
 			{
-				FString PackageName;
-				FString CDOName;
-				bool bDoNotify = false;
-				if (GIsInitialLoad && GetOutermost()->HasAnyPackageFlags(PKG_CompiledIn) && !GetOutermost()->HasAnyPackageFlags(PKG_RuntimeGenerated))
-				{
-					PackageName = GetOutermost()->GetFName().ToString();
-					CDOName = GetDefaultObjectName().ToString();
-					NotifyRegistrationEvent(*PackageName, *CDOName, ENotifyRegistrationType::NRT_ClassCDO, ENotifyRegistrationPhase::NRP_Started);
-					bDoNotify = true;
-				}
-
 				// RF_ArchetypeObject flag is often redundant to RF_ClassDefaultObject, but we need to tag
 				// the CDO as RF_ArchetypeObject in order to propagate that flag to any default sub objects.
 				ClassDefaultObject = StaticAllocateObject(this, GetOuter(), NAME_None, EObjectFlags(RF_Public|RF_ClassDefaultObject|RF_ArchetypeObject));
@@ -4148,9 +4134,13 @@ UObject* UClass::CreateDefaultObject()
 					InitOptions |= EObjectInitializerOptions::InitializeProperties;
 				}
 				(*ClassConstructor)(FObjectInitializer(ClassDefaultObject, ParentDefaultObject, InitOptions));
-				if (bDoNotify)
+				if (GetOutermost()->HasAnyPackageFlags(PKG_CompiledIn) && !GetOutermost()->HasAnyPackageFlags(PKG_RuntimeGenerated))
 				{
-					NotifyRegistrationEvent(*PackageName, *CDOName, ENotifyRegistrationType::NRT_ClassCDO, ENotifyRegistrationPhase::NRP_Finished);
+					TCHAR PackageName[FName::StringBufferSize];
+					TCHAR CDOName[FName::StringBufferSize];
+					GetOutermost()->GetFName().ToString(PackageName);
+					GetDefaultObjectName().ToString(CDOName);
+					NotifyRegistrationEvent(PackageName, CDOName, ENotifyRegistrationType::NRT_ClassCDO, ENotifyRegistrationPhase::NRP_Finished);
 				}
 				ClassDefaultObject->PostCDOContruct();
 			}
diff --git a/Engine/Source/Runtime/CoreUObject/Private/UObject/DeferredRegistry.h b/Engine/Source/Runtime/CoreUObject/Private/UObject/DeferredRegistry.h
--- a/Engine/Source/Runtime/CoreUObject/Private/UObject/DeferredRegistry.h
+++ b/Engine/Source/Runtime/CoreUObject/Private/UObject/DeferredRegistry.h
@@ -215,7 +215,7 @@ public:
 		for (int32 Index = ProcessedRegistrations; Index < Num; ++Index)
 		{
 			TType* Object = OuterRegister(Registrations[Index]);
-			InOnRegistration(*Object);
+			InOnRegistration(Registrations[Index].PackageName, *Object);
 		}
 		if (UpdateCounter)
 		{
diff --git a/Engine/Source/Runtime/CoreUObject/Private/UObject/UObjectBase.cpp b/Engine/Source/Runtime/CoreUObject/Private/UObject/UObjectBase.cpp
--- a/Engine/Source/Runtime/CoreUObject/Private/UObject/UObjectBase.cpp
+++ b/Engine/Source/Runtime/CoreUObject/Private/UObject/UObjectBase.cpp
@@ -570,7 +570,6 @@ void UObjectCompiledInDeferStruct(class UScriptStruct* (*InRegister)(), const TC
 
 class UScriptStruct *GetStaticStruct(class UScriptStruct *(*InRegister)(), UObject* StructOuter, const TCHAR* StructName)
 {
-	NotifyRegistrationEvent(*StructOuter->GetOutermost()->GetName(), StructName, ENotifyRegistrationType::NRT_Struct, ENotifyRegistrationPhase::NRP_Started);
 	UScriptStruct *Result = (*InRegister)();
 	NotifyRegistrationEvent(*StructOuter->GetOutermost()->GetName(), StructName, ENotifyRegistrationType::NRT_Struct, ENotifyRegistrationPhase::NRP_Finished);
 	return Result;
@@ -595,7 +594,6 @@ void UObjectCompiledInDeferEnum(class UEnum* (*InRegister)(), const TCHAR* Packa
 
 class UEnum *GetStaticEnum(class UEnum *(*InRegister)(), UObject* EnumOuter, const TCHAR* EnumName)
 {
-	NotifyRegistrationEvent(*EnumOuter->GetOutermost()->GetName(), EnumName, ENotifyRegistrationType::NRT_Enum, ENotifyRegistrationPhase::NRP_Started);
 	UEnum *Result = (*InRegister)();
 	NotifyRegistrationEvent(*EnumOuter->GetOutermost()->GetName(), EnumName, ENotifyRegistrationType::NRT_Enum, ENotifyRegistrationPhase::NRP_Finished);
 	return Result;
@@ -793,17 +791,22 @@ static void UObjectLoadAllCompiledInDefaultProperties()
 
 	static FName LongEnginePackageName(TEXT("/Script/Engine"));
 
-	FClassDeferredRegistry& Registry = FClassDeferredRegistry::Get();
+	FClassDeferredRegistry& ClassRegistry = FClassDeferredRegistry::Get();
 
-	if (Registry.HasPendingRegistrations())
+	if (ClassRegistry.HasPendingRegistrations())
 	{
 		SCOPED_BOOT_TIMING("UObjectLoadAllCompiledInDefaultProperties");
 		TArray<UClass*> NewClasses;
 		TArray<UClass*> NewClassesInCoreUObject;
 		TArray<UClass*> NewClassesInEngine;
-		Registry.DoPendingOuterRegistrations(true, [&NewClasses, &NewClassesInCoreUObject, &NewClassesInEngine](UClass& Class) -> void
+		ClassRegistry.DoPendingOuterRegistrations(true, [&NewClasses, &NewClassesInCoreUObject, &NewClassesInEngine](const TCHAR* PackageName, UClass& Class) -> void
 			{
-				UE_LOG(LogUObjectBootstrap, Verbose, TEXT("UObjectLoadAllCompiledInDefaultProperties After Registrant %s %s"), *Class.GetOutermost()->GetName(), *Class.GetName());
+				UE_LOG(LogUObjectBootstrap, Verbose, TEXT("UObjectLoadAllCompiledInDefaultProperties After Registrant %s %s"), PackageName, *Class.GetName());
+
+				TCHAR ClassName[FName::StringBufferSize];
+				Class.GetFName().ToString(ClassName);
+				NotifyRegistrationEvent(PackageName, ClassName, ENotifyRegistrationType::NRT_Class, ENotifyRegistrationPhase::NRP_Finished);
+
 				if (Class.GetOutermost()->GetFName() == GLongCoreUObjectPackageName)
 				{
 					NewClassesInCoreUObject.Add(&Class);
@@ -1218,4 +1221,4 @@ TMap<FName, FName>& GetConvertedDynamicPackageNameToTypeName()
 {
 	static TMap<FName, FName> ConvertedDynamicPackageNameToTypeName;
 	return ConvertedDynamicPackageNameToTypeName;
-}
\ No newline at end of file
+}
diff --git a/Engine/Source/Runtime/CoreUObject/Private/UObject/UObjectGlobals.cpp b/Engine/Source/Runtime/CoreUObject/Private/UObject/UObjectGlobals.cpp
--- a/Engine/Source/Runtime/CoreUObject/Private/UObject/UObjectGlobals.cpp
+++ b/Engine/Source/Runtime/CoreUObject/Private/UObject/UObjectGlobals.cpp
@@ -5436,17 +5436,24 @@ namespace UECodeGen_Private
 		TArray<UFunction*> Delegates;
 		Delegates.Reserve(Params.NumSingletons);
 #endif
+		TCHAR PackageName[FName::StringBufferSize];
+		NewPackage->GetFName().ToString(PackageName);
 		for (UObject* (*const *SingletonFunc)() = Params.SingletonFuncArray, *(*const *SingletonFuncEnd)() = SingletonFunc + Params.NumSingletons; SingletonFunc != SingletonFuncEnd; ++SingletonFunc)
 		{
-#if WITH_RELOAD
 			UObject* Object = (*SingletonFunc)();
+#if WITH_RELOAD
 			if (UFunction* Function = Cast<UFunction>(Object))
 			{
 				Delegates.Add(Function);
 			}
-#else
-			(*SingletonFunc)();
 #endif
+			if (Object->GetOuter() == NewPackage)
+			{
+				// Notify loader of new top level noexport objects like UScriptStruct, UDelegateFunction and USparseDelegateFunction
+				TCHAR ObjectName[FName::StringBufferSize];
+				Object->GetFName().ToString(ObjectName);
+				NotifyRegistrationEvent(PackageName, ObjectName, ENotifyRegistrationType::NRT_NoExportObject, ENotifyRegistrationPhase::NRP_Finished);
+			}
 		}
 #if WITH_RELOAD
 		NewPackage->SetReloadDelegates(MoveTemp(Delegates));
diff --git a/Engine/Source/Runtime/CoreUObject/Public/Serialization/AsyncPackageLoader.h b/Engine/Source/Runtime/CoreUObject/Public/Serialization/AsyncPackageLoader.h
--- a/Engine/Source/Runtime/CoreUObject/Public/Serialization/AsyncPackageLoader.h
+++ b/Engine/Source/Runtime/CoreUObject/Public/Serialization/AsyncPackageLoader.h
@@ -27,8 +27,6 @@ extern const FName PrestreamPackageClassNameLoad;
 /** Returns true if we're inside a FGCScopeLock */
 extern bool IsGarbageCollectionLocked();
 
-bool IsFullyLoadedObj(UObject* Obj);
-
 bool IsNativeCodePackage(UPackage* Package);
 
 /** Checks if the object can have PostLoad called on the Async Loading Thread */
@@ -55,17 +53,6 @@ class IAsyncPackageLoader;
 class FPackageIndex;
 class LinkerInstancingContext;
 
-class IEDLBootNotificationManager
-{
-public:
-	virtual ~IEDLBootNotificationManager() = default;
-
-	virtual bool AddWaitingPackage(void* Pkg, FName PackageName, FName ObjectName, FPackageIndex Import, bool bIgnoreMissingPackage) = 0;
-	virtual bool ConstructWaitingBootObjects() = 0;
-	virtual bool FireCompletedCompiledInImports(bool bFinalRun = false) = 0;
-	virtual bool IsWaitingForSomething() = 0;
-};
-
 /** Structure that holds the async loading thread ini settings */
 struct FAsyncLoadingThreadSettings
 {
@@ -221,7 +208,9 @@ public:
 
 	virtual void NotifyUnreachableObjects(const TArrayView<FUObjectItem*>& UnreachableObjects) = 0;
 
-	virtual void FireCompletedCompiledInImport(void* AsyncPackage, FPackageIndex Import) = 0;
+	virtual void NotifyRegistrationEvent(const TCHAR* PackageName, const TCHAR* Name, ENotifyRegistrationType NotifyRegistrationType, ENotifyRegistrationPhase NotifyRegistrationPhase, UObject* (*InRegister)(), bool InbDynamic) = 0;
+
+	virtual void NotifyRegistrationComplete() = 0;
 
 protected:
 	static int32 GetNextRequestId();
diff --git a/Engine/Source/Runtime/CoreUObject/Public/UObject/LinkerLoad.h b/Engine/Source/Runtime/CoreUObject/Public/UObject/LinkerLoad.h
--- a/Engine/Source/Runtime/CoreUObject/Public/UObject/LinkerLoad.h
+++ b/Engine/Source/Runtime/CoreUObject/Public/UObject/LinkerLoad.h
@@ -1366,6 +1366,7 @@ enum class ENotifyRegistrationType
 	NRT_Struct,
 	NRT_Enum,
 	NRT_Package,
+	NRT_NoExportObject,
 };
 
 enum class ENotifyRegistrationPhase
