AsyncLoading2:
Merge loaded package store with global import store
Use global import store to keep track of pinned exports instead of relying on FindObjectWithOuter

diff --git a/Engine/Source/Runtime/CoreUObject/Private/Serialization/AsyncLoading2.cpp b/Engine/Source/Runtime/CoreUObject/Private/Serialization/AsyncLoading2.cpp
--- a/Engine/Source/Runtime/CoreUObject/Private/Serialization/AsyncLoading2.cpp
+++ b/Engine/Source/Runtime/CoreUObject/Private/Serialization/AsyncLoading2.cpp
@@ -64,6 +64,7 @@
 #include "Modules/ModuleManager.h"
 #include "Containers/SpscQueue.h"
 #include "Misc/PathViews.h"
+#include "UObject/LinkerLoad.h"
 
 #if UE_BUILD_DEVELOPMENT || UE_BUILD_DEBUG
 PRAGMA_DISABLE_OPTIMIZATION
@@ -450,198 +451,214 @@
 using FUnreachablePackages = TArray<UPackage*>;
 using FUnreachablePublicExports = TArray<FUnreachablePublicExport>;
 
-struct FGlobalImportStore
-{
-	TMap<FPackageObjectIndex, UObject*> ScriptObjects;
-	TMap<FPublicExportKey, int32> PublicExportToObjectIndex;
-	TMap<int32, FPublicExportKey> ObjectIndexToPublicExport;
-
-	FGlobalImportStore()
-	{
-		PublicExportToObjectIndex.Reserve(32768);
-		ObjectIndexToPublicExport.Reserve(32768);
-	}
-
-	TArray<FPackageId> RemovePublicExports(const FUnreachablePublicExports& PublicExports)
-	{
-		TArray<FPackageId> PackageIds;
-		TArray<FPublicExportKey> PublicExportKeys;
-		PublicExportKeys.Reserve(PublicExports.Num());
-		PackageIds.Reserve(PublicExports.Num());
-
-		for (const FUnreachablePublicExport& Item : PublicExports)
-		{
-			int32 ObjectIndex = Item.Key;
-			FPublicExportKey PublicExportKey;
-			if (ObjectIndexToPublicExport.RemoveAndCopyValue(ObjectIndex, PublicExportKey))
-			{
-				PublicExportKeys.Emplace(PublicExportKey);
-#if DO_CHECK
-				{
-					UObject* GCObject = Item.Value;
-
-					checkf(GCObject->HasAllFlags(GlobalImportObjectConditionalFlags),
-						TEXT("The serialized GC Object '%s' with (ObjectFlags=%x, InternalObjectFlags=%x)and id 0x%llX:0x%llX is currently missing RF_WasLoaded or RF_Public. ")
-						TEXT("The flags must have been altered incorrectly by a higher level system after the object was loaded. ")
-						TEXT("This may cause memory corruption."),
-						*GCObject->GetFullName(),
-						GCObject->GetFlags(),
-						GCObject->GetInternalFlags(),
-						PublicExportKey.GetPackageId().Value(), PublicExportKey.GetExportHash());
-
-					UObject* ExistingObject = FindPublicExportObjectUnchecked(PublicExportKey);
-					checkf(ExistingObject,
-						TEXT("The serialized GC object '%s' with flags (ObjectFlags=%x, InternalObjectFlags=%x) and id 0x%llX:0x%llX is missing in ImportStore. ")
-						TEXT("Reason unknown. Double delete? Bug or hash collision?"),
-						*GCObject->GetFullName(),
-						GCObject->GetFlags(),
-						GCObject->GetInternalFlags(),
-						PublicExportKey.GetPackageId().Value(),
-						PublicExportKey.GetExportHash());
-
-					checkf(ExistingObject && ExistingObject == GCObject,
-						TEXT("The serialized GC Object '%s' with flags (ObjectFlags=%x, InternalObjectFlags=%x) and id 0x%llX:0x%llX is not matching the object '%s' in ImportStore. "),
-						TEXT("Reason unknown. Overwritten after it was added? Bug or hash collision?"),
-						*GCObject->GetFullName(),
-						GCObject->GetFlags(),
-						GCObject->GetInternalFlags(),
-						PublicExportKey.GetPackageId().Value(),
-						PublicExportKey.GetExportHash(),
-						*ExistingObject->GetFullName());
-				}
-#endif
-			}
-		}
-
-		FPackageId LastPackageId;
-		for (const FPublicExportKey& PublicExportKey : PublicExportKeys)
-		{
-			PublicExportToObjectIndex.Remove(PublicExportKey);
-			FPackageId PackageId = PublicExportKey.GetPackageId();
-			if (!(PackageId == LastPackageId)) // fast approximation of Contains()
-			{
-				LastPackageId = PackageId;
-				PackageIds.Emplace(LastPackageId);
-			}
-		}
-		return PackageIds;
-	}
-
-	inline UObject* FindPublicExportObjectUnchecked(const FPublicExportKey& Key)
-	{
-		int32* FindObjectIndex = PublicExportToObjectIndex.Find(Key);
-		if (!FindObjectIndex)
-		{
-			return nullptr;
-		}
-		FUObjectItem* ObjectItem = GUObjectArray.IndexToObject(*FindObjectIndex);
-		if (!ObjectItem)
-		{
-			return nullptr;
-		}
-		return static_cast<UObject*>(ObjectItem->Object);
-	}
-
-	inline UObject* FindPublicExportObject(const FPublicExportKey& Key)
-	{
-		UObject* Object = FindPublicExportObjectUnchecked(Key);
-		checkf(!Object || !Object->IsUnreachable(), TEXT("%s"), Object ? *Object->GetFullName() : TEXT("null"));
-		return Object;
-	}
-
-	inline UObject* FindScriptImportObject(FPackageObjectIndex GlobalIndex)
-	{
-		check(GlobalIndex.IsScriptImport());
-		UObject* Object = nullptr;
-		Object = ScriptObjects.FindRef(GlobalIndex);
-		return Object;
-	}
-
-	void StoreGlobalObject(FPackageId PackageId, uint64 ExportHash, UObject* Object)
-	{
-		check(PackageId.IsValid());
-		check(ExportHash != 0);
-		int32 ObjectIndex = GUObjectArray.ObjectToIndex(Object);
-		FPublicExportKey Key = FPublicExportKey::MakeKey(PackageId, ExportHash);
-#if DO_CHECK
-		{
-			UObject* ExistingObject = FindPublicExportObjectUnchecked(Key);
-			if (ExistingObject)
-			{
-				checkf(ExistingObject == Object,
-					TEXT("The constructed serialized object '%s' with index %d and id 0x%llX:0x%llX collides with the object '%s' in ImportStore. ")
-					TEXT("Reason unknown. Bug or hash collision?"),
-					Object ? *Object->GetFullName() : TEXT("null"),
-					ObjectIndex,
-					Key.GetPackageId().Value(), Key.GetExportHash(),
-					ExistingObject ? *ExistingObject->GetFullName() : TEXT("null"));
-			}
-
-			FPublicExportKey* ExistingKey = ObjectIndexToPublicExport.Find(ObjectIndex);
-			if (ExistingKey)
-			{
-				checkf(*ExistingKey == Key,
-					TEXT("The constructed serialized object '%s' with index %d and id 0x%llX:0x%llX collides with the object '%s' in ImportStore. ")
-					TEXT("Reason unknown. Bug or hash collision?"),
-					Object ? *Object->GetFullName() : TEXT("null"),
-					ObjectIndex,
-					Key.GetPackageId().Value(), Key.GetExportHash(),
-					ExistingObject ? *ExistingObject->GetFullName() : TEXT("null"));
-			}
-		}
-#endif
-		PublicExportToObjectIndex.Add(Key, ObjectIndex);
-		ObjectIndexToPublicExport.Add(ObjectIndex, Key);
-	}
-
-	void FindAllScriptObjects();
-
-	void AddScriptObject(FStringView PackageName, FStringView Name)
-	{
-		TStringBuilder<FName::StringBufferSize> FullName;
-		FPathViews::Append(FullName, PackageName);
-		FPathViews::Append(FullName, Name);
-		FPackageObjectIndex GlobalImportIndex = FPackageObjectIndex::FromScriptPath(FullName);
-
-		UObject* Package = StaticFindObjectFastInternal(UPackage::StaticClass(), nullptr, FName(PackageName), true);
-		// check(Package);
-		if (!ensureMsgf(Package, TEXT("Failed to find package when registering script object '%.*s/%.*s'"),
-			PackageName.Len(), PackageName.GetData(), Name.Len(), Name.GetData()))
-		{
-			return;
-		}
-
-		UObject* Object = StaticFindObjectFastInternal(nullptr, Package, FName(Name));
-		// check(Object);
-		if (!ensureMsgf(Object, TEXT("Failed to find object when registering script object '%.*s/%.*s'"),
-			PackageName.Len(), PackageName.GetData(), Name.Len(), Name.GetData()))
-		{
-			return;
-		}
-
-#if WITH_EDITOR
-		FPackageObjectIndex PackageGlobalImportIndex = FPackageObjectIndex::FromScriptPath(PackageName);
-		ScriptObjects.Add(PackageGlobalImportIndex, Package);
-#endif
-		ScriptObjects.Add(GlobalImportIndex, Object);
-
-		TStringBuilder<FName::StringBufferSize> SubObjectName;
-		ForEachObjectWithOuter(Object, [this,&SubObjectName](UObject* SubObject)
-		{
-			if (SubObject->HasAnyFlags(RF_Public))
-			{
-				SubObjectName.Reset();
-				SubObject->GetPathName(nullptr, SubObjectName);
-				FPackageObjectIndex SubObjectGlobalImportIndex = FPackageObjectIndex::FromScriptPath(SubObjectName);
-				ScriptObjects.Add(SubObjectGlobalImportIndex, SubObject);
-			}
-		}, /* bIncludeNestedObjects*/ true);
-	}
-};
-
 class FLoadedPackageRef
 {
+private:
+	friend class FLoadedPackageStore;
+
+	class FPublicExportMap
+	{
+	public:
+		FPublicExportMap()
+		{
+		}
+
+		FPublicExportMap(const FPublicExportMap&) = delete;
+
+		FPublicExportMap(FPublicExportMap&& Other)
+		{
+			Allocation = Other.Allocation;
+			Count = Other.Count;
+			SingleItemValue = Other.SingleItemValue;
+			Other.Allocation = nullptr;
+			Other.Count = 0;
+		};
+
+		FPublicExportMap& operator=(const FPublicExportMap&) = delete;
+
+		FPublicExportMap& operator=(FPublicExportMap&& Other)
+		{
+			if (Count > 1)
+			{
+				FMemory::Free(Allocation);
+			}
+			Allocation = Other.Allocation;
+			Count = Other.Count;
+			SingleItemValue = Other.SingleItemValue;
+			Other.Allocation = nullptr;
+			Other.Count = 0;
+			return *this;
+		}
+	
+		~FPublicExportMap()
+		{
+			if (Count > 1)
+			{
+				FMemory::Free(Allocation);
+			}
+		}
+		void Grow(int32 NewCount)
+		{
+			if (NewCount <= Count)
+			{
+				return;
+			}
+			if (NewCount > 1)
+			{
+				TArrayView<uint64> OldKeys = GetKeys();
+				TArrayView<int32> OldValues = GetValues();
+				const uint64 OldKeysSize = Count * sizeof(uint64);
+				const uint64 NewKeysSize = NewCount * sizeof(uint64);
+				const uint64 OldValuesSize = Count * sizeof(int32);
+				const uint64 NewValuesSize = NewCount * sizeof(int32);
+				const uint64 KeysToAddSize = NewKeysSize - OldKeysSize;
+				const uint64 ValuesToAddSize = NewValuesSize - OldValuesSize;
+
+				uint8* NewAllocation = reinterpret_cast<uint8*>(FMemory::Malloc(NewKeysSize + NewValuesSize));
+				FMemory::Memzero(NewAllocation, KeysToAddSize); // Insert new keys initialized to zero
+				FMemory::Memcpy(NewAllocation + KeysToAddSize, OldKeys.GetData(), OldKeysSize); // Copy old keys
+				FMemory::Memset(NewAllocation + NewKeysSize, 0xFF, ValuesToAddSize); // Insert new values initialized to -1
+				FMemory::Memcpy(NewAllocation + NewKeysSize + ValuesToAddSize, OldValues.GetData(), OldValuesSize); // Copy old values
+				if (Count > 1)
+				{
+					FMemory::Free(Allocation);
+				}
+				Allocation = NewAllocation;
+			}
+			Count = NewCount;
+		}
+
+		void Store(uint64 ExportHash, UObject* Object)
+		{
+			TArrayView<uint64> Keys = GetKeys();
+			TArrayView<int32> Values = GetValues();
+			int32 Index = Algo::LowerBound(Keys, ExportHash);
+			if (Index < Count && Keys[Index] == ExportHash)
+			{
+				// Slot already exists so reuse it
+				Values[Index] = GUObjectArray.ObjectToIndex(Object);
+				return;
+			}
+			if (Count == 0 || Keys[0] != 0)
+			{
+				// No free slots so we need to add one (will be inserted at the beginning of the array)
+				Grow(Count + 1);
+				Keys = GetKeys();
+				Values = GetValues();
+			}
+			else
+			{
+				--Index; // Update insertion index to one before the lower bound item
+			}
+			if (Index > 0)
+			{
+				// Move items down
+				FMemory::Memcpy(Keys.GetData(), Keys.GetData() + 1, Index * sizeof(uint64));
+				FMemory::Memcpy(Values.GetData(), Values.GetData() + 1, Index * sizeof(int32));
+			}
+			Keys[Index] = ExportHash;
+			Values[Index] = GUObjectArray.ObjectToIndex(Object);
+		}
+
+		void Remove(uint64 ExportHash)
+		{
+			TArrayView<uint64> Keys = GetKeys();
+			int32 Index = Algo::LowerBound(Keys, ExportHash);
+			if (Index < Count && Keys[Index] == ExportHash)
+			{
+				TArrayView<int32> Values = GetValues();
+				Values[Index] = -1;
+			}
+		}
+
+		UObject* Find(uint64 ExportHash)
+		{
+			TArrayView<uint64> Keys = GetKeys();
+			int32 Index = Algo::LowerBound(Keys, ExportHash);
+			if (Index < Count && Keys[Index] == ExportHash)
+			{
+				TArrayView<int32> Values = GetValues();
+				int32 ObjectIndex = Values[Index];
+				if (ObjectIndex >= 0)
+				{
+					return static_cast<UObject*>(GUObjectArray.IndexToObject(ObjectIndex)->Object);
+				}
+			}
+			return nullptr;
+		}
+
+		void PinForGC()
+		{
+			for (int32 ObjectIndex : GetValues())
+			{
+				if (ObjectIndex >= 0)
+				{
+					UObject* Object = static_cast<UObject*>(GUObjectArray.IndexToObject(ObjectIndex)->Object);
+					checkf(!Object->HasAnyInternalFlags(EInternalObjectFlags::LoaderImport), TEXT("%s"), *Object->GetFullName());
+					Object->SetInternalFlags(EInternalObjectFlags::LoaderImport);
+				}
+			}
+		}
+
+		void UnpinForGC()
+		{
+			for (int32 ObjectIndex : GetValues())
+			{
+				if (ObjectIndex >= 0)
+				{
+					UObject* Object = static_cast<UObject*>(GUObjectArray.IndexToObject(ObjectIndex)->Object);
+					checkf(Object->HasAnyInternalFlags(EInternalObjectFlags::LoaderImport), TEXT("%s"), *Object->GetFullName());
+					Object->AtomicallyClearInternalFlags(EInternalObjectFlags::LoaderImport);
+				}
+			}
+		}
+
+#if DO_CHECK
+		void VerifyAllObjectsRemoved()
+		{
+			for (int32 ObjectIndex : GetValues())
+			{
+				check(ObjectIndex < 0);
+			}
+		}
+#endif
+
+	private:
+		TArrayView<uint64> GetKeys()
+		{
+			if (Count == 1)
+			{
+				return MakeArrayView(&SingleItemKey, 1);
+			}
+			else
+			{
+				return MakeArrayView(reinterpret_cast<uint64*>(Allocation), Count);
+			}
+		}
+
+		TArrayView<int32> GetValues()
+		{
+			if (Count == 1)
+			{
+				return MakeArrayView(&SingleItemValue, 1);
+			}
+			else
+			{
+				return MakeArrayView(reinterpret_cast<int32*>(Allocation + Count * sizeof(uint64)), Count);
+			}
+		}
+
+		union
+		{
+			uint8* Allocation = nullptr;
+			uint64 SingleItemKey;
+		};
+		int32 Count = 0;
+		int32 SingleItemValue = -1;
+	};
+
 	UPackage* Package = nullptr;
+	FPublicExportMap PublicExportMap;
 	int32 RefCount = 0;
 	bool bAreAllPublicExportsLoaded = false;
 	bool bIsMissing = false;
@@ -649,47 +666,21 @@
 	bool bHasBeenLoadedDebug = false;
 
 public:
+	FLoadedPackageRef()
+	{
+	}
+
+	FLoadedPackageRef(const FLoadedPackageRef& Other) = delete;
+
+	FLoadedPackageRef(FLoadedPackageRef&& Other) = default;
+
+	FLoadedPackageRef& operator=(const FLoadedPackageRef& Other) = delete;
+
+	FLoadedPackageRef& operator=(FLoadedPackageRef&& Other) = default;
+	
 	inline int32 GetRefCount() const
 	{
 		return RefCount;
 	}
 
-	inline bool AddRef()
-	{
-		++RefCount;
-		// is this the first reference to a package that has been loaded earlier?
-		return RefCount == 1 && Package;
-	}
-
-	inline bool ReleaseRef(FPackageId FromPackageId, FPackageId ToPackageId)
-	{
-		check(RefCount > 0);
-		--RefCount;
-
-#if DO_CHECK
-		ensureMsgf(!bHasBeenLoadedDebug || bAreAllPublicExportsLoaded || bIsMissing || bHasFailed,
-			TEXT("LoadedPackageRef from None (0x%llX) to %s (0x%llX) should not have been released when the package is not complete.")
-			TEXT("RefCount=%d, AreAllExportsLoaded=%d, IsMissing=%d, HasFailed=%d, HasBeenLoaded=%d"),
-			FromPackageId.Value(),
-			Package ? *Package->GetName() : TEXT("None"),
-			ToPackageId.Value(),
-			RefCount,
-			bAreAllPublicExportsLoaded,
-			bIsMissing,
-			bHasFailed,
-			bHasBeenLoadedDebug);
-
-		if (bAreAllPublicExportsLoaded)
-		{
-			check(!bIsMissing);
-		}
-		if (bIsMissing)
-		{
-			check(!bAreAllPublicExportsLoaded);
-		}
-#endif
-		// is this the last reference to a loaded package?
-		return RefCount == 0 && Package;
-	}
-
 	inline UPackage* GetPackage() const
@@ -732,14 +723,6 @@
 		bHasBeenLoadedDebug = true;
 	}
 
-	inline void ClearAllPublicExportsLoaded()
-	{
-		check(!bIsMissing);
-		check(Package);
-		bIsMissing = false;
-		bAreAllPublicExportsLoaded = false;
-	}
-
 	inline void SetIsMissingPackage()
 	{
 		check(!bAreAllPublicExportsLoaded);
@@ -758,19 +741,73 @@
 	{
 		bHasFailed = true;
 	}
+
+	void ReserveSpaceForPublicExports(int32 PublicExportCount)
+	{
+		PublicExportMap.Grow(PublicExportCount);
+	}
+
+	void StorePublicExport(uint64 ExportHash, UObject* Object)
+	{
+		PublicExportMap.Store(ExportHash, Object);
+	}
+
+	void RemovePublicExport(uint64 ExportHash)
+	{
+		check(!bIsMissing);
+		check(Package);
+		bAreAllPublicExportsLoaded = false;
+		PublicExportMap.Remove(ExportHash);
+	}
+
+	UObject* GetPublicExport(uint64 ExportHash)
+	{
+		return PublicExportMap.Find(ExportHash);
+	}
+
+	void PinPublicExportsForGC()
+	{
+		UE_ASYNC_UPACKAGE_DEBUG(Package);
+
+		if (GUObjectArray.IsDisregardForGC(Package))
+		{
+			return;
+		}
+		PublicExportMap.PinForGC();
+		checkf(!Package->HasAnyInternalFlags(EInternalObjectFlags::LoaderImport), TEXT("%s"), *Package->GetFullName());
+		Package->SetInternalFlags(EInternalObjectFlags::LoaderImport);
+	}
+
+	void UnpinPublicExportsForGC()
+	{
+		UE_ASYNC_UPACKAGE_DEBUG(Package);
+
+		if (GUObjectArray.IsDisregardForGC(Package))
+		{
+			return;
+		}
+		PublicExportMap.UnpinForGC();
+		checkf(Package->HasAnyInternalFlags(EInternalObjectFlags::LoaderImport), TEXT("%s"), *Package->GetFullName());
+		Package->AtomicallyClearInternalFlags(EInternalObjectFlags::LoaderImport);
+	}
+
+#if DO_CHECK
+	void VerifyAllPublicExportsRemoved()
+	{
+		PublicExportMap.VerifyAllObjectsRemoved();
+	}
+#endif
 };
 
 class FLoadedPackageStore
 {
 private:
-	// Packages in active loading or completely loaded packages, with Desc.DiskPackageName as key.
+	// Packages in active loading or completely loaded packages, with Desc.UPackageId as key.
 	// Does not track temp packages with custom UPackage names, since they are never imorted by other packages.
 	TMap<FPackageId, FLoadedPackageRef> Packages;
-	FPackageStore& PackageStore;
 
 public:
-	FLoadedPackageStore(FPackageStore& InPackageStore)
-		: PackageStore(InPackageStore)
+	FLoadedPackageStore()
 	{
 		Packages.Reserve(32768);
 	}
@@ -785,16 +822,61 @@
 		return Packages.Find(PackageId);
 	}
 
-	inline FLoadedPackageRef& GetPackageRef(FPackageId PackageId)
+	inline FLoadedPackageRef& FindPackageRefChecked(FPackageId PackageId)
 	{
-		return Packages.FindOrAdd(PackageId);
+		FLoadedPackageRef* PackageRef = FindPackageRef(PackageId);
+		UE_CLOG(!PackageRef, LogStreaming, Fatal, TEXT("FindPackageRefChecked: Package with id %d has been deleted"), PackageId.ValueForDebugging());
+		return *PackageRef;
 	}
 
-	inline int32 RemovePackage(FPackageId PackageId)
+	inline FLoadedPackageRef& AddPackageRef(FPackageId PackageId)
 	{
-		FLoadedPackageRef Ref;
-		bool bRemoved = Packages.RemoveAndCopyValue(PackageId, Ref);
-		return bRemoved ? Ref.GetRefCount() : -1;
+		LLM_SCOPE_BYNAME(TEXT("AsyncLoadPackageStore"));
+
+		FLoadedPackageRef& PackageRef = Packages.FindOrAdd(PackageId);
+		// is this the first reference to a package that has been loaded earlier?
+		if (PackageRef.RefCount == 0 && PackageRef.Package)
+		{
+			PackageRef.PinPublicExportsForGC();
+		}
+		++PackageRef.RefCount;
+		return PackageRef;
 	}
 
+	inline void ReleasePackageRef(FPackageId PackageId, FPackageId FromPackageId = FPackageId())
+	{
+		FLoadedPackageRef& PackageRef = FindPackageRefChecked(PackageId);
+
+		check(PackageRef.RefCount > 0);
+		--PackageRef.RefCount;
+
+#if DO_CHECK
+		ensureMsgf(!PackageRef.bHasBeenLoadedDebug || PackageRef.bAreAllPublicExportsLoaded || PackageRef.bIsMissing || PackageRef.bHasFailed,
+			TEXT("LoadedPackageRef from None (0x%llX) to %s (0x%llX) should not have been released when the package is not complete.")
+			TEXT("RefCount=%d, AreAllExportsLoaded=%d, IsMissing=%d, HasFailed=%d, HasBeenLoaded=%d"),
+			FromPackageId.Value(),
+			PackageRef.Package ? *PackageRef.Package->GetName() : TEXT("None"),
+			PackageId.Value(),
+			PackageRef.RefCount,
+			PackageRef.bAreAllPublicExportsLoaded,
+			PackageRef.bIsMissing,
+			PackageRef.bHasFailed,
+			PackageRef.bHasBeenLoadedDebug);
+
+		if (PackageRef.bAreAllPublicExportsLoaded)
+		{
+			check(!PackageRef.bIsMissing);
+		}
+		if (PackageRef.bIsMissing)
+		{
+			check(!PackageRef.bAreAllPublicExportsLoaded);
+		}
+#endif
+		// is this the last reference to a loaded package?
+		if (PackageRef.RefCount == 0 && PackageRef.Package)
+		{
+			PackageRef.UnpinPublicExportsForGC();
+		}
+	}
+
 #if ALT2_VERIFY_ASYNC_FLAGS
@@ -813,49 +895,239 @@
 	}
 #endif
 
-	void RemovePackage(UPackage* Package)
-	{
-		UE_ASYNC_UPACKAGE_DEBUG(Package);
-		check(IsGarbageCollecting());
-
-		if (!Package->CanBeImported())
-		{
-			return;
-		}
-
-		int32 RefCount = RemovePackage(Package->GetPackageId());
-		if (RefCount > 0)
-		{
-			UE_LOG(LogStreaming, Error,
-				TEXT("RemovePackage: %s %s (0x%llX) - with (ObjectFlags=%x, InternalObjectFlags=%x) - ")
-				TEXT("Package destroyed while still being referenced, RefCount %d > 0."),
-				*Package->GetName(),
-				*Package->GetLoadedPath().GetDebugName(), Package->GetPackageId().Value(),
-				Package->GetFlags(), Package->GetInternalFlags(), RefCount);
-			checkf(false, TEXT("Package %s destroyed with RefCount"), *Package->GetName());
-		}
-	}
-
 	void RemovePackages(const FUnreachablePackages& UnreachablePackages)
 	{
+		check(IsGarbageCollecting());
+
 		const int32 PackageCount = UnreachablePackages.Num();
-		for (int32 Index = 0; Index < PackageCount; ++Index)
+		for (UPackage* Package : UnreachablePackages)
 		{
-			RemovePackage(UnreachablePackages[Index]);
+			UE_ASYNC_UPACKAGE_DEBUG(Package);
+			if (Package->CanBeImported())
+			{
+				FLoadedPackageRef PackageRef;
+				bool bRemoved = Packages.RemoveAndCopyValue(Package->GetPackageId(), PackageRef);
+				if (bRemoved)
+				{
+					if (PackageRef.RefCount > 0)
+					{
+						UE_LOG(LogStreaming, Error,
+							TEXT("RemovePackage: %s %s (0x%llX) - with (ObjectFlags=%x, InternalObjectFlags=%x) - ")
+							TEXT("Package destroyed while still being referenced, RefCount %d > 0."),
+							*Package->GetName(),
+							*Package->GetLoadedPath().GetDebugName(), Package->GetPackageId().Value(),
+							Package->GetFlags(), Package->GetInternalFlags(), PackageRef.RefCount);
+						checkf(false, TEXT("Package %s destroyed with RefCount"), *Package->GetName());
+					}
+#if DO_CHECK
+					PackageRef.VerifyAllPublicExportsRemoved();
+#endif
+				}
+			}
+		}
+	}
+};
+
+class FGlobalImportStore
+{
+private:
+	FLoadedPackageStore& LoadedPackageStore;
+	TMap<FPackageObjectIndex, UObject*> ScriptObjects;
+	TMap<int32, FPublicExportKey> ObjectIndexToPublicExport;
+
+public:
+	FGlobalImportStore(FLoadedPackageStore& InLoadedPackageStore)
+		: LoadedPackageStore(InLoadedPackageStore)
+	{
+		ObjectIndexToPublicExport.Reserve(32768);
+	}
+
+	int32 GetStoredScriptObjectsCount() const
+	{
+		return ScriptObjects.Num();
+	}
+
+	uint32 GetStoredScriptObjectsAllocatedSize() const
+	{
+		return ScriptObjects.GetAllocatedSize();
+	}
+
+	int32 GetStoredPublicExportsCount() const
+	{
+		return ObjectIndexToPublicExport.Num();
+	}
+
+	void RemovePublicExports(const FUnreachablePublicExports& PublicExports)
+	{
+		TArray<FPublicExportKey> PublicExportKeys;
+		PublicExportKeys.Reserve(PublicExports.Num());
+
+		for (const FUnreachablePublicExport& Item : PublicExports)
+		{
+			int32 ObjectIndex = Item.Key;
+			FPublicExportKey PublicExportKey;
+			if (ObjectIndexToPublicExport.RemoveAndCopyValue(ObjectIndex, PublicExportKey))
+			{
+				PublicExportKeys.Emplace(PublicExportKey);
+#if DO_CHECK
+				{
+					UObject* GCObject = Item.Value;
+
+					checkf(GCObject->HasAllFlags(GlobalImportObjectConditionalFlags),
+						TEXT("The serialized GC Object '%s' with (ObjectFlags=%x, InternalObjectFlags=%x)and id 0x%llX:0x%llX is currently missing RF_WasLoaded or RF_Public. ")
+						TEXT("The flags must have been altered incorrectly by a higher level system after the object was loaded. ")
+						TEXT("This may cause memory corruption."),
+						*GCObject->GetFullName(),
+						GCObject->GetFlags(),
+						GCObject->GetInternalFlags(),
+						PublicExportKey.GetPackageId().Value(), PublicExportKey.GetExportHash());
+
+					UObject* ExistingObject = FindPublicExportObjectUnchecked(PublicExportKey);
+					checkf(ExistingObject,
+						TEXT("The serialized GC object '%s' with flags (ObjectFlags=%x, InternalObjectFlags=%x) and id 0x%llX:0x%llX is missing in ImportStore. ")
+						TEXT("Reason unknown. Double delete? Bug or hash collision?"),
+						*GCObject->GetFullName(),
+						GCObject->GetFlags(),
+						GCObject->GetInternalFlags(),
+						PublicExportKey.GetPackageId().Value(),
+						PublicExportKey.GetExportHash());
+
+					checkf(ExistingObject && ExistingObject == GCObject,
+						TEXT("The serialized GC Object '%s' with flags (ObjectFlags=%x, InternalObjectFlags=%x) and id 0x%llX:0x%llX is not matching the object '%s' in ImportStore. "),
+						TEXT("Reason unknown. Overwritten after it was added? Bug or hash collision?"),
+						*GCObject->GetFullName(),
+						GCObject->GetFlags(),
+						GCObject->GetInternalFlags(),
+						PublicExportKey.GetPackageId().Value(),
+						PublicExportKey.GetExportHash(),
+						*ExistingObject->GetFullName());
+				}
+#endif
+			}
+		}
+
+		FPackageId LastPackageId;
+		FLoadedPackageRef* PackageRef = nullptr;
+		for (const FPublicExportKey& PublicExportKey : PublicExportKeys)
+		{
+			FPackageId PackageId = PublicExportKey.GetPackageId();
+			if (PackageId != LastPackageId)
+			{
+				LastPackageId = PackageId;
+				PackageRef = LoadedPackageStore.FindPackageRef(PackageId);
+			}
+			check(PackageRef)
+			{
+				PackageRef->RemovePublicExport(PublicExportKey.GetExportHash());
+			}
 		}
 	}
 
-	void ClearAllPublicExportsLoaded(const TArray<FPackageId>& PackageIds)
+	inline UObject* FindPublicExportObjectUnchecked(const FPublicExportKey& Key)
 	{
-		const int32 PackageCount = PackageIds.Num();
-		const bool bForceSingleThreaded = PackageCount < 1024;
-		ParallelFor(PackageCount, [this, &PackageIds](int32 Index)
+		FLoadedPackageRef* PackageRef = LoadedPackageStore.FindPackageRef(Key.GetPackageId());
+		if (!PackageRef)
 		{
-			if (FLoadedPackageRef* PackageRef = FindPackageRef(PackageIds[Index]))
+			return nullptr;
+		}
+		return PackageRef->GetPublicExport(Key.GetExportHash());
+	}
+
+	inline UObject* FindPublicExportObject(const FPublicExportKey& Key)
+	{
+		UObject* Object = FindPublicExportObjectUnchecked(Key);
+		checkf(!Object || !Object->IsUnreachable(), TEXT("%s"), Object ? *Object->GetFullName() : TEXT("null"));
+		return Object;
+	}
+
+	inline UObject* FindScriptImportObject(FPackageObjectIndex GlobalIndex)
+	{
+		check(GlobalIndex.IsScriptImport());
+		UObject* Object = nullptr;
+		Object = ScriptObjects.FindRef(GlobalIndex);
+		return Object;
+	}
+
+	void StoreGlobalObject(FPackageId PackageId, uint64 ExportHash, UObject* Object)
+	{
+		check(PackageId.IsValid());
+		check(ExportHash != 0);
+		int32 ObjectIndex = GUObjectArray.ObjectToIndex(Object);
+		FPublicExportKey Key = FPublicExportKey::MakeKey(PackageId, ExportHash);
+#if DO_CHECK
+		{
+			UObject* ExistingObject = FindPublicExportObjectUnchecked(Key);
+			if (ExistingObject)
 			{
-				PackageRef->ClearAllPublicExportsLoaded();
+				checkf(ExistingObject == Object,
+					TEXT("The constructed serialized object '%s' with index %d and id 0x%llX:0x%llX collides with the object '%s' in ImportStore. ")
+					TEXT("Reason unknown. Bug or hash collision?"),
+					Object ? *Object->GetFullName() : TEXT("null"),
+					ObjectIndex,
+					Key.GetPackageId().Value(), Key.GetExportHash(),
+					ExistingObject ? *ExistingObject->GetFullName() : TEXT("null"));
 			}
-		}, bForceSingleThreaded);
+
+			FPublicExportKey* ExistingKey = ObjectIndexToPublicExport.Find(ObjectIndex);
+			if (ExistingKey)
+			{
+				checkf(*ExistingKey == Key,
+					TEXT("The constructed serialized object '%s' with index %d and id 0x%llX:0x%llX collides with the object '%s' in ImportStore. ")
+					TEXT("Reason unknown. Bug or hash collision?"),
+					Object ? *Object->GetFullName() : TEXT("null"),
+					ObjectIndex,
+					Key.GetPackageId().Value(), Key.GetExportHash(),
+					ExistingObject ? *ExistingObject->GetFullName() : TEXT("null"));
+			}
+		}
+#endif
+		FLoadedPackageRef& PackageRef = LoadedPackageStore.FindPackageRefChecked(Key.GetPackageId());
+		PackageRef.StorePublicExport(ExportHash, Object);
+		ObjectIndexToPublicExport.Add(ObjectIndex, Key);
+	}
+
+	void FindAllScriptObjects();
+
+	void AddScriptObject(FStringView PackageName, FStringView Name)
+	{
+		TStringBuilder<FName::StringBufferSize> FullName;
+		FPathViews::Append(FullName, PackageName);
+		FPathViews::Append(FullName, Name);
+		FPackageObjectIndex GlobalImportIndex = FPackageObjectIndex::FromScriptPath(FullName);
+
+		UObject* Package = StaticFindObjectFastInternal(UPackage::StaticClass(), nullptr, FName(PackageName), true);
+		// check(Package);
+		if (!ensureMsgf(Package, TEXT("Failed to find package when registering script object '%.*s/%.*s'"),
+			PackageName.Len(), PackageName.GetData(), Name.Len(), Name.GetData()))
+		{
+			return;
+		}
+
+		UObject* Object = StaticFindObjectFastInternal(nullptr, Package, FName(Name));
+		// check(Object);
+		if (!ensureMsgf(Object, TEXT("Failed to find object when registering script object '%.*s/%.*s'"),
+			PackageName.Len(), PackageName.GetData(), Name.Len(), Name.GetData()))
+		{
+			return;
+		}
+
+#if WITH_EDITOR
+		FPackageObjectIndex PackageGlobalImportIndex = FPackageObjectIndex::FromScriptPath(PackageName);
+		ScriptObjects.Add(PackageGlobalImportIndex, Package);
+#endif
+		ScriptObjects.Add(GlobalImportIndex, Object);
+
+		TStringBuilder<FName::StringBufferSize> SubObjectName;
+		ForEachObjectWithOuter(Object, [this, &SubObjectName](UObject* SubObject)
+			{
+				if (SubObject->HasAnyFlags(RF_Public))
+				{
+					SubObjectName.Reset();
+					SubObject->GetPathName(nullptr, SubObjectName);
+					FPackageObjectIndex SubObjectGlobalImportIndex = FPackageObjectIndex::FromScriptPath(SubObjectName);
+					ScriptObjects.Add(SubObjectGlobalImportIndex, SubObject);
+				}
+			}, /* bIncludeNestedObjects*/ true);
 	}
 };
 
@@ -967,7 +1239,7 @@
 			{
 				UE_LOG(LogStreaming, Log, TEXT("Package %s has a dependency on pending script CDO for '%s' (0x%llX)"),
 					*Header.PackageName.ToString(), *Class->GetFullName(), Index.Value());
-					Classes.AddUnique(Class);
+				Classes.AddUnique(Class);
 			}
 		}
 	}
@@ -977,58 +1249,12 @@
 		GlobalImportStore.StoreGlobalObject(PackageId, ExportHash, Object);
 	}
 
-private:
-	void PinPublicExportsForGC(UPackage* ImportedPackage)
-	{
-		UE_ASYNC_UPACKAGE_DEBUG(ImportedPackage);
-		
-		if (GUObjectArray.IsDisregardForGC(ImportedPackage))
-		{
-			return;
-		}
-		ForEachObjectWithOuter(ImportedPackage, [](UObject* Object)
-		{
-			if (Object->HasAllFlags(GlobalImportObjectConditionalFlags))
-			{
-				checkf(!Object->HasAnyInternalFlags(EInternalObjectFlags::LoaderImport), TEXT("%s"), *Object->GetFullName());
-				Object->SetInternalFlags(EInternalObjectFlags::LoaderImport);
-			}
-		}, /* bIncludeNestedObjects*/ true);
-		checkf(!ImportedPackage->HasAnyInternalFlags(EInternalObjectFlags::LoaderImport), TEXT("%s"), *ImportedPackage->GetFullName());
-		ImportedPackage->SetInternalFlags(EInternalObjectFlags::LoaderImport);
-	}
-
-	void UnpinPublicExportsForGC(UPackage* ImportedPackage)
-	{
-		UE_ASYNC_UPACKAGE_DEBUG(ImportedPackage);
-
-		if (GUObjectArray.IsDisregardForGC(ImportedPackage))
-		{
-			return;
-		}
-		ForEachObjectWithOuter(ImportedPackage, [](UObject* Object)
-		{
-			if (Object->HasAllFlags(GlobalImportObjectConditionalFlags))
-			{
-				// RF_Public can be added to objects after they've been created and in that case they might not have the LoaderImport flag set
-				//checkf(Object->HasAnyInternalFlags(EInternalObjectFlags::LoaderImport), TEXT("%s"), *Object->GetFullName());
-				Object->AtomicallyClearInternalFlags(EInternalObjectFlags::LoaderImport);
-			}
-		}, /* bIncludeNestedObjects*/ true);
-		checkf(ImportedPackage->HasAnyInternalFlags(EInternalObjectFlags::LoaderImport), TEXT("%s"), *ImportedPackage->GetFullName());
-		ImportedPackage->AtomicallyClearInternalFlags(EInternalObjectFlags::LoaderImport);
-	}
-
 public:
 	void AddImportedPackageReferences(const TArrayView<const FPackageId>& ImportedPackageIds)
 	{
 		for (const FPackageId& ImportedPackageId : ImportedPackageIds)
 		{
-			FLoadedPackageRef& PackageRef = LoadedPackageStore.GetPackageRef(ImportedPackageId);
-			if (PackageRef.AddRef())
-			{
-				PinPublicExportsForGC(PackageRef.GetPackage());
-			}
+			LoadedPackageStore.AddPackageRef(ImportedPackageId);
 		}
 	}
 
@@ -1036,12 +1262,8 @@
 	{
 		if (Desc.bCanBeImported)
 		{
-			FLoadedPackageRef& PackageRef = LoadedPackageStore.GetPackageRef(Desc.UPackageId);
+			FLoadedPackageRef& PackageRef = LoadedPackageStore.AddPackageRef(Desc.UPackageId);
 			PackageRef.ClearErrorFlags();
-			if (PackageRef.AddRef())
-			{
-				PinPublicExportsForGC(PackageRef.GetPackage());
-			}
 		}
 	}
 
@@ -1049,11 +1271,7 @@
 	{
 		for (const FPackageId& ImportedPackageId : ImportedPackageIds)
 		{
-			FLoadedPackageRef& PackageRef = LoadedPackageStore.GetPackageRef(ImportedPackageId);
-			if (PackageRef.ReleaseRef(Desc.UPackageId, ImportedPackageId))
-			{
-				UnpinPublicExportsForGC(PackageRef.GetPackage());
-			}
+			LoadedPackageStore.ReleasePackageRef(ImportedPackageId, Desc.UPackageId);
 		}
 	}
 
@@ -1061,12 +1279,7 @@
 	{
 		if (Desc.bCanBeImported)
 		{
-			// clear own reference, and possible all async flags if no remaining ref count
-			FLoadedPackageRef& PackageRef = LoadedPackageStore.GetPackageRef(Desc.UPackageId);
-			if (PackageRef.ReleaseRef(Desc.UPackageId, Desc.UPackageId))
-			{
-				UnpinPublicExportsForGC(PackageRef.GetPackage());
-			}
+			LoadedPackageStore.ReleasePackageRef(Desc.UPackageId);
 		}
 	}
 };
@@ -3337,6 +3550,7 @@
 			}
 		}
 	}
+	ScriptObjects.Shrink();
 }
 
 void FAsyncPackage2::ImportPackagesRecursiveInner(FIoBatch& IoBatch, FPackageStore& PackageStore, const TArrayView<const FPackageId>& ImportedPackageIds, int32& ImportedPackageIndex)
@@ -3355,7 +3569,7 @@
 			}
 		}
 		
-		FLoadedPackageRef& PackageRef = ImportStore.LoadedPackageStore.GetPackageRef(ImportedPackageId);
+		FLoadedPackageRef& ImportedPackageRef = ImportStore.LoadedPackageStore.FindPackageRefChecked(ImportedPackageId);
 		FPackageStoreEntry ImportedPackageEntry;
 		EPackageStoreEntryStatus ImportedPackageStatus = PackageStore.GetPackageStoreEntry(ImportedPackageIdToLoad, ImportedPackageEntry);
 
@@ -3363,7 +3577,7 @@
 		{
 			UE_ASYNC_PACKAGE_LOG(Warning, Desc, TEXT("ImportPackages: SkipPackage"),
 				TEXT("Skipping non mounted imported package with id '0x%llX'"), ImportedPackageId.Value());
-			PackageRef.SetIsMissingPackage();
+			ImportedPackageRef.SetIsMissingPackage();
 			Data.ImportedAsyncPackages[ImportedPackageIndex++] = nullptr;
 			continue;
 		}
@@ -3371,7 +3585,7 @@
 		else if (!ImportedPackageEntry.UncookedPackageName.IsNone())
 		{
 			UPackage* UncookedPackage = nullptr;
-			if (!PackageRef.AreAllPublicExportsLoaded())
+			if (!ImportedPackageRef.AreAllPublicExportsLoaded())
 			{
 				UE_ASYNC_PACKAGE_LOG(Verbose, Desc, TEXT("ImportPackages: LoadUncookedImport"), TEXT("Loading imported uncooked package '%s' '0x%llX'"), *ImportedPackageEntry.UncookedPackageName.ToString(), ImportedPackageId.ValueForDebugging());
 				check(IsInGameThread());
@@ -3420,16 +3634,16 @@
 						}
 					}, /* bIncludeNestedObjects*/ true);
 				}
-				PackageRef.SetPackage(UncookedPackage);
-				PackageRef.SetAllPublicExportsLoaded();
+				ImportedPackageRef.SetPackage(UncookedPackage);
+				ImportedPackageRef.SetAllPublicExportsLoaded();
 			}
 			else
 			{
-				UncookedPackage = PackageRef.GetPackage();
+				UncookedPackage = ImportedPackageRef.GetPackage();
 			}
 			if (!UncookedPackage)
 			{
-				PackageRef.SetHasFailed();
+				ImportedPackageRef.SetHasFailed();
 				UE_ASYNC_PACKAGE_LOG(Warning, Desc, TEXT("ImportPackages: SkipPackage"),
 					TEXT("Failed to load uncooked imported package with id '0x%llX' ('%s')"), ImportedPackageId.Value(), *ImportedPackageEntry.UncookedPackageName.ToString());
 			}
@@ -3441,7 +3655,7 @@
 		FAsyncPackage2* ImportedPackage = nullptr;
 		bool bInserted = false;
 		FAsyncPackageDesc2 PackageDesc = FAsyncPackageDesc2::FromPackageImport(Desc.ReferencerRequestId, Desc.Priority, ImportedPackageId, ImportedPackageIdToLoad, ImportedPackageUPackageName);
-		if (PackageRef.AreAllPublicExportsLoaded())
+		if (ImportedPackageRef.AreAllPublicExportsLoaded())
 		{
 			ImportedPackage = AsyncLoadingThread.FindAsyncPackage(ImportedPackageId);
 			if (!ImportedPackage)
@@ -3674,9 +3888,8 @@
 	{
 		if (Package->Desc.bCanBeImported)
 		{
-			FLoadedPackageRef* PackageRef = Package->ImportStore.LoadedPackageStore.FindPackageRef(Package->Desc.UPackageId);
-			check(PackageRef);
-			PackageRef->SetHasFailed();
+			FLoadedPackageRef& PackageRef = Package->ImportStore.LoadedPackageStore.FindPackageRefChecked(Package->Desc.UPackageId);
+			PackageRef.SetHasFailed();
 		}
 	}
 	else
@@ -3691,6 +3904,35 @@
 			ReadAsyncPackageHeader(*Package->OptionalSegmentSerializationState, *OptionalSegmentHeaderData);
 		}
 #endif
+		if (Package->Desc.bCanBeImported)
+		{
+			int32 PublicExportsCount = 0;
+			for (const FExportMapEntry& Export : Package->HeaderData.ExportMap)
+			{
+				if (Export.PublicExportHash)
+				{
+					++PublicExportsCount;
+				}
+			}
+#if WITH_EDITOR
+			if (Package->OptionalSegmentHeaderData.IsSet())
+			{
+				for (const FExportMapEntry& Export : Package->OptionalSegmentHeaderData->ExportMap)
+				{
+					if (Export.PublicExportHash)
+					{
+						++PublicExportsCount;
+					}
+				}
+			}
+#endif
+			FLoadedPackageRef& PackageRef = Package->ImportStore.LoadedPackageStore.FindPackageRefChecked(Package->Desc.UPackageId);
+			if (PublicExportsCount)
+			{
+				PackageRef.ReserveSpaceForPublicExports(PublicExportsCount);
+			}
+		}
+
 		for (int32 ExportBundleIndex = 0; ExportBundleIndex < Package->Data.ExportBundleCount; ++ExportBundleIndex)
 		{
 			const FAsyncPackageHeaderData* HeaderData = &Package->HeaderData;
@@ -4301,8 +4543,7 @@
 
 	if (!Package->bLoadHasFailed && Package->Desc.bCanBeImported)
 	{
-		FLoadedPackageRef& PackageRef =
-			Package->AsyncLoadingThread.LoadedPackageStore.GetPackageRef(Package->Desc.UPackageId);
+		FLoadedPackageRef& PackageRef = Package->AsyncLoadingThread.LoadedPackageStore.FindPackageRefChecked(Package->Desc.UPackageId);
 		PackageRef.SetAllPublicExportsLoaded();
 	}
 
@@ -5222,7 +5463,7 @@
 	, IoDispatcher(InIoDispatcher)
 	, UncookedPackageLoader(InUncookedPackageLoader)
 	, PackageStore(FPackageStore::Get())
-	, LoadedPackageStore(PackageStore)
+	, GlobalImportStore(LoadedPackageStore)
 {
 #if !WITH_IOSTORE_IN_EDITOR
 	IsEventDrivenLoaderEnabled(); // make sure the one time init inside runs
@@ -5346,7 +5587,6 @@
 {
 	TRACE_CPUPROFILER_EVENT_SCOPE(FinalizeInitialLoad);
 	GlobalImportStore.FindAllScriptObjects(); // for verification only
-	GlobalImportStore.ScriptObjects.Shrink();
 	bHasRegisteredAllScriptObjects = true;
 
 	check(PendingCDOs.Num() == 0);
@@ -5356,7 +5596,7 @@
 
 	UE_LOG(LogStreaming, Display,
 		TEXT("AsyncLoading2 - InitialLoad Finalized: Registered %d public script object entries (%.2f KB)"),
-		GlobalImportStore.ScriptObjects.Num(), (float)GlobalImportStore.ScriptObjects.GetAllocatedSize() / 1024.f);
+		GlobalImportStore.GetStoredScriptObjectsCount(), (float)GlobalImportStore.GetStoredScriptObjectsAllocatedSize() / 1024.f);
 }
 
 uint32 FAsyncLoadingThread2::Run()
@@ -5688,22 +5928,16 @@
 {
 	TRACE_CPUPROFILER_EVENT_SCOPE(RemoveUnreachableObjects);
 
-	TArray<FPackageId> PublicExportPackages;
 	if (PublicExports.Num() > 0)
 	{
 		// TRACE_CPUPROFILER_EVENT_SCOPE(RemovePublicExports);
-		PublicExportPackages = GlobalImportStore.RemovePublicExports(PublicExports);
+		GlobalImportStore.RemovePublicExports(PublicExports);
 	}
 	if (Packages.Num() > 0)
 	{
 		// TRACE_CPUPROFILER_EVENT_SCOPE(RemovePackages);
 		LoadedPackageStore.RemovePackages(Packages);
 	}
-	if (PublicExportPackages.Num() > 0)
-	{
-		// TRACE_CPUPROFILER_EVENT_SCOPE(ClearAllPublicExportsLoaded);
-		LoadedPackageStore.ClearAllPublicExportsLoaded(PublicExportPackages);
-	}
 }
 
 void FAsyncLoadingThread2::NotifyUnreachableObjects(const TArrayView<FUObjectItem*>& UnreachableObjects)
@@ -5726,13 +5960,13 @@
 	if (PackageCount > 0 || PublicExportCount > 0)
 	{
 		const int32 OldLoadedPackageCount = LoadedPackageStore.NumTracked();
-		const int32 OldPublicExportCount = GlobalImportStore.PublicExportToObjectIndex.Num();
+		const int32 OldPublicExportCount = GlobalImportStore.GetStoredPublicExportsCount();
 
 		const double RemoveStartTime = FPlatformTime::Seconds();
 		RemoveUnreachableObjects(PublicExports, Packages);
 
 		const int32 NewLoadedPackageCount = LoadedPackageStore.NumTracked();
-		const int32 NewPublicExportCount = GlobalImportStore.PublicExportToObjectIndex.Num();
+		const int32 NewPublicExportCount = GlobalImportStore.GetStoredPublicExportsCount();
 		const int32 RemovedLoadedPackageCount = OldLoadedPackageCount - NewLoadedPackageCount;
 		const int32 RemovedPublicExportCount = OldPublicExportCount - NewPublicExportCount;
 
