Improve tracking of the LoadTime value in CookProfiler

diff --git a/Engine/Source/Developer/TraceServices/Private/Analyzers/CookAnalysis.cpp b/Engine/Source/Developer/TraceServices/Private/Analyzers/CookAnalysis.cpp
--- a/Engine/Source/Developer/TraceServices/Private/Analyzers/CookAnalysis.cpp
+++ b/Engine/Source/Developer/TraceServices/Private/Analyzers/CookAnalysis.cpp
@@ -42,7 +42,8 @@ bool FCookAnalyzer::OnEvent(uint16 RouteId, EStyle Style, const FOnEventContext&
 		EventData.GetString("Name", Name);
 
 		TraceServices::FProviderEditScopeLock ProviderEditScope(CookProfilerProvider);
-		CookProfilerProvider.AddPackage(Id, Name);
+		FPackageData* Package = CookProfilerProvider.EditPackage(Id);
+		Package->Name = Session.StoreString(Name);
 		break;
 	}
 	case RouteId_PackageStat:
@@ -59,7 +60,8 @@ bool FCookAnalyzer::OnEvent(uint16 RouteId, EStyle Style, const FOnEventContext&
 			{
 			case EPackageEventStatType::LoadPackage:
 			{
-				Package->LoadTime = Value;
+				// We measure the Loadtime in multiple scopes so we receive many LoadPackage events.
+				Package->LoadTime += Value;
 				break;
 			}
 			case EPackageEventStatType::SavePackage:
diff --git a/Engine/Source/Developer/TraceServices/Private/Model/CookProfilerProvider.cpp b/Engine/Source/Developer/TraceServices/Private/Model/CookProfilerProvider.cpp
--- a/Engine/Source/Developer/TraceServices/Private/Model/CookProfilerProvider.cpp
+++ b/Engine/Source/Developer/TraceServices/Private/Model/CookProfilerProvider.cpp
@@ -10,6 +10,13 @@ thread_local FProviderLock* GThreadCurrentCookProviderLock;
 thread_local int32 GThreadCurrentReadCookProviderLockCount;
 thread_local int32 GThreadCurrentWriteCookProviderLockCount;
 
+const TCHAR* GUnknownPackage = TEXT("Unknown Package");
+
+FPackageData::FPackageData(uint64 InId)
+	: Id(InId)
+	, Name(GUnknownPackage)
+{}
+
 void FCookProfilerProvider::BeginEdit() const
 {
 	Lock.BeginWrite(GThreadCurrentCookProviderLock, GThreadCurrentReadCookProviderLockCount, GThreadCurrentWriteCookProviderLockCount);
@@ -46,12 +53,6 @@ FCookProfilerProvider::FCookProfilerProvider(IAnalysisSession& InSession)
 	
 }
 
-void FCookProfilerProvider::AddPackage(uint64 Id, TStringView<TCHAR> Name)
-{
-	PackageIdToIndexMap.Add(Id, Packages.Num());
-	Packages.Emplace(Id, Session.StoreString(Name));
-}
-
 void FCookProfilerProvider::EnumeratePackages(double StartTime, double EndTime, EnumeratePackagesCallback Callback) const
 {
 	for (const FPackageData& Package : Packages)
@@ -70,15 +71,25 @@ uint32 FCookProfilerProvider::GetNumPackages() const
 
 FPackageData* FCookProfilerProvider::EditPackage(uint64 Id)
 {
-	uint32* Index = PackageIdToIndexMap.Find(Id);
+	uint32 Index = FindOrAddPackage(Id);
 
-	if (Index && *Index < Packages.Num())
+	FPackageData& Package = Packages[Index];
+	return &Package;
+}
+
+uint32 FCookProfilerProvider::FindOrAddPackage(uint64 Id)
+{
+	uint32* Index = PackageIdToIndexMap.Find(Id);
+	if (Index != nullptr)
 	{
-		FPackageData& Package = Packages[*Index];
-		return &Package;
+		return *Index;
+	}
+	else
+	{
+		PackageIdToIndexMap.Add(Id, Packages.Num());
+		Packages.Emplace(Id);
+		return Packages.Num() - 1;
 	}
-
-	return nullptr;
 }
 
 } // namespace TraceServices
diff --git a/Engine/Source/Developer/TraceServices/Private/Model/CookProfilerProviderPrivate.h b/Engine/Source/Developer/TraceServices/Private/Model/CookProfilerProviderPrivate.h
--- a/Engine/Source/Developer/TraceServices/Private/Model/CookProfilerProviderPrivate.h
+++ b/Engine/Source/Developer/TraceServices/Private/Model/CookProfilerProviderPrivate.h
@@ -23,14 +23,15 @@ public:
 	explicit FCookProfilerProvider(IAnalysisSession& Session);
 	virtual ~FCookProfilerProvider() {}
 
-	// Functions to add or edit data.
-	void AddPackage(uint64 Id, TStringView<TCHAR> Name);
 	FPackageData* EditPackage(uint64 Id);
 
 	// Functions to query the data.
 	virtual void EnumeratePackages(double StartTime, double EndTime, EnumeratePackagesCallback Callback) const override;
 	virtual uint32 GetNumPackages() const;
 
+private:
+	uint32 FindOrAddPackage(uint64 Id);
+
 private:
 	IAnalysisSession& Session;
 	TMap<uint64, uint32> PackageIdToIndexMap;
diff --git a/Engine/Source/Developer/TraceServices/Public/TraceServices/Model/CookProfilerProvider.h b/Engine/Source/Developer/TraceServices/Public/TraceServices/Model/CookProfilerProvider.h
--- a/Engine/Source/Developer/TraceServices/Public/TraceServices/Model/CookProfilerProvider.h
+++ b/Engine/Source/Developer/TraceServices/Public/TraceServices/Model/CookProfilerProvider.h
@@ -11,10 +11,7 @@ namespace TraceServices
 
 struct FPackageData
 {
-	FPackageData(uint64 InId, const TCHAR* InName)
-		: Id(InId)
-		, Name(InName)
-	{}
+	FPackageData(uint64 InId);
 	uint64 Id = -1;
 	const TCHAR* Name;
 	double LoadTime = 0;
diff --git a/Engine/Source/Editor/UnrealEd/Private/CookOnTheFlyServer.cpp b/Engine/Source/Editor/UnrealEd/Private/CookOnTheFlyServer.cpp
--- a/Engine/Source/Editor/UnrealEd/Private/CookOnTheFlyServer.cpp
+++ b/Engine/Source/Editor/UnrealEd/Private/CookOnTheFlyServer.cpp
@@ -3605,7 +3605,6 @@ bool UCookOnTheFlyServer::LoadPackageForCooking(UE::Cook::FPackageData& PackageD
 	UE::Cook::FPackageData* ReportingPackageData)
 {
 	UE_SCOPED_HIERARCHICAL_COOKTIMER_AND_DURATION(LoadPackageForCooking, DetailedCookStats::TickCookOnTheSideLoadPackagesTimeSec);
-	UE_SCOPED_COOK_STAT(PackageData.GetPackageName().ToUnstableInt(), EPackageEventStatType::LoadPackage);
 
 	check(PackageTracker->LoadingPackageData == nullptr);
 	PackageTracker->LoadingPackageData = &PackageData;
diff --git a/Engine/Source/Editor/UnrealEd/Private/Cooker/CookTypes.h b/Engine/Source/Editor/UnrealEd/Private/Cooker/CookTypes.h
--- a/Engine/Source/Editor/UnrealEd/Private/Cooker/CookTypes.h
+++ b/Engine/Source/Editor/UnrealEd/Private/Cooker/CookTypes.h
@@ -240,14 +240,14 @@ namespace UE::Cook
 inline void RouteBeginCacheForCookedPlatformData(UObject* Obj, const ITargetPlatform* TargetPlatform)
 {
 	UE_SCOPED_TEXT_COOKTIMER(*WriteToString<128>(GetClassTraceScope(Obj), TEXT("_BeginCacheForCookedPlatformData")));
-	UE_SCOPED_COOK_STAT(Obj->GetPackage()->GetFName().ToUnstableInt(), EPackageEventStatType::BeginCacheForCookedPlatformData);
+	UE_SCOPED_COOK_STAT(Obj->GetPackage()->GetFName(), EPackageEventStatType::BeginCacheForCookedPlatformData);
 	Obj->BeginCacheForCookedPlatformData(TargetPlatform);
 }
 
 inline bool RouteIsCachedCookedPlatformDataLoaded(UObject* Obj, const ITargetPlatform* TargetPlatform)
 {
 	UE_SCOPED_TEXT_COOKTIMER(*WriteToString<128>(GetClassTraceScope(Obj), TEXT("_IsCachedCookedPlatformDataLoaded")));
-	UE_SCOPED_COOK_STAT(Obj->GetPackage()->GetFName().ToUnstableInt(), EPackageEventStatType::IsCachedCookedPlatformDataLoaded);
+	UE_SCOPED_COOK_STAT(Obj->GetPackage()->GetFName(), EPackageEventStatType::IsCachedCookedPlatformDataLoaded);
 	return Obj->IsCachedCookedPlatformDataLoaded(TargetPlatform);
 }
 
diff --git a/Engine/Source/Runtime/Core/Public/ProfilingDebugging/CookStats.h b/Engine/Source/Runtime/Core/Public/ProfilingDebugging/CookStats.h
--- a/Engine/Source/Runtime/Core/Public/ProfilingDebugging/CookStats.h
+++ b/Engine/Source/Runtime/Core/Public/ProfilingDebugging/CookStats.h
@@ -459,12 +459,29 @@ private:
 	EPackageEventStatType StatType;
 };
 
-#define UE_SCOPED_COOK_STAT(Id, StatType) \
-	FScopedCookStat  PREPROCESSOR_JOIN(__CookTimerScope, __LINE__)(Id, StatType);
+#define UE_SCOPED_COOK_STAT(Name, StatType) \
+	FScopedCookStat  PREPROCESSOR_JOIN(__CookTimerScope, __LINE__)(Name.ToUnstableInt(), StatType);
+
+#define UE_MULTI_SCOPED_COOK_STAT_INIT() \
+	TOptional<FScopedCookStat> __CookStatScope; \
+	FName __PrevPackageName = NAME_None;
+
+#define UE_MULTI_SCOPED_COOK_STAT(Name, StatType) \
+	if (__PrevPackageName != Name) \
+	{ \
+		__CookStatScope.Emplace(Name.ToUnstableInt(), EPackageEventStatType::LoadPackage); \
+		__PrevPackageName = Name; \
+	}
 
+#define UE_MULTI_SCOPED_COOK_STAT_RESET() \
+	__CookStatScope.Reset(); \
+	__PrevPackageName = NAME_None;
 #else
 #define COOK_STAT(...)
 #define UE_SCOPED_COOK_STAT(...)
+#define UE_MULTI_SCOPED_COOK_STAT_INIT(...)
+#define UE_MULTI_SCOPED_COOK_STAT(...)
+#define UE_MULTI_SCOPED_COOK_STAT_RESET(...)
 
 #define TracePackage(...)
 #define TracePackageStat(...)
diff --git a/Engine/Source/Runtime/CoreUObject/Private/Serialization/AsyncLoading.cpp b/Engine/Source/Runtime/CoreUObject/Private/Serialization/AsyncLoading.cpp
--- a/Engine/Source/Runtime/CoreUObject/Private/Serialization/AsyncLoading.cpp
+++ b/Engine/Source/Runtime/CoreUObject/Private/Serialization/AsyncLoading.cpp
@@ -32,6 +32,7 @@
 #include "Misc/Paths.h"
 #include "Serialization/AsyncLoadingThread.h"
 #include "Misc/ExclusiveLoadPackageTimeTracker.h"
+#include "ProfilingDebugging/CookStats.h"
 #include "ProfilingDebugging/LoadTimeTracker.h"
 #include "HAL/ThreadHeartBeat.h"
 #include "HAL/ExceptionHandling.h"
@@ -4719,6 +4720,7 @@ EAsyncPackageState::Type FAsyncLoadingThread::ProcessLoadedPackages(bool bUseTim
 			|| GEventDrivenLoaderEnabled)
 		{
 			SCOPED_LOADTIMER(ProcessLoadedPackagesTime);
+			UE_SCOPED_COOK_STAT(Package->GetPackageName(), EPackageEventStatType::LoadPackage);
 			Result = Package->PostLoadDeferredObjects(TickStartTime, bUseTimeLimit, TimeLimit);
 			// if we didn't complete the post-load, we ran out of time for this slice so break immediately
 			if (Result != EAsyncPackageState::Complete)
@@ -4786,6 +4788,8 @@ EAsyncPackageState::Type FAsyncLoadingThread::ProcessLoadedPackages(bool bUseTim
 	TSet<UPackage*> CompletedUPackages;
 	for (FAsyncPackage* Package : CompletedPackages)
 	{
+		UE_SCOPED_COOK_STAT(Package->GetPackageName(), EPackageEventStatType::LoadPackage);
+
 		// Call external callbacks
 		const bool bInternalCallbacks = false;
 		const EAsyncLoadingResult::Type LoadingResult = Package->HasLoadFailed() ? EAsyncLoadingResult::Failed : EAsyncLoadingResult::Succeeded;
@@ -4829,6 +4833,8 @@ EAsyncPackageState::Type FAsyncLoadingThread::ProcessLoadedPackages(bool bUseTim
 		for (int32 PackageIndex = 0; PackageIndex < PackagesToDelete.Num(); ++PackageIndex)
 		{
 			FAsyncPackage* Package = PackagesToDelete[PackageIndex];
+			UE_SCOPED_COOK_STAT(Package->GetPackageName(), EPackageEventStatType::LoadPackage);
+
 			if (Package->GetDependencyRefCount() == 0 && !Package->IsBeingProcessedRecursively())
 			{
 				bool bSafeToDelete = false;
@@ -5872,6 +5878,7 @@ EAsyncPackageState::Type FAsyncPackage::TickAsyncPackage(bool InbUseTimeLimit, b
 
 	SCOPE_CYCLE_COUNTER(STAT_FAsyncPackage_Tick);
 	SCOPED_LOADTIMER(Package_Tick);
+	UE_SCOPED_COOK_STAT(Desc.Name, EPackageEventStatType::LoadPackage);
 
 	// Whether we should execute the next step.
 	EAsyncPackageState::Type LoadingState = EAsyncPackageState::Complete;
@@ -7351,6 +7358,8 @@ int32 FAsyncLoadingThread::LoadPackage(const FPackagePath& InPackagePath, FName
 	}
 
 	FName PackageName = InCustomName.IsNone() ? InPackagePath.GetPackageFName() : InCustomName;
+	UE_SCOPED_COOK_STAT(PackageName, EPackageEventStatType::LoadPackage);
+
 	if ( FCoreDelegates::OnAsyncLoadPackage.IsBound() )
 	{
 		FCoreDelegates::OnAsyncLoadPackage.Broadcast(PackageName.ToString());
diff --git a/Engine/Source/Runtime/CoreUObject/Private/UObject/Linker.cpp b/Engine/Source/Runtime/CoreUObject/Private/UObject/Linker.cpp
--- a/Engine/Source/Runtime/CoreUObject/Private/UObject/Linker.cpp
+++ b/Engine/Source/Runtime/CoreUObject/Private/UObject/Linker.cpp
@@ -18,6 +18,7 @@
 #include "Logging/TokenizedMessage.h"
 #include "Logging/MessageLog.h"
 #include "Misc/UObjectToken.h"
+#include "ProfilingDebugging/CookStats.h"
 #include "UObject/CoreRedirects.h"
 #include "UObject/LinkerManager.h"
 #include "UObject/UObjectThreadContext.h"
@@ -729,6 +730,7 @@ FLinkerLoad* GetPackageLinker
 		return nullptr;
 	}
 
+	UE_SCOPED_COOK_STAT(PackageFName, EPackageEventStatType::LoadPackage);
 	// Process any package redirects
 	{
 		const FCoreRedirectObjectName NewPackageName = FCoreRedirects::GetRedirectedName(ECoreRedirectFlags::Type_Package, FCoreRedirectObjectName(NAME_None, NAME_None, PackageFName));
diff --git a/Engine/Source/Runtime/CoreUObject/Private/UObject/LinkerLoad.cpp b/Engine/Source/Runtime/CoreUObject/Private/UObject/LinkerLoad.cpp
--- a/Engine/Source/Runtime/CoreUObject/Private/UObject/LinkerLoad.cpp
+++ b/Engine/Source/Runtime/CoreUObject/Private/UObject/LinkerLoad.cpp
@@ -25,6 +25,7 @@
 #include "ProfilingDebugging/DebuggingDefines.h"
 #include "Logging/MessageLog.h"
 #include "Logging/TokenizedMessage.h"
+#include "ProfilingDebugging/CookStats.h"
 #include "UObject/LinkerPlaceholderBase.h"
 #include "UObject/LinkerPlaceholderClass.h"
 #include "UObject/LinkerPlaceholderExportObject.h"
@@ -3642,6 +3643,7 @@ int32 FLinkerLoad::LoadMetaDataFromExportMap(bool bForcePreload)
 void FLinkerLoad::LoadAllObjects(bool bForcePreload)
 {
 	SCOPED_LOADTIMER(LinkerLoad_LoadAllObjects);
+	UE_SCOPED_COOK_STAT(LinkerRoot->GetFName(), EPackageEventStatType::LoadPackage);
 #if WITH_EDITOR
 	TOptional<FScopedSlowTask> SlowTask;
 	if (ShouldCreateThrottledSlowTask())
diff --git a/Engine/Source/Runtime/CoreUObject/Private/UObject/Obj.cpp b/Engine/Source/Runtime/CoreUObject/Private/UObject/Obj.cpp
--- a/Engine/Source/Runtime/CoreUObject/Private/UObject/Obj.cpp
+++ b/Engine/Source/Runtime/CoreUObject/Private/UObject/Obj.cpp
@@ -46,6 +46,7 @@
 #include "UObject/LinkerLoad.h"
 #include "Misc/RedirectCollector.h"
 #include "UObject/GCScopeLock.h"
+#include "ProfilingDebugging/CookStats.h"
 #include "ProfilingDebugging/LoadTimeTracker.h"
 #include "UObject/GCObject.h"
 #include "Serialization/ArchiveUObjectFromStructuredArchive.h"
@@ -1107,7 +1108,9 @@ void UObject::ConditionalPostLoad()
 #if WITH_EDITOR
 				SCOPED_LOADTIMER_TEXT(*WriteToString<128>(GetClassTraceScope(this), TEXTVIEW("_PostLoad")));
 #endif
-				LLM_SCOPED_TAG_WITH_OBJECT_IN_SET(GetOutermost(), ELLMTagSet::Assets);
+				UPackage* Package = GetPackage();
+				UE_SCOPED_COOK_STAT(Package->GetFName(), EPackageEventStatType::LoadPackage);
+				LLM_SCOPED_TAG_WITH_OBJECT_IN_SET(Package, ELLMTagSet::Assets);
 				LLM_SCOPED_TAG_WITH_OBJECT_IN_SET(GetClass(), ELLMTagSet::AssetClasses);
 
 				PostLoad();
diff --git a/Engine/Source/Runtime/CoreUObject/Private/UObject/SavePackage2.cpp b/Engine/Source/Runtime/CoreUObject/Private/UObject/SavePackage2.cpp
--- a/Engine/Source/Runtime/CoreUObject/Private/UObject/SavePackage2.cpp
+++ b/Engine/Source/Runtime/CoreUObject/Private/UObject/SavePackage2.cpp
@@ -2557,7 +2557,7 @@ FSavePackageResultStruct UPackage::Save2(UPackage* InPackage, UObject* InAsset,
 	COOK_STAT(FScopedDurationTimer FuncSaveTimer(FSavePackageStats::SavePackageTimeSec));
 	COOK_STAT(FSavePackageStats::NumPackagesSaved++);
 	SCOPED_SAVETIMER(UPackage_Save2);
-	UE_SCOPED_COOK_STAT(InPackage->GetFName().ToUnstableInt(), EPackageEventStatType::SavePackage);
+	UE_SCOPED_COOK_STAT(InPackage->GetFName(), EPackageEventStatType::SavePackage);
 
 	FSaveContext SaveContext(InPackage, InAsset, InFilename, SaveArgs);
 
diff --git a/Engine/Source/Runtime/CoreUObject/Private/UObject/UObjectGlobals.cpp b/Engine/Source/Runtime/CoreUObject/Private/UObject/UObjectGlobals.cpp
--- a/Engine/Source/Runtime/CoreUObject/Private/UObject/UObjectGlobals.cpp
+++ b/Engine/Source/Runtime/CoreUObject/Private/UObject/UObjectGlobals.cpp
@@ -1932,6 +1932,7 @@ void EndLoad(FUObjectSerializeContext* LoadContext, TArray<UPackage*>* OutLoaded
 			// Sort by Filename and Offset.
 			ObjLoaded.StableSort(FCompareUObjectByLinkerAndOffset());
 
+			UE_MULTI_SCOPED_COOK_STAT_INIT()
 			// Finish loading everything.
 			{
 				SCOPED_LOADTIMER(PreLoadAndSerialize);
@@ -1942,8 +1943,14 @@ void EndLoad(FUObjectSerializeContext* LoadContext, TArray<UPackage*>* OutLoaded
 					UObject* Obj = ObjLoaded[i];
 					if (Obj->HasAnyFlags(RF_NeedLoad))
 					{
-						check(Obj->GetLinker());
-						UE_TRACK_REFERENCING_PACKAGE_DELAYED(AccessRefScope, Obj->GetOutermost());
+						FLinkerLoad* Linker = Obj->GetLinker();
+						check(Linker);
+
+						UPackage* Package = Linker->LinkerRoot;
+						check(Package);
+
+						UE_MULTI_SCOPED_COOK_STAT(Package->GetFName(), EPackageEventStatType::LoadPackage);
+						UE_TRACK_REFERENCING_PACKAGE_DELAYED(AccessRefScope, Package);
 #if WITH_EDITOR
 						if (SlowTask)
 						{
@@ -1952,9 +1959,10 @@ void EndLoad(FUObjectSerializeContext* LoadContext, TArray<UPackage*>* OutLoaded
 							SlowTask->TickProgress();
 						}
 #endif
-						Obj->GetLinker()->Preload(Obj);
+						Linker->Preload(Obj);
 					}
 				}
+				UE_MULTI_SCOPED_COOK_STAT_RESET();
 			}
 
 			// Start over again as new objects have been loaded that need to have "Preload" called on them before
@@ -2023,10 +2031,12 @@ void EndLoad(FUObjectSerializeContext* LoadContext, TArray<UPackage*>* OutLoaded
 				SCOPED_LOADTIMER(PostLoadInstance);
 				for (UObject* Obj : ObjLoaded)
 				{
+					UE_MULTI_SCOPED_COOK_STAT(Obj->GetPackage()->GetFName(), EPackageEventStatType::LoadPackage);
 					UClass* ObjClass = Obj->GetClass();
 					check(ObjClass);
 					ObjClass->PostLoadInstance(Obj);
 				}
+				UE_MULTI_SCOPED_COOK_STAT_RESET();
 			}
 
 			// Create clusters after all objects have been loaded
