Automatic sidecar generation of uncooked objects using an asset type allowlist

diff --git a/Engine/Source/Runtime/CoreUObject/Private/UObject/LinkerSave.cpp b/Engine/Source/Runtime/CoreUObject/Private/UObject/LinkerSave.cpp
--- a/Engine/Source/Runtime/CoreUObject/Private/UObject/LinkerSave.cpp
+++ b/Engine/Source/Runtime/CoreUObject/Private/UObject/LinkerSave.cpp
@@ -465,6 +465,31 @@ void FLinkerSave::SetUseUnversionedPropertySerialization(bool bInUseUnversioned)
 	}
 }
 
+void FLinkerSave::SetFilterEditorOnly(bool bInFilterEditorOnly)
+{
+	FArchiveUObject::SetFilterEditorOnly(bInFilterEditorOnly);
+	if (Saver)
+	{
+		Saver->SetFilterEditorOnly(bInFilterEditorOnly);
+	}
+	if (bInFilterEditorOnly)
+	{
+		Summary.SetPackageFlags(Summary.GetPackageFlags() | PKG_FilterEditorOnly);
+		if (LinkerRoot)
+		{
+			LinkerRoot->SetPackageFlags(PKG_FilterEditorOnly);
+		}
+	}
+	else
+	{
+		Summary.SetPackageFlags(Summary.GetPackageFlags() & ~PKG_FilterEditorOnly);
+		if (LinkerRoot)
+		{
+			LinkerRoot->ClearPackageFlags(PKG_FilterEditorOnly);
+		}
+	}
+}
+
 #if WITH_EDITORONLY_DATA
 
 static FIoChunkId CreateDerivedDataChunkId(const FPackageId PackageId, const int32 ChunkIndex)
diff --git a/Engine/Source/Runtime/CoreUObject/Private/UObject/SavePackage/PackageHarvester.cpp b/Engine/Source/Runtime/CoreUObject/Private/UObject/SavePackage/PackageHarvester.cpp
--- a/Engine/Source/Runtime/CoreUObject/Private/UObject/SavePackage/PackageHarvester.cpp
+++ b/Engine/Source/Runtime/CoreUObject/Private/UObject/SavePackage/PackageHarvester.cpp
@@ -71,14 +71,12 @@ bool ConditionallyExcludeObjectForTarget(FSaveContext& SaveContext, UObject* Obj
 		return false;
 	}
 
-	//@note: currently always query excluded content against the default/main save harvesting context
-	// this is because generated excluded marks are not aware of the different harvesting context yet
 	bool bExcluded = false;
-	if (SaveContext.IsExcluded(Obj))
+	if (SaveContext.GetHarvestedRealm(HarvestingContext).IsExcluded(Obj))
 	{
 		return true;
 	}
-	else if (!SaveContext.IsIncluded(Obj, HarvestingContext))
+	else if (!SaveContext.GetHarvestedRealm(HarvestingContext).IsIncluded(Obj))
 	{
 		const EObjectMark ExcludedObjectMarks = SaveContext.GetExcludedObjectMarks();
 		const ITargetPlatform* TargetPlatform = SaveContext.GetTargetPlatform();
@@ -126,7 +124,7 @@ bool ConditionallyExcludeObjectForTarget(FSaveContext& SaveContext, UObject* Obj
 		}
 		if (bExcluded)
 		{
-			SaveContext.AddExcluded(Obj);
+			SaveContext.GetHarvestedRealm(HarvestingContext).AddExcluded(Obj);
 		}
 	}
 	return bExcluded;
@@ -153,18 +151,26 @@ bool DoesObjectNeedLoadForEditorGame(UObject* InObject)
 
 FPackageHarvester::FExportScope::FExportScope(FPackageHarvester& InHarvester, const FExportWithContext& InToProcess, bool bIsEditorOnlyObject)
 	: Harvester(InHarvester)
-	, PreviousContext(InHarvester.CurrentExportHarvestingRealm)
+	, PreviousRealm(InHarvester.CurrentExportHarvestingRealm)
+	, bPreviousFilterEditorOnly(InHarvester.IsFilterEditorOnly())
 {
 	check(Harvester.CurrentExportDependencies.CurrentExport == nullptr);
 	Harvester.CurrentExportDependencies = { InToProcess.Export };
 	Harvester.CurrentExportHarvestingRealm = InToProcess.HarvestedFromRealm;
 	Harvester.bIsEditorOnlyExportOnStack = bIsEditorOnlyObject;
+
+	// if we are auto generating optional package, then do no filter editor properties for that harvest
+	if (Harvester.SaveContext.IsSaveAutoOptional() && InToProcess.HarvestedFromRealm == ESaveRealm::Optional)
+	{
+		Harvester.SetFilterEditorOnly(false);
+	}
 }
 
 FPackageHarvester::FExportScope::~FExportScope()
 {
 	Harvester.AppendCurrentExportDependencies();
-	Harvester.CurrentExportHarvestingRealm = PreviousContext;
+	Harvester.CurrentExportHarvestingRealm = PreviousRealm;
+	Harvester.SetFilterEditorOnly(bPreviousFilterEditorOnly);
 }
 
 FPackageHarvester::FPackageHarvester(FSaveContext& InContext)
@@ -329,7 +335,7 @@ void FPackageHarvester::TryHarvestExport(UObject* InObject)
 		// If we have a illegal ref reason, record it
 		if (Reason != EIllegalRefReason::None)
 		{
-			SaveContext.RecordIllegalReference(CurrentExportDependencies.CurrentExport, InObject, EIllegalRefReason::ReferenceToOptional);
+			SaveContext.RecordIllegalReference(CurrentExportDependencies.CurrentExport, InObject, Reason);
 		}
 	}
 }
@@ -457,7 +463,7 @@ void FPackageHarvester::MarkSearchableName(const UObject* TypeObject, const FNam
 FArchive& FPackageHarvester::operator<<(UObject*& Obj)
 {	
 	// if the object is null or already marked excluded, we can skip the harvest
-	if (!Obj || SaveContext.IsExcluded(Obj))
+	if (!Obj || SaveContext.GetHarvestedRealm(CurrentExportHarvestingRealm).IsExcluded(Obj))
 	{
 		return *this;
 	}
@@ -492,7 +498,7 @@ FArchive& FPackageHarvester::operator<<(UObject*& Obj)
 		return bIsNative;
 	};
 
-	if (SaveContext.IsIncluded(Obj, CurrentExportHarvestingRealm))
+	if (SaveContext.GetHarvestedRealm(CurrentExportHarvestingRealm).IsIncluded(Obj))
 	{
 		HarvestDependency(Obj, IsObjNative(Obj));
 	}
@@ -544,7 +550,8 @@ FArchive& FPackageHarvester::operator<<(FSoftObjectPath& Value)
 			HarvestPackageHeaderName(PackageName);
 			SaveContext.GetHarvestedRealm(CurrentExportHarvestingRealm).GetSoftPackageReferenceList().Add(PackageName);
 #if WITH_EDITORONLY_DATA
-			if (CollectType != ESoftObjectPathCollectType::EditorOnlyCollect && !bIsEditorOnlyExportOnStack)
+			if (CollectType != ESoftObjectPathCollectType::EditorOnlyCollect && !bIsEditorOnlyExportOnStack 
+				&& CurrentExportHarvestingRealm != ESaveRealm::Optional)
 #endif
 			{
 				SaveContext.GetHarvestedRealm(ESaveRealm::Game).GetSoftPackageReferenceList().Add(PackageName);
diff --git a/Engine/Source/Runtime/CoreUObject/Private/UObject/SavePackage/PackageHarvester.h b/Engine/Source/Runtime/CoreUObject/Private/UObject/SavePackage/PackageHarvester.h
--- a/Engine/Source/Runtime/CoreUObject/Private/UObject/SavePackage/PackageHarvester.h
+++ b/Engine/Source/Runtime/CoreUObject/Private/UObject/SavePackage/PackageHarvester.h
@@ -42,7 +42,8 @@ public:
 		~FExportScope();
 	private:
 		FPackageHarvester& Harvester;
-		ESaveRealm PreviousContext;
+		ESaveRealm PreviousRealm;
+		bool bPreviousFilterEditorOnly;
 	};
 
 	class FIgnoreDependenciesScope
diff --git a/Engine/Source/Runtime/CoreUObject/Private/UObject/SavePackage/SaveContext.cpp b/Engine/Source/Runtime/CoreUObject/Private/UObject/SavePackage/SaveContext.cpp
--- a/Engine/Source/Runtime/CoreUObject/Private/UObject/SavePackage/SaveContext.cpp
+++ b/Engine/Source/Runtime/CoreUObject/Private/UObject/SavePackage/SaveContext.cpp
@@ -2,7 +2,9 @@
 
 #include "SaveContext.h"
 
+#include "Misc/ConfigCacheIni.h"
 #include "Serialization/PackageWriter.h"
+#include "UObject/UObjectGlobals.h"
 
 
 
@@ -125,6 +127,44 @@ FSaveContext::ESaveableStatus FSaveContext::GetSaveableStatusNoOuter(UObject* Ob
 	return ESaveableStatus::Success;
 }
 
+namespace
+{
+	TArray<UClass*> AutomaticOptionalInclusionAssetTypeList;
+}
+
+void FSaveContext::SetupHarvestingRealms()
+{
+	// Create the different harvesting realms
+	HarvestedRealms.AddDefaulted((uint32)ESaveRealm::RealmCount);
+
+	// if cooking the default harvesting context is Game, otherwise it's the editor context
+	CurrentHarvestingRealm = IsCooking() ? ESaveRealm::Game : ESaveRealm::Editor;
+
+	// Generate the automatic optional context inclusion asset list
+	static bool bAssetListGenerated = [](TArray<UClass*>& OutAssetList)
+	{
+		TArray<FString> AssetList;
+		GConfig->GetArray(TEXT("CookSettings"), TEXT("AutomaticOptionalInclusionAssetType"), AssetList, GEditorIni);
+		for (const FString& AssetType : AssetList)
+		{
+			if (UClass* AssetClass = FindObject<UClass>(nullptr, *AssetType, true))
+			{
+				OutAssetList.Add(AssetClass);
+			}
+			else
+			{
+				UE_LOG(LogSavePackage, Warning, TEXT("The asset type '%s' was not found while building the allowlist for automatic optional data inclusion list."), *AssetType);
+			}
+		}
+		return true;
+	}(AutomaticOptionalInclusionAssetTypeList);
+
+	if (bAssetListGenerated && Asset)
+	{
+		bIsSaveAutoOptional = IsCooking() && IsSaveOptional() && AutomaticOptionalInclusionAssetTypeList.Contains(Asset->GetClass());
+	}
+}
+
 const TCHAR* LexToString(FSaveContext::ESaveableStatus Status)
 {
 	static_assert(static_cast<int32>(FSaveContext::ESaveableStatus::__Count) == 7);
diff --git a/Engine/Source/Runtime/CoreUObject/Private/UObject/SavePackage/SaveContext.h b/Engine/Source/Runtime/CoreUObject/Private/UObject/SavePackage/SaveContext.h
--- a/Engine/Source/Runtime/CoreUObject/Private/UObject/SavePackage/SaveContext.h
+++ b/Engine/Source/Runtime/CoreUObject/Private/UObject/SavePackage/SaveContext.h
@@ -419,7 +419,7 @@ public:
 		ObjectSaveContext.Set(InPackage, GetTargetPlatform(), TargetPackagePath, SaveArgs.SaveFlags);
 
 		// Setup the harvesting flags and generate the context for harvesting the package
-		CreateHarvestingRealms();
+		SetupHarvestingRealms();
 	} 
 
 	~FSaveContext()
@@ -467,7 +467,8 @@ public:
 
 	EObjectMark GetExcludedObjectMarks() const
 	{
-		return ExcludedObjectMarks;
+		// When considering excluded objects for a platform, do not consider editor only object in the optional context as excluded
+		return (CurrentHarvestingRealm == ESaveRealm::Optional) ? (EObjectMark)(ExcludedObjectMarks & ~EObjectMark::OBJECTMARK_EditorOnly) : ExcludedObjectMarks;
 	}
 
 	EObjectFlags GetTopLevelFlags() const
@@ -575,6 +576,11 @@ public:
 		return !!(SaveArgs.SaveFlags & SAVE_Optional);
 	}
 
+	bool IsSaveAutoOptional() const
+	{
+		return bIsSaveAutoOptional;
+	}
+
 	bool IsConcurrent() const
 	{
 		return !!(SaveArgs.SaveFlags & SAVE_Concurrent);
@@ -665,6 +671,7 @@ public:
 		return CurrentHarvestingRealm;
 	}
 
+	/** Returns which save context should be saved. */
 	TArray<ESaveRealm> GetHarvestedRealmsToSave();
 
 	void MarkUnsaveable(UObject* InObject);
@@ -719,9 +726,9 @@ public:
 		return GetHarvestedRealm().IsExport(InObject);
 	}
 
-	bool IsIncluded(UObject* InObject, ESaveRealm InContext = ESaveRealm::None) const
+	bool IsIncluded(UObject* InObject) const
 	{
-		return GetHarvestedRealm(InContext).IsIncluded(InObject);
+		return GetHarvestedRealm().IsIncluded(InObject);
 	}
 
 	bool IsExcluded(UObject* InObject) const
@@ -974,15 +981,8 @@ public:
 	FSavePackageOutputFileArray AdditionalPackageFiles;
 private:
 
-	// Create the needed harvesting context depending on the save context options
-	void CreateHarvestingRealms()
-	{
-		// Create the different harvesting realms
-		HarvestedRealms.AddDefaulted((uint32)ESaveRealm::RealmCount);
-	
-		// if cooking the default harvesting context is Game, otherwise it's the editor context
-		CurrentHarvestingRealm = IsCooking() ? ESaveRealm::Game : ESaveRealm::Editor;
-	}
+	// Create the harvesting contexts and automatic optional context gathering options
+	void SetupHarvestingRealms();
 		
 	friend class FPackageHarvester;
 
@@ -1005,6 +1005,7 @@ private:
 	bool bNeedPreSaveCleanup = false;
 	bool bGenerateFileStub = false;
 	bool bIgnoreHeaderDiffs = false;
+	bool bIsSaveAutoOptional = false;
 
 	// Config classes shared with the old Save
 	FCanSkipEditorReferencedPackagesWhenCooking SkipEditorRefCookingSetting;
diff --git a/Engine/Source/Runtime/CoreUObject/Private/UObject/SavePackage2.cpp b/Engine/Source/Runtime/CoreUObject/Private/UObject/SavePackage2.cpp
--- a/Engine/Source/Runtime/CoreUObject/Private/UObject/SavePackage2.cpp
+++ b/Engine/Source/Runtime/CoreUObject/Private/UObject/SavePackage2.cpp
@@ -310,36 +310,45 @@ ESavePackageResult HarvestPackage(FSaveContext& SaveContext)
 	EObjectFlags TopLevelFlags = SaveContext.GetTopLevelFlags();
 	UObject* Asset = SaveContext.GetAsset();
 
-	// if no top level flags are passed, process just the provided package asset
-	if (TopLevelFlags == RF_NoFlags)
+	auto TryHarvestRootObject = [&Harvester, &SaveContext](UObject* InRoot)
 	{
-		Harvester.TryHarvestExport(Asset);
-		while (FPackageHarvester::FExportWithContext ExportContext = Harvester.PopExportToProcess())
+		Harvester.TryHarvestExport(InRoot);
+		// if we are automatically generating an optional package output, re-harvest objects with that realm as the default
+		if (SaveContext.IsSaveAutoOptional())
 		{
-			Harvester.ProcessExport(ExportContext);
+			FSaveContext::FSetSaveRealmToSaveScope(SaveContext, ESaveRealm::Optional);
+			Harvester.TryHarvestExport(InRoot);
 		}
+	};
+
+	// if no top level flags are passed, only use the provided package asset as root
+	if (TopLevelFlags == RF_NoFlags)
+	{
+		TryHarvestRootObject(Asset);
 	}
-	// Otherwise process all objects which have the relevant flags
+	// Otherwise use all objects which have the relevant flags
 	else
 	{
-		ForEachObjectWithPackage(SaveContext.GetPackage(), [&Harvester, TopLevelFlags](UObject* InObject)
+		ForEachObjectWithPackage(SaveContext.GetPackage(), [&TryHarvestRootObject, TopLevelFlags](UObject* InObject)
 			{
 				if (InObject->HasAnyFlags(TopLevelFlags))
 				{
-					Harvester.TryHarvestExport(InObject);
+					TryHarvestRootObject(InObject);
 				}
 				return true;
 			}, true/*bIncludeNestedObjects */, RF_Transient);
-		while (FPackageHarvester::FExportWithContext ExportContext = Harvester.PopExportToProcess())
-		{
-			Harvester.ProcessExport(ExportContext);
-		}
+	}
+	// Now process harvested roots
+	while (FPackageHarvester::FExportWithContext ExportContext = Harvester.PopExportToProcess())
+	{
+		Harvester.ProcessExport(ExportContext);
 	}
 
 	// If we have a valid optional context and we are saving it,
 	// transform any harvested non optional export into imports
 	// Mark other optional import package as well
-	if (SaveContext.IsSaveOptional() &&
+	if (!SaveContext.IsSaveAutoOptional() &&
+		SaveContext.IsSaveOptional() &&
 		SaveContext.IsCooking() &&
 		SaveContext.GetHarvestedRealm(ESaveRealm::Optional).GetExports().Num())
 	{
@@ -969,7 +978,15 @@ ESavePackageResult BuildLinker(FSaveContext& SaveContext)
 		SaveContext.GetLinker()->Summary.SetCustomVersionContainer(SaveContext.GetCustomVersions());
 
 		SaveContext.GetLinker()->SetPortFlags(SaveContext.GetPortFlags());
-		SaveContext.GetLinker()->SetFilterEditorOnly(SaveContext.IsFilterEditorOnly());
+		if (SaveContext.IsSaveAutoOptional() && SaveContext.GetCurrentHarvestingRealm() == ESaveRealm::Optional)
+		{
+			// Do not filter editor only data when automatically creating optional data using full uncooked objects
+			SaveContext.GetLinker()->SetFilterEditorOnly(false);
+		}
+		else
+		{
+			SaveContext.GetLinker()->SetFilterEditorOnly(SaveContext.IsFilterEditorOnly());
+		}
 		SaveContext.GetLinker()->SetCookData(SaveContext.GetCookData());
 
 		bool bUseUnversionedProperties = SaveContext.IsSaveUnversionedProperties();
@@ -1867,7 +1884,7 @@ ESavePackageResult WriteAdditionalExportFiles(FSaveContext& SaveContext)
 
 	if (SaveContext.IsCooking() && SaveContext.AdditionalFilesFromExports.Num() > 0)
 	{
-		checkf(SaveContext.GetCurrentHarvestingRealm() != ESaveRealm::Optional , TEXT("Addtional export files is currently unsupported with optional package multi output"));
+		checkf(SaveContext.GetCurrentHarvestingRealm() != ESaveRealm::Optional , TEXT("Addtional export files is currently unsupported with optional package multi output, Package %s"), *SaveContext.GetPackage()->GetName());
 		IPackageWriter* PackageWriter = SavePackageContext ? SavePackageContext->PackageWriter : nullptr;
 		checkf(PackageWriter, TEXT("Cooking requires a PackageWriter"));
 		for (FLargeMemoryWriter& Writer : SaveContext.AdditionalFilesFromExports)
@@ -1954,6 +1971,15 @@ ESavePackageResult UpdatePackageHeader(FStructuredArchive::FRecord& StructuredAr
 		{
 			PackageFlags |= PKG_ContainsNoAsset;
 		}
+		// Take the Linker FilterEditorOnlyData setting over the package flags to set this flag in the summary
+		if (Linker->IsFilterEditorOnly())
+		{
+			PackageFlags |= PKG_FilterEditorOnly;
+		}
+		else
+		{
+			PackageFlags &= ~PKG_FilterEditorOnly;
+		}
 		Linker->Summary.SetPackageFlags(PackageFlags);
 
 		// @todo: custom versions: when can this be checked?
diff --git a/Engine/Source/Runtime/CoreUObject/Public/UObject/LinkerSave.h b/Engine/Source/Runtime/CoreUObject/Public/UObject/LinkerSave.h
--- a/Engine/Source/Runtime/CoreUObject/Public/UObject/LinkerSave.h
+++ b/Engine/Source/Runtime/CoreUObject/Public/UObject/LinkerSave.h
@@ -158,6 +158,11 @@ public:
 	 * This sets it on itself, the summary, the actual Saver Archive if any and set the proper associated flag on the LinkerRoot
 	 */
 	virtual void SetUseUnversionedPropertySerialization(bool bInUseUnversioned) override;
+	/**
+	 * Sets whether we should be filtering editor only.
+	 * This sets it on itself, the summary, the actual Saver Archive if any and set the proper associated flag on the LinkerRoot
+	 */
+	virtual void SetFilterEditorOnly(bool bInFilterEditorOnly) override;
 
 
 #if WITH_EDITOR
