Fix root auto optional scope not being properly set during package harvest

diff --git a/Engine/Source/Runtime/CoreUObject/Private/UObject/SavePackage2.cpp b/Engine/Source/Runtime/CoreUObject/Private/UObject/SavePackage2.cpp
--- a/Engine/Source/Runtime/CoreUObject/Private/UObject/SavePackage2.cpp
+++ b/Engine/Source/Runtime/CoreUObject/Private/UObject/SavePackage2.cpp
@@ -316,7 +316,8 @@ ESavePackageResult HarvestPackage(FSaveContext& SaveContext)
 		// if we are automatically generating an optional package output, re-harvest objects with that realm as the default
 		if (SaveContext.IsSaveAutoOptional())
 		{
-			FSaveContext::FSetSaveRealmToSaveScope(SaveContext, ESaveRealm::Optional);
+			//@todo: FH: the top level might have to be altered to use the editor one when auto harvesting optional data
+			FSaveContext::FSetSaveRealmToSaveScope RealmScope(SaveContext, ESaveRealm::Optional);
 			Harvester.TryHarvestExport(InRoot);
 		}
 	};
