ANY_PACKAGE Warnings: In TryConvertShortTypeNameToPathName suppress the ensure if AmbiguousMessageVerbosity specifies logging at higher verbosity than error.
#rnx

diff --git a/Engine/Source/Runtime/CoreUObject/Private/UObject/Class.cpp b/Engine/Source/Runtime/CoreUObject/Private/UObject/Class.cpp
--- a/Engine/Source/Runtime/CoreUObject/Private/UObject/Class.cpp
+++ b/Engine/Source/Runtime/CoreUObject/Private/UObject/Class.cpp
@@ -6061,7 +6061,12 @@ FTopLevelAssetPath UClass::TryConvertShortTypeNameToPathName(UClass* TypeClass,
 	else
 	{
 		FTopLevelAssetPath Result;
-		UField* FoundType = (UField*)StaticFindFirstObject(TypeClass, *InShortTypeName, EFindFirstObjectOptions::EnsureIfAmbiguous | EFindFirstObjectOptions::NativeFirst, AmbiguousMessageVerbosity, AmbiguousClassMessage);
+		EFindFirstObjectOptions Options = EFindFirstObjectOptions::NativeFirst;
+		if (AmbiguousMessageVerbosity != ELogVerbosity::NoLogging && AmbiguousMessageVerbosity <= ELogVerbosity::Error)
+		{
+			Options |= EFindFirstObjectOptions::EnsureIfAmbiguous;
+		}
+		UField* FoundType = (UField*)StaticFindFirstObject(TypeClass, *InShortTypeName, Options, AmbiguousMessageVerbosity, AmbiguousClassMessage);
 		if (FoundType)
 		{
 			// UField does not define a GetFieldPathName so do what GetStructPathName does (this is the only place that needs it)
