AsyncLoading2: Make sure to call PostLoadInstances and CreateClusters before ClearConstructedObjects, because after ClearConstructedObjects unreferenced private objects in the package may be garbage collected.
Fix the underlying problem for the workaround in 19840451.

#rnx

diff --git a/Engine/Source/Runtime/CoreUObject/Private/Serialization/AsyncLoading2.cpp b/Engine/Source/Runtime/CoreUObject/Private/Serialization/AsyncLoading2.cpp
--- a/Engine/Source/Runtime/CoreUObject/Private/Serialization/AsyncLoading2.cpp
+++ b/Engine/Source/Runtime/CoreUObject/Private/Serialization/AsyncLoading2.cpp
@@ -2057,8 +2057,6 @@ private:
 	bool						bLoadHasFailed = false;
 	/** True if this package was created by this async package */
 	bool						bCreatedLinkerRoot = false;
-	/** True if package contains objects that can be put into GC clusters */
-	bool						bHasClusterObjects = false;
 	/** List of all request handles */
 	TArray<int32, TInlineAllocator<2>> RequestIDs;
 	/** List of ConstructedObjects = Exports + UPackage + ObjectsCreatedFromExports */
@@ -5171,54 +5169,92 @@ EAsyncPackageState::Type FAsyncLoadingThread2::ProcessLoadedPackagesFromGameThre
 			SCOPED_LOADTIMER(ProcessLoadedPackagesTime);
 			FAsyncPackage2* Package = LoadedPackagesToProcess[PackageIndex];
 			UE_ASYNC_PACKAGE_DEBUG(Package->Desc);
-			check(Package->AsyncPackageLoadingState == EAsyncPackageLoadingState2::Finalize);
+			check(Package->AsyncPackageLoadingState >= EAsyncPackageLoadingState2::Finalize &&
+				  Package->AsyncPackageLoadingState <= EAsyncPackageLoadingState2::CreateClusters);
 
-			TArray<UObject*> CDODefaultSubobjects;
-			// Clear async loading flags (we still want RF_Async, but EInternalObjectFlags::AsyncLoading can be cleared)
-			for (const FExportObject& Export : Package->Data.Exports)
+			if (Package->AsyncPackageLoadingState == EAsyncPackageLoadingState2::Finalize)
 			{
-				if (Export.bFiltered | Export.bExportLoadFailed)
+				TArray<UObject*> CDODefaultSubobjects;
+				// Clear async loading flags (we still want RF_Async, but EInternalObjectFlags::AsyncLoading can be cleared)
+				for (const FExportObject& Export : Package->Data.Exports)
 				{
-					continue;
-				}
+					if (Export.bFiltered | Export.bExportLoadFailed)
+					{
+						continue;
+					}
 
-				UObject* Object = Export.Object;
+					UObject* Object = Export.Object;
 
-				// CDO need special handling, no matter if it's listed in DeferredFinalizeObjects
-				UObject* CDOToHandle = ((Object != nullptr) && Object->HasAnyFlags(RF_ClassDefaultObject)) ? Object : nullptr;
+					// CDO need special handling, no matter if it's listed in DeferredFinalizeObjects
+					UObject* CDOToHandle = ((Object != nullptr) && Object->HasAnyFlags(RF_ClassDefaultObject)) ? Object : nullptr;
 
-				// Clear AsyncLoading in CDO's subobjects.
-				if (CDOToHandle != nullptr)
-				{
-					CDOToHandle->GetDefaultSubobjects(CDODefaultSubobjects);
-					for (UObject* SubObject : CDODefaultSubobjects)
+					// Clear AsyncLoading in CDO's subobjects.
+					if (CDOToHandle != nullptr)
 					{
-						if (SubObject && SubObject->HasAnyInternalFlags(EInternalObjectFlags::AsyncLoading))
+						CDOToHandle->GetDefaultSubobjects(CDODefaultSubobjects);
+						for (UObject* SubObject : CDODefaultSubobjects)
 						{
-							SubObject->AtomicallyClearInternalFlags(EInternalObjectFlags::AsyncLoading);
+							if (SubObject && SubObject->HasAnyInternalFlags(EInternalObjectFlags::AsyncLoading))
+							{
+								SubObject->AtomicallyClearInternalFlags(EInternalObjectFlags::AsyncLoading);
+							}
 						}
+						CDODefaultSubobjects.Reset();
 					}
-					CDODefaultSubobjects.Reset();
 				}
+
+				Package->AsyncPackageLoadingState = EAsyncPackageLoadingState2::PostLoadInstances;
 			}
 
-			Package->FinishUPackage();
+			if (Package->AsyncPackageLoadingState == EAsyncPackageLoadingState2::PostLoadInstances)
+			{
+				SCOPE_CYCLE_COUNTER(STAT_FAsyncPackage_PostLoadInstancesGameThread);
+				if (Package->PostLoadInstances(ThreadState) == EAsyncPackageState::Complete)
+				{
+					Package->AsyncPackageLoadingState = EAsyncPackageLoadingState2::CreateClusters;
+				}
+				else
+				{
+					// PostLoadInstances timed out
+					Result = EAsyncPackageState::TimeOut;
+				}
+			}
 
-			if (!Package->bLoadHasFailed && CanCreateObjectClusters())
+
+			if (Package->AsyncPackageLoadingState == EAsyncPackageLoadingState2::CreateClusters)
 			{
-				for (const FExportObject& Export : Package->Data.Exports)
+				SCOPE_CYCLE_COUNTER(STAT_FAsyncPackage_CreateClustersGameThread);
+				if (Package->bLoadHasFailed || !CanCreateObjectClusters())
 				{
-					if (!(Export.bFiltered | Export.bExportLoadFailed) && Export.Object->CanBeClusterRoot())
-					{
-						Package->bHasClusterObjects = true;
-						break;
-					}
+					Package->AsyncPackageLoadingState = EAsyncPackageLoadingState2::Complete;
+				}
+				else if (Package->CreateClusters(ThreadState) == EAsyncPackageState::Complete)
+				{
+					// All clusters created, it's safe to delete the package
+					Package->AsyncPackageLoadingState = EAsyncPackageLoadingState2::Complete;
+				}
+				else
+				{
+					// Cluster creation timed out
+					Result = EAsyncPackageState::TimeOut;
 				}
 			}
 
 			FSoftObjectPath::InvalidateTag();
 			FUniqueObjectGuid::InvalidateTag();
 
+			// push stats so that we don't overflow number of tags per thread during blocking loading
+			LLM_PUSH_STATS_FOR_ASSET_TAGS();
+
+			if (Result == EAsyncPackageState::TimeOut)
+			{
+				break;
+			}
+
+			check(Package->AsyncPackageLoadingState == EAsyncPackageLoadingState2::Complete);
+
+			Package->FinishUPackage();
+
 			{
 				FScopeLock LockAsyncPackages(&AsyncPackagesCritical);
 				AsyncPackageLookup.Remove(Package->Desc.UPackageId);
@@ -5241,8 +5277,6 @@ EAsyncPackageState::Type FAsyncLoadingThread2::ProcessLoadedPackagesFromGameThre
 
 			TRACE_LOADTIME_END_LOAD_ASYNC_PACKAGE(Package);
 
-			check(Package->AsyncPackageLoadingState == EAsyncPackageLoadingState2::Finalize);
-			Package->AsyncPackageLoadingState = EAsyncPackageLoadingState2::PostLoadInstances;
 			PackagesReadyForCallback.Add(Package);
 		}
 
@@ -5296,68 +5330,14 @@ EAsyncPackageState::Type FAsyncLoadingThread2::ProcessLoadedPackagesFromGameThre
 		for (int32 PackageIndex = 0; PackageIndex < CompletedPackages.Num(); ++PackageIndex)
 		{
 			FAsyncPackage2* Package = CompletedPackages[PackageIndex];
-			{
-				bool bSafeToDelete = false;
-				if (Package->AsyncPackageLoadingState == EAsyncPackageLoadingState2::PostLoadInstances)
-				{
-					SCOPE_CYCLE_COUNTER(STAT_FAsyncPackage_PostLoadInstancesGameThread);
-					// This package will create GC clusters but first check if all dependencies of this package have been fully loaded
-					if (Package->PostLoadInstances(ThreadState) == EAsyncPackageState::Complete)
-					{
-						// PostLoadInstances routed to all objects, safe to proceed to creating clusters or directly to completion
-						if (Package->bHasClusterObjects)
-						{
-							Package->AsyncPackageLoadingState = EAsyncPackageLoadingState2::CreateClusters;
-						}
-						else
-						{
-							Package->AsyncPackageLoadingState = EAsyncPackageLoadingState2::Complete;
-						}
-					}
-					else
-					{
-						// PostLoadInstances timed out
-						Result = EAsyncPackageState::TimeOut;
-						break;
-					}
-				}
-				else if (Package->AsyncPackageLoadingState == EAsyncPackageLoadingState2::CreateClusters)
-				{
-					SCOPE_CYCLE_COUNTER(STAT_FAsyncPackage_CreateClustersGameThread);
-					// This package will create GC clusters but first check if all dependencies of this package have been fully loaded
-					if (Package->CreateClusters(ThreadState) == EAsyncPackageState::Complete)
-					{
-						// All clusters created, it's safe to delete the package
-						bSafeToDelete = true;
-						Package->AsyncPackageLoadingState = EAsyncPackageLoadingState2::Complete;
-					}
-					else
-					{
-						// Cluster creation timed out
-						Result = EAsyncPackageState::TimeOut;
-						break;
-					}
-				}
-				else
-				{
-					// No clusters to create so it's safe to delete
-					bSafeToDelete = true;
-				}
-
-				if (bSafeToDelete)
-				{
-					UE_ASYNC_PACKAGE_DEBUG(Package->Desc);
-					check(Package->AsyncPackageLoadingState == EAsyncPackageLoadingState2::Complete);
-					Package->AsyncPackageLoadingState = EAsyncPackageLoadingState2::DeferredDelete;
+			UE_ASYNC_PACKAGE_DEBUG(Package->Desc);
 
-					CompletedPackages.RemoveAtSwap(PackageIndex--);
-					Package->ClearImportedPackages();
-					Package->ReleaseRef();
-				}
-			}
+			check(Package->AsyncPackageLoadingState == EAsyncPackageLoadingState2::Complete);
+			Package->AsyncPackageLoadingState = EAsyncPackageLoadingState2::DeferredDelete;
 
-			// push stats so that we don't overflow number of tags per thread during blocking loading
-			LLM_PUSH_STATS_FOR_ASSET_TAGS();
+			CompletedPackages.RemoveAtSwap(PackageIndex--);
+			Package->ClearImportedPackages();
+			Package->ReleaseRef();
 		}
 		
 		if (!bLocalDidSomething)
@@ -6385,14 +6365,8 @@ EAsyncPackageState::Type FAsyncPackage2::PostLoadInstances(FAsyncLoadingThreadSt
 
 		if (!(Export.bFiltered | Export.bExportLoadFailed))
 		{
-			if (IsValid(Export.Object))
-			{
-				UClass* ObjClass = Export.Object->GetClass();
-				if (IsValid(ObjClass))
-				{
-					ObjClass->PostLoadInstance(Export.Object);
-				}
-			}
+			UClass* ObjClass = Export.Object->GetClass();
+			ObjClass->PostLoadInstance(Export.Object);
 		}
 	}
 	return PostLoadInstanceIndex == ExportCount ? EAsyncPackageState::Complete : EAsyncPackageState::TimeOut;
@@ -6401,13 +6375,17 @@ EAsyncPackageState::Type FAsyncPackage2::PostLoadInstances(FAsyncLoadingThreadSt
 EAsyncPackageState::Type FAsyncPackage2::CreateClusters(FAsyncLoadingThreadState2& ThreadState)
 {
 	const int32 ExportCount = Data.Exports.Num();
-	while (DeferredClusterIndex < ExportCount && !ThreadState.IsTimeLimitExceeded(TEXT("CreateClusters")))
+	while (DeferredClusterIndex < ExportCount)
 	{
 		const FExportObject& Export = Data.Exports[DeferredClusterIndex++];
 
 		if (!(Export.bFiltered | Export.bExportLoadFailed) && Export.Object->CanBeClusterRoot())
 		{
 			Export.Object->CreateCluster();
+			if (DeferredClusterIndex < ExportCount && ThreadState.IsTimeLimitExceeded(TEXT("CreateClusters")))
+			{
+				break;
+			}
 		}
 	}
 
