Use UE_ASSET_LOG in LinkerLoad

#rnx

diff --git a/Engine/Source/Runtime/CoreUObject/Private/UObject/LinkerLoad.cpp b/Engine/Source/Runtime/CoreUObject/Private/UObject/LinkerLoad.cpp
--- a/Engine/Source/Runtime/CoreUObject/Private/UObject/LinkerLoad.cpp
+++ b/Engine/Source/Runtime/CoreUObject/Private/UObject/LinkerLoad.cpp
@@ -639,11 +639,11 @@ FLinkerLoad* FLinkerLoad::CreateLinkerAsync(FUObjectSerializeContext* LoadContex
 	{
 		if (GEventDrivenLoaderEnabled)
 		{
-			UE_LOG(LogStreaming, Fatal, TEXT("FLinkerLoad::CreateLinkerAsync: Found existing linker for '%s'"), *Parent->GetName());
+			UE_ASSET_LOG(LogStreaming, Fatal, Parent, TEXT("FLinkerLoad::CreateLinkerAsync: Found existing linker"));
 		}
 		else
 		{
-			UE_LOG(LogStreaming, Log, TEXT("FLinkerLoad::CreateLinkerAsync: Found existing linker for '%s'"), *Parent->GetName());
+			UE_ASSET_LOG(LogStreaming, Log, Parent, TEXT("FLinkerLoad::CreateLinkerAsync: Found existing linker"));
 		}		
 	}
 
@@ -904,9 +904,9 @@ FLinkerLoad::FLinkerLoad(UPackage* InParent, const FPackagePath& InPackagePath,
 {
 	static_assert((ExportHashCount & (ExportHashCount - 1)) == 0, "ExportHashCount must be power of two");
 
-	if (InPackagePath.GetHeaderExtension() == EPackageExtension::Unspecified)
+	if (PackagePath.GetHeaderExtension() == EPackageExtension::Unspecified)
 	{
-		UE_LOG(LogPackageName, Error, TEXT("PackagePath %s is missing header extension when assigned to LinkerLoad"), *InPackagePath.GetDebugName());
+		UE_ASSET_LOG(LogPackageName, Error, PackagePath, TEXT("PackagePath is missing header extension when assigned to LinkerLoad"));
 	}
 
 	INC_DWORD_STAT(STAT_LinkerCount);
@@ -980,7 +980,7 @@ bool FLinkerLoad::IsTimeLimitExceeded( const TCHAR* CurrentTask, int32 Granulari
 			// Log single operations that take longer than timelimit.
 			if( (CurrentTime - TickStartTime) > (2.5 * TimeLimit) )
 			{
-				UE_LOG(LogStreaming, Log, TEXT("FLinkerLoad: %s took (less than) %5.2f ms"), 
+				UE_ASSET_LOG(LogStreaming, Log, PackagePath, TEXT("FLinkerLoad: %s took (less than) %5.2f ms"), 
 					CurrentTask, 
 					(CurrentTime - TickStartTime) * 1000);
 			}
@@ -1057,7 +1057,7 @@ FLinkerLoad::ELinkerStatus FLinkerLoad::CreateLoader(
 				bCanUseAsyncLoader = false;
 				if (!bRetryWithNormalArchive)
 				{
-					UE_LOG(LogLinker, Warning, TEXT("Error opening file '%s'."), *GetDebugName());
+					UE_ASSET_LOG(LogLinker, Warning, PackagePath, TEXT("Error opening file."));
 					return LINKER_Failed;
 				}
 			}
@@ -1083,7 +1083,7 @@ FLinkerLoad::ELinkerStatus FLinkerLoad::CreateLoader(
 				{
 					DestroyLoader();
 				}
-				UE_LOG(LogLinker, Warning, TEXT("Error opening file '%s'."), *GetDebugName());
+				UE_ASSET_LOG(LogLinker, Warning, PackagePath, TEXT("Error opening file."));
 				return LINKER_Failed;
 			}
 
@@ -1183,7 +1183,7 @@ FLinkerLoad::ELinkerStatus FLinkerLoad::CreateLoader(
 				if (Size <= 0)
 				{
 					DestroyLoader();
-					UE_LOG(LogLinker, Warning, TEXT("Error opening file '%s'."), *GetDebugName());
+					UE_ASSET_LOG(LogLinker, Warning, PackagePath, TEXT("Error opening file."));
 					return LINKER_Failed;
 				}
 				// Precache up to one ECC block before serializing package file summary.
@@ -1248,7 +1248,7 @@ FLinkerLoad::ELinkerStatus FLinkerLoad::SerializePackageFileSummaryInternal()
 
 			if (IsEngineContent)
 			{
-				UE_LOG(LogLinker, Warning, TEXT("The file %s is Engine content that was saved with a licensee flag. This can result in the file failing to open on licensee builds"), *GetDebugName());
+				UE_ASSET_LOG(LogLinker, Warning, PackagePath, TEXT("The file is Engine content that was saved with a licensee flag. This can result in the file failing to open on licensee builds"));
 			}
 		}
 	}
@@ -1291,7 +1291,10 @@ FLinkerLoad::ELinkerStatus FLinkerLoad::SerializePackageFileSummaryInternal()
 			}
 			FORCEINLINE operator bool() const { return bDoWarn; }
 		} ZeroEngineVersionWarningEnabled;
-		UE_CLOG(ZeroEngineVersionWarningEnabled, LogLinker, Warning, TEXT("Asset '%s' has been saved with empty engine version. The asset will be loaded but may be incompatible."), *GetDebugName());
+		if (ZeroEngineVersionWarningEnabled)
+		{
+			UE_ASSET_LOG(LogLinker, Warning, PackagePath, TEXT("Asset has been saved with empty engine version. The asset will be loaded but may be incompatible."));
+		}
 	}
 
 	// Don't load packages that were saved with package version newer than the current one.
@@ -1373,7 +1376,7 @@ FLinkerLoad::ELinkerStatus FLinkerLoad::SerializePackageFileSummary()
 	{
 		if (Loader->IsError())
 		{
-			UE_LOG(LogLinker, Warning, TEXT("The file '%s' contains unrecognizable data, check that it is of the expected type."), *GetDebugName());
+			UE_ASSET_LOG(LogLinker, Warning, PackagePath, TEXT("The file contains unrecognizable data, check that it is of the expected type."));
 			return LINKER_Failed;
 		}
 		if (bIsAsyncLoader)
@@ -1427,19 +1430,19 @@ FLinkerLoad::ELinkerStatus FLinkerLoad::UpdateFromPackageFileSummary()
 			{
 				// Loading a package with custom version that we don't know about!
 				// Temporarily just warn and continue. @todo: this needs to be fixed properly
-				UE_LOG(LogLinker, Warning, TEXT("Package %s was saved with a custom version that is not present. Tag %s  Version %d"), *GetDebugName(),
+				UE_ASSET_LOG(LogLinker, Warning, PackagePath, TEXT("Package was saved with a custom version that is not present. Tag %s  Version %d"),
 					*Diff.Version->Key.ToString(), Diff.Version->Version);
 			}
 			else if (Diff.Type == ECustomVersionDifference::Invalid)
 			{
-				UE_LOG(LogLinker, Error, TEXT("Package %s was saved with an invalid custom version. Tag %s  Version %d"), *GetDebugName(), *Diff.Version->Key.ToString(), Diff.Version->Version);
+				UE_ASSET_LOG(LogLinker, Error, PackagePath, TEXT("Package was saved with an invalid custom version. Tag %s  Version %d"), *Diff.Version->Key.ToString(), Diff.Version->Version);
 				return LINKER_Failed;
 			}
 			else if (Diff.Type == ECustomVersionDifference::Newer)
 			{
 				FCustomVersion LatestVersion = FCurrentCustomVersions::Get(Diff.Version->Key).GetValue();
 				// Loading a package with a newer custom version than the current one.
-				UE_LOG(LogLinker, Error, TEXT("Package %s was saved with a newer custom version than the current. Tag %s Name '%s' PackageVersion %d  MaxExpected %d"), *GetDebugName(),
+				UE_ASSET_LOG(LogLinker, Error, PackagePath, TEXT("Package was saved with a newer custom version than the current. Tag %s Name '%s' PackageVersion %d  MaxExpected %d"),
 					*Diff.Version->Key.ToString(), *LatestVersion.GetFriendlyName().ToString(), Diff.Version->Version, LatestVersion.Version);
 				return LINKER_Failed;
 			}
@@ -2695,10 +2698,10 @@ void FLinkerLoad::GatherImportDependencies(int32 ImportIndex, TSet<FDependencyRe
 			// don't print out for intrinsic native classes
 			if (!Import.XObject || !(Import.XObject->GetClass()->HasAnyClassFlags(CLASS_Intrinsic)))
 			{
-				UE_LOG(LogLinker, Warning, TEXT("VerifyImportInner failed [(%x, %d), (%x, %d)] for %s with linker: %s"), 
+				UE_ASSET_LOG(LogLinker, Warning, PackagePath, TEXT("VerifyImportInner failed [(%x, %d), (%x, %d)] for %s"), 
 					Import.XObject, Import.XObject ? (Import.XObject->IsNative() ? 1 : 0) : 0, 
 					Import.SourceLinker, Import.SourceIndex, 
-					*GetImportFullName(ImportIndex), *GetDebugName() );
+					*GetImportFullName(ImportIndex));
 			}
 			EndLoad(LoadContext);
 			return;
@@ -2708,7 +2711,7 @@ void FLinkerLoad::GatherImportDependencies(int32 ImportIndex, TSet<FDependencyRe
 		// only object we should create are one FLinkerLoad for source linker
 		if (GUObjectArray.GetObjectArrayNum() - NumObjectsBefore > 2)
 		{
-			UE_LOG(LogLinker, Warning, TEXT("Created %d objects checking %s"), GUObjectArray.GetObjectArrayNum() - NumObjectsBefore, *GetImportFullName(ImportIndex));
+			UE_ASSET_LOG(LogLinker, Warning, PackagePath, TEXT("Created %d objects checking %s"), GUObjectArray.GetObjectArrayNum() - NumObjectsBefore, *GetImportFullName(ImportIndex));
 		}
 #endif
 	}
@@ -2717,7 +2720,7 @@ void FLinkerLoad::GatherImportDependencies(int32 ImportIndex, TSet<FDependencyRe
 	FDependencyRef NewRef;
 	if (Import.XObject)
 	{
-		UE_LOG(LogLinker, Warning, TEXT("Using non-native XObject %s!!!"), *Import.XObject->GetFullName());
+		UE_ASSET_LOG(LogLinker, Warning, PackagePath, TEXT("Using non-native XObject %s!!!"), *Import.XObject->GetFullName());
 		NewRef.Linker = Import.XObject->GetLinker();
 		NewRef.ExportIndex = Import.XObject->GetLinkerIndex();
 	}
@@ -2867,7 +2870,7 @@ FLinkerLoad::EVerifyResult FLinkerLoad::VerifyImport(int32 ImportIndex)
 			// then do the throw here
 			if (bCrashOnFail)
 			{
-				UE_LOG(LogLinker, Fatal,  TEXT("Failed import: %s %s (file %s)"), *Import.ClassName.ToString(), *GetImportFullName(ImportIndex), *Import.SourceLinker->GetDebugName() );
+				UE_ASSET_LOG(LogLinker, Fatal, PackagePath, TEXT("Failed import: %s %s (file %s)"), *Import.ClassName.ToString(), *GetImportFullName(ImportIndex), *Import.SourceLinker->GetDebugName());
 				return Result;
 			}
 			// otherwise just printout warnings, and if in the editor, popup the EdLoadWarnings box
@@ -2914,10 +2917,9 @@ FLinkerLoad::EVerifyResult FLinkerLoad::VerifyImport(int32 ImportIndex)
 					// print warning about missing class
 					if (!FindClass)
 					{
-						UE_LOG(LogLinker, Warning, TEXT("Missing Class %s for '%s' referenced by package '%s'.  Classes should not be removed if referenced by content; mark the class 'deprecated' instead."),
+						UE_ASSET_LOG(LogLinker, Warning, PackagePath, TEXT("Missing Class %s for '%s'. Classes should not be removed if referenced by content; mark the class 'deprecated' instead."),
 							*OriginalImport.ClassName.ToString(),
-							*GetImportFullName(ImportIndex),
-							*LinkerRoot->GetName());
+							*GetImportFullName(ImportIndex));
 					}
 				}
 #endif // WITH_EDITOR
@@ -2993,12 +2995,11 @@ void StaticFindAllImportObjects(TArray<FObjectImport>& InOutImportMap, FPackageI
 					}
 					else
 					{
-						UE_LOG(LogLinker, Warning, TEXT("Failed to resolve import '%s' ('%d') in outer '%s' ('%d') in cooked package '%s'"),
+						UE_ASSET_LOG(LogLinker, Warning, Package, TEXT("Failed to resolve import '%s' ('%d') in outer '%s' ('%d') in cooked package"),
 							*ObjectImport.ObjectName.ToString(),
 							Inner.ToImport(),
 							*OuterObject->GetName(),
-							Outer.ToImport(),
-							*Package->GetName());
+							Outer.ToImport());
 					}
 				}
 			}
@@ -3055,7 +3056,7 @@ bool FLinkerLoad::VerifyImportInner(const int32 ImportIndex, FString& WarningSuf
 					if (!FLinkerLoad::IsKnownMissingPackage(PackageToLoad))
 					{
 						FLinkerLoad::AddKnownMissingPackage(PackageToLoad);
-						UE_LOG(LogLinker, Warning, TEXT("VerifyImport for '%s': Failed to find script package for import object '%s'"), *GetDebugName(), *GetImportFullName(ImportIndex));
+						UE_ASSET_LOG(LogLinker, Warning, PackagePath, TEXT("VerifyImport: Failed to find script package for import object '%s'"), *GetImportFullName(ImportIndex));
 					}
 					return nullptr;
 				}
@@ -3114,7 +3115,7 @@ bool FLinkerLoad::VerifyImportInner(const int32 ImportIndex, FString& WarningSuf
 			if (!FLinkerLoad::IsKnownMissingPackage(PackageToLoad))
 			{
 				FLinkerLoad::AddKnownMissingPackage(PackageToLoad);
-				UE_LOG(LogLinker, Warning, TEXT("VerifyImport: Failed to load package for import object '%s'"), *GetImportFullName(ImportIndex));
+				UE_ASSET_LOG(LogLinker, Warning, PackagePath, TEXT("VerifyImport: Failed to load package for import object '%s'"), *GetImportFullName(ImportIndex));
 			}
 			return nullptr;
 		}
@@ -3194,7 +3195,7 @@ bool FLinkerLoad::VerifyImportInner(const int32 ImportIndex, FString& WarningSuf
 	// Find or load the linker load that contains the FObjectExport for this import
 	if (Import.OuterIndex.IsNull() && Import.ClassName!=NAME_Package )
 	{
-		UE_LOG(LogLinker, Error, TEXT("%s has an inappropriate outermost, it was probably saved with a deprecated outer (file: %s)"), *ImportObjectName, *GetDebugName());
+		UE_ASSET_LOG(LogLinker, Error, PackagePath, TEXT("%s has an inappropriate outermost, it was probably saved with a deprecated outer"), *ImportObjectName);
 		Import.SourceLinker = NULL;
 		return false;
 	}
@@ -3276,13 +3277,13 @@ bool FLinkerLoad::VerifyImportInner(const int32 ImportIndex, FString& WarningSuf
 		{
 			if (!Import.SourceLinker->bHasFoundExistingExports)
 			{
-				UE_LOG(LogLinker, Log, TEXT("Source linker '%s' imported from '%s' has not processed all header information, ticking it now"), *GetNameSafe(Import.SourceLinker->LinkerRoot), *GetNameSafe(LinkerRoot));
+				UE_ASSET_LOG(LogLinker, Log, PackagePath, TEXT("Source linker '%s' has not processed all header information, ticking it now"), *GetNameSafe(Import.SourceLinker->LinkerRoot));
 				// This means that the source linker timed out during it's async loading tick and that the header information hasn't been fully processed yet.
 				// Make sure that the header information is available but don't process any imports or exports (LOAD_NoVerify).
 				TGuardValue<uint32> LinkerLoadFlagGuard(Import.SourceLinker->LoadFlags, Import.SourceLinker->LoadFlags | LOAD_NoVerify);
 				if (Import.SourceLinker->Tick(0.f, false, false, nullptr) == LINKER_Failed)
 				{
-					UE_LOG(LogLinker, Warning, TEXT("Failed ticking import source linker '%s'"), *GetNameSafe(Import.SourceLinker->LinkerRoot));
+					UE_ASSET_LOG(LogLinker, Warning, PackagePath, TEXT("Failed ticking import source linker '%s'"), *GetNameSafe(Import.SourceLinker->LinkerRoot));
 					return false;
 				}
 			}
@@ -3353,8 +3354,8 @@ bool FLinkerLoad::VerifyImportInner(const int32 ImportIndex, FString& WarningSuf
 									{
 										if ((SourceExportOuter.ClassName != OuterImport.ClassName) || (SourceExportOuter.ClassPackage != OuterImport.ClassPackage))
 										{
-											UE_LOG(LogLinker, Warning, TEXT("Resolved outer import with a different class: import class '%s.%s', package class '%s.%s'. Resave %s to fix."), 
-												*OuterImport.ClassPackage.ToString(), *OuterImport.ClassName.ToString(), *SourceExportOuter.ClassPackage.ToString(), *SourceExportOuter.ClassName.ToString(), *GetNameSafe(LinkerRoot));
+											UE_ASSET_LOG(LogLinker, Warning, PackagePath, TEXT("Resolved outer import with a different class: import class '%s.%s', package class '%s.%s'. Resave to fix."), 
+												*OuterImport.ClassPackage.ToString(), *OuterImport.ClassName.ToString(), *SourceExportOuter.ClassPackage.ToString(), *SourceExportOuter.ClassName.ToString());
 										}
 									}
 								}
@@ -3383,17 +3384,17 @@ bool FLinkerLoad::VerifyImportInner(const int32 ImportIndex, FString& WarningSuf
 								FObjectExport& Export = ExportMap[i];
 								if ( Export.SuperIndex == FoundIndex )
 								{
-									UE_LOG(LogLinker, Log, TEXT("Private import was referenced by export '%s' (parent)"), *Export.ObjectName.ToString());
+									UE_ASSET_LOG(LogLinker, Log, PackagePath, TEXT("Private import was referenced by export '%s' (parent)"), *Export.ObjectName.ToString());
 									SafeReplace = false;
 								}
 								else if ( Export.ClassIndex == FoundIndex )
 								{
-									UE_LOG(LogLinker, Log, TEXT("Private import was referenced by export '%s' (class)"), *Export.ObjectName.ToString());
+									UE_ASSET_LOG(LogLinker, Log, PackagePath, TEXT("Private import was referenced by export '%s' (class)"), *Export.ObjectName.ToString());
 									SafeReplace = false;
 								}
 								else if ( Export.OuterIndex == FoundIndex )
 								{
-									UE_LOG(LogLinker, Log, TEXT("Private import was referenced by export '%s' (outer)"), *Export.ObjectName.ToString());
+									UE_ASSET_LOG(LogLinker, Log, PackagePath, TEXT("Private import was referenced by export '%s' (outer)"), *Export.ObjectName.ToString());
 									SafeReplace = false;
 								}
 							}
@@ -3404,7 +3405,7 @@ bool FLinkerLoad::VerifyImportInner(const int32 ImportIndex, FString& WarningSuf
 									FObjectImport& TestImport = ImportMap[i];
 									if ( TestImport.OuterIndex == FoundIndex )
 									{
-										UE_LOG(LogLinker, Log, TEXT("Private import was referenced by import '%s' (outer)"), *ImportObjectName);
+										UE_ASSET_LOG(LogLinker, Log, PackagePath, TEXT("Private import was referenced by import '%s' (outer)"), *ImportObjectName);
 										SafeReplace = false;
 									}
 								}
@@ -3412,7 +3413,7 @@ bool FLinkerLoad::VerifyImportInner(const int32 ImportIndex, FString& WarningSuf
 
 							if ( !SafeReplace )
 							{
-								UE_LOG(LogLinker, Warning, TEXT("%s"), *FString::Printf( TEXT("Can't import private object %s %s"), *Import.ClassName.ToString(), *GetImportFullName(ImportIndex) ) );
+								UE_ASSET_LOG(LogLinker, Warning, PackagePath, TEXT("Can't import private object %s %s"), *Import.ClassName.ToString(), *GetImportFullName(ImportIndex));
 								return false;
 							}
 							else
@@ -3429,8 +3430,8 @@ bool FLinkerLoad::VerifyImportInner(const int32 ImportIndex, FString& WarningSuf
 						// Found the FObjectExport for this import
 						if ((Import.ClassName != Import.SourceLinker->GetExportClassName(j)) || (Import.ClassPackage != Import.SourceLinker->GetExportClassPackage(j)))
 						{
-							UE_LOG(LogLinker, Warning, TEXT("Resolved import with a different class: import class '%s.%s', package class '%s.%s'. Resave %s to fix."), 
-								*Import.ClassPackage.ToString(), *Import.ClassName.ToString(), *Import.SourceLinker->GetExportClassPackage(j).ToString(), *Import.SourceLinker->GetExportClassName(j).ToString(), *GetNameSafe(LinkerRoot));
+							UE_ASSET_LOG(LogLinker, Warning, PackagePath, TEXT("Resolved import with a different class: import class '%s.%s', package class '%s.%s'. Resave to fix."), 
+								*Import.ClassPackage.ToString(), *Import.ClassName.ToString(), *Import.SourceLinker->GetExportClassPackage(j).ToString(), *Import.SourceLinker->GetExportClassName(j).ToString());
 						}
 
 						Import.SourceIndex = j;
@@ -3838,8 +3839,8 @@ int32 FLinkerLoad::FindExportIndex( FName ClassName, FName ClassPackage, FName O
 			{
 				if ((ClassPackage != GetExportClassPackage(i)) || (ClassName != GetExportClassName(i)))
 				{
-					UE_LOG(LogLinker, Warning, TEXT("Resolved export with a different class: export class '%s.%s', package class '%s.%s'. Resave %s to fix."), 
-						 *GetExportClassPackage(i).ToString(), *GetExportClassName(i).ToString(), *ClassPackage.ToString(), *ClassName.ToString(), *GetNameSafe(LinkerRoot));
+					UE_ASSET_LOG(LogLinker, Warning, PackagePath, TEXT("Resolved export with a different class: export class '%s.%s', package class '%s.%s'. Resave to fix."), 
+						 *GetExportClassPackage(i).ToString(), *GetExportClassName(i).ToString(), *ClassPackage.ToString(), *ClassName.ToString());
 				}
 
 				return i;
@@ -3946,7 +3947,7 @@ UObject* FLinkerLoad::Create( UClass* ObjectClass, FName ObjectName, UObject* Ou
 	if (Index != INDEX_NONE)
 	{
 #if DEBUG_PRINT_NONQUALIFIED_RESULT
-		UE_LOG(LogLinker, Warning, TEXT("Using a non-qualified name (would have) found: %s"), *GetExportFullName(Index));
+		UE_ASSET_LOG(LogLinker, Warning, PackagePath, TEXT("Using a non-qualified name (would have) found: %s"), *GetExportFullName(Index));
 #endif
 #if FIND_OBJECT_NONQUALIFIED
 		return (InLoadFlags & LOAD_Verify) ? INVALID_OBJECT : CreateExport(Index);
@@ -3958,7 +3959,7 @@ UObject* FLinkerLoad::Create( UClass* ObjectClass, FName ObjectName, UObject* Ou
 	// if we are checking for failure cases, and we failed, throw an error
 	if( Checked )
 	{
-		UE_LOG(LogLinker, Warning, TEXT("%s"), *FString::Printf( TEXT("%s %s not found for creation"), *ObjectClass->GetName(), *ObjectName.ToString() ));
+		UE_ASSET_LOG(LogLinker, Warning, PackagePath, TEXT("%s %s not found for creation"), *ObjectClass->GetName(), *ObjectName.ToString());
 	}
 	return NULL;
 }
@@ -4271,11 +4272,11 @@ void FLinkerLoad::Preload( UObject* Object )
 				{
 					if (Object->GetClass()->HasAnyClassFlags(CLASS_Deprecated))
 					{
-						UE_LOG(LogLinker, Warning, TEXT("%s"), *FString::Printf( TEXT("%s: Serial size mismatch: Got %d, Expected %d"), *Object->GetFullName(), (int32)SizeSerialized, Export.SerialSize ) );
+						UE_ASSET_LOG(LogLinker, Warning, PackagePath, TEXT("%s: Serial size mismatch: Got %d, Expected %d"), *Object->GetFullName(), (int32)SizeSerialized, Export.SerialSize);
 					}
 					else
 					{
-						UE_LOG(LogLinker, Fatal, TEXT("%s"), *FString::Printf( TEXT("%s: Serial size mismatch: Got %d, Expected %d"), *Object->GetFullName(), (int32)SizeSerialized, Export.SerialSize ) );
+						UE_ASSET_LOG(LogLinker, Fatal, PackagePath, TEXT("%s: Serial size mismatch: Got %d, Expected %d"), *Object->GetFullName(), (int32)SizeSerialized, Export.SerialSize);
 					}
 				}
 
@@ -4451,7 +4452,7 @@ UObject* FLinkerLoad::CreateExport( int32 Index )
 
 			// otherwise, return NULL and let the calling code determine what to do
 			FString OuterName = Export.OuterIndex.IsNull() ? LinkerRoot->GetFullName() : GetFullImpExpName(Export.OuterIndex);
-			UE_LOG(LogLinker, Warning, TEXT("CreateExport: Failed to load Outer for resource because its class is a redirector '%s': %s"), *Export.ObjectName.ToString(), *OuterName);
+			UE_ASSET_LOG(LogLinker, Warning, PackagePath, TEXT("CreateExport: Failed to load Outer for resource because its class is a redirector '%s': %s"), *Export.ObjectName.ToString(), *OuterName);
 			return nullptr;
 		}
 
@@ -4477,7 +4478,7 @@ UObject* FLinkerLoad::CreateExport( int32 Index )
 						if (FuncClass && FuncClass->ClassGeneratedBy && !FuncClass->ClassGeneratedBy->HasAnyFlags(RF_BeingRegenerated))
 						{
 							// If this is a function (NOT being regenerated) whose parent has been removed, give it a NULL parent, as we would have in the script compiler.
-							UE_LOG(LogLinker, Display, TEXT("CreateExport: Failed to load Parent for %s; removing parent information, but keeping function"), *GetExportFullName(Index));
+							UE_ASSET_LOG(LogLinker, Display, PackagePath, TEXT("CreateExport: Failed to load Parent for %s; removing parent information, but keeping function"), *GetExportFullName(Index));
 						}
 					}
 #endif
@@ -4499,7 +4500,7 @@ UObject* FLinkerLoad::CreateExport( int32 Index )
 							UClass* StructClass = Cast<UClass>(ObjOuter);
 							if (StructClass && StructClass->ClassGeneratedBy && !StructClass->ClassGeneratedBy->HasAnyFlags(RF_BeingRegenerated))
 							{
-								UE_LOG(LogLinker, Display, TEXT("CreateExport: Failed to load Parent for %s; resaving the parents of %s will remove this message"), *GetExportFullName(Index), *StructClass->ClassGeneratedBy->GetFullName());
+								UE_ASSET_LOG(LogLinker, Display, PackagePath, TEXT("CreateExport: Failed to load Parent for %s; resaving the parents of %s will remove this message"), *GetExportFullName(Index), *StructClass->ClassGeneratedBy->GetFullName());
 								bFailedToLoadGeneratedStruct = true;
 							}
 						}
@@ -4507,7 +4508,7 @@ UObject* FLinkerLoad::CreateExport( int32 Index )
 					
 					if (!bFailedToLoadGeneratedStruct && !FLinkerLoad::IsKnownMissingPackage(*GetExportFullName(Index)))
 					{
-						UE_LOG(LogLinker, Warning, TEXT("CreateExport: Failed to load Parent for %s"), *GetExportFullName(Index));
+						UE_ASSET_LOG(LogLinker, Warning, PackagePath, TEXT("CreateExport: Failed to load Parent for %s"), *GetExportFullName(Index));
 					}
 #endif
 					return nullptr;
@@ -4539,7 +4540,7 @@ UObject* FLinkerLoad::CreateExport( int32 Index )
 		// we warn about it and remove the flag.
 		if( (Export.ObjectFlags & RF_MarkAsNative) != 0 && !LoadClass->IsChildOf(UField::StaticClass()) )
 		{
-			UE_LOG(LogLinker, Warning,TEXT("%s %s has RF_MarkAsNative set but is not a UField derived class"),*LoadClass->GetName(),*Export.ObjectName.ToString());
+			UE_ASSET_LOG(LogLinker, Warning,PackagePath, TEXT("%s %s has RF_MarkAsNative set but is not a UField derived class"),*LoadClass->GetName(),*Export.ObjectName.ToString());
 			// Remove RF_MarkAsNative;
 			Export.ObjectFlags = EObjectFlags(Export.ObjectFlags & ~RF_MarkAsNative);
 		}
@@ -4767,7 +4768,7 @@ UObject* FLinkerLoad::CreateExport( int32 Index )
 
 		if (ActualObjectWithTheName && !ActualObjectWithTheName->GetClass()->IsChildOf(LoadClass))
 		{
-			UE_LOG(LogLinker, Error, TEXT("Failed import: class '%s' name '%s' outer '%s'. There is another object (of '%s' class) at the path."),
+			UE_ASSET_LOG(LogLinker, Error, PackagePath, TEXT("Failed import: class '%s' name '%s' outer '%s'. There is another object (of '%s' class) at the path."),
 				*LoadClass->GetName(), *Export.ObjectName.ToString(), *ThisParent->GetName(), *ActualObjectWithTheName->GetClass()->GetName());
 			return NULL;
 		}
@@ -4920,7 +4921,7 @@ UObject* FLinkerLoad::CreateExport( int32 Index )
 		}
 		else
 		{
-			UE_LOG(LogLinker, Warning, TEXT("FLinker::CreatedExport failed to construct object %s %s"), *LoadClass->GetName(), *Export.ObjectName.ToString() );
+			UE_ASSET_LOG(LogLinker, Warning, PackagePath, TEXT("FLinker::CreatedExport failed to construct object %s %s"), *LoadClass->GetName(), *Export.ObjectName.ToString() );
 		}
 
 		if ( Export.Object != NULL )
@@ -5031,7 +5032,7 @@ UObject* FLinkerLoad::CreateImport( int32 Index )
 	if (Import.XObject != nullptr && Import.XObject->HasAnyInternalFlags(EInternalObjectFlags::Unreachable))
 	{
 		// This is just a safeguard to catch potential bugs that should have been fixed by calling UnhashUnreachableObjects in Async Loading code
-		UE_LOG(LogLinker, Warning, TEXT("Unreachable object found when creating import %s from linker %s"), *Import.XObject->GetFullName(), *GetArchiveName());
+		UE_ASSET_LOG(LogLinker, Warning, PackagePath, TEXT("Unreachable object found when creating import %s"), *Import.XObject->GetFullName());
 		Import.XObject = nullptr;
 	}
 
@@ -5097,7 +5098,7 @@ UObject* FLinkerLoad::CreateImport( int32 Index )
 						{
 							// This can happen when deleting native properties or restructing blueprints. If there is an actual problem it will be caught when trying to resolve the outer itself
 							FString OuterName = Import.OuterIndex.IsNull() ? LinkerRoot->GetFullName() : GetFullImpExpName(Import.OuterIndex);
-							UE_LOG(LogLinker, Verbose, TEXT("CreateImport: Failed to load Outer for resource '%s': %s"), *Import.ObjectName.ToString(), *OuterName);
+							UE_ASSET_LOG(LogLinker, Verbose, PackagePath, TEXT("CreateImport: Failed to load Outer for resource '%s': %s"), *Import.ObjectName.ToString(), *OuterName);
 							return NULL;
 						}
 	
@@ -5187,7 +5188,7 @@ UObject* FLinkerLoad::IndexToObject( FPackageIndex Index )
 					ErrorCaption,
 					GEngineIni);
 
-				UE_LOG( LogLinker, Error, TEXT("Invalid export object index=%d while reading %s. File is most likely corrupted. Please verify your installation."), Index.ToExport(), *GetDebugName() );
+				UE_ASSET_LOG(LogLinker, Error, PackagePath, TEXT("Invalid export object index=%d. File is most likely corrupted. Please verify your installation."), Index.ToExport());
 
 				if (GLog)
 				{
@@ -5222,7 +5223,7 @@ UObject* FLinkerLoad::IndexToObject( FPackageIndex Index )
 					ErrorCaption,
 					GEngineIni);
 
-				UE_LOG( LogLinker, Error, TEXT("Invalid import object index=%d while reading %s. File is most likely corrupted. Please verify your installation."), Index.ToImport(), *GetDebugName() );
+				UE_ASSET_LOG(LogLinker, Error, PackagePath, TEXT("Invalid import object index=%d. File is most likely corrupted. Please verify your installation."), Index.ToImport());
 
 				if (GLog)
 				{
@@ -5256,7 +5257,7 @@ void FLinkerLoad::DetachExport( int32 i )
 	check(Export.Object);
 	if( !Export.Object->IsValidLowLevel() )
 	{
-		UE_LOG(LogLinker, Fatal, TEXT("Linker object %s %s.%s is invalid"), *GetExportClassName(i).ToString(), *LinkerRoot->GetName(), *Export.ObjectName.ToString() );
+		UE_ASSET_LOG(LogLinker, Fatal, PackagePath, TEXT("Linker object %s %s is invalid"), *GetExportClassName(i).ToString(), *Export.ObjectName.ToString());
 	}
 	{
 		const FLinkerLoad* ActualLinker = Export.Object->GetLinker();
@@ -5268,13 +5269,13 @@ void FLinkerLoad::DetachExport( int32 i )
 			UE_LOG(LogLinker, Log, TEXT("Linker LinkerRoot : %s"), Object->GetLinker() ? *Object->GetLinker()->LinkerRoot->GetFullName() : TEXT("None"));
 			//UE_LOG(LogLinker, Log, TEXT("Detach Linker     : %s"), *GetFullName() );
 			UE_LOG(LogLinker, Log, TEXT("Detach LinkerRoot : %s"), *LinkerRoot->GetFullName());
-			UE_LOG(LogLinker, Fatal, TEXT("Linker object %s %s.%s mislinked!"), *GetExportClassName(i).ToString(), *LinkerRoot->GetName(), *Export.ObjectName.ToString());
+			UE_ASSET_LOG(LogLinker, Fatal, PackagePath, TEXT("Linker object %s %s mislinked!"), *GetExportClassName(i).ToString(), *Export.ObjectName.ToString());
 		}
 	}
 
 	if (Export.Object->GetLinkerIndex() == -1)
 	{
-		UE_LOG(LogLinker, Warning, TEXT("Linker object %s %s.%s was already detached."), *GetExportClassName(i).ToString(), *LinkerRoot->GetName(), *Export.ObjectName.ToString());
+		UE_ASSET_LOG(LogLinker, Warning, PackagePath, TEXT("Linker object %s %s was already detached."), *GetExportClassName(i).ToString(), *Export.ObjectName.ToString());
 	}
 	else
 	{
@@ -5406,7 +5407,7 @@ void FLinkerLoad::DetachBulkData( FBulkData* BulkData, bool bEnsureBulkDataIsLoa
 	int32 RemovedCount = BulkDataLoaders.Remove( BulkData );
 	if( RemovedCount!=1 )
 	{	
-		UE_LOG(LogLinker, Fatal, TEXT("Detachment inconsistency: %i (%s)"), RemovedCount, *GetDebugName() );
+		UE_ASSET_LOG(LogLinker, Fatal, PackagePath, TEXT("Detachment inconsistency: %i"), RemovedCount);
 	}
 	BulkData->DetachFromArchive( this, bEnsureBulkDataIsLoaded );
 }
@@ -5416,7 +5417,7 @@ void FLinkerLoad::DetachBulkData(UE::Serialization::FEditorBulkData* BulkData, b
 	int32 RemovedCount = EditorBulkDataLoaders.Remove(BulkData);
 	if (RemovedCount != 1)
 	{
-		UE_LOG(LogLinker, Fatal, TEXT("Detachment inconsistency: %i (%s)"), RemovedCount, *GetDebugName());
+		UE_ASSET_LOG(LogLinker, Fatal, PackagePath, TEXT("Detachment inconsistency: %i"), RemovedCount);
 	}
 	BulkData->DetachFromDisk(this, bEnsureBulkDataIsLoaded);
 }
@@ -5497,7 +5498,7 @@ FArchive& FLinkerLoad::operator<<( FObjectPtr& ObjectPtr )
 
 void FLinkerLoad::BadNameIndexError(int32 NameIndex)
 {
-	UE_LOG(LogLinker, Error, TEXT("Bad name index %i/%i"), NameIndex, NameMap.Num());
+	UE_ASSET_LOG(LogLinker, Error, PackagePath, TEXT("Bad name index %i/%i"), NameIndex, NameMap.Num());
 }
 
 /**
@@ -5962,7 +5963,7 @@ FLinkerLoad::ELinkerStatus FLinkerLoad::FixupExportMap()
 			FObjectExport &Export = ExportMap[ExportMapIdx];
 			if (!IsValidPackageIndex(Export.ClassIndex))
 			{
-				UE_LOG(LogLinker, Warning, TEXT("Bad class index found on export %d"), ExportMapIdx);
+				UE_ASSET_LOG(LogLinker, Warning, PackagePath, TEXT("Bad class index found on export %d"), ExportMapIdx);
 				return LINKER_Failed;
 			}
 			FName NameClass = GetExportClassName(ExportMapIdx);
diff --git a/Engine/Source/Runtime/CoreUObject/Private/UObject/UObjectGlobals.cpp b/Engine/Source/Runtime/CoreUObject/Private/UObject/UObjectGlobals.cpp
--- a/Engine/Source/Runtime/CoreUObject/Private/UObject/UObjectGlobals.cpp
+++ b/Engine/Source/Runtime/CoreUObject/Private/UObject/UObjectGlobals.cpp
@@ -4948,6 +4948,11 @@ FString FAssetMsg::FormatPathForAssetLog(const TCHAR* InPath)
 	return FilePath;
 }
 
+FString FAssetMsg::FormatPathForAssetLog(const FPackagePath& InPath)
+{
+	return FormatPathForAssetLog(*InPath.GetDebugName());
+}
+
 /**
  * Format the path of the passed in object
  */
@@ -4961,6 +4966,11 @@ FString FAssetMsg::GetAssetLogString(const TCHAR* Path, const FString& Message)
 	return FString::Printf(TEXT(ASSET_LOG_FORMAT_STRING_ANSI "%s"), *FAssetMsg::FormatPathForAssetLog(Path), *Message);
 }
 
+FString FAssetMsg::GetAssetLogString(const FPackagePath& Path, const FString& Message)
+{
+	return FString::Printf(ASSET_LOG_FORMAT_STRING TEXT("%s"), *FAssetMsg::FormatPathForAssetLog(Path), *Message);
+}
+
 FString FAssetMsg::GetAssetLogString(const UObject* Object, const FString& Message)
 {
 	return ensure(Object) ? GetAssetLogString(*Object->GetOutermost()->GetName(), Message) : FString();
diff --git a/Engine/Source/Runtime/CoreUObject/Public/UObject/UObjectGlobals.h b/Engine/Source/Runtime/CoreUObject/Public/UObject/UObjectGlobals.h
--- a/Engine/Source/Runtime/CoreUObject/Public/UObject/UObjectGlobals.h
+++ b/Engine/Source/Runtime/CoreUObject/Public/UObject/UObjectGlobals.h
@@ -2841,10 +2841,14 @@ struct FAssetMsg
 	/** Formats a path for the UE_ASSET_LOG macro */
 	static COREUOBJECT_API FString FormatPathForAssetLog(const TCHAR* Path);
 
+	/** Formats a path for the UE_ASSET_LOG macro */
+	static COREUOBJECT_API FString FormatPathForAssetLog(const FPackagePath& Path);
+
 	/** If possible, finds a path to the underlying asset for the provided object and formats it for the UE_ASSET_LOG macro */
 	static COREUOBJECT_API FString FormatPathForAssetLog(const UObject* Object);
 
 	static COREUOBJECT_API FString GetAssetLogString(const TCHAR* Path, const FString& Message);
+	static COREUOBJECT_API FString GetAssetLogString(const FPackagePath& Path, const FString& Message);
 	static COREUOBJECT_API FString GetAssetLogString(const UObject* Object, const FString& Message);
 };
 
