Adding an optional (default disabled) CSV stat for tracking total UObject count.

diff --git a/Engine/Source/Runtime/CoreUObject/Private/UObject/UObjectBase.cpp b/Engine/Source/Runtime/CoreUObject/Private/UObject/UObjectBase.cpp
--- a/Engine/Source/Runtime/CoreUObject/Private/UObject/UObjectBase.cpp
+++ b/Engine/Source/Runtime/CoreUObject/Private/UObject/UObjectBase.cpp
@@ -16,6 +16,7 @@
 #include "UObject/Class.h"
 #include "UObject/DeferredRegistry.h"
 #include "UObject/UObjectIterator.h"
+#include "UObject/UObjectStats.h"
 #include "UObject/Package.h"
 #include "Templates/Casts.h"
 #include "UObject/GCObject.h"
@@ -32,6 +33,12 @@ DECLARE_CYCLE_STAT(TEXT("CreateStatID"), STAT_CreateStatID, STATGROUP_StatSystem
 
 DEFINE_LOG_CATEGORY_STATIC(LogUObjectBootstrap, Display, Display);
 
+#if CSV_PROFILER && CSV_TRACK_UOBJECT_COUNT
+namespace UObjectStats
+{
+	COREUOBJECT_API std::atomic<int32> GUObjectCount;
+}
+#endif
 
 /** Whether uobject system is initialized.												*/
 namespace Internal
@@ -97,7 +104,11 @@ UObjectBase::UObjectBase( EObjectFlags InFlags )
 ,	InternalIndex		(INDEX_NONE)
 ,	ClassPrivate		(nullptr)
 ,	OuterPrivate		(nullptr)
-{}
+{
+#if CSV_PROFILER && CSV_TRACK_UOBJECT_COUNT
+	UObjectStats::IncrementUObjectCount();
+#endif
+}
 
 /**
  * Constructor used by StaticAllocateObject
@@ -116,6 +127,10 @@ UObjectBase::UObjectBase(UClass* InClass, EObjectFlags InFlags, EInternalObjectF
 	check(ClassPrivate);
 	// Add to global table.
 	AddObject(InName, InInternalFlags);
+	
+#if CSV_PROFILER && CSV_TRACK_UOBJECT_COUNT
+	UObjectStats::IncrementUObjectCount();
+#endif
 }
 
 
@@ -132,6 +147,10 @@ UObjectBase::~UObjectBase()
 		check(GetFName() == NAME_None);
 		GUObjectArray.FreeUObjectIndex(this);
 	}
+
+#if CSV_PROFILER && CSV_TRACK_UOBJECT_COUNT
+	UObjectStats::DecrementUObjectCount();
+#endif
 }
 
 
diff --git a/Engine/Source/Runtime/CoreUObject/Public/UObject/UObjectStats.h b/Engine/Source/Runtime/CoreUObject/Public/UObject/UObjectStats.h
new file mode 100644
--- /dev/null
+++ b/Engine/Source/Runtime/CoreUObject/Public/UObject/UObjectStats.h
@@ -0,0 +1,33 @@
+// Copyright Epic Games, Inc. All Rights Reserved.
+
+#pragma once
+
+#include "HAL/Platform.h"
+#include "ProfilingDebugging/CsvProfilerConfig.h"
+#include <atomic>
+
+#if !defined(CSV_TRACK_UOBJECT_COUNT)
+#define CSV_TRACK_UOBJECT_COUNT 0
+#endif
+
+#if CSV_PROFILER && CSV_TRACK_UOBJECT_COUNT
+namespace UObjectStats
+{
+	extern COREUOBJECT_API std::atomic<int32> GUObjectCount;
+
+	FORCEINLINE void IncrementUObjectCount()
+	{
+		GUObjectCount.fetch_add(1, std::memory_order_relaxed);
+	}
+
+	FORCEINLINE void DecrementUObjectCount()
+	{
+		GUObjectCount.fetch_sub(1, std::memory_order_relaxed);
+	}
+
+	FORCEINLINE int32 GetUObjectCount()
+	{
+		return GUObjectCount.load(std::memory_order_relaxed);
+	}
+}
+#endif
diff --git a/Engine/Source/Runtime/Engine/Private/LevelTick.cpp b/Engine/Source/Runtime/Engine/Private/LevelTick.cpp
--- a/Engine/Source/Runtime/Engine/Private/LevelTick.cpp
+++ b/Engine/Source/Runtime/Engine/Private/LevelTick.cpp
@@ -15,6 +15,7 @@
 #include "Modules/ModuleManager.h"
 #include "UObject/UObjectGlobals.h"
 #include "UObject/UObjectBaseUtility.h"
+#include "UObject/UObjectStats.h"
 #include "UObject/GarbageCollection.h"
 #include "EngineStats.h"
 #include "EngineGlobals.h"
@@ -78,6 +79,9 @@ CSV_DECLARE_CATEGORY_MODULE_EXTERN(CORE_API, Basic);
 CSV_DEFINE_CATEGORY_MODULE(ENGINE_API, Ticks, true);
 CSV_DEFINE_CATEGORY_MODULE(ENGINE_API, ActorCount, true);
 
+#if CSV_PROFILER && CSV_TRACK_UOBJECT_COUNT
+CSV_DEFINE_CATEGORY_MODULE(ENGINE_API, ObjectCount, true);
+#endif
 
 // this will log out all of the objects that were ticked in the FDetailedTickStats struct so you can isolate what is expensive
 #define LOG_DETAILED_DUMPSTATS 0
@@ -1284,6 +1288,11 @@ static void RecordWorldCountsToCSV(UWorld* World, bool bDoingActorTicks)
 			static FName TotalActorCountStatName(TEXT("TotalActorCount"));
 			FCsvProfiler::Get()->RecordCustomStat(TotalActorCountStatName, CSV_CATEGORY_INDEX(ActorCount), CSVActorTotalCount, ECsvCustomStatOp::Set);
 		}
+
+#if CSV_TRACK_UOBJECT_COUNT
+		static const FName TotalObjectCountStatName(TEXT("Total"));
+		FCsvProfiler::Get()->RecordCustomStat(TotalObjectCountStatName, CSV_CATEGORY_INDEX(ObjectCount), UObjectStats::GetUObjectCount(), ECsvCustomStatOp::Set);
+#endif
 	}
 }
 #endif // (CSV_PROFILER && !UE_BUILD_SHIPPING)
