Update FPackageTrailerBuilder to take and store an FString as context for error messages rather than a FName that is assumed to be the package name.

#rnx

- Ideally the VA tool would not need to deal with this like package mount points (the less project/plugin set up, the faster the tool can run) so we want to move back towards using file paths rather than FPackageNames to reference things where possible. Since virtualization does not need to do anything with the package data itself we can somewhat ignore a lot of Unreal conventions and treat the package files as just files that happen to have a package trailer on the end.
- The package builder was previously storing the name of it's package via an FName and converting it to a string when we needed to print a message. Now we store a generic FString of 'debug context' instead which  can be anything.
- Although it is unlikely that anyone is using the FName overloads externally we need to deprecate them rather than just delete them.

diff --git a/Engine/Source/Developer/Virtualization/Private/PackageSubmissionChecks.cpp b/Engine/Source/Developer/Virtualization/Private/PackageSubmissionChecks.cpp
--- a/Engine/Source/Developer/Virtualization/Private/PackageSubmissionChecks.cpp
+++ b/Engine/Source/Developer/Virtualization/Private/PackageSubmissionChecks.cpp
@@ -596,7 +596,7 @@ void VirtualizePackages(const TArray<FString>& FilesToSubmit, TArray<FText>& Out
 			return;
 		}
 
-		FPackageTrailerBuilder TrailerBuilder = FPackageTrailerBuilder::CreateFromTrailer(PackageInfo.Trailer, *PackageAr, PackagePath.GetPackageFName());
+		FPackageTrailerBuilder TrailerBuilder = FPackageTrailerBuilder::CreateFromTrailer(PackageInfo.Trailer, *PackageAr, PackagePath.GetPackageName());
 		if (!TrailerBuilder.BuildAndAppendTrailer(nullptr, *CopyAr))
 		{
 			FText Message = FText::Format(	LOCTEXT("Virtualization_TrailerAppend", "Failed to append the trailer to '{0}'"),
diff --git a/Engine/Source/Runtime/CoreUObject/Private/UObject/PackageTrailer.cpp b/Engine/Source/Runtime/CoreUObject/Private/UObject/PackageTrailer.cpp
--- a/Engine/Source/Runtime/CoreUObject/Private/UObject/PackageTrailer.cpp
+++ b/Engine/Source/Runtime/CoreUObject/Private/UObject/PackageTrailer.cpp
@@ -114,13 +114,13 @@ void FLookupTableEntry::Serialize(FArchive& Ar, EPackageTrailerVersion PackageTr
 
 } // namespace Private
 
-FPackageTrailerBuilder FPackageTrailerBuilder::CreateFromTrailer(const FPackageTrailer& Trailer, FArchive& Ar, const FName& PackageName)
+FPackageTrailerBuilder FPackageTrailerBuilder::CreateFromTrailer(const FPackageTrailer& Trailer, FArchive& Ar, FString DebugContext)
 {
-	FPackageTrailerBuilder Builder(PackageName);
+	FPackageTrailerBuilder Builder(MoveTemp(DebugContext));
 
 	for (const Private::FLookupTableEntry& Entry : Trailer.Header.PayloadLookupTable)
 	{
-		checkf(!Entry.Identifier.IsZero(), TEXT("PackageTrailer for package should not contain invalid FIoHash entry. Package '%s'"), *PackageName.ToString());
+		checkf(!Entry.Identifier.IsZero(), TEXT("PackageTrailer for package should not contain invalid FIoHash entry. Package '%s'"), *Builder.GetDebugContext());
 
 		switch (Entry.AccessMode)
 		{
@@ -153,20 +153,20 @@ FPackageTrailerBuilder FPackageTrailerBuilder::CreateFromTrailer(const FPackageT
 	return Builder;
 }
 
-TUniquePtr<UE::FPackageTrailerBuilder> FPackageTrailerBuilder::CreateReferenceToTrailer(const class FPackageTrailer& Trailer, const FName& PackageName)
+TUniquePtr<UE::FPackageTrailerBuilder> FPackageTrailerBuilder::CreateReferenceToTrailer(const class FPackageTrailer& Trailer, FString DebugContext)
 {
-	TUniquePtr<UE::FPackageTrailerBuilder> Builder = MakeUnique<UE::FPackageTrailerBuilder>(PackageName);
+	TUniquePtr<UE::FPackageTrailerBuilder> Builder = MakeUnique<UE::FPackageTrailerBuilder>(MoveTemp(DebugContext));
 
 	for (const Private::FLookupTableEntry& Entry : Trailer.Header.PayloadLookupTable)
 	{
-		checkf(!Entry.Identifier.IsZero(), TEXT("PackageTrailer for package should not contain invalid FIoHash entry. Package '%s'"), *PackageName.ToString());
+		checkf(!Entry.Identifier.IsZero(), TEXT("PackageTrailer for package should not contain invalid FIoHash entry. Package '%s'"), *Builder->GetDebugContext());
 
 		switch (Entry.AccessMode)
 		{
 			case EPayloadAccessMode::Local:
 			{
 				const int64 AbsoluteOffset = Trailer.FindPayloadOffsetInFile(Entry.Identifier);
-				checkf(AbsoluteOffset != INDEX_NONE, TEXT("PackageTrailer for package should not contain invalid payload offsets. Package '%s'"), *PackageName.ToString());
+				checkf(AbsoluteOffset != INDEX_NONE, TEXT("PackageTrailer for package should not contain invalid payload offsets. Package '%s'"), *Builder->GetDebugContext());
 
 				Builder->ReferencedEntries.Add(Entry.Identifier, ReferencedEntry(AbsoluteOffset, Entry.CompressedSize, Entry.RawSize));
 			}
@@ -174,7 +174,7 @@ TUniquePtr<UE::FPackageTrailerBuilder> FPackageTrailerBuilder::CreateReferenceTo
 			
 			case EPayloadAccessMode::Referenced:
 			{
-				checkf(false, TEXT("Attempting to create a reference to a trailer that already contains reference payload entries. Package '%s'"), *PackageName.ToString());
+				checkf(false, TEXT("Attempting to create a reference to a trailer that already contains reference payload entries. Package '%s'"), *Builder->GetDebugContext());
 			}
 			break;
 			
@@ -195,8 +195,13 @@ TUniquePtr<UE::FPackageTrailerBuilder> FPackageTrailerBuilder::CreateReferenceTo
 	return Builder;
 }
 
-FPackageTrailerBuilder::FPackageTrailerBuilder(const FName& InPackageName)
-	: PackageName(InPackageName)
+FPackageTrailerBuilder::FPackageTrailerBuilder(const FName& PackageName)
+	: DebugContext(PackageName.ToString())
+{
+}
+
+FPackageTrailerBuilder::FPackageTrailerBuilder(FString&& InDebugContext)
+	: DebugContext(InDebugContext)
 {
 }
 
@@ -253,7 +258,7 @@ bool FPackageTrailerBuilder::BuildAndAppendTrailer(FLinkerSave* Linker, FArchive
 
 	for (const TPair<FIoHash, LocalEntry>& It : LocalEntries)
 	{
-		checkf(!It.Key.IsZero(), TEXT("PackageTrailer should not contain invalid FIoHash values. Package '%s'"), *PackageName.ToString());
+		checkf(!It.Key.IsZero(), TEXT("PackageTrailer should not contain invalid FIoHash values. Package '%s'"), *DebugContext);
 
 		Private::FLookupTableEntry& Entry = Trailer.Header.PayloadLookupTable.AddDefaulted_GetRef();
 		Entry.Identifier = It.Key;
@@ -269,7 +274,7 @@ bool FPackageTrailerBuilder::BuildAndAppendTrailer(FLinkerSave* Linker, FArchive
 
 	for (const TPair<FIoHash, ReferencedEntry>& It : ReferencedEntries)
 	{
-		checkf(!It.Key.IsZero(), TEXT("PackageTrailer should not contain invalid FIoHash values. Package '%s'"), *PackageName.ToString());
+		checkf(!It.Key.IsZero(), TEXT("PackageTrailer should not contain invalid FIoHash values. Package '%s'"), *DebugContext);
 
 		Private::FLookupTableEntry& Entry = Trailer.Header.PayloadLookupTable.AddDefaulted_GetRef();
 		Entry.Identifier = It.Key;
@@ -281,7 +286,7 @@ bool FPackageTrailerBuilder::BuildAndAppendTrailer(FLinkerSave* Linker, FArchive
 
 	for (const TPair<FIoHash, VirtualizedEntry>& It : VirtualizedEntries)
 	{
-		checkf(!It.Key.IsZero(), TEXT("PackageTrailer should not contain invalid FIoHash values. Package '%s'"), *PackageName.ToString());
+		checkf(!It.Key.IsZero(), TEXT("PackageTrailer should not contain invalid FIoHash values. Package '%s'"), *DebugContext);
 
 		Private::FLookupTableEntry& Entry = Trailer.Header.PayloadLookupTable.AddDefaulted_GetRef();
 		Entry.Identifier = It.Key;
@@ -297,7 +302,10 @@ bool FPackageTrailerBuilder::BuildAndAppendTrailer(FLinkerSave* Linker, FArchive
 
 	DataArchive << Trailer.Header;
 
-	checkf((Trailer.TrailerPositionInFile + Trailer.Header.HeaderLength) == DataArchive.Tell(), TEXT("Header length was calculated as %d bytes but we wrote %" INT64_FMT " bytes!"), Trailer.Header.HeaderLength, DataArchive.Tell() - Trailer.TrailerPositionInFile);
+	checkf((Trailer.TrailerPositionInFile + Trailer.Header.HeaderLength) == DataArchive.Tell(), 
+		TEXT("Header length was calculated as %d bytes but we wrote %" INT64_FMT " bytes!"), 
+		Trailer.Header.HeaderLength, 
+		DataArchive.Tell() - Trailer.TrailerPositionInFile);
 
 	const int64 PayloadPosInFile = DataArchive.Tell();
 
@@ -306,12 +314,18 @@ bool FPackageTrailerBuilder::BuildAndAppendTrailer(FLinkerSave* Linker, FArchive
 		DataArchive << It.Value.Payload;
 	}
 
-	checkf((PayloadPosInFile + Trailer.Header.PayloadsDataLength) == DataArchive.Tell(), TEXT("Total payload length was calculated as %" INT64_FMT " bytes but we wrote %" INT64_FMT " bytes!"), Trailer.Header.PayloadsDataLength, DataArchive.Tell() - PayloadPosInFile);
+	checkf((PayloadPosInFile + Trailer.Header.PayloadsDataLength) == DataArchive.Tell(), 
+		TEXT("Total payload length was calculated as %" INT64_FMT " bytes but we wrote %" INT64_FMT " bytes!"), 
+		Trailer.Header.PayloadsDataLength, 
+		DataArchive.Tell() - PayloadPosInFile);
 
 	FPackageTrailer::FFooter Footer = Trailer.CreateFooter();
 	DataArchive << Footer;
 
-	checkf((Trailer.TrailerPositionInFile + Footer.TrailerLength) == DataArchive.Tell(), TEXT("Trailer length was calculated as %" INT64_FMT " bytes but we wrote %" INT64_FMT " bytes!"), Footer.TrailerLength, DataArchive.Tell() - Trailer.TrailerPositionInFile);
+	checkf((Trailer.TrailerPositionInFile + Footer.TrailerLength) == DataArchive.Tell(), 
+		TEXT("Trailer length was calculated as %" INT64_FMT " bytes but we wrote %" INT64_FMT " bytes!"), 
+		Footer.TrailerLength, 
+		DataArchive.Tell() - Trailer.TrailerPositionInFile);
 
 	// Invoke any registered callbacks and pass in the trailer, this allows the callbacks to poll where 
 	// in the output archive the payload has been stored.
@@ -369,7 +383,11 @@ void FPackageTrailerBuilder::RemoveDuplicateEntries()
 	{
 		if (VirtualizedEntries.Contains(Iter.Key()))
 		{	
-			UE_LOG(LogSerialization, Verbose, TEXT("Replacing localized payload '%s' with the virtualized version when building the package trailer for '%s'"), *LexToString(Iter.Key()) , *PackageName.ToString());
+			UE_LOG(LogSerialization, Verbose, 
+				TEXT("Replacing localized payload '%s' with the virtualized version when building the package trailer for '%s'"), 
+				*LexToString(Iter.Key()) , 
+				*DebugContext);
+
 			Iter.RemoveCurrent();
 		}
 	}
@@ -378,7 +396,11 @@ void FPackageTrailerBuilder::RemoveDuplicateEntries()
 	{
 		if (VirtualizedEntries.Contains(Iter.Key()))
 		{
-			UE_LOG(LogSerialization, Verbose, TEXT("Replacing localized payload '%s' with the virtualized version when building the package trailer for '%s'"), *LexToString(Iter.Key()), *PackageName.ToString());
+			UE_LOG(LogSerialization, Verbose, 
+				TEXT("Replacing localized payload '%s' with the virtualized version when building the package trailer for '%s'"), 
+				*LexToString(Iter.Key()), 
+				*DebugContext);
+
 			Iter.RemoveCurrent();
 		}
 	}
diff --git a/Engine/Source/Runtime/CoreUObject/Private/UObject/SavePackage2.cpp b/Engine/Source/Runtime/CoreUObject/Private/UObject/SavePackage2.cpp
--- a/Engine/Source/Runtime/CoreUObject/Private/UObject/SavePackage2.cpp
+++ b/Engine/Source/Runtime/CoreUObject/Private/UObject/SavePackage2.cpp
@@ -896,8 +896,8 @@ ESavePackageResult CreateLinker(FSaveContext& SaveContext)
 		{
 			// The package trailer is not supported for text based assets yet
 			if (!SaveContext.IsTextFormat() && !SaveContext.IsProceduralSave())
-			{
-				SaveContext.GetLinker()->PackageTrailerBuilder = MakeUnique<UE::FPackageTrailerBuilder>(SaveContext.GetPackage()->GetFName());
+			{		
+				SaveContext.GetLinker()->PackageTrailerBuilder = MakeUnique<UE::FPackageTrailerBuilder>(SaveContext.GetPackage()->GetName());
 			}
 			else if ((SaveContext.GetSaveArgs().SaveFlags & SAVE_BulkDataByReference) != 0)
 			{
@@ -905,7 +905,7 @@ ESavePackageResult CreateLinker(FSaveContext& SaveContext)
 				{
 					if (const UE::FPackageTrailer* Trailer = LinkerLoad->GetPackageTrailer())
 					{
-						SaveContext.GetLinker()->PackageTrailerBuilder = UE::FPackageTrailerBuilder::CreateReferenceToTrailer(*Trailer, SaveContext.GetPackage()->GetFName());
+						SaveContext.GetLinker()->PackageTrailerBuilder = UE::FPackageTrailerBuilder::CreateReferenceToTrailer(*Trailer, SaveContext.GetPackage()->GetName());
 					}
 				}
 			}
diff --git a/Engine/Source/Runtime/CoreUObject/Public/UObject/PackageTrailer.h b/Engine/Source/Runtime/CoreUObject/Public/UObject/PackageTrailer.h
--- a/Engine/Source/Runtime/CoreUObject/Public/UObject/PackageTrailer.h
+++ b/Engine/Source/Runtime/CoreUObject/Public/UObject/PackageTrailer.h
@@ -217,9 +217,15 @@ public:
 	 * 
 	 * @param Trailer		The trailer to create the builder from
 	 * @param Ar			An archive that the trailer can use to load payloads from 
-	 * @param PackageName	The name of the package that owns the trailer. Used for error messages.
+	 * @param DebugContext	The name or path of the of the file that owns the trailer. Used for error messages.
 	 */
-	[[nodiscard]] static FPackageTrailerBuilder CreateFromTrailer(const class FPackageTrailer& Trailer, FArchive& Ar, const FName& PackageName);
+	[[nodiscard]] static FPackageTrailerBuilder CreateFromTrailer(const class FPackageTrailer& Trailer, FArchive& Ar, FString DebugContext);
+
+	UE_DEPRECATED(5.1, "Use the overload that takes a FString instead of an FName for the last parameter")
+	static FPackageTrailerBuilder CreateFromTrailer(const class FPackageTrailer& Trailer, FArchive& Ar, const FName& PackageName)
+	{
+		return FPackageTrailerBuilder::CreateFromTrailer(Trailer, Ar, PackageName.ToString());
+	}
 
 	/**
 	 * Creates a builder from a pre-existing FPackageTrailer that will will reference the local payloads of the
@@ -227,12 +233,20 @@ public:
 	 * This means that there is no need to load the payloads.
 	 *
 	 * @param Trailer		The trailer to create the reference from.
-	 * @param PackageName	The name of the package that owns the trailer. Used for error messages.
+	 * @param DebugContext	The name or path of the of the file that owns the trailer. Used for error messages.
 	 */
-	[[nodiscard]] static TUniquePtr<UE::FPackageTrailerBuilder> CreateReferenceToTrailer(const class FPackageTrailer& Trailer, const FName& PackageName);
+	[[nodiscard]] static TUniquePtr<UE::FPackageTrailerBuilder> CreateReferenceToTrailer(const class FPackageTrailer& Trailer, FString DebugContext);
+	
+	UE_DEPRECATED(5.1, "Use the overload that takes a FString instead of an FName for the last parameter")
+	static TUniquePtr<UE::FPackageTrailerBuilder> CreateReferenceToTrailer(const class FPackageTrailer& Trailer, const FName& PackageName)
+	{
+		return FPackageTrailerBuilder::CreateReferenceToTrailer(Trailer, PackageName.ToString());
+	}
 
 	FPackageTrailerBuilder() = delete;
+	UE_DEPRECATED(5.1, "Use the overload that takes a FString instead of an FName")
 	FPackageTrailerBuilder(const FName& InPackageName);
+	FPackageTrailerBuilder(FString&& DebugContext);
 	~FPackageTrailerBuilder() = default;
 
 	// Methods that can be called while building the trailer
@@ -283,6 +297,12 @@ public:
 	/** Returns the number of payload entries in the builder with the access mode EPayloadAccessMode::Virtualized */
 	[[nodiscard]] int32 GetNumVirtualizedPayloads() const;
 
+	/** Returns the debug context associated with the builder, used for adding further description to error messages */
+	[[nodiscard]] const FString& GetDebugContext() const
+	{
+		return DebugContext;
+	}
+
 private:
 	
 	/** All of the data required to add a payload that is stored locally within the trailer */
@@ -338,8 +358,8 @@ private:
 
 	// Members used when building the trailer
 
-	/** Name of the package the trailer is being built for, used to give meaningful error messages */
-	FName PackageName;
+	/** Context used when giving error messages so that the user can identify the cause of problems */
+	FString DebugContext;
 
 	/** Payloads that will be stored locally when the trailer is written to disk */
 	TMap<FIoHash, LocalEntry> LocalEntries;
