Fix constructing a FSoftObjectPath from a FObjectPtr.
The original code was trying to build the path to the object without resolving it, but that seems impossible with redirectors, etc.

#rnx

diff --git a/Engine/Source/Runtime/CoreUObject/Private/UObject/ObjectPtr.cpp b/Engine/Source/Runtime/CoreUObject/Private/UObject/ObjectPtr.cpp
deleted file mode 100644
diff --git a/Engine/Source/Runtime/CoreUObject/Public/UObject/ObjectPtr.h b/Engine/Source/Runtime/CoreUObject/Public/UObject/ObjectPtr.h
--- a/Engine/Source/Runtime/CoreUObject/Public/UObject/ObjectPtr.h
+++ b/Engine/Source/Runtime/CoreUObject/Public/UObject/ObjectPtr.h
@@ -134,7 +134,7 @@ public:
 
 	FORCEINLINE bool IsResolved() const { return IsObjectHandleResolved(Handle); }
 
-	COREUOBJECT_API FString GetPath() const;
+	FORCEINLINE FString GetPath() const { return Get()->GetPathName(); }
 
 	FORCEINLINE FObjectHandle GetHandle() const { return Handle; }
 	FORCEINLINE FObjectHandle& GetHandleRef() const { return Handle; }
diff --git a/Engine/Source/Runtime/CoreUObject/Public/UObject/SoftObjectPath.h b/Engine/Source/Runtime/CoreUObject/Public/UObject/SoftObjectPath.h
--- a/Engine/Source/Runtime/CoreUObject/Public/UObject/SoftObjectPath.h
+++ b/Engine/Source/Runtime/CoreUObject/Public/UObject/SoftObjectPath.h
@@ -53,15 +53,13 @@ struct COREUOBJECT_API FSoftObjectPath
 	FSoftObjectPath(const UObject* InObject);
 
 	FSoftObjectPath(FObjectPtr InObject)
-	{
-		SetPath(InObject.GetPath());
-	}
+		: FSoftObjectPath(InObject.Get())
+	{}
 
 	template <typename T>
 	FSoftObjectPath(const TObjectPtr<T>& InObject)
-	{
-		SetPath(InObject.GetPath());
-	}
+		: FSoftObjectPath(InObject.Get())
+	{}
 
 	~FSoftObjectPath() {}
 	
diff --git a/Engine/Source/Runtime/CoreUObject/Tests/ObjectPtrTest.cpp b/Engine/Source/Runtime/CoreUObject/Tests/ObjectPtrTest.cpp
--- a/Engine/Source/Runtime/CoreUObject/Tests/ObjectPtrTest.cpp
+++ b/Engine/Source/Runtime/CoreUObject/Tests/ObjectPtrTest.cpp
@@ -207,7 +207,7 @@ TEST_CASE_METHOD(FObjectPtrTestBase, "CoreUObject::TObjectPtr::Soft Object Path"
 
 	FSoftObjectPath DefaultSoftObjPath(DefaultSoftObjPtr);
 
-	ObjectRefMetrics.TestNumResolves(TEXT("Unexpected resolve count after initializing an FSoftObjectPath from an FObjectPtr"), 0);
+	ObjectRefMetrics.TestNumResolves(TEXT("Unexpected resolve count after initializing an FSoftObjectPath from an FObjectPtr"), 1);
 
 	TEST_EQUAL_STR(TEXT("Soft object path constructed from an FObjectPtr does not have the expected path value"), TEXT("/Engine/Test/ObjectPtrSoftObjectPath/Transient.TestSoftObject"), *DefaultSoftObjPath.ToString());
 
diff --git a/Engine/Source/Runtime/Engine/Private/EdGraph/EdGraphPin.cpp b/Engine/Source/Runtime/Engine/Private/EdGraph/EdGraphPin.cpp
--- a/Engine/Source/Runtime/Engine/Private/EdGraph/EdGraphPin.cpp
+++ b/Engine/Source/Runtime/Engine/Private/EdGraph/EdGraphPin.cpp
@@ -747,9 +747,9 @@ bool UEdGraphPin::HasAnyConnections() const
 
 FString UEdGraphPin::GetDefaultAsString() const
 {
-	if(!DefaultObject.IsNullNoResolve())
+	if(DefaultObject)
 	{
-		return DefaultObject.GetPath();
+		return DefaultObject->GetPathName();
 	}
 	else if(!DefaultTextValue.IsEmpty())
 	{
