Fix PVS static analysis warning about potential nullptr in disabled test

diff --git a/Engine/Source/Runtime/CoreUObject/Private/Serialization/UnversionedPropertySerializationTest.cpp b/Engine/Source/Runtime/CoreUObject/Private/Serialization/UnversionedPropertySerializationTest.cpp
--- a/Engine/Source/Runtime/CoreUObject/Private/Serialization/UnversionedPropertySerializationTest.cpp
+++ b/Engine/Source/Runtime/CoreUObject/Private/Serialization/UnversionedPropertySerializationTest.cpp
@@ -398,7 +398,9 @@ struct FUnversionedPropertyTest : public FUnversionedPropertyTestInput
 					"%s loaded an instance equal to the original. "
 					"Inner mismatch in %s for the %s %s with UPS/TPS/Original values %s/%s/%s"), 
 					*Property->GetClass()->GetName(), *Struct->GetName(), *Property->GetName(), OkPaths,
-					VersionedUnversionedDiff.MismatchKind, *VersionedUnversionedDiff.GetType(), *VersionedUnversionedDiff.Property->GetName(), *VersionedValue, *UnversionedValue, *OriginalValue);
+					VersionedUnversionedDiff.MismatchKind, *VersionedUnversionedDiff.GetType(),
+					VersionedUnversionedDiff.Property ? *VersionedUnversionedDiff.Property->GetName() : TEXT("(null)"),
+					*VersionedValue, *UnversionedValue, *OriginalValue);
 			}
 		}		
 	}
