Updated comment on FScriptMapHelper/FScriptSetHelper FindInternalIndex() APIs to redirect potential callers to the more efficient method for iteration.

diff --git a/Engine/Source/Runtime/CoreUObject/Public/UObject/UnrealType.h b/Engine/Source/Runtime/CoreUObject/Public/UObject/UnrealType.h
--- a/Engine/Source/Runtime/CoreUObject/Public/UObject/UnrealType.h
+++ b/Engine/Source/Runtime/CoreUObject/Public/UObject/UnrealType.h
@@ -4698,7 +4698,7 @@ public:
 	/** 
 	 * Maps have gaps in their indices, so this function translates a logical index (ie. Nth element) 
 	 * to an internal index that can be used for the other functions in this class.
-	 * NOTE: This is slow, do not use this for iteration!
+	 * NOTE: This is slow, do not use this for iteration! Use CreateIterator() instead.
 	 */
 	int32 FindInternalIndex(int32 LogicalIdx) const
 	{
@@ -5391,7 +5391,7 @@ public:
 	/**
 	 * Maps have gaps in their indices, so this function translates a logical index (ie. Nth element)
 	 * to an internal index that can be used for the other functions in this class.
-	 * NOTE: This is slow, do not use this for iteration!
+	 * NOTE: This is slow, do not use this for iteration! Use CreateIterator() instead.
 	 */
 	int32 FindInternalIndex(int32 LogicalIdx) const
 	{
