Add metadata to UAnimDataModel to prevent the debugger from hanging while generating debug info

diff --git a/Engine/Source/Editor/UnrealEd/Private/Kismet2/KismetDebugUtilities.cpp b/Engine/Source/Editor/UnrealEd/Private/Kismet2/KismetDebugUtilities.cpp
--- a/Engine/Source/Editor/UnrealEd/Private/Kismet2/KismetDebugUtilities.cpp
+++ b/Engine/Source/Editor/UnrealEd/Private/Kismet2/KismetDebugUtilities.cpp
@@ -2157,6 +2157,14 @@ void FPropertyInstanceInfo::PopulateChildren(FPropertyInstance PropertyInstance,
 		return;
 	}
 
+	if (UObject* ResolvedObject = Object.Get())
+	{
+		if (ResolvedObject->GetClass()->HasMetaDataHierarchical("DebugTreeLeaf"))
+		{
+			return;
+		}
+	}
+
 	if (const FArrayProperty* ArrayProperty = CastField<FArrayProperty>(Property.Get()))
 	{
 		checkSlow(ArrayProperty->Inner);
diff --git a/Engine/Source/Runtime/CoreUObject/Public/UObject/ObjectMacros.h b/Engine/Source/Runtime/CoreUObject/Public/UObject/ObjectMacros.h
--- a/Engine/Source/Runtime/CoreUObject/Public/UObject/ObjectMacros.h
+++ b/Engine/Source/Runtime/CoreUObject/Public/UObject/ObjectMacros.h
@@ -1143,6 +1143,9 @@ namespace UM
 
 		/// [ClassMetadata] Used for Actor and Component classes. If the native class cannot tick, Blueprint generated classes based this Actor or Component can never tick even if bCanBlueprintsTickByDefault is true.
 		ChildCannotTick,
+		
+		/// [ClassMetadata] Used for objects with hundreds of nested properties like animation data to stop the editor from hanging to build debug data.
+		DebugTreeLeaf,
 
 		/// [ClassMetadata] Used to make the first subclass of a class ignore all inherited showCategories and hideCategories commands
 		IgnoreCategoryKeywordsInSubclasses,
diff --git a/Engine/Source/Runtime/Engine/Classes/Animation/AnimData/AnimDataModel.h b/Engine/Source/Runtime/Engine/Classes/Animation/AnimData/AnimDataModel.h
--- a/Engine/Source/Runtime/Engine/Classes/Animation/AnimData/AnimDataModel.h
+++ b/Engine/Source/Runtime/Engine/Classes/Animation/AnimData/AnimDataModel.h
@@ -71,7 +71,7 @@ struct FAnimationCurveIdentifier;
  * The Model represents the source data for animations. It contains both bone animation data as well as animated curves.
  * They are currently only a sub-object of a AnimSequenceBase instance. The instance derives all runtime data from the source data. 
  */
-UCLASS(BlueprintType)
+UCLASS(BlueprintType, meta=(DebugTreeLeaf))
 class ENGINE_API UAnimDataModel : public UObject
 {
 	GENERATED_BODY()
