[KillPendingKill] Disable use of PersistentGarbage flag because it's too aggressive at removing reported garbage references.

diff --git a/Engine/Source/Runtime/CoreUObject/Private/UObject/GarbageCollection.cpp b/Engine/Source/Runtime/CoreUObject/Private/UObject/GarbageCollection.cpp
--- a/Engine/Source/Runtime/CoreUObject/Private/UObject/GarbageCollection.cpp
+++ b/Engine/Source/Runtime/CoreUObject/Private/UObject/GarbageCollection.cpp
@@ -1005,7 +1005,8 @@ public:
 					return;
 				}
 
-				if (bGarbageReferenceTrackingEnabled && !ObjectItem->HasAnyFlags(EInternalObjectFlags::PersistentGarbage))
+				// Disable checking this flag because it's too aggressive at removing references from reporting.
+				if (bGarbageReferenceTrackingEnabled /*&& !ObjectItem->HasAnyFlags(EInternalObjectFlags::PersistentGarbage)*/)
 				{
 					HandleGarbageReference(ObjectsToSerializeStruct, ReferencingObject, Object, TokenIndex);
 				}
@@ -2502,7 +2503,8 @@ void FGCArrayPool::DumpGarbageReferencers(TArray<FGCArrayStruct*>& AllArrays)
 			for (FGarbageReferenceInfo& GarbageReference : ArrayStruct->GarbageReferences)
 			{
 				// No need to spam with references that would never get released anyway (even if we still had PendingKill) because they were referenced by a persistent reference 
-				if (!GarbageReference.GarbageObject->HasAnyInternalFlags(EInternalObjectFlags::PersistentGarbage))
+				// Disable checking this flag because it's too aggressive at removing references from reporting.
+				// if (!GarbageReference.GarbageObject->HasAnyInternalFlags(EInternalObjectFlags::PersistentGarbage))
 				{
 					UE_LOG(LogGarbage, Error, TEXT("Reachable garbage object: %s"), *GarbageReference.GarbageObject->GetFullName());
 					UE_LOG(LogGarbage, Error, TEXT("Referenced by:            %s"), *GarbageReference.GetReferencingObjectInfo());
@@ -2525,7 +2527,8 @@ void FGCArrayPool::DumpGarbageReferencers(TArray<FGCArrayStruct*>& AllArrays)
 			for (FGarbageReferenceInfo& GarbageReference : ArrayStruct->GarbageReferences)
 			{
 				// No need to spam with references that would never get released anyway (even if we still had PendingKill) because they were referenced by a persistent reference 
-				if (!GarbageReference.GarbageObject->HasAnyInternalFlags(EInternalObjectFlags::PersistentGarbage))
+				// Disable checking this flag because it's too aggressive at removing references from reporting.
+				// if (!GarbageReference.GarbageObject->HasAnyInternalFlags(EInternalObjectFlags::PersistentGarbage))
 				{
 					FKey Key(GarbageReference.bReferencerUObject ? (void*)GarbageReference.Referencer.Object->GetClass() : (void*)  GarbageReference.Referencer.GCObject, GarbageReference.PropertyName);
 					if (Seen.Contains(Key) == false)
