Add more explicit message for property validation checks.
Design cannot determine what is wrong with their assets without engineering assistence.
This will now point them in the right direction.

Change error from:
Cannot import unqualified delegate name; no object to search.

to:
Cannot import unqualified delegate name; no object to search. Delegate=%s SignatureFunction=%s"),

Adding this information helped me diagnose the problem from the editor without needing a c++ debugger attached.

Issue appears to be related to copying a struct with delegates.
#tests PIE and running content validation to figure out what the problem is.

diff --git a/Engine/Source/Runtime/CoreUObject/Private/UObject/PropertyHelper.h b/Engine/Source/Runtime/CoreUObject/Private/UObject/PropertyHelper.h
--- a/Engine/Source/Runtime/CoreUObject/Private/UObject/PropertyHelper.h
+++ b/Engine/Source/Runtime/CoreUObject/Private/UObject/PropertyHelper.h
@@ -216,7 +216,10 @@ namespace DelegatePropertyTools
 			// if neither are valid, then the importing fails
 			if (Parent == nullptr)
 			{
-				ErrorText->Logf(ELogVerbosity::Warning, TEXT("Cannot import unqualified delegate name; no object to search"));
+				FString DelegateString = Delegate.GetFunctionName().ToString();
+				FString SignatureFunctionString = GetNameSafe(SignatureFunction);
+				ErrorText->Logf(ELogVerbosity::Warning, TEXT("Cannot import unqualified delegate name; no object to search. Delegate=%s SignatureFunction=%s"),
+					*DelegateString, *SignatureFunctionString);
 				return nullptr;
 			}
 			// since we don't support nested components, we only need to check one level deep
