[core] Removed header extension test in operator== to match WITH_EDITOR code.

diff --git a/Engine/Source/Runtime/CoreUObject/Private/Misc/PackagePath.cpp b/Engine/Source/Runtime/CoreUObject/Private/Misc/PackagePath.cpp
--- a/Engine/Source/Runtime/CoreUObject/Private/Misc/PackagePath.cpp
+++ b/Engine/Source/Runtime/CoreUObject/Private/Misc/PackagePath.cpp
@@ -1117,7 +1117,7 @@ bool FPackagePath::TryMatchCase(const FPackagePath& SourcePackagePath, FStringVi
 
 bool FPackagePath::operator==(const FPackagePath& Other) const
 {
-	return Other.PackageName == PackageName && Other.HeaderExtension == HeaderExtension;
+	return Other.PackageName == PackageName;
 }
 
 FArchive& operator<<(FArchive& Ar, FPackagePath& PackagePath)
