Fix a crash issue when importing creative islands. This can occur when LoadObjectWhenImportingT3D is used to find an object.
There is no reason that the ObjectClass is unable to be nullptr and it is expected callers can find an object without a given class passed in as in the StaticFindObjectFast function allows for the class to be null.

diff --git a/Engine/Source/Runtime/CoreUObject/Private/UObject/UObjectGlobals.cpp b/Engine/Source/Runtime/CoreUObject/Private/UObject/UObjectGlobals.cpp
--- a/Engine/Source/Runtime/CoreUObject/Private/UObject/UObjectGlobals.cpp
+++ b/Engine/Source/Runtime/CoreUObject/Private/UObject/UObjectGlobals.cpp
@@ -1132,7 +1132,6 @@ bool ParseObject( const TCHAR* Stream, const TCHAR* Match, UClass* Class, UObjec
 UObject* StaticLoadObjectInternal(UClass* ObjectClass, UObject* InOuter, const TCHAR* InName, const TCHAR* Filename, uint32 LoadFlags, UPackageMap* Sandbox, bool bAllowObjectReconciliation, const FLinkerInstancingContext* InstancingContext)
 {
 	SCOPE_CYCLE_COUNTER(STAT_LoadObject);
-	check(ObjectClass);
 	check(InName);
 
 	FScopedLoadingState ScopedLoadingState(InName);
@@ -1178,7 +1177,7 @@ UObject* StaticLoadObjectInternal(UClass* ObjectClass, UObject* InOuter, const T
 			if (!Result && !(LoadFlags & LOAD_NoRedirects))
 			{
 				UObjectRedirector* Redirector = FindObjectFast<UObjectRedirector>(InOuter, *StrName);
-				if (Redirector && Redirector->DestinationObject && Redirector->DestinationObject->IsA(ObjectClass))
+				if (Redirector && Redirector->DestinationObject && Redirector->DestinationObject->IsA(ObjectClass ? ObjectClass : UObject::StaticClass()))
 				{
 					return Redirector->DestinationObject;
 				}
@@ -1230,7 +1229,7 @@ UObject* StaticLoadObject(UClass* ObjectClass, UObject* InOuter, const TCHAR* In
 		{
 			// we haven't created or found the object, error
 			FFormatNamedArguments Arguments;
-			Arguments.Add(TEXT("ClassName"), FText::FromString(ObjectClass->GetName()));
+			Arguments.Add(TEXT("ClassName"), ObjectClass ? FText::FromString(ObjectClass->GetName()) : NSLOCTEXT("Core", "None", "None"));
 			Arguments.Add(TEXT("OuterName"), InOuter ? FText::FromString(InOuter->GetPathName()) : NSLOCTEXT("Core", "None", "None"));
 			Arguments.Add(TEXT("ObjectName"), FText::FromString(ObjectName));
 			const FString Error = FText::Format(NSLOCTEXT("Core", "ObjectNotFound", "Failed to find object '{ClassName} {OuterName}.{ObjectName}'"), Arguments).ToString();
