[KillPendingKill] Downgrade error on garbage refs to warning so that it doesn't register as causing things like cook commandlets to fail.

diff --git a/Engine/Source/Runtime/CoreUObject/Private/UObject/GarbageCollection.cpp b/Engine/Source/Runtime/CoreUObject/Private/UObject/GarbageCollection.cpp
--- a/Engine/Source/Runtime/CoreUObject/Private/UObject/GarbageCollection.cpp
+++ b/Engine/Source/Runtime/CoreUObject/Private/UObject/GarbageCollection.cpp
@@ -2506,9 +2506,9 @@ void FGCArrayPool::DumpGarbageReferencers(TArray<FGCArrayStruct*>& AllArrays)
 				// Disable checking this flag because it's too aggressive at removing references from reporting.
 				// if (!GarbageReference.GarbageObject->HasAnyInternalFlags(EInternalObjectFlags::PersistentGarbage))
 				{
-					UE_LOG(LogGarbage, Error, TEXT("Reachable garbage object: %s"), *GarbageReference.GarbageObject->GetFullName());
-					UE_LOG(LogGarbage, Error, TEXT("Referenced by:            %s"), *GarbageReference.GetReferencingObjectInfo());
-					UE_LOG(LogGarbage, Error, TEXT(""));
+					UE_LOG(LogGarbage, Warning, TEXT("Reachable garbage object: %s"), *GarbageReference.GarbageObject->GetFullName());
+					UE_LOG(LogGarbage, Warning, TEXT("Referenced by:            %s"), *GarbageReference.GetReferencingObjectInfo());
+					UE_LOG(LogGarbage, Warning, TEXT(""));
 					NumGarbageReferences++;
 				}
 			}
@@ -2533,9 +2533,9 @@ void FGCArrayPool::DumpGarbageReferencers(TArray<FGCArrayStruct*>& AllArrays)
 					FKey Key(GarbageReference.bReferencerUObject ? (void*)GarbageReference.Referencer.Object->GetClass() : (void*)  GarbageReference.Referencer.GCObject, GarbageReference.PropertyName);
 					if (Seen.Contains(Key) == false)
 					{
-						UE_LOG(LogGarbage, Error, TEXT("Reachable garbage object: %s"), *GarbageReference.GarbageObject->GetFullName());
-						UE_LOG(LogGarbage, Error, TEXT("Referenced by:            %s"), *GarbageReference.GetReferencingObjectInfo());
-						UE_LOG(LogGarbage, Error, TEXT(""));
+						UE_LOG(LogGarbage, Warning, TEXT("Reachable garbage object: %s"), *GarbageReference.GarbageObject->GetFullName());
+						UE_LOG(LogGarbage, Warning, TEXT("Referenced by:            %s"), *GarbageReference.GetReferencingObjectInfo());
+						UE_LOG(LogGarbage, Warning, TEXT(""));
 						Seen.Add(Key);
 						ReportedGarbageReferences++;
 					}
