Minor QoL improvements for Blueprint stack traces.

* Warnings will always display the name of the function that tripped the warning. Note that stacks are still only displayed when there's an error or ScriptStackOnWarnings is used.
* Complete stack traces have an additional "<---" to highlight the top of the stack.
* Full paths are used, which helps track down the package that has the Blueprint.

eg:

LogScript: Warning: Script Msg: Divide by zero: Divide_IntInt
LogScript: Warning: Script call stack:
    /Game/Tests/Rendering/CornellBox.CornellBox_C.NewFunction_2 <---
    /Game/Tests/Rendering/CornellBox.CornellBox_C.NewFunction_1
    /Game/Tests/Rendering/CornellBox.CornellBox_C.NewFunction_0
    /Game/Tests/Rendering/CornellBox.CornellBox_C.ExecuteUbergraph_CornellBox
    /Game/Tests/Rendering/CornellBox.CornellBox_C.ReceiveBeginPlay

and

LogScript: Warning: Script Msg: Divide by zero: Divide_IntInt
LogScript: Warning: Last function called:
    /Game/Tests/Rendering/CornellBox.CornellBox_C.NewFunction_2

diff --git a/Engine/Source/Runtime/CoreUObject/Private/UObject/ScriptCore.cpp b/Engine/Source/Runtime/CoreUObject/Private/UObject/ScriptCore.cpp
--- a/Engine/Source/Runtime/CoreUObject/Private/UObject/ScriptCore.cpp
+++ b/Engine/Source/Runtime/CoreUObject/Private/UObject/ScriptCore.cpp
@@ -519,7 +519,7 @@ static bool ShowKismetScriptStackOnWarnings()
 	return ShowScriptStackForScriptWarning;
 }
 
-FString FFrame::GetScriptCallstack(bool bReturnEmpty)
+FString FFrame::GetScriptCallstack(bool bReturnEmpty, bool bTopOfStackOnly)
 {
 	FString ScriptStack;
 
@@ -527,10 +527,20 @@ FString FFrame::GetScriptCallstack(bool bReturnEmpty)
 	FBlueprintContextTracker& BlueprintExceptionTracker = FBlueprintContextTracker::Get();
 	if (BlueprintExceptionTracker.ScriptStack.Num() > 0)
 	{
-		for (int32 i = BlueprintExceptionTracker.ScriptStack.Num() - 1; i >= 0; --i)
+		const bool bDisplayArrow = (BlueprintExceptionTracker.ScriptStack.Num() > 1) && !bTopOfStackOnly;
+		const int32 TopOfStackIndex = BlueprintExceptionTracker.ScriptStack.Num() - 1;
+		int32 i = TopOfStackIndex;
+
+		do
 		{
-			ScriptStack += TEXT("\t") + BlueprintExceptionTracker.ScriptStack[i]->GetStackDescription() + TEXT("\n");
-		}
+			ScriptStack += TEXT("\t") + BlueprintExceptionTracker.ScriptStack[i]->GetStackDescription();
+			if ((i == TopOfStackIndex) && bDisplayArrow)
+			{
+				ScriptStack += TEXT(" <---");
+			}
+			ScriptStack += TEXT("\n");
+			--i;
+		} while ((i >= 0) && !bTopOfStackOnly);
 	}
 	else if (!bReturnEmpty)
 	{
@@ -548,7 +558,7 @@ FString FFrame::GetScriptCallstack(bool bReturnEmpty)
 
 FString FFrame::GetStackDescription() const
 {
-	return Node->GetOuter()->GetName() + TEXT(".") + Node->GetName();
+	return Node->GetOuter()->GetPathName() + TEXT(".") + Node->GetName();
 }
 
 #if DO_BLUEPRINT_GUARD
@@ -612,6 +622,11 @@ void FFrame::KismetExecutionMessage(const TCHAR* Message, ELogVerbosity::Type Ve
 		ScriptStack = TEXT("Script call stack:\n");
 		ScriptStack += GetScriptCallstack();
 	}
+	else if (Verbosity == ELogVerbosity::Warning)
+	{
+		ScriptStack = TEXT("Last function called:\n");
+		ScriptStack += GetScriptCallstack(false, true);
+	}
 #endif
 
 	if (Verbosity == ELogVerbosity::Fatal)
diff --git a/Engine/Source/Runtime/CoreUObject/Public/UObject/Stack.h b/Engine/Source/Runtime/CoreUObject/Public/UObject/Stack.h
--- a/Engine/Source/Runtime/CoreUObject/Public/UObject/Stack.h
+++ b/Engine/Source/Runtime/CoreUObject/Public/UObject/Stack.h
@@ -240,7 +240,7 @@ public:
 	* 
 	* @param	bReturnEmpty if true, returns empty string when no script callstack found
 	**/
-	COREUOBJECT_API static FString GetScriptCallstack(bool bReturnEmpty = false);
+	COREUOBJECT_API static FString GetScriptCallstack(bool bReturnEmpty = false, bool bTopOfStackOnly = false);
 
 	/** 
 	 * This will return a string of the form "ScopeName.FunctionName" associated with this stack frame:
