AsyncLoading2: Fix TSAN warning about data race when starting the async loading thread.
Move Array.Empty() calls to release memory of already empty game thread initial load arrays from ALT to GT.
Array.Empty() was writing ArrayNum=0 which was also being read by [REDACTED_PERSON]() on the GT.

#rnx

diff --git a/Engine/Source/Runtime/CoreUObject/Private/Serialization/AsyncLoading2.cpp b/Engine/Source/Runtime/CoreUObject/Private/Serialization/AsyncLoading2.cpp
--- a/Engine/Source/Runtime/CoreUObject/Private/Serialization/AsyncLoading2.cpp
+++ b/Engine/Source/Runtime/CoreUObject/Private/Serialization/AsyncLoading2.cpp
@@ -2326,7 +2326,7 @@ private:
 	TSpscQueue<FPackageRequest> PackageRequestQueue;
 	TArray<FAsyncPackage2*> PendingPackages;
 
-	/** Initial load pending CDOs */
+	/** [GAME THREAD] Initial load pending CDOs */
 	TMap<UClass*, TArray<FEventLoadNode2*>> PendingCDOs;
 	TArray<UClass*> PendingCDOsRecursiveStack;
 
@@ -5513,6 +5513,12 @@ void FAsyncLoadingThread2::StartThread()
 	// Make sure the GC sync object is created before we start the thread (apparently this can happen before we call InitUObject())
 	FGCCSyncObject::Create();
 
+	// Clear game thread initial load arrays
+	check(PendingCDOs.Num() == 0);
+	PendingCDOs.Empty();
+	check(PendingCDOsRecursiveStack.Num() == 0);
+	PendingCDOsRecursiveStack.Empty();
+
 	if (!FAsyncLoadingThreadSettings::Get().bAsyncLoadingThreadEnabled)
 	{
 		FinalizeInitialLoad();
@@ -5575,11 +5581,6 @@ void FAsyncLoadingThread2::FinalizeInitialLoad()
 	GlobalImportStore.FindAllScriptObjects(); // for verification only
 	bHasRegisteredAllScriptObjects = true;
 
-	check(PendingCDOs.Num() == 0);
-	PendingCDOs.Empty();
-	check(PendingCDOsRecursiveStack.Num() == 0);
-	PendingCDOsRecursiveStack.Empty();
-
 	UE_LOG(LogStreaming, Display,
 		TEXT("AsyncLoading2 - InitialLoad Finalized: Registered %d public script object entries (%.2f KB)"),
 		GlobalImportStore.GetStoredScriptObjectsCount(), (float)GlobalImportStore.GetStoredScriptObjectsAllocatedSize() / 1024.f);
