Add EPropertyFlags::CPF_NonNullable that can be used with FObjectProperty (and similar) to ensure that setting nullptr to the property is a no-op and the set has no effect

diff --git a/Engine/Source/Runtime/CoreUObject/Private/UObject/PropertyBaseObject.cpp b/Engine/Source/Runtime/CoreUObject/Private/UObject/PropertyBaseObject.cpp
--- a/Engine/Source/Runtime/CoreUObject/Private/UObject/PropertyBaseObject.cpp
+++ b/Engine/Source/Runtime/CoreUObject/Private/UObject/PropertyBaseObject.cpp
@@ -408,7 +408,7 @@ const TCHAR* FObjectPropertyBase::ImportText_Internal( const TCHAR* InBuffer, vo
 
 	if (PropertyPointerType == EPropertyPointerType::Container && HasSetter())
 	{
-		SetValue_InContainer(ContainerOrPropertyPtr, &Result); // @todo: need to make sure SetValue follows whatever SetObjectPropertyValue does for object properties
+		SetObjectPropertyValue_InContainer(ContainerOrPropertyPtr, Result);
 	}
 	else
 	{
@@ -603,9 +603,9 @@ bool FObjectPropertyBase::AllowObjectTypeReinterpretationTo(const FObjectPropert
 	return false;
 }
 
-void FObjectPropertyBase::CheckValidObject(void* Value) const
+void FObjectPropertyBase::CheckValidObject(void* ValueAddress, UObject* OldValue) const
 {
-	UObject *Object = GetObjectPropertyValue(Value);
+	UObject *Object = GetObjectPropertyValue(ValueAddress);
 	if (Object)
 	{
 		//
@@ -639,14 +639,30 @@ void FObjectPropertyBase::CheckValidObject(void* Value) const
 
 		if ((PropertyClass != nullptr) && !ObjectClass->IsChildOf(PropertyClass) && !ObjectClass->GetAuthoritativeClass()->IsChildOf(PropertyClass) && !bIsReplacingClassRefs && !bIsDeferringValueLoad)
 		{
-			UE_LOG(LogProperty, Warning,
-				TEXT("Serialized %s for a property of %s. Reference will be nullptred.\n    Property = %s\n    Item = %s"),
-				*Object->GetClass()->GetFullName(),
-				*PropertyClass->GetFullName(),
-				*GetFullName(),
-				*Object->GetFullName()
-			);
-			SetObjectPropertyValue(Value, nullptr);
+			if (!HasAnyPropertyFlags(CPF_NonNullable))
+			{
+				UE_LOG(LogProperty, Warning,
+					TEXT("Serialized %s for a property of %s. Reference will be nullptred.\n    Property = %s\n    Item = %s"),
+					*Object->GetClass()->GetFullName(),
+					*PropertyClass->GetFullName(),
+					*GetFullName(),
+					*Object->GetFullName()
+				);
+				SetObjectPropertyValue(ValueAddress, nullptr);
+			}
+			else
+			{
+				checkf(OldValue, TEXT("CheckValidObject(\"%s\") trying to assign null object value to non-nullable property \"%s\""), *Object->GetFullName(), *GetFullName());
+				UE_LOG(LogProperty, Warning,
+					TEXT("Serialized %s for a property of %s. Reference will be reverted back to %s.\n    Property = %s\n    Item = %s"),
+					*Object->GetClass()->GetFullName(),
+					*PropertyClass->GetFullName(),
+					*OldValue->GetFullName(),
+					*GetFullName(),
+					*Object->GetFullName()
+				);
+				SetObjectPropertyValue(ValueAddress, OldValue);
+			}
 		}
 	}
 }
diff --git a/Engine/Source/Runtime/CoreUObject/Private/UObject/PropertyClassPtr.cpp b/Engine/Source/Runtime/CoreUObject/Private/UObject/PropertyClassPtr.cpp
--- a/Engine/Source/Runtime/CoreUObject/Private/UObject/PropertyClassPtr.cpp
+++ b/Engine/Source/Runtime/CoreUObject/Private/UObject/PropertyClassPtr.cpp
@@ -58,12 +58,26 @@ UObject* FClassPtrProperty::GetObjectPropertyValue_InContainer(const void* Conta
 
 void FClassPtrProperty::SetObjectPropertyValue(void* PropertyValueAddress, UObject* Value) const
 {
-	SetPropertyValue(PropertyValueAddress, TCppType(Value));
+	if (Value || !HasAnyPropertyFlags(CPF_NonNullable))
+	{
+		SetPropertyValue(PropertyValueAddress, TCppType(Value));
+	}
+	else
+	{
+		UE_LOG(LogProperty, Verbose /*Warning*/, TEXT("Trying to assign null object value to non-nullable \"%s\""), *GetFullName());
+	}
 }
 
 void FClassPtrProperty::SetObjectPropertyValue_InContainer(void* ContainerAddress, UObject* Value, int32 ArrayIndex) const
 {
-	SetWrappedObjectPropertyValue_InContainer<FObjectPtr>(ContainerAddress, Value, ArrayIndex);
+	if (Value || !HasAnyPropertyFlags(CPF_NonNullable))
+	{
+		SetWrappedObjectPropertyValue_InContainer<FObjectPtr>(ContainerAddress, Value, ArrayIndex);
+	}
+	else
+	{
+		UE_LOG(LogProperty, Verbose /*Warning*/, TEXT("Trying to assign null object value to non-nullable \"%s\""), *GetFullName());
+	}
 }
 
 uint32 FClassPtrProperty::GetValueTypeHashInternal(const void* Src) const
diff --git a/Engine/Source/Runtime/CoreUObject/Private/UObject/PropertyLazyObjectPtr.cpp b/Engine/Source/Runtime/CoreUObject/Private/UObject/PropertyLazyObjectPtr.cpp
--- a/Engine/Source/Runtime/CoreUObject/Private/UObject/PropertyLazyObjectPtr.cpp
+++ b/Engine/Source/Runtime/CoreUObject/Private/UObject/PropertyLazyObjectPtr.cpp
@@ -59,7 +59,7 @@ void FLazyObjectProperty::SerializeItem( FStructuredArchive::FSlot Slot, void* V
 
 		if ((UnderlyingArchive.IsLoading() || UnderlyingArchive.IsModifyingWeakAndStrongReferences()) && ObjectValue != GetObjectPropertyValue(Value))
 		{
-			CheckValidObject(Value);
+			CheckValidObject(Value, ObjectValue);
 		}
 	}
 	else
@@ -117,12 +117,26 @@ UObject* FLazyObjectProperty::GetObjectPropertyValue_InContainer(const void* Con
 
 void FLazyObjectProperty::SetObjectPropertyValue(void* PropertyValueAddress, UObject* Value) const
 {
-	SetPropertyValue(PropertyValueAddress, TCppType(Value));
+	if (Value || !HasAnyPropertyFlags(CPF_NonNullable))
+	{
+		SetPropertyValue(PropertyValueAddress, TCppType(Value));
+	}
+	else
+	{
+		UE_LOG(LogProperty, Verbose /*Warning*/, TEXT("Trying to assign null object value to non-nullable \"%s\""), *GetFullName());
+	}
 }
 
 void FLazyObjectProperty::SetObjectPropertyValue_InContainer(void* ContainerAddress, UObject* Value, int32 ArrayIndex) const
 {
-	SetWrappedObjectPropertyValue_InContainer<FLazyObjectPtr>(ContainerAddress, Value, ArrayIndex);
+	if (Value || !HasAnyPropertyFlags(CPF_NonNullable))
+	{
+		SetWrappedObjectPropertyValue_InContainer<FLazyObjectPtr>(ContainerAddress, Value, ArrayIndex);
+	}
+	else
+	{
+		UE_LOG(LogProperty, Verbose /*Warning*/, TEXT("Trying to assign null object value to non-nullable \"%s\""), *GetFullName());
+	}
 }
 
 bool FLazyObjectProperty::AllowCrossLevel() const
diff --git a/Engine/Source/Runtime/CoreUObject/Private/UObject/PropertyObject.cpp b/Engine/Source/Runtime/CoreUObject/Private/UObject/PropertyObject.cpp
--- a/Engine/Source/Runtime/CoreUObject/Private/UObject/PropertyObject.cpp
+++ b/Engine/Source/Runtime/CoreUObject/Private/UObject/PropertyObject.cpp
@@ -67,7 +67,7 @@ EConvertFromTypeResult FObjectProperty::ConvertFromType(const FPropertyTag& Tag,
 		SetPropertyValue_InContainer(Data, PreviousValueObj, Tag.ArrayIndex);
 
 		// Validate the type is proper
-		CheckValidObject(GetPropertyValuePtr_InContainer(Data, Tag.ArrayIndex));
+		CheckValidObject(GetPropertyValuePtr_InContainer(Data, Tag.ArrayIndex), PreviousValueObj);
 
 		return EConvertFromTypeResult::Converted;
 	}
@@ -83,7 +83,7 @@ EConvertFromTypeResult FObjectProperty::ConvertFromType(const FPropertyTag& Tag,
 		}
 
 		SetPropertyValue_InContainer(Data, ObjectValue, Tag.ArrayIndex);
-		CheckValidObject(GetPropertyValuePtr_InContainer(Data, Tag.ArrayIndex));
+		CheckValidObject(GetPropertyValuePtr_InContainer(Data, Tag.ArrayIndex), ObjectValue);
 		return EConvertFromTypeResult::Converted;
 	}
 
@@ -102,7 +102,7 @@ void FObjectProperty::SerializeItem( FStructuredArchive::FSlot Slot, void* Value
 
 		if(!UnderlyingArchive.IsSaving())
 		{
-			CheckValidObject(ObjectPtr);
+			CheckValidObject(ObjectPtr, *ObjectPtr);
 		}
 	}
 	else
@@ -137,7 +137,7 @@ void FObjectProperty::SerializeItem( FStructuredArchive::FSlot Slot, void* Value
 			//        to accommodate this (as it depends on finding itself as the set value)
 	#endif // USE_CIRCULAR_DEPENDENCY_LOAD_DEFERRING
 
-			CheckValidObject(Value);
+			CheckValidObject(Value, CurrentValue);
 		}
 	}
 }
@@ -148,11 +148,10 @@ const TCHAR* FObjectProperty::ImportText_Internal(const TCHAR* Buffer, void* Con
 	if (Result)
 	{
 		void* Data = PointerToValuePtr(ContainerOrPropertyPtr, PropertyPointerType);
-		CheckValidObject(Data);
-
-#if USE_CIRCULAR_DEPENDENCY_LOAD_DEFERRING
 		UObject* ObjectValue = GetObjectPropertyValue(Data);
+		CheckValidObject(Data, ObjectValue);
 
+#if USE_CIRCULAR_DEPENDENCY_LOAD_DEFERRING		
 		if (ULinkerPlaceholderClass* PlaceholderClass = Cast<ULinkerPlaceholderClass>(ObjectValue))
 		{
 			// we use this tracker mechanism to help record the instance that is
@@ -199,7 +198,14 @@ UObject* FObjectProperty::GetObjectPropertyValue_InContainer(const void* Contain
 
 void FObjectProperty::SetObjectPropertyValue(void* PropertyValueAddress, UObject* Value) const
 {
-	SetPropertyValue(PropertyValueAddress, Value);
+	if (Value || !HasAnyPropertyFlags(CPF_NonNullable))
+	{
+		SetPropertyValue(PropertyValueAddress, Value);
+	}
+	else
+	{
+		UE_LOG(LogProperty, Verbose /*Warning*/, TEXT("Trying to assign null object value to non-nullable \"%s\""), *GetFullName());
+	}
 }
 
 void FObjectProperty::SetObjectPropertyValue_InContainer(void* ContainerAddress, UObject* Value, int32 ArrayIndex) const
@@ -214,5 +220,12 @@ void FObjectProperty::SetObjectPropertyValue_InContainer(void* ContainerAddress,
 		}
 		UObject* Obj = nullptr;
 	};
-	SetWrappedObjectPropertyValue_InContainer<FWrappedObjectPtr>(ContainerAddress, Value, ArrayIndex);
+	if (Value || !HasAnyPropertyFlags(CPF_NonNullable))
+	{
+		SetWrappedObjectPropertyValue_InContainer<FWrappedObjectPtr>(ContainerAddress, Value, ArrayIndex);
+	}
+	else
+	{
+		UE_LOG(LogProperty, Verbose /*Warning*/, TEXT("Trying to assign null object value to non-nullable \"%s\""), *GetFullName());
+	}
 }
diff --git a/Engine/Source/Runtime/CoreUObject/Private/UObject/PropertyObjectPtr.cpp b/Engine/Source/Runtime/CoreUObject/Private/UObject/PropertyObjectPtr.cpp
--- a/Engine/Source/Runtime/CoreUObject/Private/UObject/PropertyObjectPtr.cpp
+++ b/Engine/Source/Runtime/CoreUObject/Private/UObject/PropertyObjectPtr.cpp
@@ -55,6 +55,7 @@ void FObjectPtrProperty::StaticSerializeItem(const FObjectPropertyBase* ObjectPr
 {
 	FArchive& UnderlyingArchive = Slot.GetUnderlyingArchive();
 	FObjectPtr* ObjectPtr = (FObjectPtr*)GetPropertyValuePtr(Value);
+	UObject* CurrentValue = IsObjectHandleResolved(ObjectPtr->GetHandle()) ? ObjectPtr->Get() : nullptr;
 
 	if (UnderlyingArchive.IsObjectReferenceCollector())
 	{
@@ -62,7 +63,7 @@ void FObjectPtrProperty::StaticSerializeItem(const FObjectPropertyBase* ObjectPr
 
 		if(!UnderlyingArchive.IsSaving() && IsObjectHandleResolved(ObjectPtr->GetHandle()))
 		{
-			ObjectProperty->CheckValidObject(ObjectPtr);
+			ObjectProperty->CheckValidObject(ObjectPtr, CurrentValue);
 		}
 	}
 	else
@@ -96,7 +97,7 @@ void FObjectPtrProperty::StaticSerializeItem(const FObjectPropertyBase* ObjectPr
 			//        to accommodate this (as it depends on finding itself as the set value)
 	#endif // USE_CIRCULAR_DEPENDENCY_LOAD_DEFERRING
 
-			ObjectProperty->CheckValidObject(Value);
+			ObjectProperty->CheckValidObject(Value, CurrentValue);
 		}
 	}
 }
@@ -149,12 +150,26 @@ UObject* FObjectPtrProperty::GetObjectPropertyValue_InContainer(const void* Cont
 
 void FObjectPtrProperty::SetObjectPropertyValue(void* PropertyValueAddress, UObject* Value) const
 {
-	SetPropertyValue(PropertyValueAddress, TCppType(Value));
+	if (Value || !HasAnyPropertyFlags(CPF_NonNullable))
+	{
+		SetPropertyValue(PropertyValueAddress, TCppType(Value));
+	}
+	else
+	{
+		UE_LOG(LogProperty, Verbose /*Warning*/, TEXT("Trying to assign null object value to non-nullable \"%s\""), *GetFullName());
+	}
 }
 
 void FObjectPtrProperty::SetObjectPropertyValue_InContainer(void* ContainerAddress, UObject* Value, int32 ArrayIndex) const
 {
-	SetWrappedObjectPropertyValue_InContainer<FObjectPtr>(ContainerAddress, Value, ArrayIndex);
+	if (Value || !HasAnyPropertyFlags(CPF_NonNullable))
+	{
+		SetWrappedObjectPropertyValue_InContainer<FObjectPtr>(ContainerAddress, Value, ArrayIndex);
+	}
+	else
+	{
+		UE_LOG(LogProperty, Verbose /*Warning*/, TEXT("Trying to assign null object value to non-nullable \"%s\""), *GetFullName());
+	}
 }
 
 bool FObjectPtrProperty::AllowCrossLevel() const
diff --git a/Engine/Source/Runtime/CoreUObject/Private/UObject/PropertySoftObjectPtr.cpp b/Engine/Source/Runtime/CoreUObject/Private/UObject/PropertySoftObjectPtr.cpp
--- a/Engine/Source/Runtime/CoreUObject/Private/UObject/PropertySoftObjectPtr.cpp
+++ b/Engine/Source/Runtime/CoreUObject/Private/UObject/PropertySoftObjectPtr.cpp
@@ -53,6 +53,12 @@ bool FSoftObjectProperty::Identical( const void* A, const void* B, uint32 PortFl
 	return ObjectA.GetUniqueID() == ObjectB.GetUniqueID();
 }
 
+void FSoftObjectProperty::LinkInternal(FArchive& Ar)
+{
+	checkf(!HasAnyPropertyFlags(CPF_NonNullable), TEXT("Soft Object Properties can't be non nullable but \"%s\" is marked as CPF_NonNullable"), *GetFullName());
+	Super::LinkInternal(Ar);
+}
+
 void FSoftObjectProperty::SerializeItem( FStructuredArchive::FSlot Slot, void* Value, void const* Defaults ) const
 {
 	FArchive& UnderlyingArchive = Slot.GetUnderlyingArchive();
@@ -71,7 +77,7 @@ void FSoftObjectProperty::SerializeItem( FStructuredArchive::FSlot Slot, void* V
 		{
 			if (OldValue.GetUniqueID() != ((FSoftObjectPtr*)Value)->GetUniqueID())
 			{
-				CheckValidObject(Value);
+				CheckValidObject(Value, nullptr); // FSoftObjectProperty is never non-nullable at this point so it's ok to pass null as the current value
 			}
 		}
 #endif
diff --git a/Engine/Source/Runtime/CoreUObject/Private/UObject/PropertyWeakObjectPtr.cpp b/Engine/Source/Runtime/CoreUObject/Private/UObject/PropertyWeakObjectPtr.cpp
--- a/Engine/Source/Runtime/CoreUObject/Private/UObject/PropertyWeakObjectPtr.cpp
+++ b/Engine/Source/Runtime/CoreUObject/Private/UObject/PropertyWeakObjectPtr.cpp
@@ -46,6 +46,12 @@ FString FWeakObjectProperty::GetCPPMacroType( FString& ExtendedTypeText ) const
 	return TEXT("WEAKOBJECT");
 }
 
+void FWeakObjectProperty::LinkInternal(FArchive& Ar)
+{
+	checkf(!HasAnyPropertyFlags(CPF_NonNullable), TEXT("Weak Object Properties can't be non nullable but \"%s\" is marked as CPF_NonNullable"), *GetFullName());
+	Super::LinkInternal(Ar);
+}
+
 void FWeakObjectProperty::SerializeItem( FStructuredArchive::FSlot Slot, void* Value, void const* Defaults ) const
 {
 	FArchive& UnderlyingArchive = Slot.GetUnderlyingArchive();
@@ -54,7 +60,7 @@ void FWeakObjectProperty::SerializeItem( FStructuredArchive::FSlot Slot, void* V
 	Slot << *(FWeakObjectPtr*)Value;
 	if ((UnderlyingArchive.IsLoading() || UnderlyingArchive.IsModifyingWeakAndStrongReferences()) && ObjectValue != GetObjectPropertyValue(Value))
 	{
-		CheckValidObject(Value);
+		CheckValidObject(Value, nullptr); // FWeakObjectProperty is never non-nullable at this point so it's ok to pass null as the current value
 	}
 }
 
diff --git a/Engine/Source/Runtime/CoreUObject/Public/UObject/ObjectMacros.h b/Engine/Source/Runtime/CoreUObject/Public/UObject/ObjectMacros.h
--- a/Engine/Source/Runtime/CoreUObject/Public/UObject/ObjectMacros.h
+++ b/Engine/Source/Runtime/CoreUObject/Public/UObject/ObjectMacros.h
@@ -403,7 +403,7 @@ enum EPropertyFlags : uint64
 	CPF_ZeroConstructor					= 0x0000000000000200,	///< memset is fine for construction
 	CPF_ReturnParm						= 0x0000000000000400,	///< Return value.
 	CPF_DisableEditOnTemplate			= 0x0000000000000800,	///< Disable editing of this property on an archetype/sub-blueprint
-	//CPF_      						= 0x0000000000001000,	///< 
+	CPF_NonNullable						= 0x0000000000001000,	///< Object property can never be null
 	CPF_Transient   					= 0x0000000000002000,	///< Property is transient: shouldn't be saved or loaded, except for Blueprint CDOs.
 	CPF_Config      					= 0x0000000000004000,	///< Property should be loaded/saved as permanent profile.
 	//CPF_								= 0x0000000000008000,	///< 
diff --git a/Engine/Source/Runtime/CoreUObject/Public/UObject/UnrealType.h b/Engine/Source/Runtime/CoreUObject/Public/UObject/UnrealType.h
--- a/Engine/Source/Runtime/CoreUObject/Public/UObject/UnrealType.h
+++ b/Engine/Source/Runtime/CoreUObject/Public/UObject/UnrealType.h
@@ -2713,7 +2713,7 @@ public:
 	FORCEINLINE void SetPropertyClass(UClass* NewPropertyClass) { PropertyClass = NewPropertyClass; }
 #endif // USE_CIRCULAR_DEPENDENCY_LOAD_DEFERRING
 
-	virtual void CheckValidObject(void* Value) const;
+	virtual void CheckValidObject(void* ValueAddress, UObject* OldValue) const;
 
 protected:
 	virtual bool AllowCrossLevel() const;
@@ -3002,6 +3002,8 @@ PRAGMA_ENABLE_DEPRECATION_WARNINGS
 	// FProperty interface
 	virtual void SerializeItem(FStructuredArchive::FSlot Slot, void* Value, void const* Defaults) const override;
 	virtual void EmitReferenceInfo(UClass& OwnerClass, int32 BaseOffset, TArray<const FStructProperty*>& EncounteredStructProps, FGCStackSizeHelper& StackSizeHelper) override;
+protected:
+	virtual void LinkInternal(FArchive& Ar) override;
 private:
 	virtual uint32 GetValueTypeHashInternal(const void* Src) const override;
 public:
@@ -3129,6 +3131,7 @@ PRAGMA_ENABLE_DEPRECATION_WARNINGS
 	virtual void SerializeItem( FStructuredArchive::FSlot Slot, void* Value, void const* Defaults) const override;
 	virtual bool NetSerializeItem(FArchive& Ar, UPackageMap* Map, void* Data, TArray<uint8> * MetaData = NULL) const override;
 protected:
+	virtual void LinkInternal(FArchive& Ar) override;
 	virtual void ExportText_Internal( FString& ValueStr, const void* PropertyValueOrContainer, EPropertyPointerType PropertyPointerType, const void* DefaultValue, UObject* Parent, int32 PortFlags, UObject* ExportRootScope ) const override;
 	virtual const TCHAR* ImportText_Internal(const TCHAR* Buffer, void* ContainerOrPropertyPtr, EPropertyPointerType PropertyPointerType, UObject* OwnerObject, int32 PortFlags, FOutputDevice* ErrorText) const override;
 public:
