Fixing non unity build errors

diff --git a/Engine/Source/Runtime/CoreUObject/Private/UObject/PropertyObject.cpp b/Engine/Source/Runtime/CoreUObject/Private/UObject/PropertyObject.cpp
--- a/Engine/Source/Runtime/CoreUObject/Private/UObject/PropertyObject.cpp
+++ b/Engine/Source/Runtime/CoreUObject/Private/UObject/PropertyObject.cpp
@@ -8,6 +8,7 @@
 #include "UObject/LinkerPlaceholderBase.h"
 #include "UObject/LinkerPlaceholderExportObject.h"
 #include "UObject/LinkerPlaceholderClass.h"
+#include "UObject/PropertyHelper.h"
 
 /*-----------------------------------------------------------------------------
 	FObjectProperty.
