Add special handling to FSoftObjectPath import that treats () like an empty string to avoid a confusing ensure.
Other types like gameplay tags already do this to handle the hardcoded "empty struct" used by blueprints and utilities

diff --git a/Engine/Source/Runtime/CoreUObject/Private/UObject/SoftObjectPath.cpp b/Engine/Source/Runtime/CoreUObject/Private/UObject/SoftObjectPath.cpp
--- a/Engine/Source/Runtime/CoreUObject/Private/UObject/SoftObjectPath.cpp
+++ b/Engine/Source/Runtime/CoreUObject/Private/UObject/SoftObjectPath.cpp
@@ -375,6 +375,25 @@ bool FSoftObjectPath::ImportTextItem(const TCHAR*& Buffer, int32 PortFlags, UObj
 	}
 	else
 	{
+		if (*Buffer == TCHAR('('))
+		{
+			// Blueprints and other utilities may pass in () as a hardcoded value for an empty struct, so treat that like an empty string
+			Buffer++;
+
+			if (*Buffer == TCHAR(')'))
+			{
+				Buffer++;
+				Reset();
+				return true;
+			}
+			else
+			{
+				// Fall back to the default struct parsing, which will print an error message
+				Buffer--;
+				return false;
+			}
+		}
+
 		if (*Buffer == TCHAR('\''))
 		{
 			// A ' token likely means we're looking at a path string in the form "Texture2d'/Game/UI/HUD/Actions/Barrel'" and we need to read and append the path part
