TWeakInterfacePtr::GetWeakObjectPtr

diff --git a/Engine/Source/Runtime/CoreUObject/Public/UObject/WeakInterfacePtr.h b/Engine/Source/Runtime/CoreUObject/Public/UObject/WeakInterfacePtr.h
--- a/Engine/Source/Runtime/CoreUObject/Public/UObject/WeakInterfacePtr.h
+++ b/Engine/Source/Runtime/CoreUObject/Public/UObject/WeakInterfacePtr.h
@@ -189,6 +189,11 @@ struct TWeakInterfacePtr
 		return TScriptInterface<T>();
 	}
 
+	FORCEINLINE TWeakObjectPtr<UObject> GetWeakObjectPtr() const
+	{
+		return ObjectInstance;
+	}
+
 private:
 	TWeakObjectPtr<UObject> ObjectInstance;
 	T* InterfaceInstance;
