Set archetype flag before duplicating subobjects, to allow PostEditProperties etc. to read it.

diff --git a/Engine/Source/Runtime/CoreUObject/Private/UObject/PropertyBaseObject.cpp b/Engine/Source/Runtime/CoreUObject/Private/UObject/PropertyBaseObject.cpp
--- a/Engine/Source/Runtime/CoreUObject/Private/UObject/PropertyBaseObject.cpp
+++ b/Engine/Source/Runtime/CoreUObject/Private/UObject/PropertyBaseObject.cpp
@@ -395,15 +395,17 @@ const TCHAR* FObjectPropertyBase::ImportText_Internal( const TCHAR* InBuffer, vo
 			ExistingObject->Rename(nullptr, nullptr, REN_DontCreateRedirectors | REN_ForceNoResetLoaders);
 		}
 
-		Result = DuplicateObject<UObject>(Result, Parent, DesiredName);
+		FObjectDuplicationParameters ObjectDuplicationParams = InitStaticDuplicateObjectParams(Result, Parent, DesiredName);
+		EnumRemoveFlags(ObjectDuplicationParams.FlagMask, RF_ArchetypeObject);
 		if (Parent->IsTemplate())
 		{
-			Result->SetFlags(RF_ArchetypeObject);
+			EnumAddFlags(ObjectDuplicationParams.ApplyFlags, RF_ArchetypeObject);
 		}
 		else
 		{
-			Result->ClearFlags(RF_ArchetypeObject);
+			EnumRemoveFlags(ObjectDuplicationParams.ApplyFlags, RF_ArchetypeObject);
 		}
+		Result = StaticDuplicateObjectEx(ObjectDuplicationParams);
 	}
 
 	if (PropertyPointerType == EPropertyPointerType::Container && HasSetter())
@@ -641,13 +643,13 @@ void FObjectPropertyBase::CheckValidObject(void* ValueAddress, UObject* OldValue
 		{
 			if (!HasAnyPropertyFlags(CPF_NonNullable))
 			{
-				UE_LOG(LogProperty, Warning,
-					TEXT("Serialized %s for a property of %s. Reference will be nullptred.\n    Property = %s\n    Item = %s"),
-					*Object->GetClass()->GetFullName(),
-					*PropertyClass->GetFullName(),
-					*GetFullName(),
-					*Object->GetFullName()
-				);
+			UE_LOG(LogProperty, Warning,
+				TEXT("Serialized %s for a property of %s. Reference will be nullptred.\n    Property = %s\n    Item = %s"),
+				*Object->GetClass()->GetFullName(),
+				*PropertyClass->GetFullName(),
+				*GetFullName(),
+				*Object->GetFullName()
+			);
 				SetObjectPropertyValue(ValueAddress, nullptr);
 			}
 			else
