fixed a TSAN warning about not synced concurrent access to `FAsyncLoadingThread2::AsyncLoadingThreadID` by making it atomic

diff --git a/Engine/Source/Runtime/CoreUObject/Private/Serialization/AsyncLoading2.cpp b/Engine/Source/Runtime/CoreUObject/Private/Serialization/AsyncLoading2.cpp
--- a/Engine/Source/Runtime/CoreUObject/Private/Serialization/AsyncLoading2.cpp
+++ b/Engine/Source/Runtime/CoreUObject/Private/Serialization/AsyncLoading2.cpp
@@ -67,6 +67,8 @@
 #include "Misc/PathViews.h"
 #include "UObject/LinkerLoad.h"
 
+#include <atomic>
+
 #if UE_BUILD_DEVELOPMENT || UE_BUILD_DEBUG
 PRAGMA_DISABLE_OPTIMIZATION
 #endif
@@ -2313,7 +2315,7 @@ private:
 	TSet<FAsyncPackage2*> PackagesToCancel;
 
 	/** Async loading thread ID */
-	uint32 AsyncLoadingThreadID;
+	std::atomic<uint32> AsyncLoadingThreadID;
 
 	/** I/O Dispatcher */
 	FIoDispatcher& IoDispatcher;
