switched from TQueue<Mpsc> to TMpscQueue because it's a faster and modern version, and TSAN doesn't complain about it

diff --git a/Engine/Source/Runtime/CoreUObject/Private/Serialization/AsyncLoading2.cpp b/Engine/Source/Runtime/CoreUObject/Private/Serialization/AsyncLoading2.cpp
--- a/Engine/Source/Runtime/CoreUObject/Private/Serialization/AsyncLoading2.cpp
+++ b/Engine/Source/Runtime/CoreUObject/Private/Serialization/AsyncLoading2.cpp
@@ -66,6 +66,7 @@
 #include "Containers/SpscQueue.h"
 #include "Misc/PathViews.h"
 #include "UObject/LinkerLoad.h"
+#include "Containers/SpscQueue.h"
 
 #include <atomic>
 
@@ -2277,7 +2278,7 @@ private:
 	TSet<FWeakObjectPtr> LoadedAssets;
 #endif
 	/** [ASYNC/GAME THREAD] Packages to be deleted from async thread */
-	TQueue<FAsyncPackage2*, EQueueMode::Spsc> DeferredDeletePackages;
+	TSpscQueue<FAsyncPackage2*> DeferredDeletePackages;
 	
 	struct FFailedPackageRequest
 	{
@@ -2292,7 +2293,7 @@ private:
 	/** Packages in active loading with GetAsyncPackageId() as key */
 	TMap<FPackageId, FAsyncPackage2*> AsyncPackageLookup;
 
-	TQueue<FAsyncPackage2*, EQueueMode::Mpsc> ExternalReadQueue;
+	TSpscQueue<FAsyncPackage2*> ExternalReadQueue;
 	TAtomic<int32> PendingIoRequestsCounter{ 0 };
 	
 	/** List of all pending package requests */
@@ -5646,14 +5647,15 @@ uint32 FAsyncLoadingThread2::Run()
 					{
 						bPopped = false;
 						{
-							FAsyncPackage2* Package = nullptr;
-							if (ExternalReadQueue.Peek(Package))
+							FAsyncPackage2** Package = ExternalReadQueue.Peek();
+							if (Package != nullptr)
 							{
 								TRACE_CPUPROFILER_EVENT_SCOPE(PollExternalReads);
-								EAsyncPackageState::Type Result = Package->ProcessExternalReads(FAsyncPackage2::ExternalReadAction_Poll);
+								check(*Package);
+								EAsyncPackageState::Type Result = (*Package)->ProcessExternalReads(FAsyncPackage2::ExternalReadAction_Poll);
 								if (Result == EAsyncPackageState::Complete)
 								{
-									ExternalReadQueue.Pop();
+									ExternalReadQueue.Dequeue();
 									bPopped = true;
 									bDidSomething = true;
 								}
@@ -5721,14 +5723,13 @@ uint32 FAsyncLoadingThread2::Run()
 					TRACE_CPUPROFILER_EVENT_SCOPE(WaitingForIo);
 					Waiter.Wait();
 				}
-				else if (ExternalReadQueue.Peek(Package))
+				else if (ExternalReadQueue.Dequeue(Package))
 				{
 					TRACE_CPUPROFILER_EVENT_SCOPE(AsyncLoadingTime);
 					TRACE_CPUPROFILER_EVENT_SCOPE(WaitingForExternalReads);
 
 					EAsyncPackageState::Type Result = Package->ProcessExternalReads(FAsyncPackage2::ExternalReadAction_Wait);
 					check(Result == EAsyncPackageState::Complete);
-					ExternalReadQueue.Pop();
 				}
 				else
 				{
