[Backout] - CL20829812

Original CL Desc
-----------------------------------------------------------------
switched from TQueue<Mpsc> to TMpscQueue because it's a faster and modern version, and TSAN doesn't complain about it

diff --git a/Engine/Source/Runtime/CoreUObject/Private/Serialization/AsyncLoading2.cpp b/Engine/Source/Runtime/CoreUObject/Private/Serialization/AsyncLoading2.cpp
--- a/Engine/Source/Runtime/CoreUObject/Private/Serialization/AsyncLoading2.cpp
+++ b/Engine/Source/Runtime/CoreUObject/Private/Serialization/AsyncLoading2.cpp
@@ -66,7 +66,6 @@
 #include "Containers/SpscQueue.h"
 #include "Misc/PathViews.h"
 #include "UObject/LinkerLoad.h"
-#include "Containers/SpscQueue.h"
 
 #include <atomic>
 
@@ -2278,7 +2277,7 @@ private:
 	TSet<FWeakObjectPtr> LoadedAssets;
 #endif
 	/** [ASYNC/GAME THREAD] Packages to be deleted from async thread */
-	TSpscQueue<FAsyncPackage2*> DeferredDeletePackages;
+	TQueue<FAsyncPackage2*, EQueueMode::Spsc> DeferredDeletePackages;
 	
 	struct FFailedPackageRequest
 	{
@@ -2293,7 +2292,7 @@ private:
 	/** Packages in active loading with GetAsyncPackageId() as key */
 	TMap<FPackageId, FAsyncPackage2*> AsyncPackageLookup;
 
-	TSpscQueue<FAsyncPackage2*> ExternalReadQueue;
+	TQueue<FAsyncPackage2*, EQueueMode::Mpsc> ExternalReadQueue;
 	TAtomic<int32> PendingIoRequestsCounter{ 0 };
 	
 	/** List of all pending package requests */
@@ -5647,15 +5646,14 @@ uint32 FAsyncLoadingThread2::Run()
 					{
 						bPopped = false;
 						{
-							FAsyncPackage2** Package = ExternalReadQueue.Peek();
-							if (Package != nullptr)
+							FAsyncPackage2* Package = nullptr;
+							if (ExternalReadQueue.Peek(Package))
 							{
 								TRACE_CPUPROFILER_EVENT_SCOPE(PollExternalReads);
-								check(*Package);
-								EAsyncPackageState::Type Result = (*Package)->ProcessExternalReads(FAsyncPackage2::ExternalReadAction_Poll);
+								EAsyncPackageState::Type Result = Package->ProcessExternalReads(FAsyncPackage2::ExternalReadAction_Poll);
 								if (Result == EAsyncPackageState::Complete)
 								{
-									ExternalReadQueue.Dequeue();
+									ExternalReadQueue.Pop();
 									bPopped = true;
 									bDidSomething = true;
 								}
@@ -5723,13 +5721,14 @@ uint32 FAsyncLoadingThread2::Run()
 					TRACE_CPUPROFILER_EVENT_SCOPE(WaitingForIo);
 					Waiter.Wait();
 				}
-				else if (ExternalReadQueue.Dequeue(Package))
+				else if (ExternalReadQueue.Peek(Package))
 				{
 					TRACE_CPUPROFILER_EVENT_SCOPE(AsyncLoadingTime);
 					TRACE_CPUPROFILER_EVENT_SCOPE(WaitingForExternalReads);
 
 					EAsyncPackageState::Type Result = Package->ProcessExternalReads(FAsyncPackage2::ExternalReadAction_Wait);
 					check(Result == EAsyncPackageState::Complete);
+					ExternalReadQueue.Pop();
 				}
 				else
 				{
