Ensure extra GC logging is removed in shipping.

diff --git a/Engine/Source/Runtime/CoreUObject/Private/UObject/GarbageCollection.cpp b/Engine/Source/Runtime/CoreUObject/Private/UObject/GarbageCollection.cpp
--- a/Engine/Source/Runtime/CoreUObject/Private/UObject/GarbageCollection.cpp
+++ b/Engine/Source/Runtime/CoreUObject/Private/UObject/GarbageCollection.cpp
@@ -259,6 +259,9 @@ static FAutoConsoleVariableRef CMultithreadedDestructionEnabled(
 	ECVF_Default
 );
 
+#if UE_BUILD_SHIPPING
+static constexpr int32 GGarbageReferenceTrackingEnabled = 0;
+#else
 int32 GGarbageReferenceTrackingEnabled = 0;
 static FAutoConsoleVariableRef CGarbageReferenceTrackingEnabled(
 	TEXT("gc.GarbageReferenceTrackingEnabled"),
@@ -266,6 +269,7 @@ static FAutoConsoleVariableRef CGarbageReferenceTrackingEnabled(
 	TEXT("Causes the Garbage Collector to track and log unreleased garbage objects. If 1, will dump every reference. If 2, will dump a sample of the references to highlight problematic properties."),
 	ECVF_Default
 );
+#endif
 
 #if !UE_BUILD_SHIPPING
 static FAutoConsoleCommand CmdCalculateTokenStreamSize(
@@ -916,6 +920,7 @@ public:
 
 	FORCENOINLINE void HandleGarbageReference(FGCArrayStruct& ObjectsToSerializeStruct, const UObject* const ReferencingObject, UObject*& Object, const int32 TokenIndex)
 	{
+#if !UE_BUILD_SHIPPING
 		FName PropertyName;
 		UObject* GarbageReferencer = ReferencingObject ? const_cast<UObject*>(ReferencingObject) : ObjectsToSerializeStruct.GetReferencingObject();
 		if (IsValid(GarbageReferencer))
@@ -936,6 +941,7 @@ public:
 				ObjectsToSerializeStruct.GarbageReferences.Add(FGarbageReferenceInfo(GCObjectReferencer, Object, Object));
 			}
 		}
+#endif
 	}
 
 #if ENABLE_GC_HISTORY
@@ -2482,6 +2488,7 @@ void FGCArrayPool::DumpGarbageReferencers(TArray<FGCArrayStruct*>& AllArrays)
 		return;
 	}
 
+#if !UE_BUILD_SHIPPING
 	// We don't care about leaks when engine exit was requested since the final GC pass will destroy everything anyway
 	// We still want to clear the GarbageReferences array though
 	if (IsEngineExitRequested())
@@ -2546,6 +2553,7 @@ void FGCArrayPool::DumpGarbageReferencers(TArray<FGCArrayStruct*>& AllArrays)
 		}
 		UE_CLOG(TotalGarbageReferences > 0, LogGarbage, Log, TEXT("Reported %d/%d garbage references in %f ms."), ReportedGarbageReferences, TotalGarbageReferences, (FPlatformTime::Seconds() - StartTime) * 1000);
 	}
+#endif
 }
 
 void FGCArrayPool::UpdateGCHistory(TArray<FGCArrayStruct*>& AllArrays)
