Restore old FBulkData::Realloc behavior in the editor from before CL 20391050

#rnx

- Bulkdata2 has an error at runtime to catch the use of ::Realloc on bulkdata objects that represented a payload on disk, when this was merged with the original bulkdata code we should've removed the check from the editor.
- Note: That calling realloc on a bulkdata instance that still represents a payload on disk can have some weird effects as ::realloc will change the size, but not detach the bulkdata, so if we do try to later load it off disk, the size value will not match anymore. In the older code we would have a "SizeOnDisk" member that would still allow us to safely load the payload off disk, but then the payload and the size that we return will still be mismatched.
-- I will add a jira for further work on this, it is likely that calling ::Realloc will in future detach the bulkdata from disk to prevent potential problems in the future.

diff --git a/Engine/Source/Runtime/CoreUObject/Private/Serialization/BulkData.cpp b/Engine/Source/Runtime/CoreUObject/Private/Serialization/BulkData.cpp
--- a/Engine/Source/Runtime/CoreUObject/Private/Serialization/BulkData.cpp
+++ b/Engine/Source/Runtime/CoreUObject/Private/Serialization/BulkData.cpp
@@ -994,7 +994,9 @@ void* FBulkData::Realloc(int64 ElementCount, int64 ElementSize)
 	DECLARE_SCOPE_CYCLE_COUNTER(TEXT("FBulkData::Realloc"), STAT_UBD_Realloc, STATGROUP_Memory);
 
 	checkf(IsLocked(), TEXT("BulkData must be locked for 'write' before reallocating!"));
+#if !WITH_EDITOR
 	checkf(!CanLoadFromDisk(), TEXT("Cannot re-allocate a FBulkDataBase object that represents a file on disk!"));
+#endif //!WITH_EDITOR
 
 	// We might want to consider this a valid use case if anyone can come up with one?
 	checkf(!IsUsingIODispatcher(), TEXT("Attempting to re-allocate data loaded from the IoDispatcher"));
