Modulate excluded object mark based on the save realm
Do not save gatherable text when cooking with editor only data enabled

diff --git a/Engine/Source/Runtime/CoreUObject/Private/UObject/SavePackage/PackageHarvester.cpp b/Engine/Source/Runtime/CoreUObject/Private/UObject/SavePackage/PackageHarvester.cpp
--- a/Engine/Source/Runtime/CoreUObject/Private/UObject/SavePackage/PackageHarvester.cpp
+++ b/Engine/Source/Runtime/CoreUObject/Private/UObject/SavePackage/PackageHarvester.cpp
@@ -78,7 +78,7 @@ bool ConditionallyExcludeObjectForTarget(FSaveContext& SaveContext, UObject* Obj
 	}
 	else if (!SaveContext.GetHarvestedRealm(HarvestingContext).IsIncluded(Obj))
 	{
-		const EObjectMark ExcludedObjectMarks = SaveContext.GetExcludedObjectMarks();
+		const EObjectMark ExcludedObjectMarks = SaveContext.GetExcludedObjectMarks(HarvestingContext);
 		const ITargetPlatform* TargetPlatform = SaveContext.GetTargetPlatform();
 		EObjectMark ObjectMarks = GenerateMarksForObject(Obj, TargetPlatform);
 		if (!(ObjectMarks & ExcludedObjectMarks))
diff --git a/Engine/Source/Runtime/CoreUObject/Private/UObject/SavePackage/SaveContext.h b/Engine/Source/Runtime/CoreUObject/Private/UObject/SavePackage/SaveContext.h
--- a/Engine/Source/Runtime/CoreUObject/Private/UObject/SavePackage/SaveContext.h
+++ b/Engine/Source/Runtime/CoreUObject/Private/UObject/SavePackage/SaveContext.h
@@ -465,10 +465,10 @@ public:
 		return TargetPackagePath;
 	}
 
-	EObjectMark GetExcludedObjectMarks() const
+	EObjectMark GetExcludedObjectMarks(ESaveRealm HarvestingRealm) const
 	{
 		// When considering excluded objects for a platform, do not consider editor only object in the optional context as excluded
-		return (CurrentHarvestingRealm == ESaveRealm::Optional) ? (EObjectMark)(ExcludedObjectMarks & ~EObjectMark::OBJECTMARK_EditorOnly) : ExcludedObjectMarks;
+		return (HarvestingRealm == ESaveRealm::Optional) ? (EObjectMark)(ExcludedObjectMarks & ~EObjectMark::OBJECTMARK_EditorOnly) : ExcludedObjectMarks;
 	}
 
 	EObjectFlags GetTopLevelFlags() const
diff --git a/Engine/Source/Runtime/CoreUObject/Private/UObject/SavePackage2.cpp b/Engine/Source/Runtime/CoreUObject/Private/UObject/SavePackage2.cpp
--- a/Engine/Source/Runtime/CoreUObject/Private/UObject/SavePackage2.cpp
+++ b/Engine/Source/Runtime/CoreUObject/Private/UObject/SavePackage2.cpp
@@ -1558,7 +1558,9 @@ void SavePreloadDependencies(FStructuredArchive::FRecord& StructuredArchiveRoot,
 void WriteGatherableText(FStructuredArchive::FRecord& StructuredArchiveRoot, FSaveContext& SaveContext)
 {
 	FStructuredArchive::FStream Stream = StructuredArchiveRoot.EnterStream(TEXT("GatherableTextData"));
-	if (!SaveContext.IsFilterEditorOnly()
+	// Do not gather text data during cooking since the data isn't only scrubbed off of editor packages
+	if (!SaveContext.IsCooking()
+		&& !SaveContext.IsFilterEditorOnly()
 		// We can only cache packages that:
 		//	1) Don't contain script data, as script data is very volatile and can only be safely gathered after it's been compiled (which happens automatically on asset load).
 		//	2) Don't contain text keyed with an incorrect package localization ID, as these keys will be changed later during save.
