Do not exclude "NotForPlatform" objects as well when in the optional context

diff --git a/Engine/Source/Runtime/CoreUObject/Private/UObject/SavePackage/SaveContext.h b/Engine/Source/Runtime/CoreUObject/Private/UObject/SavePackage/SaveContext.h
--- a/Engine/Source/Runtime/CoreUObject/Private/UObject/SavePackage/SaveContext.h
+++ b/Engine/Source/Runtime/CoreUObject/Private/UObject/SavePackage/SaveContext.h
@@ -467,8 +467,8 @@ public:
 
 	EObjectMark GetExcludedObjectMarks(ESaveRealm HarvestingRealm) const
 	{
-		// When considering excluded objects for a platform, do not consider editor only object in the optional context as excluded
-		return (HarvestingRealm == ESaveRealm::Optional) ? (EObjectMark)(ExcludedObjectMarks & ~EObjectMark::OBJECTMARK_EditorOnly) : ExcludedObjectMarks;
+		// When considering excluded objects for a platform, do not consider editor only object and not for platform objects in the optional context as excluded
+		return (HarvestingRealm == ESaveRealm::Optional) ? (EObjectMark)(ExcludedObjectMarks & ~(EObjectMark::OBJECTMARK_EditorOnly|EObjectMark::OBJECTMARK_NotForTargetPlatform)) : ExcludedObjectMarks;
 	}
 
 	EObjectFlags GetTopLevelFlags() const
