Restore old FBulkData allocation behavior in the editor from before CL 20391050

#rnx

- Bulkdata2 made a distinction between allocating for the first time and re-allocating were as the original bulkdata system only called realloc even if we knew that we have no pre-existing allocation.
- Using the Bulkdata2 system in the editor was causing some code paths to crash as they were relying on being able to serialize a bulkdata instance repeatedly without resetting it.
- For now, I have changed the code to default to the old behavior so we only ever call reallocate.
-- A better fix in the future would be to ::Reset the bulkdata before a loading serialization call, but I will add a new JIRA ticket for that.

diff --git a/Engine/Source/Runtime/CoreUObject/Private/Serialization/BulkData.cpp b/Engine/Source/Runtime/CoreUObject/Private/Serialization/BulkData.cpp
--- a/Engine/Source/Runtime/CoreUObject/Private/Serialization/BulkData.cpp
+++ b/Engine/Source/Runtime/CoreUObject/Private/Serialization/BulkData.cpp
@@ -360,15 +360,6 @@ void FBulkData::FAllocatedPtr::Free(FBulkData* Owner)
 	}
 }
 
-void* FBulkData::FAllocatedPtr::AllocateData(FBulkData* Owner, SIZE_T SizeInBytes)
-{
-	checkf(Allocation.RawData == nullptr, TEXT("Trying to allocate a BulkData object without freeing it first!"));
-
-	Allocation.RawData = FMemory::Malloc(SizeInBytes, DEFAULT_ALIGNMENT);
-
-	return Allocation.RawData;
-}
-
 void* FBulkData::FAllocatedPtr::ReallocateData(FBulkData* Owner, SIZE_T SizeInBytes)
 {
 	checkf(!Owner->IsDataMemoryMapped(),  TEXT("Trying to reallocate a memory mapped BulkData object without freeing it first!"));
@@ -616,7 +607,7 @@ FBulkData& FBulkData::operator=( const FBulkData& Other )
 	{
 		if (const void* Src = Other.GetDataBufferReadOnly())
 		{
-			void* Dst = AllocateData(GetBulkDataSize());
+			void* Dst = ReallocateData(GetBulkDataSize());
 			FMemory::Memcpy(Dst, Src, GetBulkDataSize());
 		}
 	}
@@ -1417,7 +1408,7 @@ void FBulkData::Serialize(FArchive& Ar, UObject* Owner, bool bAttemptFileMapping
 				checkf(bAttemptFileMapping == false || bIsUsingIoDispatcher == false, TEXT("Trying to memory map inline bulk data '%s' which is not supported when loading from I/O store"), *Owner->GetPackage()->GetFName().ToString());
 
 				const int64 BulkDataSize = GetBulkDataSize();
-				void* DataBuffer = AllocateData(BulkDataSize);
+				void* DataBuffer = ReallocateData(BulkDataSize);
 				SerializeBulkData(Ar, DataBuffer, BulkDataSize, BulkMeta.GetFlags());
 
 				ConditionalSetInlineAlwaysAllowDiscard(bIsUsingIoDispatcher);
diff --git a/Engine/Source/Runtime/CoreUObject/Public/Serialization/BulkData.h b/Engine/Source/Runtime/CoreUObject/Public/Serialization/BulkData.h
--- a/Engine/Source/Runtime/CoreUObject/Public/Serialization/BulkData.h
+++ b/Engine/Source/Runtime/CoreUObject/Public/Serialization/BulkData.h
@@ -533,7 +533,6 @@ class COREUOBJECT_API FBulkData
 
 		void Free(FBulkData* Owner);
 
-		void* AllocateData(FBulkData* Owner, SIZE_T SizeInBytes); 
 		void* ReallocateData(FBulkData* Owner, SIZE_T SizeInBytes);
 
 		void SetData(FBulkData* Owner, void* Buffer);
@@ -1075,9 +1074,6 @@ private:
 	/** Sets whether inline bulk data is allowed to be unloaded or not */
 	void ConditionalSetInlineAlwaysAllowDiscard(bool bPackageUsesIoStore);
 
-	/** Free allocatged data */
-	inline void* AllocateData(SIZE_T SizeInBytes) { return DataAllocation.AllocateData(this, SizeInBytes); }
-
 	/** Reallocate bulk data */
 	inline void* ReallocateData(SIZE_T SizeInBytes) { return DataAllocation.ReallocateData(this, SizeInBytes); }
 
