BulkData: Fixed loading of iodispatcher bulkdata in editor builds

diff --git a/Engine/Source/Runtime/CoreUObject/Private/Serialization/BulkData.cpp b/Engine/Source/Runtime/CoreUObject/Private/Serialization/BulkData.cpp
--- a/Engine/Source/Runtime/CoreUObject/Private/Serialization/BulkData.cpp
+++ b/Engine/Source/Runtime/CoreUObject/Private/Serialization/BulkData.cpp
@@ -1957,13 +1957,6 @@ bool FBulkData::TryLoadDataIntoMemory(FIoBuffer Dest)
 		
 			return true;
 		}
-		else
-		{
-			UE_LOG(LogSerialization, Error, TEXT("Attempted to load bulk data without an attached archive. ")
-				TEXT("Most likely the bulk data was loaded twice on console, which is not supported"));
-
-			return false;
-		}
 	}
 #endif // WITH_EDITOR
 	   
