fix crash when saving a transient asset

#rnx

diff --git a/Engine/Source/Runtime/CoreUObject/Private/UObject/SavePackage/PackageHarvester.cpp b/Engine/Source/Runtime/CoreUObject/Private/UObject/SavePackage/PackageHarvester.cpp
--- a/Engine/Source/Runtime/CoreUObject/Private/UObject/SavePackage/PackageHarvester.cpp
+++ b/Engine/Source/Runtime/CoreUObject/Private/UObject/SavePackage/PackageHarvester.cpp
@@ -699,7 +699,7 @@ FString FPackageHarvester::GetUnsaveableReason(UObject* Required)
 		}
 		else
 		{
-			ReasonText =  FString::Printf(TEXT("It %s"), *LexToString(Status));
+			ReasonText =  FString::Printf(TEXT("It %s"), LexToString(Status));
 		}
 	}
 	else if (ConditionallyExcludeObjectForTarget(SaveContext, Required, HarvestContext))
