Fixed deadlock if FLongPackagePathsSingleton isn't yet initialised.

The constructor for FLongPackagePathsSingleton would take a write lock. This would block however as there are Get() call sites where a read lock is already being held. But as the singleton construction is thread-safe there is no need to take a write lock in constructor and it has been removed. To mitigate the error returning in the future the Get() is called before taking any read locks. The lock was renamed and moved into FLongPackagePathsSingleton instead of a global.

diff --git a/Engine/Source/Runtime/CoreUObject/Private/Misc/PackageName.cpp b/Engine/Source/Runtime/CoreUObject/Private/Misc/PackageName.cpp
--- a/Engine/Source/Runtime/CoreUObject/Private/Misc/PackageName.cpp
+++ b/Engine/Source/Runtime/CoreUObject/Private/Misc/PackageName.cpp
@@ -36,8 +36,6 @@ DEFINE_LOG_CATEGORY(LogPackageName);
 
 #define LOCTEXT_NAMESPACE "PackageNames"
 
-static FRWLock ContentMountPointCriticalSection;
-
 /** Event that is triggered when a new content path is mounted */
 FPackageName::FOnContentPathMountedEvent FPackageName::OnContentPathMountedEvent;
 
@@ -181,6 +179,8 @@ struct FPathPair
 
 struct FLongPackagePathsSingleton
 {
+	mutable FRWLock MountLock;
+
 	FString ConfigRootPath;
 	FString EngineRootPath;
 	FString GameRootPath;
@@ -285,7 +285,7 @@ struct FLongPackagePathsSingleton
 		FPathPair Pair(RootPath, RelativeContentPath, AbsolutePath.ToString());
 		
 		{
-			FWriteScopeLock ScopeLock(ContentMountPointCriticalSection);
+			FWriteScopeLock ScopeLock(MountLock);
 			ContentRootToPath.Insert(Pair, 0);
 			ContentPathToRoot.Insert(Pair, 0);
 			MountPointRootPaths.Add(RootPath);
@@ -313,7 +313,7 @@ struct FLongPackagePathsSingleton
 
 		bool bFirePathDismountedDelegate = false;
 		{
-			FWriteScopeLock ScopeLock(ContentMountPointCriticalSection);
+			FWriteScopeLock ScopeLock(MountLock);
 			if ( MountPointRootPaths.Remove(RootPath) > 0 )
 			{
 				// Absolute path doesn't need to be computed here as we are removing mount point
@@ -338,7 +338,7 @@ struct FLongPackagePathsSingleton
 	// Checks whether the specific root path is a valid mount point.
 	bool MountPointExists(const FString& RootPath)
 	{
-		FReadScopeLock ScopeLock(ContentMountPointCriticalSection);
+		FReadScopeLock ScopeLock(MountLock);
 		return GetValidLongPackageRoots(true/*bIncludeReadOnlyRoots*/).Contains(RootPath);
 	}
 
@@ -391,8 +391,6 @@ private:
 		GameConfigPathRebased  = RebasedGameDir / TEXT("Config/");
 		GameScriptPathRebased  = RebasedGameDir / TEXT("Script/");
 		GameSavedPathRebased   = RebasedGameDir / TEXT("Saved/");
-		
-		FWriteScopeLock ScopeLock(ContentMountPointCriticalSection);
 
 		ContentPathToRoot.Empty(13);
 		TStringBuilder<256> AbsolutePath;
@@ -441,7 +439,7 @@ private:
 	{
 		UE_LOG(LogPackageName, Log, TEXT("Valid mount points:"));
 
-		FReadScopeLock ScopeLock(ContentMountPointCriticalSection);
+		FReadScopeLock ScopeLock(MountLock);
 		for (const auto& Pair : ContentRootToPath)
 		{
 			UE_LOG(LogPackageName, Log, TEXT("	'%s' -> '%s'"), *Pair.RootPath, *Pair.ContentPath);
@@ -506,7 +504,7 @@ private:
 			
 			if (IsRootPath)
 			{
-				FReadScopeLock ScopeLock(ContentMountPointCriticalSection);
+				FReadScopeLock ScopeLock(MountLock);
 				
 				RootPath = Path;
 				auto Elements = Get().ContentRootToPath.FilterByPredicate([&RootPath](const FPathPair& elem){ return elem.RootPath.Equals(RootPath); });
@@ -530,7 +528,7 @@ private:
 			}
 			else
 			{
-				FReadScopeLock ScopeLock(ContentMountPointCriticalSection);
+				FReadScopeLock ScopeLock(MountLock);
 			
 				ContentPath = Path;
 				auto Elements = Get().ContentPathToRoot.FilterByPredicate([&ContentPath](const FPathPair& elem){ return elem.ContentPath.Equals(ContentPath); });
@@ -586,7 +584,7 @@ void FPackageName::InternalFilenameToLongPackageName(FStringView InFilename, FSt
 	// Convert to relative path if it's not already a long package name
 	bool bIsValidLongPackageName = false;
 	{
-		FReadScopeLock ScopeLock(ContentMountPointCriticalSection);
+		FReadScopeLock ScopeLock(Paths.MountLock);
 		for (const auto& Pair : Paths.ContentRootToPath)
 		{
 			if (FPathViews::IsParentPathOf(Pair.RootPath, Filename))
@@ -632,7 +630,7 @@ void FPackageName::InternalFilenameToLongPackageName(FStringView InFilename, FSt
 	}
 
 	{
-		FReadScopeLock ScopeLock(ContentMountPointCriticalSection);
+		FReadScopeLock ScopeLock(Paths.MountLock);
 		for (const auto& Pair : Paths.ContentPathToRoot)
 		{
 			FStringView RelPath;
@@ -655,7 +653,7 @@ void FPackageName::InternalFilenameToLongPackageName(FStringView InFilename, FSt
 			FPaths::NormalizeFilename(Filename);
 			Result = FPathViews::GetBaseFilenameWithPath(Filename);
 			{
-				FReadScopeLock ScopeLock(ContentMountPointCriticalSection);
+				FReadScopeLock ScopeLock(Paths.MountLock);
 				for (const auto& Pair : Paths.ContentPathToRoot)
 				{
 					FStringView RelPath;
@@ -749,8 +747,9 @@ FString FPackageName::FilenameToLongPackageName(const FString& InFilename)
 	{
 		TStringBuilder<128> ContentRoots;
 		{
-			FReadScopeLock ScopeLock(ContentMountPointCriticalSection);
-			for (const FPathPair& Pair : FLongPackagePathsSingleton::Get().ContentPathToRoot)
+			const auto& Paths = FLongPackagePathsSingleton::Get();
+			FReadScopeLock ScopeLock(Paths.MountLock);
+			for (const FPathPair& Pair : Paths.ContentPathToRoot)
 			{
 				ContentRoots << TEXT("\n\t\t") << Pair.ContentPath;
 				ContentRoots << TEXT("\n\t\t\t") << Pair.AbsolutePath;
@@ -776,7 +775,7 @@ FString FPackageName::FilenameToLongPackageName(const FString& InFilename)
 bool FPackageName::TryConvertLongPackageNameToFilename(const FString& InLongPackageName, FString& OutFilename, const FString& InExtension)
 {
 	const auto& Paths = FLongPackagePathsSingleton::Get();
-	FReadScopeLock ScopeLock(ContentMountPointCriticalSection);
+	FReadScopeLock ScopeLock(Paths.MountLock);
 	for (const auto& Pair : Paths.ContentRootToPath)
 	{
 		FStringView RelPath;
@@ -868,7 +867,7 @@ bool FPackageName::SplitLongPackageName(const FString& InLongPackageName, FStrin
 
 	const bool bIncludeReadOnlyRoots = true;
 
-	FReadScopeLock ScopeLock(ContentMountPointCriticalSection);
+	FReadScopeLock ScopeLock(Paths.MountLock);
 	const TArray<FString>& ValidRoots = Paths.GetValidLongPackageRoots(bIncludeReadOnlyRoots);
 
 	// Check to see whether our package came from a valid root
@@ -1061,9 +1060,9 @@ bool FPackageName::IsValidLongPackageName(FStringView InLongPackageName, bool bI
 		{
 			if (Reason == EErrorCode::PackageNamePathNotMounted)
 			{
-				FReadScopeLock ScopeLock(ContentMountPointCriticalSection);
-
 				const FLongPackagePathsSingleton& Paths = FLongPackagePathsSingleton::Get();
+
+				FReadScopeLock ScopeLock(Paths.MountLock);
 				const TArray<FString>& ValidRoots = Paths.GetValidLongPackageRoots(bIncludeReadOnlyRoots);
 				if (ValidRoots.Num() == 0)
 				{
@@ -1111,8 +1110,8 @@ bool FPackageName::IsValidLongPackageName(FStringView InLongPackageName, bool bI
 	}
 
 	// Check valid roots
-	FReadScopeLock ScopeLock(ContentMountPointCriticalSection);
 	const FLongPackagePathsSingleton& Paths = FLongPackagePathsSingleton::Get();
+	FReadScopeLock ScopeLock(Paths.MountLock);
 	const TArray<FString>& ValidRoots = Paths.GetValidLongPackageRoots(bIncludeReadOnlyRoots);
 	bool bValidRoot = false;
 	for (int32 RootIdx = 0; RootIdx < ValidRoots.Num(); ++RootIdx)
@@ -1195,7 +1194,7 @@ bool FPackageName::IsValidObjectPath(const FString& InObjectPath, FText* OutReas
 bool FPackageName::IsValidPath(const FString& InPath)
 {
 	const FLongPackagePathsSingleton& Paths = FLongPackagePathsSingleton::Get();
-	FReadScopeLock ScopeLock(ContentMountPointCriticalSection);
+	FReadScopeLock ScopeLock(Paths.MountLock);
 	for (const FPathPair& Pair : Paths.ContentRootToPath)
 	{
 		if (FPathViews::IsParentPathOf(Pair.RootPath, InPath))
@@ -1227,7 +1226,7 @@ FName FPackageName::GetPackageMountPoint(const FString& InPackagePath, bool InWi
 {
 	FLongPackagePathsSingleton& Paths = FLongPackagePathsSingleton::Get();
 	
-	FReadScopeLock ScopeLock(ContentMountPointCriticalSection);
+	FReadScopeLock ScopeLock(Paths.MountLock);
 	const TArray<FString>& MountPoints = Paths.GetValidLongPackageRoots(true);
 
 	int32 WithoutSlashes = InWithoutSlashes ? 1 : 0;
@@ -1397,7 +1396,7 @@ bool FPackageName::TryGetMountPointForPath(FStringView InFilePathOrPackageName,
 	TStringBuilder<256> PossibleAbsFilePath;
 	FPathViews::ToAbsolutePath(InFilePathOrPackageName, PossibleAbsFilePath);
 	const FLongPackagePathsSingleton& Paths = FLongPackagePathsSingleton::Get();
-	FReadScopeLock ScopeLock(ContentMountPointCriticalSection);
+	FReadScopeLock ScopeLock(Paths.MountLock);
 	for (const auto& Pair : Paths.ContentRootToPath)
 	{
 		FStringView RelPath;
@@ -2151,7 +2150,7 @@ void FPackageName::QueryRootContentPaths(TArray<FString>& OutRootContentPaths, b
 	const FLongPackagePathsSingleton& Paths = FLongPackagePathsSingleton::Get();
 
 	{
-		FReadScopeLock ScopeLock(ContentMountPointCriticalSection);
+		FReadScopeLock ScopeLock(Paths.MountLock);
 		OutRootContentPaths = Paths.GetValidLongPackageRoots( bIncludeReadOnlyRoots );
 	}
 
