Enable backwards compatibility for cooked AssetRegistries in cooked builds

diff --git a/Engine/Source/Runtime/CoreUObject/Private/AssetRegistry/AssetData.cpp b/Engine/Source/Runtime/CoreUObject/Private/AssetRegistry/AssetData.cpp
--- a/Engine/Source/Runtime/CoreUObject/Private/AssetRegistry/AssetData.cpp
+++ b/Engine/Source/Runtime/CoreUObject/Private/AssetRegistry/AssetData.cpp
@@ -208,20 +208,18 @@ void FAssetData::SerializeForCacheInternal(FArchive& Ar, FAssetRegistryVersion::
 	Ar << ObjectPath;
 	Ar << PackagePath;
 
-	// Serialize the asset class. Legacy path is editor only for DiffAssetRegistriesCommandlet
-#if WITH_EDITORONLY_DATA
-	if (Version < FAssetRegistryVersion::ClassPaths)
+	// Serialize the asset class.
+	if (Version >= FAssetRegistryVersion::ClassPaths)
+	{
+		Ar << AssetClassPath;
+	}
+	else
 	{
 PRAGMA_DISABLE_DEPRECATION_WARNINGS
 		Ar << AssetClass;
 		AssetClassPath = FAssetData::TryConvertShortClassNameToPathName(AssetClass, ELogVerbosity::NoLogging);
 PRAGMA_ENABLE_DEPRECATION_WARNINGS
 	}
-	else
-#endif 
-	{
-		Ar << AssetClassPath;
-	}
 
 	// These are derived from ObjectPath, we manually serialize them because they get pooled
 	Ar << PackageName;
diff --git a/Engine/Source/Runtime/CoreUObject/Private/AssetRegistry/AssetDataTagMap.cpp b/Engine/Source/Runtime/CoreUObject/Private/AssetRegistry/AssetDataTagMap.cpp
--- a/Engine/Source/Runtime/CoreUObject/Private/AssetRegistry/AssetDataTagMap.cpp
+++ b/Engine/Source/Runtime/CoreUObject/Private/AssetRegistry/AssetDataTagMap.cpp
@@ -1093,7 +1093,7 @@ namespace FixedTagPrivate
 		static constexpr uint32 EndMagic		= 0x87654321;
 		static constexpr uint32 MaxViewAlignment = 16;
 	public:
-		TSerializer(FArchive& InAr) // , FAssetRegistryVersion::Type InVersion = FAssetRegistryVersion::LatestVersion
+		TSerializer(FArchive& InAr)
 			: Ar(InAr)
 		{
 		}
@@ -1243,13 +1243,11 @@ namespace FixedTagPrivate
 		// branches in performance critical serialization code.
 		// Note that we branch in this function only because it's the only fucntion used to load legacy 
 		// asset registries in DiffAssetRegistries commandlet hence it's editor only code too
-#if WITH_EDITORONLY_DATA
 		if (Version < FAssetRegistryVersion::ClassPaths)
 		{
-			TSerializer<FAssetRegistryVersion::ObjectResourceOptionalVersionChange>(Ar).Load(*Store);
+			TSerializer<FAssetRegistryVersion::AddedChunkHashes>(Ar).Load(*Store);
 		}
 		else // Add more branches should serialization format change
-#endif // WITH_EDITORONLY_DATA
 		{			
 			TSerializer<FAssetRegistryVersion::LatestVersion>(Ar).Load(*Store);
 		}
