Allow BP object and field names to be used with core redirects

These allow various characters that are otherwise invalid from C++

diff --git a/Engine/Source/Runtime/CoreUObject/Private/UObject/CoreRedirects.cpp b/Engine/Source/Runtime/CoreUObject/Private/UObject/CoreRedirects.cpp
--- a/Engine/Source/Runtime/CoreUObject/Private/UObject/CoreRedirects.cpp
+++ b/Engine/Source/Runtime/CoreUObject/Private/UObject/CoreRedirects.cpp
@@ -178,12 +178,13 @@ void FCoreRedirectObjectName::UnionFieldsInline(const FCoreRedirectObjectName& O
 	}
 }
 
-bool FCoreRedirectObjectName::HasValidCharacters() const
+bool FCoreRedirectObjectName::HasValidCharacters(ECoreRedirectFlags Type) const
 {
-	// ObjectNames in Blueprint may contain spaces.
 	static FString InvalidRedirectCharacters = TEXT("\"',|&!~\n\r\t@#(){}[]=;^%$`");
+	static FString InvalidObjectRedirectCharacters = TEXT(".\n\r\t"); // Object and field names in Blueprints are very permissive...
 
-	return ObjectName.IsValidXName(InvalidRedirectCharacters) && OuterName.IsValidXName(InvalidRedirectCharacters) && PackageName.IsValidXName(InvalidRedirectCharacters);
+	return ObjectName.IsValidXName(EnumHasAnyFlags(Type, ECoreRedirectFlags::Type_Object | ECoreRedirectFlags::Type_Property | ECoreRedirectFlags::Type_Function) ? InvalidObjectRedirectCharacters : InvalidRedirectCharacters)
+		&& OuterName.IsValidXName(InvalidRedirectCharacters) && PackageName.IsValidXName(InvalidRedirectCharacters);
 }
 
 bool FCoreRedirectObjectName::ExpandNames(const FString& InString, FName& OutName, FName& OutOuter, FName &OutPackage)
@@ -1060,8 +1061,8 @@ bool FCoreRedirects::AddRedirectList(TArrayView<const FCoreRedirect> Redirects,
 			continue;
 		}
 
-		if ((!NewRedirect.OldName.HasValidCharacters() && !FPackageName::IsVersePackage(NewRedirect.OldName.PackageName.ToString()))
-			|| (!NewRedirect.NewName.HasValidCharacters() && !FPackageName::IsVersePackage(NewRedirect.NewName.PackageName.ToString())))
+		if ((!NewRedirect.OldName.HasValidCharacters(NewRedirect.RedirectFlags) && !FPackageName::IsVersePackage(NewRedirect.OldName.PackageName.ToString()))
+			|| (!NewRedirect.NewName.HasValidCharacters(NewRedirect.RedirectFlags) && !FPackageName::IsVersePackage(NewRedirect.NewName.PackageName.ToString())))
 		{
 			UE_LOG(LogLinker, Error, TEXT("AddRedirectList(%s) failed to add redirector from %s to %s with invalid characters!"), *SourceString, *NewRedirect.OldName.ToString(), *NewRedirect.NewName.ToString());
 			continue;
@@ -1187,7 +1188,7 @@ bool FCoreRedirects::RemoveRedirectList(TArrayView<const FCoreRedirect> Redirect
 			continue;
 		}
 
-		if (!RedirectToRemove.OldName.HasValidCharacters() || !RedirectToRemove.NewName.HasValidCharacters())
+		if (!RedirectToRemove.OldName.HasValidCharacters(RedirectToRemove.RedirectFlags) || !RedirectToRemove.NewName.HasValidCharacters(RedirectToRemove.RedirectFlags))
 		{
 			UE_LOG(LogLinker, Error, TEXT("RemoveRedirectList(%s) failed to remove redirector from %s to %s with invalid characters!"), *SourceString, *RedirectToRemove.OldName.ToString(), *RedirectToRemove.NewName.ToString());
 			continue;
diff --git a/Engine/Source/Runtime/CoreUObject/Public/UObject/CoreRedirects.h b/Engine/Source/Runtime/CoreUObject/Public/UObject/CoreRedirects.h
--- a/Engine/Source/Runtime/CoreUObject/Public/UObject/CoreRedirects.h
+++ b/Engine/Source/Runtime/CoreUObject/Public/UObject/CoreRedirects.h
@@ -163,7 +163,7 @@ struct COREUOBJECT_API FCoreRedirectObjectName
 	}
 
 	/** Returns true if all names have valid characters */
-	bool HasValidCharacters() const;
+	bool HasValidCharacters(ECoreRedirectFlags Type) const;
 
 	/** Expand OldName/NewName as needed */
 	static bool ExpandNames(const FString& FullString, FName& OutName, FName& OutOuter, FName &OutPackage);
