Remove assert that fires when we try to point an already valid FEditorBulkData to reference older style bulkdata.

#rnx

- The assert isn't really required since we reset the editor bulkdata before we change it to reference the older style bulkdata. So the only purpose that the assert serves is to force the caller to always call ::Reset anyway.

diff --git a/Engine/Source/Runtime/CoreUObject/Private/Serialization/EditorBulkData.cpp b/Engine/Source/Runtime/CoreUObject/Private/Serialization/EditorBulkData.cpp
--- a/Engine/Source/Runtime/CoreUObject/Private/Serialization/EditorBulkData.cpp
+++ b/Engine/Source/Runtime/CoreUObject/Private/Serialization/EditorBulkData.cpp
@@ -659,10 +659,6 @@ void FEditorBulkData::LogRegisterError(UE::BulkDataRegistry::ERegisterResult Val
 void FEditorBulkData::CreateFromBulkData(FBulkData& InBulkData, const FGuid& InGuid, UObject* Owner)
 {
 	TRACE_CPUPROFILER_EVENT_SCOPE(FEditorBulkData::CreateFromBulkData);
-	
-	checkf(!BulkDataId.IsValid(), 
-		TEXT("Calling ::CreateFromBulkData on a bulkdata object that already has a valid identifier! Package: '%s'"),
-		*InBulkData.GetPackagePath().GetDebugName());
 
 	Reset();
 
