Optimizing FAssetData loading/saving by removing versioning support from the main serialization path

diff --git a/Engine/Source/Runtime/AssetRegistry/Private/AssetRegistryState.cpp b/Engine/Source/Runtime/AssetRegistry/Private/AssetRegistryState.cpp
--- a/Engine/Source/Runtime/AssetRegistry/Private/AssetRegistryState.cpp
+++ b/Engine/Source/Runtime/AssetRegistry/Private/AssetRegistryState.cpp
@@ -1048,7 +1048,7 @@ bool FAssetRegistryState::Save(FArchive& OriginalAr, const FAssetRegistrySeriali
 	for (const TPair<FName, FAssetData*>& Pair : SortedAssetsByObjectPath)
 	{
 		// Hardcoding FAssetRegistryVersion::LatestVersion here so that branches can get optimized out in the forceinlined SerializeForCache
-		Pair.Value->SerializeForCache(Ar, FAssetRegistryVersion::LatestVersion);
+		Pair.Value->SerializeForCache(Ar);
 	}
 
 	// Serialize Dependencies
@@ -1221,9 +1221,21 @@ void FAssetRegistryState::Load(Archive&& Ar, FAssetRegistryVersion::Type Version
 	TArrayView<FAssetData> PreallocatedAssetDataBuffer(new FAssetData[LocalNumAssets], LocalNumAssets);
 	PreallocatedAssetDataBuffers.Add(PreallocatedAssetDataBuffer.GetData());
 
-	for (FAssetData& NewAssetData : PreallocatedAssetDataBuffer)
+	// Optimizing serialization of latest asset data format by moving version checking out of SerializeForCache function and falling back to versioned
+	// serialization should we attempt to load an older version of AR (usually commandlets)
+	if (Version == FAssetRegistryVersion::LatestVersion)
 	{
-		NewAssetData.SerializeForCache(Ar, Version);
+		for (FAssetData& NewAssetData : PreallocatedAssetDataBuffer)
+		{
+			NewAssetData.SerializeForCache(Ar);
+		}
+	}
+	else
+	{
+		for (FAssetData& NewAssetData : PreallocatedAssetDataBuffer)
+		{
+			NewAssetData.SerializeForCacheOldVersion(Ar, Version);
+		}
 	}
 
 	SetAssetDatas(PreallocatedAssetDataBuffer, Options);
diff --git a/Engine/Source/Runtime/CoreUObject/Private/AssetRegistry/AssetData.cpp b/Engine/Source/Runtime/CoreUObject/Private/AssetRegistry/AssetData.cpp
--- a/Engine/Source/Runtime/CoreUObject/Private/AssetRegistry/AssetData.cpp
+++ b/Engine/Source/Runtime/CoreUObject/Private/AssetRegistry/AssetData.cpp
@@ -240,6 +240,15 @@ PRAGMA_ENABLE_DEPRECATION_WARNINGS
 	Ar << PackageFlags;
 }
 
+void FAssetData::SerializeForCacheWithTagsAndBundles(FArchive& Ar, void (*SerializeTagsAndBundles)(FArchive&, FAssetData&))
+{
+	SerializeForCacheInternal(Ar, FAssetRegistryVersion::LatestVersion, SerializeTagsAndBundles);
+}
+void FAssetData::SerializeForCacheOldVersionWithTagsAndBundles(FArchive& Ar, FAssetRegistryVersion::Type Version, void (*SerializeTagsAndBundles)(FArchive&, FAssetData&))
+{
+	SerializeForCacheInternal(Ar, Version, SerializeTagsAndBundles);
+}
+
 FTopLevelAssetPath FAssetData::TryConvertShortClassNameToPathName(FName InClassName, ELogVerbosity::Type FailureMessageVerbosity /*= ELogVerbosity::Warning*/)
 {
 	FTopLevelAssetPath ClassPath;
diff --git a/Engine/Source/Runtime/CoreUObject/Public/AssetRegistry/AssetData.h b/Engine/Source/Runtime/CoreUObject/Public/AssetRegistry/AssetData.h
--- a/Engine/Source/Runtime/CoreUObject/Public/AssetRegistry/AssetData.h
+++ b/Engine/Source/Runtime/CoreUObject/Public/AssetRegistry/AssetData.h
@@ -524,14 +524,31 @@ PRAGMA_ENABLE_DEPRECATION_WARNINGS
 	 * To version this data change FAssetRegistryVersion
 	 */
 	template<class Archive>
-	FORCEINLINE void SerializeForCache(Archive&& Ar, FAssetRegistryVersion::Type Version = FAssetRegistryVersion::LatestVersion)
+	FORCEINLINE void SerializeForCache(Archive&& Ar)
 	{
-		SerializeForCacheInternal(Ar, Version, [](FArchive& Ar, FAssetData& Ad) {
+		SerializeForCacheWithTagsAndBundles(Ar, [](FArchive& Ar, FAssetData& Ad) {
+			static_cast<Archive&>(Ar).SerializeTagsAndBundles(Ad);
+			});
+	}
+	/**
+	 * Serialize as part of the registry cache using legacy paths (versioned)
+	 */
+	template<class Archive>
+	FORCEINLINE void SerializeForCacheOldVersion(Archive&& Ar, FAssetRegistryVersion::Type Version = FAssetRegistryVersion::LatestVersion)
+	{
+		SerializeForCacheOldVersionWithTagsAndBundles(Ar, Version, [](FArchive& Ar, FAssetData& Ad) {
 			static_cast<Archive&>(Ar).SerializeTagsAndBundles(Ad);
 			});
 	}
 private:
-	COREUOBJECT_API void SerializeForCacheInternal(FArchive& Ar, FAssetRegistryVersion::Type Version, void (*SerializeTagsAndBundles)(FArchive& , FAssetData&));
+	/**
+	 * The actual implementation of SerializeForCache.
+	 * Note that this function is force-inlined but defined in AssetData.cpp which is fine as functions will get inlined
+	 * as long as they're defined before they are used for the first time by other functions (SerializeForCacheWithTagsAndBundles in this case)
+	 */
+	FORCEINLINE void SerializeForCacheInternal(FArchive& Ar, FAssetRegistryVersion::Type Version, void (*SerializeTagsAndBundles)(FArchive& , FAssetData&));
+	COREUOBJECT_API void SerializeForCacheWithTagsAndBundles(FArchive& Ar, void (*SerializeTagsAndBundles)(FArchive&, FAssetData&));
+	COREUOBJECT_API void SerializeForCacheOldVersionWithTagsAndBundles(FArchive& Ar, FAssetRegistryVersion::Type Version, void (*SerializeTagsAndBundles)(FArchive&, FAssetData&));
 
 	static bool DetectIsUAssetByNames(FStringView PackageName, FStringView ObjectPathName)
 	{
