Made FPropertyChangedEvent::GetArrayIndex a const function. Added FPropertyChangedEvent::GetMemberPropertyName as a complement to GetPropertyName.

#rnx

diff --git a/Engine/Source/Runtime/CoreUObject/Public/UObject/UnrealType.h b/Engine/Source/Runtime/CoreUObject/Public/UObject/UnrealType.h
--- a/Engine/Source/Runtime/CoreUObject/Public/UObject/UnrealType.h
+++ b/Engine/Source/Runtime/CoreUObject/Public/UObject/UnrealType.h
@@ -6371,7 +6371,7 @@ struct FPropertyChangedEvent
 	 * Gets the Array Index of the "current object" based on a particular name
 	 * InName - Name of the property to find the array index for
 	 */
-	int32 GetArrayIndex(const FString& InName)
+	int32 GetArrayIndex(const FString& InName) const
 	{
 		//default to unknown index
 		int32 Retval = -1;
@@ -6416,6 +6416,14 @@ struct FPropertyChangedEvent
 		return (Property != nullptr) ? Property->GetFName() : NAME_None;
 	}
 
+	/**
+	 * Simple utility to get the name of the object's member property and takes care of the possible null property.
+	 */
+	FName GetMemberPropertyName() const
+	{
+		return (MemberProperty != nullptr) ? MemberProperty->GetFName() : NAME_None;
+	}
+
 	/**
 	 * The actual property that changed
 	 */
