Add serialization support for cooking editor bulk data as empty during cooking
Actual payload support will be implemented separately

diff --git a/Engine/Source/Runtime/CoreUObject/Private/Serialization/EditorBulkData.cpp b/Engine/Source/Runtime/CoreUObject/Private/Serialization/EditorBulkData.cpp
--- a/Engine/Source/Runtime/CoreUObject/Private/Serialization/EditorBulkData.cpp
+++ b/Engine/Source/Runtime/CoreUObject/Private/Serialization/EditorBulkData.cpp
@@ -989,7 +989,6 @@ void FEditorBulkData::Serialize(FArchive& Ar, UObject* Owner, bool bAllowRegiste
 
 			// This cannot be inside the above ::IsStoredInPackageTrailer branch due to the original prototype assets using the trailer without the StoredInPackageTrailer flag
 			if (Trailer != nullptr && Trailer->FindPayloadStatus(PayloadContentId) == EPayloadStatus::StoredVirtualized)
-			{
 				// As the virtualization process happens outside of serialization we need
 				// to check with the trailer to see if the payload is virtualized or not
 				EnumAddFlags(Flags, EFlags::IsVirtualized);
@@ -1997,6 +1996,13 @@ FEditorBulkData::EFlags FEditorBulkData::BuildFlagsForSerialization(FArchive& Ar
 			EnumAddFlags(UpdatedFlags, EFlags::StoredInPackageTrailer);
 		}
 
+		// if we are cooking add the cooked flag to the serialized stream
+		// Editor bulk data might be cooked when using editor optional data at runtime
+		if (Ar.IsCooking())
+		{
+			EnumAddFlags(UpdatedFlags, EFlags::IsCooked);
+		}
+
 		return UpdatedFlags;
 	}
 	else
diff --git a/Engine/Source/Runtime/CoreUObject/Private/UObject/SavePackage/SavePackageUtilities.cpp b/Engine/Source/Runtime/CoreUObject/Private/UObject/SavePackage/SavePackageUtilities.cpp
--- a/Engine/Source/Runtime/CoreUObject/Private/UObject/SavePackage/SavePackageUtilities.cpp
+++ b/Engine/Source/Runtime/CoreUObject/Private/UObject/SavePackage/SavePackageUtilities.cpp
@@ -2103,15 +2103,16 @@ ESavePackageResult SaveBulkData(FLinkerSave* Linker, int64& InOutStartOffset, co
 			}
 			if (OptionalBulkArchive && OptionalBulkArchive->TotalSize())
 			{
-				checkf(!bIsOptionalRealm, TEXT("OptionalBulkData is currently unsupported with optional package multi output"));
-				BulkInfo.ChunkId = CreateIoChunkId(PackageId.Value(), 0, EIoChunkType::OptionalBulkData);
+				// @note FH: temporarily comment assert for not supporting optional bulk data into editor optional packages
+				//checkf(!bIsOptionalRealm, TEXT("OptionalBulkData is currently unsupported with optional package multi output for %s"), *InOuter->GetName());
+				BulkInfo.ChunkId = CreateIoChunkId(PackageId.Value(), BulkInfo.MultiOutputIndex, EIoChunkType::OptionalBulkData);
 				BulkInfo.BulkDataType = IPackageWriter::FBulkDataInfo::Optional;
 				BulkInfo.LooseFilePath = FPathViews::ChangeExtension(BulkInfo.LooseFilePath, LexToString(EPackageExtension::BulkDataOptional));
 				PackageWriter->WriteBulkData(BulkInfo, AddSizeAndConvertToIoBuffer(OptionalBulkArchive.Get()), OptionalBulkArchive->FileRegions);
 			}
 			if (MappedBulkArchive && MappedBulkArchive->TotalSize())
 			{
-				checkf(!bIsOptionalRealm, TEXT("MemoryMappedBulkData is currently unsupported with optional package multi output"));
+				checkf(!bIsOptionalRealm, TEXT("MemoryMappedBulkData is currently unsupported with optional package multi output for %s"), *InOuter->GetName());
 				BulkInfo.ChunkId = CreateIoChunkId(PackageId.Value(), 0, EIoChunkType::MemoryMappedBulkData);
 				BulkInfo.BulkDataType = IPackageWriter::FBulkDataInfo::Mmap;
 				BulkInfo.LooseFilePath = FPathViews::ChangeExtension(BulkInfo.LooseFilePath, LexToString(EPackageExtension::BulkDataMemoryMapped));
@@ -2120,7 +2121,7 @@ ESavePackageResult SaveBulkData(FLinkerSave* Linker, int64& InOutStartOffset, co
 		}
 		else
 		{
-			checkf(!bIsOptionalRealm, TEXT("Package optional package multi output is unsupported without a PackageWriter"));
+			checkf(!bIsOptionalRealm, TEXT("Package optional package multi output is unsupported without a PackageWriter for %s"), *InOuter->GetName());
 			auto WriteBulkData = [&](FLargeMemoryWriterWithRegions* Archive, const TCHAR* BulkFileExtension)
 			{
 				if (const int64 DataSize = Archive ? Archive->TotalSize() : 0)
diff --git a/Engine/Source/Runtime/CoreUObject/Public/Serialization/EditorBulkData.h b/Engine/Source/Runtime/CoreUObject/Public/Serialization/EditorBulkData.h
--- a/Engine/Source/Runtime/CoreUObject/Public/Serialization/EditorBulkData.h
+++ b/Engine/Source/Runtime/CoreUObject/Public/Serialization/EditorBulkData.h
@@ -349,6 +349,8 @@ private:
 		ReferencesWorkspaceDomain	= 1 << 8,
 		/** The payload is stored in a package trailer, so the bulkdata object will have to poll the trailer to find the payload offset */
 		StoredInPackageTrailer		= 1 << 9,
+		/** The bulkdata object was cooked. */
+		IsCooked					= 1 << 10,
 
 		TransientFlags				= HasRegistered | IsTornOff,
 	};
diff --git a/Engine/Source/Runtime/Engine/Private/EngineUtils.cpp b/Engine/Source/Runtime/Engine/Private/EngineUtils.cpp
--- a/Engine/Source/Runtime/Engine/Private/EngineUtils.cpp
+++ b/Engine/Source/Runtime/Engine/Private/EngineUtils.cpp
@@ -556,7 +556,7 @@ FStripDataFlags::FStripDataFlags( class FArchive& Ar, uint8 InClassFlags /*= 0*/
 		{
 			// When cooking GlobalStripFlags are automatically generated based on the current target
 			// platform's properties.
-			GlobalStripFlags |= Ar.CookingTarget()->HasEditorOnlyData() ? static_cast<uint8>(FStripDataFlags::EStrippedData::None) : static_cast<uint8>(FStripDataFlags::EStrippedData::Editor);
+			GlobalStripFlags |= Ar.IsFilterEditorOnly() ? static_cast<uint8>(FStripDataFlags::EStrippedData::Editor) : static_cast<uint8>(FStripDataFlags::EStrippedData::None);
 			GlobalStripFlags |= !Ar.CookingTarget()->AllowAudioVisualData() ? static_cast<uint8>(FStripDataFlags::EStrippedData::Server) : static_cast<uint8>(FStripDataFlags::EStrippedData::None);
 			ClassStripFlags = InClassFlags;
 		}
@@ -600,7 +600,7 @@ FStripDataFlags::FStripDataFlags(FStructuredArchive::FSlot Slot, uint8 InClassFl
 		{
 			// When cooking GlobalStripFlags are automatically generated based on the current target
 			// platform's properties.
-			GlobalStripFlags |= UnderlyingArchive.CookingTarget()->HasEditorOnlyData() ? static_cast<uint8>(FStripDataFlags::EStrippedData::None) : static_cast<uint8>(FStripDataFlags::EStrippedData::Editor);
+			GlobalStripFlags |= UnderlyingArchive.IsFilterEditorOnly() ? static_cast<uint8>(FStripDataFlags::EStrippedData::Editor) : static_cast<uint8>(FStripDataFlags::EStrippedData::None);
 			GlobalStripFlags |= !UnderlyingArchive.CookingTarget()->AllowAudioVisualData() ? static_cast<uint8>(FStripDataFlags::EStrippedData::Server) : static_cast<uint8>(FStripDataFlags::EStrippedData::None);
 			ClassStripFlags = InClassFlags;
 		}
