Fix merge typo eating a bracket

diff --git a/Engine/Source/Runtime/CoreUObject/Private/Serialization/EditorBulkData.cpp b/Engine/Source/Runtime/CoreUObject/Private/Serialization/EditorBulkData.cpp
--- a/Engine/Source/Runtime/CoreUObject/Private/Serialization/EditorBulkData.cpp
+++ b/Engine/Source/Runtime/CoreUObject/Private/Serialization/EditorBulkData.cpp
@@ -989,6 +989,7 @@ void FEditorBulkData::Serialize(FArchive& Ar, UObject* Owner, bool bAllowRegiste
 
 			// This cannot be inside the above ::IsStoredInPackageTrailer branch due to the original prototype assets using the trailer without the StoredInPackageTrailer flag
 			if (Trailer != nullptr && Trailer->FindPayloadStatus(PayloadContentId) == EPayloadStatus::StoredVirtualized)
+			{
 				// As the virtualization process happens outside of serialization we need
 				// to check with the trailer to see if the payload is virtualized or not
 				EnumAddFlags(Flags, EFlags::IsVirtualized);
