Clean up FEditorBulkData flags and associated serialization code.
#rnx

EFlags::TransientFlags has been removed from the enum and instead is stored as a constant. This will allow the debugger to more easily display the flag member.
Move the removal of transient bulkdata flags when saving to ::BuildFlagsForSerialization so that all flag manipulation is done in the same place.
We now update the flags at the end of serialization and remove the EFlags::IsVirtualized bit (if set) if the payload is being stored in a package trailer
-- I originally tried to do this by removing the bit in ::BuildFlagsForSerialization, but that was partially reverted in CL 19976045 as it was causing the editor bulkdata is rehydrate as we check for EFlags::IsVirtualized in several places when saving out the bulkdata
-- Moving the bit removal to the end of serialization is the easiest and safest way to fix this. However with some more work we could probably remove more branches from this logic but I'd rather move that to a separate work item.
Moved the paranoid checks to make sure that the package trailer builder was created correctly to a static method ::ValidatePackageTrailerBuilder in order to make the ::Serialize method easier to read.

diff --git a/Engine/Source/Runtime/CoreUObject/Private/Serialization/EditorBulkData.cpp b/Engine/Source/Runtime/CoreUObject/Private/Serialization/EditorBulkData.cpp
--- a/Engine/Source/Runtime/CoreUObject/Private/Serialization/EditorBulkData.cpp
+++ b/Engine/Source/Runtime/CoreUObject/Private/Serialization/EditorBulkData.cpp
@@ -428,7 +428,7 @@ FEditorBulkData& FEditorBulkData::operator=(const FEditorBulkData& Other)
 	Flags = Other.Flags;
 	CompressionSettings = Other.CompressionSettings;
 
-	EnumRemoveFlags(Flags, EFlags::TransientFlags);
+	EnumRemoveFlags(Flags, TransientFlags);
 
 	if (bTornOff)
 	{
@@ -815,7 +815,7 @@ void FEditorBulkData::Serialize(FArchive& Ar, UObject* Owner, bool bAllowRegiste
 		Ar << Flags;
 		if (Ar.IsLoading())
 		{
-			EnumRemoveFlags(Flags, EFlags::TransientFlags);
+			EnumRemoveFlags(Flags, TransientFlags);
 		}
 
 		// TODO: Can probably remove these checks before UE5 release
@@ -835,14 +835,6 @@ void FEditorBulkData::Serialize(FArchive& Ar, UObject* Owner, bool bAllowRegiste
 			checkf(Ar.IsCooking() == false, TEXT("FEditorBulkData::Serialize should not be called during a cook"));
 
 			EFlags UpdatedFlags = BuildFlagsForSerialization(Ar, bKeepFileDataByReference);
-			EnumRemoveFlags(UpdatedFlags, EFlags::TransientFlags);
-
-			// Go back in the archive and update the flags in the archive, we will only apply the updated flags to the current
-			// object later if we detect that the package saved successfully.
-			// TODO: Not a huge fan of this, might be better to find a way to build the flags during serialization and potential callbacks 
-			// later then go back and update the flags in the Ar. Applying the updated flags only if we are saving a package to disk
-			// and the save succeeds continues to make sense.
-			UpdateArchiveData(Ar, SavedFlagsPos, UpdatedFlags);
 
 			// Write out required extra data if we're saving by reference
 			bool bWriteOutPayload = true;
@@ -912,29 +904,27 @@ void FEditorBulkData::Serialize(FArchive& Ar, UObject* Owner, bool bAllowRegiste
 					UpdateArchiveData(Ar, OffsetPos, DataStartOffset);
 				}
 			}
-
-			// Make sure that the trailer builder is correct (if it is being used)
-			if (IsStoredInPackageTrailer(UpdatedFlags) && !PayloadContentId.IsZero())
+			
+			if (IsStoredInPackageTrailer(UpdatedFlags) && IsDataVirtualized(UpdatedFlags))
 			{
-				check(LinkerSave != nullptr);
-				check(LinkerSave->PackageTrailerBuilder.IsValid());
-				checkf(!(IsDataVirtualized(UpdatedFlags) && IsReferencingByPackagePath(UpdatedFlags)), TEXT("Payload cannot be both virtualized and a reference"));
+				LinkerSave->PackageTrailerBuilder->AddVirtualizedPayload(PayloadContentId, PayloadSize);
+			}
 
-				if (IsReferencingByPackagePath(UpdatedFlags))
-				{
-					check(LinkerSave->PackageTrailerBuilder->IsReferencedPayloadEntry(PayloadContentId));
-				}
-				else if (IsDataVirtualized(UpdatedFlags))
-				{
-					LinkerSave->PackageTrailerBuilder->AddVirtualizedPayload(PayloadContentId, PayloadSize);
-					check(LinkerSave->PackageTrailerBuilder->IsVirtualizedPayloadEntry(PayloadContentId));
-				}
-				else
-				{
-					check(LinkerSave->PackageTrailerBuilder->IsLocalPayloadEntry(PayloadContentId));
-				}
+			/** Beyond this point UpdatedFlags will be modified to avoid serializing some flags */
+			/** So any code actually using UpdatedFlags should come before this section of code */
+
+			ValidatePackageTrailerBuilder(LinkerSave, PayloadContentId, UpdatedFlags);
+
+			// Remove the IsVirtualized flag if we are storing the payload in a package trailer before we serialize the flags
+			// to the package export data. We will determine if a payload is virtualized or not by checking the package trailer.
+			if (IsStoredInPackageTrailer(UpdatedFlags))
+			{
+				EnumRemoveFlags(UpdatedFlags, EFlags::IsVirtualized);
 			}
-			
+
+			// Replace the flags we wrote out earlier with the updated, final values
+			UpdateArchiveData(Ar, SavedFlagsPos, UpdatedFlags);
+
 			if (CanUnloadData())
 			{
 				this->CompressionSettings.Reset();
@@ -1056,13 +1046,13 @@ void FEditorBulkData::SerializeForRegistry(FArchive& Ar)
 	{
 		check(CanSaveForRegistry());
 		EFlags FlagsForSerialize = Flags;
-		EnumRemoveFlags(FlagsForSerialize, EFlags::TransientFlags);
+		EnumRemoveFlags(FlagsForSerialize, TransientFlags);
 		Ar << FlagsForSerialize;
 	}
 	else
 	{
 		Ar << Flags;
-		EnumRemoveFlags(Flags, EFlags::TransientFlags);
+		EnumRemoveFlags(Flags, TransientFlags);
 		EnumAddFlags(Flags, EFlags::IsTornOff);
 	}
 
@@ -1533,7 +1523,7 @@ void FEditorBulkData::SerializeToLegacyPath(FLinkerSave& LinkerSave, FCompressed
 				this->PackagePath = InPackagePath;
 				check(!this->PackagePath.IsEmpty()); // LinkerSave guarantees a valid PackagePath if we're updating loaded path
 				this->OffsetInFile = DataStartOffset;
-				this->Flags = (this->Flags & EFlags::TransientFlags) | (UpdatedFlags & ~EFlags::TransientFlags);
+				this->Flags = (this->Flags & TransientFlags) | (UpdatedFlags & ~TransientFlags);
 
 				if (CanUnloadData())
 				{
@@ -1579,7 +1569,7 @@ void FEditorBulkData::SerializeToPackageTrailer(FLinkerSave& LinkerSave, FCompre
 				}
 
 				this->OffsetInFile = PayloadOffset;
-				this->Flags = (this->Flags & EFlags::TransientFlags) | (UpdatedFlags & ~EFlags::TransientFlags);
+				this->Flags = (this->Flags & TransientFlags) | (UpdatedFlags & ~TransientFlags);
 	
 				// If the payload is valid we might want to fix up the package path or virtualization flags now
 				// that the package has saved.
@@ -1997,12 +1987,15 @@ FEditorBulkData::EFlags FEditorBulkData::BuildFlagsForSerialization(FArchive& Ar
 			EnumAddFlags(UpdatedFlags, EFlags::StoredInPackageTrailer);
 		}
 
-		// if we are cooking add the cooked flag to the serialized stream
+		// If we are cooking add the cooked flag to the serialized stream
 		// Editor bulk data might be cooked when using editor optional data at runtime
 		if (Ar.IsCooking())
 		{
 			EnumAddFlags(UpdatedFlags, EFlags::IsCooked);
 		}
+		
+		// Transient flags are never serialized
+		EnumRemoveFlags(UpdatedFlags, TransientFlags);
 
 		return UpdatedFlags;
 	}
@@ -2065,6 +2058,29 @@ FText FEditorBulkData::GetCorruptedPayloadErrorMsgForSave(FLinkerSave* Linker) c
 	}
 }
 
+void FEditorBulkData::ValidatePackageTrailerBuilder(const FLinkerSave* LinkerSave, const FIoHash& Id, EFlags PayloadFlags)
+{
+	if (IsStoredInPackageTrailer(PayloadFlags) && !Id.IsZero())
+	{
+		check(LinkerSave != nullptr);
+		check(LinkerSave->PackageTrailerBuilder.IsValid());
+		checkf(!(IsDataVirtualized(PayloadFlags) && IsReferencingByPackagePath(PayloadFlags)), TEXT("Payload cannot be both virtualized and a reference"));
+
+		if (IsReferencingByPackagePath(PayloadFlags))
+		{
+			check(LinkerSave->PackageTrailerBuilder->IsReferencedPayloadEntry(Id));
+		}
+		else if (IsDataVirtualized(PayloadFlags))
+		{
+			check(LinkerSave->PackageTrailerBuilder->IsVirtualizedPayloadEntry(Id));
+		}
+		else
+		{
+			check(LinkerSave->PackageTrailerBuilder->IsLocalPayloadEntry(Id));
+		}
+	}
+}
+
 bool FEditorBulkData::ShouldUseLegacySerialization(const FLinkerSave* LinkerSave) const
 {
 	if (LinkerSave == nullptr)
diff --git a/Engine/Source/Runtime/CoreUObject/Public/Serialization/EditorBulkData.h b/Engine/Source/Runtime/CoreUObject/Public/Serialization/EditorBulkData.h
--- a/Engine/Source/Runtime/CoreUObject/Public/Serialization/EditorBulkData.h
+++ b/Engine/Source/Runtime/CoreUObject/Public/Serialization/EditorBulkData.h
@@ -350,13 +350,14 @@ private:
 		/** The payload is stored in a package trailer, so the bulkdata object will have to poll the trailer to find the payload offset */
 		StoredInPackageTrailer		= 1 << 9,
 		/** The bulkdata object was cooked. */
-		IsCooked					= 1 << 10,
-
-		TransientFlags				= HasRegistered | IsTornOff,
+		IsCooked					= 1 << 10
 	};
 
 	FRIEND_ENUM_CLASS_FLAGS(EFlags);
 
+	/** A common grouping of EFlags */
+	static constexpr EFlags TransientFlags = EFlags((uint32)EFlags::HasRegistered | (uint32)EFlags::IsTornOff);
+
 	/** Used to control what level of error reporting we return from some methods */
 	enum ErrorVerbosity
 	{
@@ -473,6 +474,20 @@ private:
 	 */
 	FText GetCorruptedPayloadErrorMsgForSave(FLinkerSave* Linker) const;
 
+	/**
+	 * A utility for validating that a package trailer builder was created correctly. If something incorrect is encountered we will assert
+	 * to prevent a corrupted package from being saved.
+	 * The main thing we check is that the payload was added as in the correct list for the correct payload storage type based on the flags 
+	 * we have for the payload.
+	 * Note that it is not expected that we will ever encounter these problems (so asserting is acceptable) and the checks could be considered
+	 * a little over cautious. We should consider just removing this check in 5.2 onwards.
+	 * 
+	 * @param LinkerSave	The linker containing the package trailer builder
+	 * @param Id			The hash of the payload we want to verify
+	 * @param PayloadFlags	The flags for the payload.
+	 */
+	static void ValidatePackageTrailerBuilder(const FLinkerSave* LinkerSave, const FIoHash& Id, EFlags PayloadFlags);
+
 	/** Returns true if we should use legacy serialization instead of the FPackageTrailer system. This can be removed when UE_ENABLE_VIRTUALIZATION_TOGGLE is removed. */
 	bool ShouldUseLegacySerialization(const FLinkerSave* LinkerSave) const;
 
