Adjust checks to allow invalid paylod when editor bulk data is cooked

diff --git a/Engine/Source/Runtime/CoreUObject/Private/Serialization/EditorBulkData.cpp b/Engine/Source/Runtime/CoreUObject/Private/Serialization/EditorBulkData.cpp
--- a/Engine/Source/Runtime/CoreUObject/Private/Serialization/EditorBulkData.cpp
+++ b/Engine/Source/Runtime/CoreUObject/Private/Serialization/EditorBulkData.cpp
@@ -819,16 +819,16 @@ void FEditorBulkData::Serialize(FArchive& Ar, UObject* Owner, bool bAllowRegiste
 		}
 
 		// TODO: Can probably remove these checks before UE5 release
-		check(!Ar.IsSaving() || GetPayloadSize() == 0 || BulkDataId.IsValid()); // Sanity check to stop us saving out bad data
-		check(!Ar.IsSaving() || GetPayloadSize() == 0 || !PayloadContentId.IsZero()); // Sanity check to stop us saving out bad data
+		check(!Ar.IsSaving() || GetPayloadSize() == 0 || BulkDataId.IsValid() || EnumHasAnyFlags(Flags, EFlags::IsCooked)); // Sanity check to stop us saving out bad data
+		check(!Ar.IsSaving() || GetPayloadSize() == 0 || !PayloadContentId.IsZero() || EnumHasAnyFlags(Flags, EFlags::IsCooked)); // Sanity check to stop us saving out bad data
 		
 		Ar << BulkDataId;
 		Ar << PayloadContentId;
 		Ar << PayloadSize;
 
 		// TODO: Can probably remove these checks before UE5 release
-		check(!Ar.IsLoading() || GetPayloadSize() == 0 || BulkDataId.IsValid()); // Sanity check to stop us loading in bad data
-		check(!Ar.IsLoading() || GetPayloadSize() == 0 || !PayloadContentId.IsZero()); // Sanity check to stop us loading in bad data
+		check(!Ar.IsLoading() || GetPayloadSize() == 0 || BulkDataId.IsValid() || EnumHasAnyFlags(Flags, EFlags::IsCooked)); // Sanity check to stop us loading in bad data
+		check(!Ar.IsLoading() || GetPayloadSize() == 0 || !PayloadContentId.IsZero() || EnumHasAnyFlags(Flags, EFlags::IsCooked)); // Sanity check to stop us loading in bad data
 
 		if (Ar.IsSaving())
 		{
