EditorBulkData Duplicate Guids: Improve feedback in warning message. Fix renames causing spurious warning message. Fix transaction serialization to not modify the BulkDataId so that multiple applications of a transaction buffer do not cause duplicates.
#rnx

diff --git a/Engine/Source/Editor/UnrealEd/Private/Commandlets/ContentCommandlets.cpp b/Engine/Source/Editor/UnrealEd/Private/Commandlets/ContentCommandlets.cpp
--- a/Engine/Source/Editor/UnrealEd/Private/Commandlets/ContentCommandlets.cpp
+++ b/Engine/Source/Editor/UnrealEd/Private/Commandlets/ContentCommandlets.cpp
@@ -283,7 +283,7 @@ int32 UResavePackagesCommandlet::InitializeResaveParameters( const TArray<FStrin
 	}
 
 	// ... if not, load in all packages
-	if( !bExplicitPackages )
+	if( !bExplicitPackages || Switches.Contains("SaveAll"))
 	{
 		UE_LOG( LogContentCommandlet, Display, TEXT( "No maps found to save when building HLODs, checking Project Settings for Directory or Asset Path(s)" ) );
 
diff --git a/Engine/Source/Editor/UnrealEd/Private/EditorDomain/BulkDataRegistryImpl.cpp b/Engine/Source/Editor/UnrealEd/Private/EditorDomain/BulkDataRegistryImpl.cpp
--- a/Engine/Source/Editor/UnrealEd/Private/EditorDomain/BulkDataRegistryImpl.cpp
+++ b/Engine/Source/Editor/UnrealEd/Private/EditorDomain/BulkDataRegistryImpl.cpp
@@ -260,8 +260,9 @@ void FBulkDataRegistryImpl::UpdateRegistrationData(UPackage* Owner, const UE::Se
 		FWriteScopeLock RegistryScopeLock(RegistryLock);
 		check(bActive); // Registrations should not come in after we destruct
 		UE::BulkDataRegistry::Private::FRegisteredBulk& RegisteredBulk = Registry.FindOrAdd(BulkData.GetIdentifier());
-		// Add the owner if we previously did not have an owner, otherwise keep the first owner
-		if (RegisteredBulk.PackageName.IsNone())
+		// Update the Owner if we have it. This is important for Package Renames, which call UpdateRegistrationData when they save
+		// to the newly renamed package, with the renamed Owner package.
+		if (!PackageName.IsNone())
 		{
 			RegisteredBulk.PackageName = PackageName;
 		}
diff --git a/Engine/Source/Runtime/CoreUObject/Private/Serialization/EditorBulkData.cpp b/Engine/Source/Runtime/CoreUObject/Private/Serialization/EditorBulkData.cpp
--- a/Engine/Source/Runtime/CoreUObject/Private/Serialization/EditorBulkData.cpp
+++ b/Engine/Source/Runtime/CoreUObject/Private/Serialization/EditorBulkData.cpp
@@ -618,7 +618,9 @@ void FEditorBulkData::LogRegisterError(UE::BulkDataRegistry::ERegisterResult Val
 					// 1) Resaveondemand does not yet work for startup packages. If the warning is coming from startup packages, run ResavePackagesCommandlet with
 					// argument "-package=<PackageName>" for each warned package.
 					// 2) The guid will only be updated on resave if both packages contributing to the collision are loaded during the resave.
-					SilenceWarningMessage = TEXT(" To silence this warning, run the ResavePackagesCommandlet with \"-autocheckout -resaveondemand=bulkdataduplicates\".");
+					SilenceWarningMessage = FString::Printf(TEXT(" To silence this warning, run the ResavePackagesCommandlet with ")
+						TEXT("\"-autocheckout -resaveondemand=bulkdataduplicates -SaveAll -Package=%s -Package=%s\"."),
+						*FPackageName::ObjectPathToPackageName(OwnerPathName), *OtherOwnerPackageName);
 				}
 				bool bSuppressWarning = false; // Allow projects to suppress the warning because they might not be able to resave packages during an integration
 				GConfig->GetBool(TEXT("CookSettings"), TEXT("BulkDataRegistrySuppressDuplicateWarning"), bSuppressWarning, GEditorIni);
@@ -726,6 +728,9 @@ void FEditorBulkData::Serialize(FArchive& Ar, UObject* Owner, bool bAllowRegiste
 
 	if (Ar.IsTransacting())
 	{
+		// We've made changes to this serialization path without versioning them, because we assume
+		// IsTransacting means !IsPersistent and therefore all loaded data was saved with the current binary.
+		check(!Ar.IsPersistent()); 
 		// Do not process the transaction if the owner is mid loading (see FUntypedBulkData::Serialize)
 		bool bNeedsTransaction = Ar.IsSaving() && (!Owner || !Owner->HasAnyFlags(RF_NeedLoad));
 
@@ -739,7 +744,9 @@ void FEditorBulkData::Serialize(FArchive& Ar, UObject* Owner, bool bAllowRegiste
 			}
 
 			Ar << Flags;
-			Ar << BulkDataId;
+			// Transactions do not save/load the BulkDataId; it is specific to an instance and is
+			// unchanged by modifications to the payload. Allowing it to save/load would allow it to be
+			// duplicated if an editor operation plays back the transaction buffer multiple times.
 			Ar << PayloadContentId;
 			Ar << PayloadSize;
 			Ar << PackagePath;
@@ -766,6 +773,11 @@ void FEditorBulkData::Serialize(FArchive& Ar, UObject* Owner, bool bAllowRegiste
 			}
 			else
 			{
+				if (PayloadSize > 0 && !BulkDataId.IsValid())
+				{
+					BulkDataId = FGuid::NewGuid();
+				}
+
 				FCompressedBuffer CompressedPayload;
 				if (bPayloadInArchive)
 				{
