Never consider optional data to be an asset (by default)

#rnx

diff --git a/Engine/Source/Runtime/CoreUObject/Private/UObject/Obj.cpp b/Engine/Source/Runtime/CoreUObject/Private/UObject/Obj.cpp
--- a/Engine/Source/Runtime/CoreUObject/Private/UObject/Obj.cpp
+++ b/Engine/Source/Runtime/CoreUObject/Private/UObject/Obj.cpp
@@ -2153,7 +2153,7 @@ bool UObject::IsAsset() const
 	// Assets are not transient or CDOs. They must be public.
 	const bool bHasValidObjectFlags = !HasAnyFlags(RF_Transient | RF_ClassDefaultObject) && HasAnyFlags(RF_Public) && IsValidChecked(this);
 
-	if ( bHasValidObjectFlags )
+	if ( bHasValidObjectFlags && !GetClass()->HasAnyClassFlags(CLASS_Optional) )
 	{
 		// Don't count objects embedded in other objects (e.g. font textures, sequences, material expressions)
 		if ( UPackage* LocalOuterPackage = dynamic_cast<UPackage*>(GetOuter()) )
