Updating FSoftObjecPtr values stored in the asset registry tags to match the current format returned by FSoftObjecPtr::ToString() (stripping class names and leaving just path names)

diff --git a/Engine/Source/Runtime/CoreUObject/Private/UObject/Obj.cpp b/Engine/Source/Runtime/CoreUObject/Private/UObject/Obj.cpp
--- a/Engine/Source/Runtime/CoreUObject/Private/UObject/Obj.cpp
+++ b/Engine/Source/Runtime/CoreUObject/Private/UObject/Obj.cpp
@@ -2067,7 +2067,20 @@ static void PostLoadAssetRegistryTagProperty(FProperty* Prop, const FAssetData&
 
 	if (Prop->HasAnyPropertyFlags(CPF_AssetRegistrySearchable))
 	{
-		if (Prop->IsA<FObjectPropertyBase>() && !Prop->IsA<FSoftObjectProperty>())
+		if (FSoftObjectProperty* SoftObjectProp  = CastField<FSoftObjectProperty>(Prop))
+		{
+			// Old files may contain legacy format of FSofObjectPtr::ToString() which used to return 
+			// an export path (ClassName'/Package/Name.ObjectName') however it now returns just a pathname (/Package/Name.ObjectName)
+			FString ExportPath = AssetData.GetTagValueRef<FString>(Prop->GetFName());
+			int32 ClassSeparatorIndex = -1;
+			if (!ExportPath.IsEmpty() && ExportPath[0] != '/' && ExportPath.FindChar('\'', ClassSeparatorIndex))
+			{
+				// Strip the class name and leave just the pathname of an object
+				FString ObjectPath = ExportPath.Mid(ClassSeparatorIndex + 1, ExportPath.Len() - ClassSeparatorIndex - 2); // -2 because we're stripping the first and last '
+				OutTagsAndValuesToUpdate.Add(UObject::FAssetRegistryTag(Prop->GetFName(), ObjectPath, TagType));
+			}
+		}
+		else if (Prop->IsA<FObjectPropertyBase>())
 		{
 			FObjectPropertyBase* PropertyObject = CastFieldChecked<FObjectPropertyBase>(Prop);
 			FString ExportPath = AssetData.GetTagValueRef<FString>(Prop->GetFName());
