AssetData FText Tags: Add FixedTagPrivate::FMarshalledText and do not call FText functions from within AssetRegistry loading/saving, since these FText functions can trigger loads and call AssetRegistry functions which would deadlock.
#rnx

diff --git a/Engine/Source/Runtime/CoreUObject/Private/AssetRegistry/AssetDataTagMap.cpp b/Engine/Source/Runtime/CoreUObject/Private/AssetRegistry/AssetDataTagMap.cpp
--- a/Engine/Source/Runtime/CoreUObject/Private/AssetRegistry/AssetDataTagMap.cpp
+++ b/Engine/Source/Runtime/CoreUObject/Private/AssetRegistry/AssetDataTagMap.cpp
@@ -5,6 +5,7 @@
 #include "Algo/Sort.h"
 #include "AssetRegistry/AssetDataTagMapSerializationDetails.h"
 #include "Misc/PackageName.h"
+#include "Misc/Optional.h"
 #include "Misc/ScopeLock.h"
 #include "Serialization/ArrayWriter.h"
 #include "Serialization/MemoryReader.h"
@@ -164,6 +165,38 @@ static FString LocalizeIfComplexString(const FString& Value)
 
 namespace FixedTagPrivate
 {
+	FMarshalledText::FMarshalledText(const FString& InComplexString)
+		: String(InComplexString)
+	{
+	}
+	FMarshalledText::FMarshalledText(FString&& InComplexString)
+		: String(MoveTemp(InComplexString))
+	{
+	}
+	FMarshalledText::FMarshalledText(const FText& InText)
+		: String(ToComplexString(InText))
+	{
+	}
+	FMarshalledText::FMarshalledText(FText&& InText)
+		: String(ToComplexString(InText))
+	{
+	}
+
+	const FString& FMarshalledText::GetAsComplexString() const
+	{
+		return String;
+	}
+
+	FText FMarshalledText::GetAsText() const
+	{
+		return FTextStringHelper::CreateFromBuffer(*String);
+	}
+
+	int32 FMarshalledText::CompareToCaseIgnored(const FMarshalledText& Other) const
+	{
+		return String.Compare(Other.String);
+	}
+
 	FArchive& operator<<(FArchive& Ar, FLegacyAssetRegistryExportPath& Path)
 	{
 		return Ar << Path.Class << Path.Object << Path.Package;
@@ -204,7 +237,7 @@ namespace FixedTagPrivate
 		return FPlatformString::Stricmp(A, B) == 0;
 	}
 
-	static bool EqualsInsensitive(const FText& A, const FText& B)
+	static bool EqualsInsensitive(const FMarshalledText& A, const FMarshalledText& B)
 	{
 		return A.CompareToCaseIgnored(B) == 0;
 	}
@@ -367,6 +400,7 @@ namespace FixedTagPrivate
 
 	//////////////////////////////////////////////////////////////////////////
 
+	template <bool bForStorage>
 	FString	FValueHandle::AsString() const
 	{
 		FStore& Store = GStores[StoreIndex];
@@ -380,13 +414,25 @@ namespace FixedTagPrivate
 		case EValueType::Name:					return Store.Names[Index].ToString();
 		case EValueType::NumberlessExportPath:	return Store.NumberlessExportPaths[Index].ToString();
 		case EValueType::ExportPath:			return Store.ExportPaths[Index].ToString();
-		case EValueType::LocalizedText:			return Store.Texts[Index].ToString();
+		case EValueType::LocalizedText:
+			return bForStorage ? Store.Texts[Index].GetAsComplexString()
+				: Store.Texts[Index].GetAsText().ToString();
 		}
 
 		check(false);
 		return FString();
 	}
 
+	FString	FValueHandle::AsDisplayString() const
+	{
+		return AsString<false /* bForStorage */>();
+	}
+
+	FString	FValueHandle::AsStorageString() const
+	{
+		return AsString<true /* bForStorage */>();
+	}
+
 	FName FValueHandle::AsName() const
 	{
 		FStore& Store = GStores[StoreIndex];
@@ -402,7 +448,7 @@ namespace FixedTagPrivate
 		case EValueType::Name:					return Store.Names[Index];
 		case EValueType::NumberlessExportPath:	return Store.NumberlessExportPaths[Index].ToName();
 		case EValueType::ExportPath:			return Store.ExportPaths[Index].ToName();
-		case EValueType::LocalizedText:			return FName(Store.Texts[Index].ToString());
+		case EValueType::LocalizedText:			return FName(Store.Texts[Index].GetAsText().ToString());
 		}
 
 		check(false);
@@ -424,7 +470,7 @@ namespace FixedTagPrivate
 		case EValueType::Name:					return MakeNumberedPath(Store.Names[Index]);
 		case EValueType::NumberlessExportPath:	return MakeNumberedPath(Store.NumberlessExportPaths[Index]);
 		case EValueType::ExportPath:			return Store.ExportPaths[Index];
-		case EValueType::LocalizedText:			return FAssetRegistryExportPath(Store.Texts[Index].ToString());
+		case EValueType::LocalizedText:			return FAssetRegistryExportPath(Store.Texts[Index].GetAsText().ToString());
 		}
 
 		check(false);
@@ -435,17 +481,21 @@ namespace FixedTagPrivate
 	{
 		if (Id.Type == EValueType::LocalizedText)
 		{
-			Out = GStores[StoreIndex].Texts[Id.Index];
+			Out = GStores[StoreIndex].Texts[Id.Index].GetAsText();
 			return true;
 		}
 
 		return false;
 	}
 
-	static FString FixedToLoose(FValueHandle Fixed)
+	bool FValueHandle::AsMarshalledText(FMarshalledText& Out) const
 	{
-		return Fixed.Id.Type == EValueType::LocalizedText
-			? ToComplexString(GStores[Fixed.StoreIndex].Texts[Fixed.Id.Index]) : Fixed.AsString();
+		if (Id.Type == EValueType::LocalizedText)
+		{
+			Out = GStores[StoreIndex].Texts[Id.Index];
+			return true;
+		}
+		return false;
 	}
 
 	static bool EqualsInsensitive(const FStringView& Str, const FAssetRegistryExportPath& Path)
@@ -480,12 +530,7 @@ namespace FixedTagPrivate
 		case EValueType::Name:					return EqualsInsensitive(Str, Store.Names[Index]);
 		case EValueType::NumberlessExportPath:	return EqualsInsensitive(Str, Store.NumberlessExportPaths[Index]);
 		case EValueType::ExportPath:			return EqualsInsensitive(Str, Store.ExportPaths[Index]);
-		case EValueType::LocalizedText:
-		{
-			FText StrAsText;
-			FromComplexString(FString(Str), StrAsText);
-			return EqualsInsensitive(StrAsText, Store.Texts[Index]);
-		}
+		case EValueType::LocalizedText:			return EqualsInsensitive(Str, *Store.Texts[Index].GetAsComplexString());
 		}
 
 		check(false);
@@ -747,19 +792,6 @@ namespace FixedTagPrivate
 		return Out;
 	}
 
-	template<class MapType>
-	static TArray<FText> FlattenAsText(const MapType& Index)
-	{
-		TArray<FText> Out;
-		Out.SetNum(Index.Num());
-		for (const TPair<FString, uint32>& Pair : Index)
-		{
-			verify(FTextStringHelper::ReadFromBuffer(*Pair.Key, Out[Pair.Value]));
-		}
-
-		return Out;
-	}
-
 	static TArray<FNumberlessPair> MakeNumberlessPairs(const TArray<FNumberedPair>& In)
 	{
 		TArray<FNumberlessPair> Out;
@@ -785,7 +817,7 @@ namespace FixedTagPrivate
 		Out.Names = Flatten(NameIndices);
 		Out.NumberlessExportPaths = Flatten(NumberlessExportPathIndices);
 		Out.ExportPaths = Flatten(ExportPathIndices);
-		Out.Texts = FlattenAsText(TextIndices);
+		Out.Texts = Flatten(TextIndices);
 		Out.AnsiStrings = AnsiStrings.FlattenAsAnsi();
 		Out.WideStrings = WideStrings.FlattenAsWide();
 
@@ -821,10 +853,10 @@ namespace FixedTagPrivate
 
 	FValueId FStoreBuilder::IndexValue(FName Key, FAssetTagValueRef Value)
 	{
-		FText Text; 
-		if (Value.TryGetAsText(Text))
+		FMarshalledText MarshalledText; 
+		if (Value.TryGetAsMarshalledText(MarshalledText))
 		{
-			return FValueId{EValueType::LocalizedText,			Index(TextIndices, ToComplexString(Text))};
+			return FValueId{EValueType::LocalizedText,			Index(TextIndices, MoveTemp(MarshalledText))};
 		}
 		else if (Options.StoreAsName.Contains(Key))
 		{
@@ -1024,21 +1056,16 @@ namespace FixedTagPrivate
 			return { LoadItem<decltype(FNumberedPair::Key)>(), LoadItem<decltype(FNumberedPair::Value)>() };
 		}
 		
-		void SaveItem(const FText& Text)
+		void SaveItem(const FMarshalledText& Text)
 		{
-			Scratch.Reset();
-			FTextStringHelper::WriteToBuffer(Scratch, Text);
-			Ar << Scratch;
+			Ar << const_cast<FString&>(Text.GetAsComplexString());
 		}
 
-		template<> FText LoadItem()
+		template<> FMarshalledText LoadItem()
 		{
-			Scratch.Reset();
-			Ar << Scratch;
-			
-			FText Out;
-			FTextStringHelper::ReadFromBuffer(*Scratch, Out);
-			return Out;
+			FString ComplexString;
+			Ar << ComplexString;
+			return FMarshalledText(MoveTemp(ComplexString));
 		}
 		
 		template<typename T>
@@ -1098,7 +1125,7 @@ namespace FixedTagPrivate
 		{
 		}
 
-		void SaveTextData(TArrayView<const FText> Texts)
+		void SaveTextData(TArrayView<const FMarshalledText> Texts)
 		{
 			FArrayWriter Data;
 			TSerializer<Version>(Data).SaveViewData(Texts);
@@ -1311,7 +1338,7 @@ const FString& FAssetTagValueRef::AsLoose() const
 
 FString FAssetTagValueRef::AsString() const
 {
-	return IsFixed() ? AsFixed().AsString() : LocalizeIfComplexString(AsLoose());
+	return IsFixed() ? AsFixed().AsDisplayString() : LocalizeIfComplexString(AsLoose());
 }
 
 FName FAssetTagValueRef::AsName() const
@@ -1332,12 +1359,30 @@ bool FAssetTagValueRef::TryGetAsText(FText& Out) const
 FText FAssetTagValueRef::AsText() const
 {
 	FText Tmp;
-	return TryGetAsText(Tmp) ? Tmp : FText::FromString(IsFixed() ? AsFixed().AsString() : AsLoose());
+	return TryGetAsText(Tmp) ? Tmp : FText::FromString(IsFixed() ? AsFixed().AsStorageString() : AsLoose());
 }
 
 FString FAssetTagValueRef::ToLoose() const
 {
-	return IsFixed() ? FixedToLoose(AsFixed()) : AsLoose();
+	return IsFixed() ? AsFixed().AsStorageString() : AsLoose();
+}
+
+bool FAssetTagValueRef::TryGetAsMarshalledText(FixedTagPrivate::FMarshalledText& Out) const
+{
+	if (IsFixed())
+	{
+		return AsFixed().AsMarshalledText(Out);
+	}
+	else
+	{
+		const FString& String = AsLoose();
+		if (FTextStringHelper::IsComplexText(*String))
+		{
+			Out = FixedTagPrivate::FMarshalledText(String);
+			return true;
+		}
+		return false;
+	}
 }
 
 bool FAssetTagValueRef::Equals(FStringView Str) const
@@ -1421,7 +1466,7 @@ FAssetDataTagMap FAssetDataTagMapSharedView::CopyMap() const
 	{
 		FAssetDataTagMap Out;
 		Out.Reserve(Num());
-		ForEach([&](TPair<FName, FAssetTagValueRef> Pair){ Out.Add(Pair.Key, FixedToLoose(Pair.Value.AsFixed())); });
+		ForEach([&](TPair<FName, FAssetTagValueRef> Pair){ Out.Add(Pair.Key, Pair.Value.AsFixed().AsStorageString()); });
 		return Out;
 	}
 	else
@@ -1438,24 +1483,9 @@ void FAssetDataTagMapSharedView::Shrink()
 	}
 }
 
-bool operator==(const FAssetDataTagMapSharedView& A, const FAssetDataTagMap& B)
-{
-	if (A.Num() != B.Num())
-	{
-		return false;
-	}
-	else if (A.IsFixed())
-	{
-		// This is very wasteful but currently only used by unit tests
-		return A.CopyMap() == B;
-	}
-		
-	return A.Num() == 0 || *A.Loose == B;
-}
-
 namespace FixedTagPrivate
 {
-	static bool MapEqualsHelper(FixedTagPrivate::FMapHandle Fixed, FAssetDataTagMap* Loose)
+	static bool MapEqualsHelper(FixedTagPrivate::FMapHandle Fixed, const FAssetDataTagMap* Loose)
 	{
 		checkSlow(Fixed.Num == Loose->Num());
 		// Since Num is the same and the maps are unique, test only whether all keys in Fixed exist with equal value
@@ -1464,7 +1494,7 @@ namespace FixedTagPrivate
 		{
 			for (FNumberlessPair Pair : Fixed.GetNumberlessView())
 			{
-				FString* LooseValue = Loose->Find(FName::CreateFromDisplayId(Pair.Key, 0));
+				const FString* LooseValue = Loose->Find(FName::CreateFromDisplayId(Pair.Key, 0));
 				if (!LooseValue || FAssetTagValueRef(Fixed.StoreIndex, Pair.Value) != *LooseValue)
 				{
 					return false;
@@ -1475,7 +1505,7 @@ namespace FixedTagPrivate
 		{
 			for (FNumberedPair Pair : Fixed.GetNumberedView())
 			{
-				FString* LooseValue = Loose->Find(Pair.Key);
+				const FString* LooseValue = Loose->Find(Pair.Key);
 				if (!LooseValue || FAssetTagValueRef(Fixed.StoreIndex, Pair.Value) != *LooseValue)
 				{
 					return false;
@@ -1486,6 +1516,20 @@ namespace FixedTagPrivate
 	}
 }
 
+bool operator==(const FAssetDataTagMapSharedView& A, const FAssetDataTagMap& B)
+{
+	if (A.Num() != B.Num())
+	{
+		return false;
+	}
+	else if (A.IsFixed())
+	{
+		return FixedTagPrivate::MapEqualsHelper(A.Fixed, &B);
+	}
+
+	return A.Num() == 0 || *A.Loose == B;
+}
+
 bool operator==(const FAssetDataTagMapSharedView& A, const FAssetDataTagMapSharedView& B)
 {
 	if (A.Num() != B.Num())
diff --git a/Engine/Source/Runtime/CoreUObject/Public/AssetRegistry/AssetDataTagMap.h b/Engine/Source/Runtime/CoreUObject/Public/AssetRegistry/AssetDataTagMap.h
--- a/Engine/Source/Runtime/CoreUObject/Public/AssetRegistry/AssetDataTagMap.h
+++ b/Engine/Source/Runtime/CoreUObject/Public/AssetRegistry/AssetDataTagMap.h
@@ -4,15 +4,20 @@
 
 #include "CoreMinimal.h"
 #include "Containers/SortedMap.h"
+#include "HAL/CriticalSection.h"
 #include "Misc/StringBuilder.h"
 #include "Templates/RefCounting.h"
 #include "Templates/TypeCompatibleBytes.h"
 #include "UObject/TopLevelAssetPath.h"
 
+class FAssetRegistryState;
 class FAssetTagValueRef;
 class FAssetDataTagMapSharedView;
 struct FAssetRegistrySerializationOptions;
 
+namespace FixedTagPrivate { class FMarshalledText; }
+namespace FixedTagPrivate { class FStoreBuilder; }
+
 /**
  * Helper class for condensing strings of these types into  1 - 3 FNames
  * [class]'[package].[object]'
@@ -111,12 +116,18 @@ namespace FixedTagPrivate
 		uint32 StoreIndex;
 		FValueId Id;
 
-		FString						AsString() const;
+		FString						AsDisplayString() const;
+		FString						AsStorageString() const;
 		FName						AsName() const;
 		FAssetRegistryExportPath	AsExportPath() const;
 		bool						AsText(FText& Out) const;
+		bool						AsMarshalledText(FMarshalledText& Out) const;
 		bool						Equals(FStringView Str) const;
 		bool						Contains(const TCHAR* Str) const;
+
+	private:
+		template <bool bForStorage>
+		FString						AsString() const;
 	};
 
 	// Handle to a tag map owned by a managed FStore
@@ -237,13 +248,22 @@ public:
 
 	FString						GetValue() const { return AsString(); }
 
-	// Get FTexts as unlocalized complex strings. For internal use only, to make new FAssetDataTagMapSharedView.
-	FString						ToLoose() const;
-
 	bool						Equals(FStringView Str) const;
 
 	UE_DEPRECATED(4.27, "Use AsString(), AsName(), AsExportPath() or AsText() instead. ")
 	operator FString () const { return AsString(); }
+
+private:
+	/** Return whether this's value is a MarshalledFText, and copy it into out parameter if so */
+	bool						TryGetAsMarshalledText(FixedTagPrivate::FMarshalledText& Out) const;
+	/**
+	 * Copy this's value (whether loose or fixed) into the loose format.
+	 * The returned loose value is in StorageFormat (e.g. complex strings) rather than display format.
+	 */
+	FString						ToLoose() const;
+
+	friend class FixedTagPrivate::FStoreBuilder;
+	friend FAssetRegistryState;
 };
 
 inline bool operator==(FAssetTagValueRef A, FStringView B) { return  A.Equals(B); }
diff --git a/Engine/Source/Runtime/CoreUObject/Public/AssetRegistry/AssetDataTagMapSerializationDetails.h b/Engine/Source/Runtime/CoreUObject/Public/AssetRegistry/AssetDataTagMapSerializationDetails.h
--- a/Engine/Source/Runtime/CoreUObject/Public/AssetRegistry/AssetDataTagMapSerializationDetails.h
+++ b/Engine/Source/Runtime/CoreUObject/Public/AssetRegistry/AssetDataTagMapSerializationDetails.h
@@ -18,6 +18,27 @@ namespace FixedTagPrivate
 		FName Object;
 	};
 
+	/**
+	 * The AssetRegistry's representation of an FText AssetData Tag value.
+	 * It can be stored and copied without being interpreted as an FText.
+	 */
+	class COREUOBJECT_API FMarshalledText
+	{
+	public:
+		FMarshalledText() = default;
+		explicit FMarshalledText(const FString& InString);
+		explicit FMarshalledText(FString&& InString);
+		explicit FMarshalledText(const FText& InText);
+		explicit FMarshalledText(FText&& InText);
+
+		const FString& GetAsComplexString() const;
+		FText GetAsText() const;
+		int32 CompareToCaseIgnored(const FMarshalledText& Other) const;
+
+	private:
+		FString String;
+	};
+
 	/// Stores a fixed set of values and all the key-values maps used for lookup
 	struct FStore
 	{
@@ -34,7 +55,7 @@ namespace FixedTagPrivate
 		TArrayView<FName> Names;
 		TArrayView<FNumberlessExportPath> NumberlessExportPaths;
 		TArrayView<FAssetRegistryExportPath> ExportPaths;
-		TArrayView<FText> Texts;
+		TArrayView<FMarshalledText> Texts;
 
 		const uint32 Index;
 		void* Data = nullptr;
@@ -86,7 +107,7 @@ namespace FixedTagPrivate
 		TArray<FName> Names;
 		TArray<FNumberlessExportPath> NumberlessExportPaths;
 		TArray<FAssetRegistryExportPath> ExportPaths;
-		TArray<FText> Texts;
+		TArray<FMarshalledText> Texts;
 	};
 
 	uint32 HashCaseSensitive(const TCHAR* Str, int32 Len);
@@ -160,6 +181,16 @@ namespace FixedTagPrivate
 			{
 				return HashCombineQuick(HashCombineQuick(GetKeyHash(Key.ClassPath.GetPackageName()), GetKeyHash(Key.ClassPath.GetAssetName())), GetKeyHash(Key.Package), GetKeyHash(Key.Object));
 			}
+
+			static bool Matches(const FMarshalledText& A, const FMarshalledText& B)
+			{
+				return Matches(A.GetAsComplexString(), B.GetAsComplexString());
+			}
+
+			static uint32 GetKeyHash(const FMarshalledText& Key)
+			{
+				return GetKeyHash(Key.GetAsComplexString());
+			}
 		};
 
 		struct FStringIndexer
@@ -181,7 +212,7 @@ namespace FixedTagPrivate
 		TMap<FName, uint32, FDefaultSetAllocator, FCaseSensitiveFuncs<uint32>> NameIndices;
 		TMap<FNumberlessExportPath, uint32, FDefaultSetAllocator, FCaseSensitiveFuncs<uint32>> NumberlessExportPathIndices;
 		TMap<FAssetRegistryExportPath, uint32, FDefaultSetAllocator, FCaseSensitiveFuncs<uint32>> ExportPathIndices;
-		TMap<FString, uint32, FDefaultSetAllocator, FCaseSensitiveFuncs<uint32>> TextIndices;
+		TMap<FMarshalledText, uint32, FDefaultSetAllocator, FCaseSensitiveFuncs<uint32>> TextIndices;
 
 		TArray<FNumberedPair> NumberedPairs;
 		TArray<FNumberedPair> NumberlessPairs; // Stored as numbered for convenience
