Fix for a regression that can block property value serialization of import text from creating a placeholder linker object in lieu of a deferred Blueprint class asset at editor load time.

diff --git a/Engine/Source/Runtime/CoreUObject/Private/Blueprint/BlueprintSupport.cpp b/Engine/Source/Runtime/CoreUObject/Private/Blueprint/BlueprintSupport.cpp
--- a/Engine/Source/Runtime/CoreUObject/Private/Blueprint/BlueprintSupport.cpp
+++ b/Engine/Source/Runtime/CoreUObject/Private/Blueprint/BlueprintSupport.cpp
@@ -820,7 +820,7 @@ FString GetPlaceholderPrefix<ULinkerPlaceholderClass>()    { return TEXT("PLACEH
 
 /** Internal utility function for spawning various type of placeholder objects. */
 template<class PlaceholderType>
-static PlaceholderType* MakeImportPlaceholder(UObject* Outer, const TCHAR* TargetObjName, int32 ImportIndex = INDEX_NONE)
+static PlaceholderType* MakeImportPlaceholder(UObject* Outer, const UClass* TargetObjType, const TCHAR* TargetObjName, int32 ImportIndex = INDEX_NONE)
 {
 	PlaceholderType* PlaceholderObj = nullptr;
 
@@ -838,6 +838,9 @@ static PlaceholderType* MakeImportPlaceholder(UObject* Outer, const TCHAR* Targe
 	// and isn't referenced by the ImportMap... instead, this should be stored 
 	// in the FLinkerLoad's ImportPlaceholders map
 
+	// Record the type of object that's being deferred
+	PlaceholderObj->DeferredObjectType = TargetObjType;
+
 	// make sure the class is fully formed (has its 
 	// CppClassStaticFunctions/ClassConstructor members set)
 	PlaceholderObj->Bind();
@@ -1102,7 +1105,7 @@ bool FLinkerLoad::DeferPotentialCircularImport(const int32 Index)
 			{
 				if (ImportClass->HasAnyClassFlags(CLASS_NeedsDeferredDependencyLoading))
 				{
-					Import.XObject = MakeImportPlaceholder<ULinkerPlaceholderClass>(LinkerRoot, *Import.ObjectName.ToString(), Index);
+					Import.XObject = MakeImportPlaceholder<ULinkerPlaceholderClass>(LinkerRoot, ImportClass, *Import.ObjectName.ToString(), Index);
 				}
 				else if (ImportClass->IsChildOf<UFunction>() && Import.OuterIndex.IsImport())
 				{
@@ -1119,7 +1122,7 @@ bool FLinkerLoad::DeferPotentialCircularImport(const int32 Index)
 						// but the DEFERRED_DEPENDENCY_CHECK may be out of date...
 						if(Cast<UClass>(FuncOuter))
 						{
-							Import.XObject = MakeImportPlaceholder<ULinkerPlaceholderFunction>(FuncOuter, *Import.ObjectName.ToString(), Index);
+							Import.XObject = MakeImportPlaceholder<ULinkerPlaceholderFunction>(FuncOuter, ImportClass, *Import.ObjectName.ToString(), Index);
 							DEFERRED_DEPENDENCY_CHECK(dynamic_cast<ULinkerPlaceholderClass*>(FuncOuter) != nullptr);
 						}
 					}
@@ -2333,7 +2336,7 @@ bool FLinkerLoad::HasPerformedFullExportResolvePass()
 	
 }
 
-UObject* FLinkerLoad::RequestPlaceholderValue(UClass* ObjectType, const TCHAR* ObjectPath)
+UObject* FLinkerLoad::RequestPlaceholderValue(const FProperty* Property, const UClass* ObjectType, const TCHAR* ObjectPath)
 {
 #if !USE_CIRCULAR_DEPENDENCY_LOAD_DEFERRING
 	return nullptr;
@@ -2353,13 +2356,40 @@ UObject* FLinkerLoad::RequestPlaceholderValue(UClass* ObjectType, const TCHAR* O
 		// handle that here as well
 		else if (ObjectType->IsChildOf<UClass>())
 		{
+			// Class property values will typically always request the base UClass type via FObjectPropertyBase, so
+			// we try and redirect to the actual class object value type to determine if a placeholder can be created.
+			// Generally, we shouldn't be in here unless we're serializing an object property's value from exported T3D.
+			if (const FClassProperty* ClassProperty = CastField<FClassProperty>(Property))
+			{
+				if (ClassProperty->MetaClass)
+				{
+					// Note: We are interested in the value's underlying class type, not the value itself (which will be an instance of the class type).
+					const UClass* ValueType = ClassProperty->MetaClass->GetClass();
+					if (const ULinkerPlaceholderClass* PlaceholderClass = Cast<ULinkerPlaceholderClass>(ClassProperty->MetaClass))
+					{
+						// If the type was deferred on import of the property value, redirect to the type of value whose load has been deferred (this should be a native UClass derivative).
+						ValueType = PlaceholderClass->DeferredObjectType.Get();
+					}
+
+					if (ValueType)
+					{
+						checkf(ValueType->IsChildOf(ObjectType),
+							TEXT("Requesting an import placeholder object for a class value type (%s) that is not a derivative of the required object type (%s)."),
+							*ValueType->GetName(),
+							*ObjectType->GetName());
+
+						ObjectType = ValueType;
+					}
+				}
+			}
+
 			const FString ObjectPathStr(ObjectPath);
 			// we don't need placeholders for native object references and for non-BP class objects (the 
 			// calling code should properly handle null return values)
 			if (!FPackageName::IsScriptPackage(ObjectPathStr) && ObjectType->HasAnyClassFlags(CLASS_NeedsDeferredDependencyLoading))
 			{
 				const FString ObjectName = FPackageName::ObjectPathToObjectName(ObjectPathStr);
-				Placeholder = MakeImportPlaceholder<ULinkerPlaceholderClass>(LinkerRoot, *ObjectName);
+				Placeholder = MakeImportPlaceholder<ULinkerPlaceholderClass>(LinkerRoot, ObjectType, *ObjectName);
 				ImportPlaceholders.Add(ObjId, Placeholder);
 			}
 		}
diff --git a/Engine/Source/Runtime/CoreUObject/Private/UObject/LinkerPlaceholderBase.h b/Engine/Source/Runtime/CoreUObject/Private/UObject/LinkerPlaceholderBase.h
--- a/Engine/Source/Runtime/CoreUObject/Private/UObject/LinkerPlaceholderBase.h
+++ b/Engine/Source/Runtime/CoreUObject/Private/UObject/LinkerPlaceholderBase.h
@@ -104,6 +104,9 @@ public:
 	/** Set by the ULinkerLoad that created this instance, tracks what import/export this was used in place of. */
 	FPackageIndex PackageIndex;
 
+	/** Also set by the ULinkerLoad that created this instance; keeps track of the type of the original object whose load is being deferred. */
+	TWeakObjectPtr<const UClass> DeferredObjectType;
+
 	/**
 	 * Attempts to find and store the referencing container object (along with  
 	 * the specified property), so that we can replace the reference at a 
diff --git a/Engine/Source/Runtime/CoreUObject/Private/UObject/PropertyBaseObject.cpp b/Engine/Source/Runtime/CoreUObject/Private/UObject/PropertyBaseObject.cpp
--- a/Engine/Source/Runtime/CoreUObject/Private/UObject/PropertyBaseObject.cpp
+++ b/Engine/Source/Runtime/CoreUObject/Private/UObject/PropertyBaseObject.cpp
@@ -533,7 +533,7 @@ UObject* FObjectPropertyBase::FindImportedObject( const FProperty* Property, UOb
 
 			if (bDeferAssetImports)
 			{
-				Result = Linker->RequestPlaceholderValue(ObjectClass, Text);
+				Result = Linker->RequestPlaceholderValue(Property, ObjectClass, Text);
 			}
 			
 			if (Result == nullptr)
diff --git a/Engine/Source/Runtime/CoreUObject/Public/UObject/LinkerLoad.h b/Engine/Source/Runtime/CoreUObject/Public/UObject/LinkerLoad.h
--- a/Engine/Source/Runtime/CoreUObject/Public/UObject/LinkerLoad.h
+++ b/Engine/Source/Runtime/CoreUObject/Public/UObject/LinkerLoad.h
@@ -1132,11 +1132,12 @@ public:
 	 * NOTE: For now, this will only produce UClass placeholders, as that is the 
 	 *       only type we've identified needing.
 	 * 
+	 * @param  Property		 The property for which you want to defer loading the value.
 	 * @param  ObjectType    The expected type of the object you want to defer loading of.
 	 * @param  ObjectPath    The full object/package path for the expected object.
 	 * @return A FLinkerPlaceholderBase UObject that can be used in place of the import dependency.
 	 */
-	UObject* RequestPlaceholderValue(UClass* ObjectType, const TCHAR* ObjectPath);
+	UObject* RequestPlaceholderValue(const FProperty* Property, const UClass* ObjectType, const TCHAR* ObjectPath);
 
 private:
 	/**
