CIS Fix - AsyncLoading2

#rnx

- MSVC Static Analysis picked up that the Class pointer was still being used in the branch where we know Class == nullptr.
- When iterating over the pending CDOs and logging them, we should be logging the class name for that CDO.
- Added an assert to check that CurrentClass is valid before using it.
- Fixed some code alignment issues that occurred in a recent code merge.

diff --git a/Engine/Source/Runtime/CoreUObject/Private/Serialization/AsyncLoading2.cpp b/Engine/Source/Runtime/CoreUObject/Private/Serialization/AsyncLoading2.cpp
--- a/Engine/Source/Runtime/CoreUObject/Private/Serialization/AsyncLoading2.cpp
+++ b/Engine/Source/Runtime/CoreUObject/Private/Serialization/AsyncLoading2.cpp
@@ -914,7 +914,7 @@ public:
 				if (bRemoved)
 				{
 					if (PackageRef.RefCount > 0)
-		{
+					{
 						UE_LOG(LogStreaming, Error,
 							TEXT("RemovePackage: %s %s (0x%llX) - with (ObjectFlags=%x, InternalObjectFlags=%x) - ")
 							TEXT("Package destroyed while still being referenced, RefCount %d > 0."),
@@ -952,14 +952,14 @@ public:
 	}
 
 	uint32 GetStoredScriptObjectsAllocatedSize() const
-		{
+	{
 		return ScriptObjects.GetAllocatedSize();
 	}
 
 	int32 GetStoredPublicExportsCount() const
-			{
+	{
 		return ObjectIndexToPublicExport.Num();
-			}
+	}
 
 	void RemovePublicExports(const FUnreachablePublicExports& PublicExports)
 	{
@@ -1233,9 +1233,9 @@ struct FPackageImportStore
 			// Filter out CDOs that are themselves classes,
 			// like Default__BlueprintGeneratedClass of type UBlueprintGeneratedClass
 			if (Class->HasAnyFlags(RF_ClassDefaultObject))
-				{
+			{
 				continue;
-				}
+			}
 
 			// Add dependency on any script CDO that has not been created and initialized yet
 			UObject* CDO = Class->GetDefaultObject(/*bCreateIfNeeded*/ false);
@@ -2623,7 +2623,7 @@ private:
 			{
 				int32 RequestId = Node->ReferencerRequestId();
 				if (RequestId > MaxRequestId)
-		{
+				{
 					MaxRequestId = RequestId;
 					Class = CurrentClass;
 				}
@@ -2655,10 +2655,12 @@ private:
 			for (TMap<UClass*, TArray<FEventLoadNode2*>>::TIterator It = PendingCDOs.CreateIterator(); It; ++It)
 			{
 				UClass* CurrentClass = It.Key();
+				check(CurrentClass != nullptr);
+
 				const TArray<FEventLoadNode2*>& Nodes = It.Value();
 				UE_LOG(LogStreaming, Warning,
 					TEXT("ProcessPendingCDOs: '%s' with %d nodes could not be processed from this stack."),
-					*Class->GetFullName(), Nodes.Num());
+					*CurrentClass->GetFullName(), Nodes.Num());
 			}
 		}
 	}
@@ -3542,7 +3544,7 @@ void FGlobalImportStore::FindAllScriptObjects()
 				FPackageObjectIndex GlobalImportIndex = FPackageObjectIndex::FromScriptPath(Name);
 				if (!ScriptObjects.Contains(GlobalImportIndex))
 				{
-				ScriptObjects.Add(GlobalImportIndex, Object);
+					ScriptObjects.Add(GlobalImportIndex, Object);
 #if !WITH_EDITOR
 					UE_LOG(LogStreaming, Warning, TEXT("Found additional public script object 0x%016llX: %s\r\n"),
 						GlobalImportIndex.Value(),
