Flag in exports brought in directly through optional object references and only make those non optional objects mandatory
Do not validate optional reference to cooked data when the cooked context is empty

diff --git a/Engine/Source/Runtime/CoreUObject/Private/UObject/SavePackage/PackageHarvester.cpp b/Engine/Source/Runtime/CoreUObject/Private/UObject/SavePackage/PackageHarvester.cpp
--- a/Engine/Source/Runtime/CoreUObject/Private/UObject/SavePackage/PackageHarvester.cpp
+++ b/Engine/Source/Runtime/CoreUObject/Private/UObject/SavePackage/PackageHarvester.cpp
@@ -182,6 +182,7 @@ FPackageHarvester::FPackageHarvester(FSaveContext& InContext)
 	this->SetIsPersistent(true);
 	ArIsObjectReferenceCollector = true;
 	ArShouldSkipBulkData = true;
+	ArIgnoreClassGeneratedByRef = SaveContext.IsCooking();
 
 	this->SetPortFlags(SaveContext.GetPortFlags());
 	this->SetFilterEditorOnly(SaveContext.IsFilterEditorOnly());
@@ -652,7 +653,9 @@ ESaveRealm FPackageHarvester::GetObjectHarvestingRealm(UObject* InObject, EIlleg
 
 void FPackageHarvester::HarvestExport(UObject* InObject, ESaveRealm InContext)
 {
-	SaveContext.GetHarvestedRealm(InContext).AddExport(InObject, !DoesObjectNeedLoadForEditorGame(InObject));
+	bool bFromOptionalRef = CurrentExportDependencies.CurrentExport && CurrentExportDependencies.CurrentExport->GetClass()->HasAnyClassFlags(CLASS_Optional);
+	SaveContext.GetHarvestedRealm(InContext)
+		.AddExport(FTaggedExport(InObject, !DoesObjectNeedLoadForEditorGame(InObject), bFromOptionalRef));
 	SaveContext.GetHarvestedRealm(InContext).GetNamesReferencedFromPackageHeader().Add(InObject->GetFName().GetDisplayIndex());
 	ExportsToProcess.Enqueue({ InObject, InContext });
 }
diff --git a/Engine/Source/Runtime/CoreUObject/Private/UObject/SavePackage/SaveContext.h b/Engine/Source/Runtime/CoreUObject/Private/UObject/SavePackage/SaveContext.h
--- a/Engine/Source/Runtime/CoreUObject/Private/UObject/SavePackage/SaveContext.h
+++ b/Engine/Source/Runtime/CoreUObject/Private/UObject/SavePackage/SaveContext.h
@@ -33,20 +33,27 @@ struct FTaggedExport
 	uint32 bNotAlwaysLoadedForEditorGame : 1;
 	/**
 	 * Indicate that this export should have a public hash even if it isn't marked as RF_Public
-	 * This will artificially mark the object RF_Public in the linker tables so the iostore generates the public hash 
+	 * This will artificially mark the object RF_Public in the linker tables so the iostore generates the public hash
 	 */
 	uint32 bGeneratePublicHash : 1;
+	/**
+	 * Indicate if the object that directly referenced this export was optional
+	 * Used to determine mandatory objects in the game save realm
+	 */
+	uint32 bFromOptionalReference : 1;
 
 	FTaggedExport()
 		: Obj(nullptr)
 		, bNotAlwaysLoadedForEditorGame(false)
 		, bGeneratePublicHash(false)
+		, bFromOptionalReference(false)
 	{}
 
-	FTaggedExport(UObject* InObj, bool bInNotAlwaysLoadedForEditorGame = true)
+	FTaggedExport(UObject* InObj, bool bInNotAlwaysLoadedForEditorGame = true, bool bInFromOptionalReference = false)
 		: Obj(InObj)
 		, bNotAlwaysLoadedForEditorGame(bInNotAlwaysLoadedForEditorGame)
 		, bGeneratePublicHash(false)
+		, bFromOptionalReference(bInFromOptionalReference)
 	{}
 
 	inline bool operator == (const FTaggedExport& Other) const
@@ -115,9 +122,9 @@ struct FHarvestedRealm
 		Imports.Add(InObject);
 	}
 
-	void AddExport(UObject* InObj, bool bNotAlwaysLoadedForEditorGame)
+	void AddExport(FTaggedExport InTagObj)
 	{
-		Exports.Add(FTaggedExport(InObj, bNotAlwaysLoadedForEditorGame));
+		Exports.Add(MoveTemp(InTagObj));
 	}
 
 	void AddExcluded(UObject* InObject)
@@ -706,9 +713,9 @@ public:
 		GetHarvestedRealm().AddImport(InObject);
 	}
 
-	void AddExport(UObject* InObj, bool bNotAlwaysLoadedForEditorGame)
+	void AddExport(FTaggedExport InTagObj)
 	{
-		GetHarvestedRealm().AddExport(InObj, bNotAlwaysLoadedForEditorGame);
+		GetHarvestedRealm().AddExport(MoveTemp(InTagObj));
 	}
 
 	void AddExcluded(UObject* InObject)
diff --git a/Engine/Source/Runtime/CoreUObject/Private/UObject/SavePackage2.cpp b/Engine/Source/Runtime/CoreUObject/Private/UObject/SavePackage2.cpp
--- a/Engine/Source/Runtime/CoreUObject/Private/UObject/SavePackage2.cpp
+++ b/Engine/Source/Runtime/CoreUObject/Private/UObject/SavePackage2.cpp
@@ -351,7 +351,8 @@ ESavePackageResult HarvestPackage(FSaveContext& SaveContext)
 	if (!SaveContext.IsSaveAutoOptional() &&
 		SaveContext.IsSaveOptional() &&
 		SaveContext.IsCooking() &&
-		SaveContext.GetHarvestedRealm(ESaveRealm::Optional).GetExports().Num())
+		SaveContext.GetHarvestedRealm(ESaveRealm::Optional).GetExports().Num() &&
+		SaveContext.GetHarvestedRealm(ESaveRealm::Game).GetExports().Num())
 	{
 		bool bHasNonOptionalSelfReference = false;
 		FHarvestedRealm& OptionalContext = SaveContext.GetHarvestedRealm(ESaveRealm::Optional);
@@ -369,8 +370,8 @@ ESavePackageResult HarvestPackage(FSaveContext& SaveContext)
 					// Flag the package itself to be an import
 					bHasNonOptionalSelfReference = true;
 				}
-				// if not found in the game context, record an illegal reference
-				else
+				// if not found in the game context and the reference directly came from an optional object, record an illegal reference
+				else if (It->bFromOptionalReference)
 				{
 					SaveContext.RecordIllegalReference(nullptr, It->Obj, EIllegalRefReason::ReferenceFromOptionalToMissingGameExport);
 				}
