Revert cl 19830106 and cl 19940134 meant to fix loading issue when the 3 loading paths in the cooked editor are involved
The fix was causing other side effects included non deserialized objects in some cases

diff --git a/Engine/Source/Runtime/CoreUObject/Private/UObject/UObjectGlobals.cpp b/Engine/Source/Runtime/CoreUObject/Private/UObject/UObjectGlobals.cpp
--- a/Engine/Source/Runtime/CoreUObject/Private/UObject/UObjectGlobals.cpp
+++ b/Engine/Source/Runtime/CoreUObject/Private/UObject/UObjectGlobals.cpp
@@ -1443,26 +1443,12 @@ UPackage* LoadPackageInternal(UPackage* InOuter, const FPackagePath& PackagePath
 		}
 
 		GSyncLoadUsingAsyncLoaderCount++;
-
-#if WITH_EDITOR
-		// Temporarily clear the list of loaded objects before entering async loading. Otherwise the (uncooked) async loader might try to preload them which in turn can lead to
-		// recursive imports of the same cooked package.
-		FUObjectSerializeContext* LoadContext = FUObjectThreadContext::Get().GetSerializeContext();
-		TArray<UObject*>& ThreadObjLoaded = LoadContext->PRIVATE_GetObjectsLoadedInternalUseOnly();
-		TArray<UObject*> SavedObjLoaded;
-		Swap(SavedObjLoaded, ThreadObjLoaded);
-#endif
 		int32 RequestID = LoadPackageAsync(PackagePath, PackageName);
 
 		if (RequestID != INDEX_NONE)
 		{
 			FlushAsyncLoading(RequestID);
 		}
-
-#if WITH_EDITOR
-		ThreadObjLoaded.Insert(SavedObjLoaded, 0);
-#endif
-
 		GSyncLoadUsingAsyncLoaderCount--;
 
 		return (InOuter ? InOuter : FindObjectFast<UPackage>(nullptr, PackageName));
