[KillPendingKill] Remove renamed packages and their exports from the async loading thread's caches.

diff --git a/Engine/Source/Runtime/CoreUObject/Internal/UObject/UObjectGlobalsInternal.h b/Engine/Source/Runtime/CoreUObject/Internal/UObject/UObjectGlobalsInternal.h
new file mode 100644
--- /dev/null
+++ b/Engine/Source/Runtime/CoreUObject/Internal/UObject/UObjectGlobalsInternal.h
@@ -0,0 +1,19 @@
+// Copyright Epic Games, Inc. All Rights Reserved.
+
+// UObjectGlobalsInternal.h - Global UObject functions for use by other engine modules but not by game modules.
+
+#pragma once
+
+#include "CoreMinimal.h"
+
+/**
+ * Global CoreUObject delegates for use by other engine modules
+ */
+struct COREUOBJECT_API FCoreUObjectInternalDelegates
+{
+	/** Called before GC verification code rename a package containing a World to try and prevent such errors blocking testing and development. */
+	DECLARE_MULTICAST_DELEGATE_OneParam(FPackageRename, UPackage*);
+	static FPackageRename& GetOnLeakedPackageRenameDelegate();
+};
+
+    
\ No newline at end of file
diff --git a/Engine/Source/Runtime/CoreUObject/Private/Serialization/AsyncLoading2.cpp b/Engine/Source/Runtime/CoreUObject/Private/Serialization/AsyncLoading2.cpp
--- a/Engine/Source/Runtime/CoreUObject/Private/Serialization/AsyncLoading2.cpp
+++ b/Engine/Source/Runtime/CoreUObject/Private/Serialization/AsyncLoading2.cpp
@@ -25,6 +25,7 @@
 #include "UObject/ObjectResource.h"
 #include "UObject/SoftObjectPath.h"
 #include "UObject/NameBatchSerialization.h"
+#include "UObject/UObjectGlobalsInternal.h"
 #include "Serialization/DeferredMessageLog.h"
 #include "UObject/UObjectThreadContext.h"
 #include "Misc/Paths.h"
@@ -901,7 +902,7 @@ public:
 
 	void RemovePackages(const FUnreachablePackages& UnreachablePackages)
 	{
-		check(IsGarbageCollecting());
+		check(FGCCSyncObject::Get().IsGCLocked() || !IsAsyncLoading());
 
 		const int32 PackageCount = UnreachablePackages.Num();
 		for (UPackage* Package : UnreachablePackages)
@@ -2583,6 +2584,8 @@ public:
 
 private:
 
+	void OnLeakedPackageRename(UPackage* Package);
+
 	void SuspendWorkers();
 	void ResumeWorkers();
 
@@ -5481,6 +5484,8 @@ FAsyncLoadingThread2::FAsyncLoadingThread2(FIoDispatcher& InIoDispatcher, IAsync
 	ThreadResumedEvent = FPlatformProcess::GetSynchEventFromPool();
 	AsyncLoadingTickCounter = 0;
 
+	FCoreUObjectInternalDelegates::GetOnLeakedPackageRenameDelegate().AddRaw(this, &FAsyncLoadingThread2::OnLeakedPackageRename);
+
 	FAsyncLoadingThreadState2::TlsSlot = FPlatformTLS::AllocTlsSlot();
 	FAsyncLoadingThreadState2::Create(GraphAllocator, IoDispatcher);
 
@@ -5502,6 +5507,7 @@ void FAsyncLoadingThread2::ShutdownLoading()
 {
 	FCoreUObjectDelegates::GetPreGarbageCollectDelegate().RemoveAll(this);
 	FCoreUObjectDelegates::GetPostGarbageCollect().RemoveAll(this);
+	FCoreUObjectInternalDelegates::GetOnLeakedPackageRenameDelegate().RemoveAll(this);
 
 	delete Thread;
 	Thread = nullptr;
@@ -5916,6 +5922,31 @@ FORCENOINLINE static void FilterUnreachableObjects(
 	}
 }
 
+void FAsyncLoadingThread2::OnLeakedPackageRename(UPackage* Package)
+{
+	check(IsInGameThread());
+	
+	if(!FGCCSyncObject::Get().IsGCLocked())
+	{
+		// Flush async loading so that we can be sure nothing is modifying the stores, and that nothing is depending on this package
+		FlushAsyncLoading(); 
+	}
+
+	// Code such as LoadMap or LevelStreaming has renamed a loaded package which was detected as leaking so that we can load another copy of it.
+	// We should not have any loading happening at present, so we can remove these objects from our stores 
+	FUnreachablePublicExports ObjectsToRemove;
+	FUnreachablePackages PackagesToRemove;
+	PackagesToRemove.Add(Package);
+	ForEachObjectWithOuter(Package, [&ObjectsToRemove](UObject* Obj) {
+		if (DO_CHECK || Obj->HasAllFlags(GlobalImportObjectConditionalFlags))
+		{
+			ObjectsToRemove.Emplace(GUObjectArray.ObjectToIndex(Obj), Obj);
+		}
+	});
+
+	RemoveUnreachableObjects(ObjectsToRemove, PackagesToRemove);
+}
+
 void FAsyncLoadingThread2::RemoveUnreachableObjects(const FUnreachablePublicExports& PublicExports, const FUnreachablePackages& Packages)
 {
 	TRACE_CPUPROFILER_EVENT_SCOPE(RemoveUnreachableObjects);
diff --git a/Engine/Source/Runtime/CoreUObject/Private/UObject/UObjectGlobals.cpp b/Engine/Source/Runtime/CoreUObject/Private/UObject/UObjectGlobals.cpp
--- a/Engine/Source/Runtime/CoreUObject/Private/UObject/UObjectGlobals.cpp
+++ b/Engine/Source/Runtime/CoreUObject/Private/UObject/UObjectGlobals.cpp
@@ -66,6 +66,7 @@
 #include "Misc/PackageAccessTracking.h"
 #include "UObject/PropertyWithSetterAndGetter.h"
 #include "UObject/AnyPackagePrivate.h"
+#include "UObject/UObjectGlobalsInternal.h"
 
 DEFINE_LOG_CATEGORY(LogUObjectGlobals);
 
@@ -175,6 +176,12 @@ FSimpleMulticastDelegate& FCoreUObjectDelegates::GetPostGarbageCollect()
 	return Delegate;
 }
 
+FCoreUObjectInternalDelegates::FPackageRename& FCoreUObjectInternalDelegates::GetOnLeakedPackageRenameDelegate()
+{
+	static FPackageRename Delegate;
+	return Delegate;
+}
+
 FCoreUObjectDelegates::FTraceExternalRootsForReachabilityAnalysisDelegate FCoreUObjectDelegates::TraceExternalRootsForReachabilityAnalysis;
 FSimpleMulticastDelegate FCoreUObjectDelegates::PostReachabilityAnalysis;
 
diff --git a/Engine/Source/Runtime/Engine/Private/UnrealEngine.cpp b/Engine/Source/Runtime/Engine/Private/UnrealEngine.cpp
--- a/Engine/Source/Runtime/Engine/Private/UnrealEngine.cpp
+++ b/Engine/Source/Runtime/Engine/Private/UnrealEngine.cpp
@@ -44,6 +44,7 @@ UnrealEngine.cpp: Implements the UEngine class and helpers.
 #include "UObject/Package.h"
 #include "UObject/MetaData.h"
 #include "UObject/ObjectMemoryAnalyzer.h"
+#include "UObject/UObjectGlobalsInternal.h"
 #include "Serialization/ArchiveCountMem.h"
 #include "Serialization/ObjectWriter.h"
 #include "Serialization/ObjectReader.h"
@@ -15532,6 +15533,7 @@ void UEngine::VerifyLoadMapWorldCleanup(FWorldContext* ForWorldContext)
 			UE_LOG(LogLoad, Warning, TEXT("Renaming leaked objects in case they need to be reloaded."));
 			for (UPackage* Pkg : LeakedPackages)
 			{
+				FCoreUObjectInternalDelegates::GetOnLeakedPackageRenameDelegate().Broadcast(Pkg);
 				FName BaseName = FName(Pkg->GetFName(), NAME_NO_NUMBER_INTERNAL);
 				TStringBuilder<NAME_SIZE> Builder;
 				Builder << TEXT("/Leaked");
diff --git a/Engine/Source/Runtime/Engine/Private/World.cpp b/Engine/Source/Runtime/Engine/Private/World.cpp
--- a/Engine/Source/Runtime/Engine/Private/World.cpp
+++ b/Engine/Source/Runtime/Engine/Private/World.cpp
@@ -33,6 +33,7 @@
 #include "Misc/ConfigCacheIni.h"
 #include "UObject/MetaData.h"
 #include "UObject/PropertyPortFlags.h"
+#include "UObject/UObjectGlobalsInternal.h"
 #include "Serialization/ArchiveReplaceObjectRef.h"
 #include "GameFramework/OnlineReplStructs.h"
 #include "GameFramework/PlayerController.h"
@@ -3371,6 +3372,7 @@ void FLevelStreamingGCHelper::PrepareStreamedOutLevelsForGC()
 				// Rename the packages that we are streaming out so that we can possibly reload another copy of them
 				for (UPackage* Package : Packages)
 				{
+					FCoreUObjectInternalDelegates::GetOnLeakedPackageRenameDelegate().Broadcast(Package);
 					FName NewName = MakeUniqueObjectName(nullptr, UPackage::StaticClass(), Package->GetFName());
 					Package->Rename(*NewName.ToString(), nullptr, REN_ForceNoResetLoaders | REN_DontCreateRedirectors | REN_NonTransactional);
 				}
