Delegate for adding additional debug info to garbage reference tracking in GC.

diff --git a/Engine/Source/Runtime/CoreUObject/Private/UObject/GarbageCollection.cpp b/Engine/Source/Runtime/CoreUObject/Private/UObject/GarbageCollection.cpp
--- a/Engine/Source/Runtime/CoreUObject/Private/UObject/GarbageCollection.cpp
+++ b/Engine/Source/Runtime/CoreUObject/Private/UObject/GarbageCollection.cpp
@@ -2519,7 +2519,11 @@ void FGCArrayPool::DumpGarbageReferencers(TArray<FGCArrayStruct*>& AllArrays)
 					NumGarbageReferences++;
 				}
 			}
-			ArrayStruct->GarbageReferences.Reset();
+			if (ArrayStruct->GarbageReferences.Num())
+			{
+				FCoreUObjectDelegates::GetGarbageCollectReportGarbageReferencers().Broadcast(ArrayStruct->GarbageReferences);
+				ArrayStruct->GarbageReferences.Reset();
+			}
 		}
 		UE_CLOG(NumGarbageReferences > 0, LogGarbage, Log, TEXT("Found %d garbage references in %f ms."), NumGarbageReferences, (FPlatformTime::Seconds() - StartTime) * 1000);
 	}
@@ -2549,7 +2553,11 @@ void FGCArrayPool::DumpGarbageReferencers(TArray<FGCArrayStruct*>& AllArrays)
 					TotalGarbageReferences++;
 				}
 			}
-			ArrayStruct->GarbageReferences.Reset();
+			if (ArrayStruct->GarbageReferences.Num())
+			{
+				FCoreUObjectDelegates::GetGarbageCollectReportGarbageReferencers().Broadcast(ArrayStruct->GarbageReferences);
+				ArrayStruct->GarbageReferences.Reset();
+			}
 		}
 		UE_CLOG(TotalGarbageReferences > 0, LogGarbage, Log, TEXT("Reported %d/%d garbage references in %f ms."), ReportedGarbageReferences, TotalGarbageReferences, (FPlatformTime::Seconds() - StartTime) * 1000);
 	}
diff --git a/Engine/Source/Runtime/CoreUObject/Private/UObject/UObjectGlobals.cpp b/Engine/Source/Runtime/CoreUObject/Private/UObject/UObjectGlobals.cpp
--- a/Engine/Source/Runtime/CoreUObject/Private/UObject/UObjectGlobals.cpp
+++ b/Engine/Source/Runtime/CoreUObject/Private/UObject/UObjectGlobals.cpp
@@ -176,6 +176,14 @@ FSimpleMulticastDelegate& FCoreUObjectDelegates::GetPostGarbageCollect()
 	return Delegate;
 }
 
+#if !UE_BUILD_SHIPPING
+FCoreUObjectDelegates::FOnReportGarbageReferencers& FCoreUObjectDelegates::GetGarbageCollectReportGarbageReferencers()
+{
+	static FCoreUObjectDelegates::FOnReportGarbageReferencers Delegate;
+	return Delegate;
+}
+#endif
+
 FCoreUObjectInternalDelegates::FPackageRename& FCoreUObjectInternalDelegates::GetOnLeakedPackageRenameDelegate()
 {
 	static FPackageRename Delegate;
diff --git a/Engine/Source/Runtime/CoreUObject/Public/UObject/UObjectGlobals.h b/Engine/Source/Runtime/CoreUObject/Public/UObject/UObjectGlobals.h
--- a/Engine/Source/Runtime/CoreUObject/Public/UObject/UObjectGlobals.h
+++ b/Engine/Source/Runtime/CoreUObject/Public/UObject/UObjectGlobals.h
@@ -2761,6 +2761,12 @@ struct COREUOBJECT_API FCoreUObjectDelegates
 	/** Called after garbage collection */
 	static FSimpleMulticastDelegate& GetPostGarbageCollect();
 
+#if !UE_BUILD_SHIPPING
+	/** Called when garbage collection detects references to objects that are marked for explicit destruction by MarkAsGarbage */
+	DECLARE_MULTICAST_DELEGATE_OneParam(FOnReportGarbageReferencers, TConstArrayView<struct FGarbageReferenceInfo>);
+	static FOnReportGarbageReferencers& GetGarbageCollectReportGarbageReferencers();
+#endif
+
 	/** Called before ConditionalBeginDestroy phase of garbage collection */
 	static FSimpleMulticastDelegate PreGarbageCollectConditionalBeginDestroy;
 
