Fix for a regression that incorrectly flags nested template objects as inherited instanced subobjects in the export map.

diff --git a/Engine/Source/Runtime/CoreUObject/Private/UObject/ObjectResource.cpp b/Engine/Source/Runtime/CoreUObject/Private/UObject/ObjectResource.cpp
--- a/Engine/Source/Runtime/CoreUObject/Private/UObject/ObjectResource.cpp
+++ b/Engine/Source/Runtime/CoreUObject/Private/UObject/ObjectResource.cpp
@@ -88,7 +88,7 @@ FObjectExport::FObjectExport( UObject* InObject, bool bInNotAlwaysLoadedForEdito
 		bNotForClient = !Object->NeedsLoadForClient();
 		bNotForServer = !Object->NeedsLoadForServer();
 		bIsAsset = Object->IsAsset();
-		bIsInheritedInstance = Object->GetArchetype()->IsDefaultSubobject();
+		bIsInheritedInstance = !Object->HasAnyFlags(RF_ArchetypeObject) && Object->GetArchetype()->IsDefaultSubobject();
 	}
 }
 
@@ -164,6 +164,12 @@ void operator<<(FStructuredArchive::FSlot Slot, FObjectExport& E)
 	if (BaseArchive.UEVer() >= EUnrealEngineObjectUE5Version::TRACK_OBJECT_EXPORT_IS_INHERITED)
 	{
 		SERIALIZE_BIT_TO_RECORD(bIsInheritedInstance);
+
+		// Initial saves of TRACK_OBJECT_EXPORT_IS_INHERITED incorrectly considered nested template objects as inherited instances; fix that up here on load.
+		if (BaseArchive.IsLoading() && BaseArchive.UEVer().ToValue() == static_cast<int32>(EUnrealEngineObjectUE5Version::TRACK_OBJECT_EXPORT_IS_INHERITED))
+		{
+			E.bIsInheritedInstance &= !(!!(E.ObjectFlags & RF_ArchetypeObject));
+		}
 	}
 
 	Record << SA_VALUE(TEXT("PackageFlags"), E.PackageFlags);
