Fixed regression in FSoftObjectPath::SetPath from ANY_PACKAGE removal

The class part of an ExportText path now also starts with a /, so the code that was trimming it off during SetPath was no longer running.

#rnx

diff --git a/Engine/Source/Runtime/CoreUObject/Private/Misc/PackageName.cpp b/Engine/Source/Runtime/CoreUObject/Private/Misc/PackageName.cpp
--- a/Engine/Source/Runtime/CoreUObject/Private/Misc/PackageName.cpp
+++ b/Engine/Source/Runtime/CoreUObject/Private/Misc/PackageName.cpp
@@ -2209,7 +2209,7 @@ bool ParseExportTextPathImpl(const T& InExportTextPath, T* OutClassName, T* OutO
 		if (OutObjectPath)
 		{
 			*OutObjectPath = InExportTextPath.Mid(Index + 1);
-			OutObjectPath->RemoveSuffix(InExportTextPath.EndsWith('\''));
+			OutObjectPath->RemoveSuffix(OutObjectPath->EndsWith('\''));
 		}
 
 		return true;
diff --git a/Engine/Source/Runtime/CoreUObject/Private/UObject/SoftObjectPath.cpp b/Engine/Source/Runtime/CoreUObject/Private/UObject/SoftObjectPath.cpp
--- a/Engine/Source/Runtime/CoreUObject/Private/UObject/SoftObjectPath.cpp
+++ b/Engine/Source/Runtime/CoreUObject/Private/UObject/SoftObjectPath.cpp
@@ -123,11 +123,8 @@ void FSoftObjectPath::SetPath(FWideStringView Path)
 	}
 	else if (ensureMsgf(!FPackageName::IsShortPackageName(Path), TEXT("Cannot create SoftObjectPath with short package name '%.*s'%s! You must pass in fully qualified package names"), Path.Len(), Path.GetData(), *GetObjectBeingSerializedForSoftObjectPath()))
 	{
-		if (Path[0] != '/')
-		{
-			// Possibly an ExportText path. Trim the ClassName.
-			Path = FPackageName::ExportTextPathToObjectPath(Path);
-		}
+		// Possibly an ExportText path. Trim the ClassName.
+		Path = FPackageName::ExportTextPathToObjectPath(Path);
 
 		int32 ColonIndex;
 		if (Path.FindChar(SUBOBJECT_DELIMITER_CHAR, ColonIndex))
@@ -172,11 +169,8 @@ void FSoftObjectPath::SetPath(FName PathName)
 
 		if (ensureMsgf(!FPackageName::IsShortPackageName(Path), TEXT("Cannot create SoftObjectPath with short package name '%s'%s! You must pass in fully qualified package names"), Path.GetData(), *GetObjectBeingSerializedForSoftObjectPath()))
 		{
-			if (Path[0] != '/')
-			{
-				// Possibly an ExportText path. Trim the ClassName.
-				Path = FPackageName::ExportTextPathToObjectPath(Path);
-			}
+			// Possibly an ExportText path. Trim the ClassName.
+			Path = FPackageName::ExportTextPathToObjectPath(Path);
 
 			int32 ColonIndex;
 			if (Path.FindChar(SUBOBJECT_DELIMITER_CHAR, ColonIndex))
