EditorBulkData can now attempt to read payloads from package files without a valid attachment to the package file on disk.

### Problem
- Because of the package trailer we can easily find out which payloads are present in a package file and where they are. This means we don't need to rely on the cached offset into the file that bulkdata generally uses for payload loading. This means as long as the payload is still actually in a payload file then we can load it no matter what.
- We will still give errors in the following scenarios
-- The payload is not stored in a package trailer (either referencing older bulkdata or an older editor bulkdata format)
-- We cannot get a read handle to the package file.
-- The payload is no longer in the package file.
- If the attachment is still valid then we can guarantee that the payload can be loaded. Without an attachment we make our best efforts but in the case of the package file being reomoved/replaced this can fail. So ::ResetLoader will still be used to indicate that we want to remove the attachment but expect the package file to be changed and to load the payload into memory for safety. Work on improving this is another work item for the future.

### Fix
- We no longer remove the package path when ::DetachFromDisk is called nor do we remove flags relating to the payload. From this point onwards we use the attached archive being null to determine how we load. We continue to set OffsetInFile to an invalid value as we can no longer safely use it.
- Add a new method ::IsAttachedToPackageFile to explicitly show when we are checking to see if the bulkdata is still 'attached' to the package file.
- Explicitly set AttachedAr to null ptr when ::DetachBulkData is called on it. In theory the ::DetachBulkData call will call FEditorBulkData::DetachFromDisk which in turn will set AttachedAr but it can look confusing when reading the code. Explicitly setting the AttachedAr member to nullptr helps make the intent clearer.
- When loading the payload from disk there is now a new code branch that is followed if we detect that we are loading from a package file that we are no longer attached to.
-- This path gives bespoke errors in cases where it is not safe for us to try loading the payload, in addition there is a bespoke error message if loading from the package trailer fails.
-- Note we can only risk loading from the package trailer as it's format allows for us to find if the payload is still on disk or not. With the older formats we will not know if we are loading garbage or not.
-- Fixed a bug with CVarShouldLoadFromTrailer. When it is true we will only attempt to load payloads from the trailer if the payload is stored in the trailer. The comment for this development CVar has been updated.
- There has been a pass made to error reporting. If we are reporting an error relating to the package we will output the package name, if we are reporting on a error relating to the file itself then we will output the package path + extension

diff --git a/Engine/Source/Runtime/CoreUObject/Private/Serialization/EditorBulkData.cpp b/Engine/Source/Runtime/CoreUObject/Private/Serialization/EditorBulkData.cpp
--- a/Engine/Source/Runtime/CoreUObject/Private/Serialization/EditorBulkData.cpp
+++ b/Engine/Source/Runtime/CoreUObject/Private/Serialization/EditorBulkData.cpp
@@ -44,7 +44,7 @@ TRACE_DECLARE_INT_COUNTER(EditorBulkData_PayloadDataPulled, TEXT("EditorBulkData
 
 namespace UE::Serialization
 {
-/** This console variable should only exist for testing */
+/** This is an experimental code path and is not expected to be used in production! */
 static TAutoConsoleVariable<bool> CVarShouldLoadFromSidecar(
 	TEXT("Serialization.LoadFromSidecar"),
 	false,
@@ -118,11 +118,11 @@ static void LogPackageOpenFailureMessage(const FPackagePath& PackagePath, EPacka
 	{
 		TCHAR SystemErrorMsg[2048] = { 0 };
 		FPlatformMisc::GetSystemErrorMessage(SystemErrorMsg, sizeof(SystemErrorMsg), SystemError);
-		UE_LOG(LogSerialization, Error, TEXT("Could not open the file '%s' for reading due to system error: '%s' (%d))"), *PackagePath.GetDebugNameWithExtension(PackageSegment), SystemErrorMsg, SystemError);
+		UE_LOG(LogSerialization, Error, TEXT("Could not open the file '%s' for reading due to system error: '%s' (%d))"), *PackagePath.GetLocalFullPath(PackageSegment), SystemErrorMsg, SystemError);
 	}
 	else
 	{
-		UE_LOG(LogSerialization, Error, TEXT("Could not open (%s) to read FEditorBulkData with an unknown error"), *PackagePath.GetDebugNameWithExtension(PackageSegment));
+		UE_LOG(LogSerialization, Error, TEXT("Could not open (%s) to read FEditorBulkData with an unknown error"), *PackagePath.GetLocalFullPath(PackageSegment));
 	}
 }
 
@@ -443,9 +443,10 @@ FEditorBulkData& FEditorBulkData::operator=(const FEditorBulkData& Other)
 
 FEditorBulkData::~FEditorBulkData()
 {
-	if (AttachedAr != nullptr)
+	if (IsAttachedToPackageFile())
 	{
 		AttachedAr->DetachBulkData(this, false);
+		AttachedAr = nullptr;
 	}
 
 	OnExitMemory();
@@ -665,8 +666,9 @@ void FEditorBulkData::CreateFromBulkData(FBulkData& InBulkData, const FGuid& InG
 	Reset();
 
 #if UE_ALLOW_LINKERLOADER_ATTACHMENT
+	check(!IsAttachedToPackageFile());
 	AttachedAr = InBulkData.AttachedAr;
-	if (AttachedAr != nullptr)
+	if (IsAttachedToPackageFile())
 	{
 		AttachedAr->AttachBulkData(this);
 	}
@@ -1017,6 +1019,12 @@ void FEditorBulkData::Serialize(FArchive& Ar, UObject* Owner, bool bAllowRegiste
 				if (!PackagePath.IsEmpty() && CacheableArchive != nullptr)
 				{
 #if UE_ALLOW_LINKERLOADER_ATTACHMENT
+					if (IsAttachedToPackageFile())
+					{
+						// TODO: Remove this when doing UE-159339
+						AttachedAr->DetachBulkData(this, false);
+					}
+
 					AttachedAr = CacheableArchive;
 					AttachedAr->AttachBulkData(this);
 #endif //VBD_ALLOW_LINKERLOADER_ATTACHMENT
@@ -1114,9 +1122,47 @@ FCompressedBuffer FEditorBulkData::LoadFromDisk() const
 	{
 		return LoadFromSidecarFile();
 	}
+	else if (!IsAttachedToPackageFile())
+	{
+		if (IsReferencingOldBulkData())
+		{
+			UE_LOG(LogSerialization, Error, 
+				TEXT("Cannot attempt to load the payload '%s' from '%s' as the package is no longer attached to the file on disk and the payload is in an old style bulkdata structure"), 
+				*LexToString(PayloadContentId), 
+				*PackagePath.GetLocalFullPath(EPackageSegment::Header));
+
+			return FCompressedBuffer();
+		}
+
+		if (!IsStoredInPackageTrailer())
+		{
+			UE_LOG(LogSerialization, Error, 
+				TEXT("Cannot attempt to load the payload '%s' from '%s' as the package is no longer attached to the file on disk and the payload is not stored in a package trailer"),
+				*LexToString(PayloadContentId),
+				*PackagePath.GetLocalFullPath(EPackageSegment::Header));
+
+			return FCompressedBuffer();
+		}
+
+		FCompressedBuffer CompressedPayload = LoadFromPackageTrailer();
+		if (!CompressedPayload.IsNull())
+		{
+			return CompressedPayload;
+		}
+		else
+		{
+			UE_LOG(LogSerialization, Error,
+				TEXT("Cannot attempt to load the payload '%s' from '%s' as the package is no longer attached to the file on disk and the payload was not found in the package trailer"),
+				*LexToString(PayloadContentId),
+				*PackagePath.GetLocalFullPath(EPackageSegment::Header));
+
+			return FCompressedBuffer();
+		}
+		
+	}
 	else
 	{
-		if (CVarShouldLoadFromTrailer.GetValueOnAnyThread())
+		if (CVarShouldLoadFromTrailer.GetValueOnAnyThread() && IsStoredInPackageTrailer())
 		{
 			return LoadFromPackageTrailer();
 		}
@@ -1131,7 +1177,7 @@ FCompressedBuffer FEditorBulkData::LoadFromPackageFile() const
 {
 	TRACE_CPUPROFILER_EVENT_SCOPE(FEditorBulkData::LoadFromPackageFile);
 
-	UE_LOG(LogSerialization, Verbose, TEXT("Attempting to load payload from the package file: '%s'"), *PackagePath.GetLocalFullPath(EPackageSegment::Header));
+	UE_LOG(LogSerialization, Verbose, TEXT("Attempting to load payload from the package file '%s'"), *PackagePath.GetLocalFullPath(EPackageSegment::Header));
 
 	// Open a reader to the file
 	TUniquePtr<FArchive> BulkArchive;
@@ -1149,8 +1195,7 @@ FCompressedBuffer FEditorBulkData::LoadFromPackageFile() const
 		// Workspace Domain file. This was only possible if PackageSegment == Header; we checked that when serializing to the EditorDomain
 		// In that case, we need to use OpenReadExternalResource to access the Workspace Domain file
 		// In the cases where *this was loaded from the WorkspaceDomain, OpenReadExternalResource and OpenReadPackage are identical.
-		BulkArchive = IPackageResourceManager::Get().OpenReadExternalResource(EPackageExternalResource::WorkspaceDomainFile,
-			PackagePath.GetPackageName());
+		BulkArchive = IPackageResourceManager::Get().OpenReadExternalResource(EPackageExternalResource::WorkspaceDomainFile, PackagePath.GetPackageName());
 	}
 
 	if (!BulkArchive.IsValid())
@@ -1159,7 +1204,7 @@ FCompressedBuffer FEditorBulkData::LoadFromPackageFile() const
 		return FCompressedBuffer();
 	}
 
-	checkf(OffsetInFile != INDEX_NONE, TEXT("Attempting to load '%s' from disk with an invalid OffsetInFile!"), *PackagePath.GetDebugNameWithExtension(EPackageSegment::Header));
+	checkf(OffsetInFile != INDEX_NONE, TEXT("Attempting to load from the package file '%s' with an invalid OffsetInFile!"), *PackagePath.GetLocalFullPath(EPackageSegment::Header));
 	// Move the correct location of the data in the file
 	BulkArchive->Seek(OffsetInFile);
 
@@ -1174,7 +1219,7 @@ FCompressedBuffer FEditorBulkData::LoadFromPackageTrailer() const
 {
 	TRACE_CPUPROFILER_EVENT_SCOPE(FEditorBulkData::LoadFromPackageTrailer);
 
-	UE_LOG(LogSerialization, Verbose, TEXT("Attempting to load payload from the package trailer: '%s'"), *PackagePath.GetLocalFullPath(EPackageSegment::Header));
+	UE_LOG(LogSerialization, Verbose, TEXT("Attempting to load a payload via the package trailer from file '%s'"), *PackagePath.GetLocalFullPath(EPackageSegment::Header));
 
 	// TODO: Could just get the trailer from the owning FLinkerLoad if still attached
 
@@ -1194,8 +1239,7 @@ FCompressedBuffer FEditorBulkData::LoadFromPackageTrailer() const
 		// Workspace Domain file. This was only possible if PackageSegment == Header; we checked that when serializing to the EditorDomain
 		// In that case, we need to use OpenReadExternalResource to access the Workspace Domain file
 		// In the cases where *this was loaded from the WorkspaceDomain, OpenReadExternalResource and OpenReadPackage are identical.
-		BulkArchive = IPackageResourceManager::Get().OpenReadExternalResource(EPackageExternalResource::WorkspaceDomainFile,
-			PackagePath.GetPackageName());
+		BulkArchive = IPackageResourceManager::Get().OpenReadExternalResource(EPackageExternalResource::WorkspaceDomainFile, PackagePath.GetPackageName());
 	}
 
 	if (!BulkArchive.IsValid())
@@ -1207,15 +1251,21 @@ FCompressedBuffer FEditorBulkData::LoadFromPackageTrailer() const
 	BulkArchive->Seek(BulkArchive->TotalSize());
 
 	FPackageTrailer Trailer;
-	
 	if (Trailer.TryLoadBackwards(*BulkArchive))
 	{
-		return Trailer.LoadLocalPayload(PayloadContentId, *BulkArchive);
+		FCompressedBuffer CompressedPayload = Trailer.LoadLocalPayload(PayloadContentId, *BulkArchive);
+
+		UE_CLOG(CompressedPayload.IsNull(), LogSerialization, Error, TEXT("Could not find the payload '%s' in the package trailer of file '%s'"),
+			*LexToString(PayloadContentId ), 
+			*PackagePath.GetLocalFullPath(EPackageSegment::Header));
+
+		return CompressedPayload;
 	}
 	else
 	{
+		UE_LOG(LogSerialization, Error, TEXT("Could not read the package trailer from file '%s'"), *PackagePath.GetLocalFullPath(EPackageSegment::Header));
 		return FCompressedBuffer();
-	}	
+	}
 }
 
 FCompressedBuffer FEditorBulkData::LoadFromSidecarFileInternal(ErrorVerbosity Verbosity) const
@@ -1230,7 +1280,7 @@ FCompressedBuffer FEditorBulkData::LoadFromSidecarFileInternal(ErrorVerbosity Ve
 
 		if (Version != FTocEntry::PayloadSidecarFileVersion)
 		{
-			UE_CLOG(Verbosity > ErrorVerbosity::None, LogSerialization, Error, TEXT("Unknown version (%u) found in '%s'"), Version, *PackagePath.GetLocalFullPath(EPackageSegment::PayloadSidecar));
+			UE_CLOG(Verbosity > ErrorVerbosity::None, LogSerialization, Error, TEXT("Unknown payload sidecar version (%u) found in '%s'"), Version, *PackagePath.GetLocalFullPath(EPackageSegment::PayloadSidecar));
 			return FCompressedBuffer();
 		}
 
@@ -1258,12 +1308,16 @@ FCompressedBuffer FEditorBulkData::LoadFromSidecarFileInternal(ErrorVerbosity Ve
 			}
 			else if(Verbosity > ErrorVerbosity::None)
 			{
-				UE_LOG(LogSerialization, Error, TEXT("Payload '%s' in '%s' has an invalid OffsetInFile!"), *LexToString(PayloadContentId), *PackagePath.GetLocalFullPath(EPackageSegment::PayloadSidecar));
+				UE_LOG(LogSerialization, Error, TEXT("Payload '%s' in '%s' has an invalid OffsetInFile!"), 
+					*LexToString(PayloadContentId), 
+					*PackagePath.GetLocalFullPath(EPackageSegment::PayloadSidecar));
 			}
 		}
 		else if(Verbosity > ErrorVerbosity::None)
 		{
-			UE_LOG(LogSerialization, Error, TEXT("Unable to find payload '%s' in '%s'"), *LexToString(PayloadContentId), *PackagePath.GetLocalFullPath(EPackageSegment::PayloadSidecar));
+			UE_LOG(LogSerialization, Error, TEXT("Unable to find payload '%s' in '%s'"), 
+				*LexToString(PayloadContentId), 
+				*PackagePath.GetLocalFullPath(EPackageSegment::PayloadSidecar));
 		}
 	}
 	else if(Verbosity > ErrorVerbosity::None)
@@ -1412,8 +1466,11 @@ FCompressedBuffer FEditorBulkData::PullData() const
 
 bool FEditorBulkData::CanUnloadData() const
 {
-	// We cannot unload the data if are unable to reload it from a file
-	return IsDataVirtualized() || (PackagePath.IsEmpty() == false && AttachedAr != nullptr);
+	// Technically if we have a valid path but are detached we can still try  to load 
+	// the payload from disk via ::LoadFromPackageTrailer but we are not guaranteed success
+	// because the package file may have changed to one without the payload in it at all.
+	// Due to this we do not allow the payload to be unloaded if we are detached.
+	return IsDataVirtualized() || (!PackagePath.IsEmpty() && IsAttachedToPackageFile());
 }
 
 bool FEditorBulkData::IsMemoryOnlyPayload() const
@@ -1426,9 +1483,10 @@ void FEditorBulkData::Reset()
 	// Unregister rather than allowing the Registry to keep our record, since we are changing the payload
 	Unregister();
 	// Note that we do not reset the BulkDataId
-	if (AttachedAr != nullptr)
+	if (IsAttachedToPackageFile())
 	{
 		AttachedAr->DetachBulkData(this, false);
+		AttachedAr = nullptr;
 	}
 
 	PayloadContentId.Reset();
@@ -1453,9 +1511,16 @@ void FEditorBulkData::DetachFromDisk(FArchive* Ar, bool bEnsurePayloadIsLoaded)
 {
 	TRACE_CPUPROFILER_EVENT_SCOPE(FEditorBulkData::DetachFromDisk);
 
-	check(Ar);
-	check(Ar == AttachedAr || AttachedAr == nullptr || AttachedAr->IsProxyOf(Ar));
+	check(Ar != nullptr);
+	check(Ar == AttachedAr || !IsAttachedToPackageFile() || AttachedAr->IsProxyOf(Ar));
 
+	// If bEnsurePayloadIsLoaded is true, then we should assume that a change to the 
+	// package file is imminent and we should load the payload into memory if possible.
+	// If it is false then either we are not expecting a change to the package file or
+	// we are not expected that the payload will be used again and can unload the 
+	// payload. If the payload is requested later we will try to load it off disk via
+	// the package trailer, but there is no guarantee that it will be present in the
+	// package file.
 	if (!IsDataVirtualized() && !PackagePath.IsEmpty())
 	{
 		if (Payload.IsNull() && bEnsurePayloadIsLoaded)
@@ -1465,10 +1530,8 @@ void FEditorBulkData::DetachFromDisk(FArchive* Ar, bool bEnsurePayloadIsLoaded)
 			Payload = CompressedPayload.Decompress();
 		}
 
-		PackagePath.Empty();
 		OffsetInFile = INDEX_NONE;
 
-		EnumRemoveFlags(Flags, EFlags::ReferencesLegacyFile | EFlags::ReferencesWorkspaceDomain | EFlags::LegacyFileIsCompressed);
 		if (EnumHasAnyFlags(Flags, EFlags::HasRegistered))
 		{
 			if (bEnsurePayloadIsLoaded)
@@ -1635,13 +1698,12 @@ void FEditorBulkData::UpdatePayloadImpl(FSharedBuffer&& InPayload, FIoHash&& InP
 	// Unregister before calling DetachBulkData; DetachFromDisk calls OnExitMemory which is incorrect since
 	// we are changing data rather than leaving memory
 	Unregister();
-	if (AttachedAr != nullptr)
+	if (IsAttachedToPackageFile())
 	{
 		AttachedAr->DetachBulkData(this, false);
+		AttachedAr = nullptr;
 	}
 
-	check(AttachedAr == nullptr);
-
 	// We only take the payload if it has a length to avoid potentially holding onto a
 	// 0 byte allocation in the FSharedBuffer
 	if (InPayload.GetSize() > 0)
@@ -1708,7 +1770,7 @@ FCompressedBuffer FEditorBulkData::GetDataInternal() const
 		check(Payload.IsNull()); //Make sure that we did not assign the buffer internally
 
 		UE_CLOG(CompressedPayload.IsNull(), LogSerialization, Error, TEXT("Failed to load payload '%s"), *LexToString(PayloadContentId));
-		UE_CLOG(!IsDataValid(*this, CompressedPayload), LogSerialization, UE_CORRUPTED_DATA_SEVERITY, TEXT("Payload '%s' loaded from '%s' is corrupt! Check the file on disk."),
+		UE_CLOG(!IsDataValid(*this, CompressedPayload), LogSerialization, UE_CORRUPTED_DATA_SEVERITY, TEXT("Payload '%s' loaded from package '%s' is corrupt! Check the package file on disk."),
 			*LexToString(PayloadContentId),
 			*PackagePath.GetDebugName());
 
@@ -2057,7 +2119,7 @@ FText FEditorBulkData::GetCorruptedPayloadErrorMsgForSave(FLinkerSave* Linker) c
 	else if(!PackagePath.IsEmpty())
 	{
 		// We don't know where we are saving to, but we do know the package where the payload came from.
-		const FText PackageName = FText::FromString(PackagePath.GetPackageName());
+		const FText PackageName = FText::FromString(PackagePath.GetDebugName());
 
 		return FText::Format(	NSLOCTEXT("Core", "Serialization_InvalidPayloadFromPkg", "Attempting to save bulkdata {0} with an invalid payload from package '{1}'. The package probably needs to be reverted/recreated to fix this."),
 								GuidID, PackageName);
diff --git a/Engine/Source/Runtime/CoreUObject/Public/Serialization/EditorBulkData.h b/Engine/Source/Runtime/CoreUObject/Public/Serialization/EditorBulkData.h
--- a/Engine/Source/Runtime/CoreUObject/Public/Serialization/EditorBulkData.h
+++ b/Engine/Source/Runtime/CoreUObject/Public/Serialization/EditorBulkData.h
@@ -445,6 +445,15 @@ private:
 		return EnumHasAnyFlags(InFlags, EFlags::StoredInPackageTrailer);
 	}
 
+	/** 
+	 * Returns true when the bulkdata has an attachment to it's package file on disk, meaning that it is safe for us to load
+	 * the payload from the file directly as we know where the payload is on disk 
+	 */
+	bool IsAttachedToPackageFile() const
+	{
+		return AttachedAr != nullptr;
+	}
+
 	void Register(UObject* Owner, const TCHAR* LogCallerName, bool bAllowUpdateId);
 	/** Used when we need to clear the registration for our bulkdataId because its ownership is transferring elsewhere. */
 	void Unregister();
@@ -475,12 +484,12 @@ private:
 	FText GetCorruptedPayloadErrorMsgForSave(FLinkerSave* Linker) const;
 
 	/**
-	 * A utility for validating that a package trailer builder was created correctly. If something incorrect is encountered we will assert
+	 * A utility for validating that a package trailer builder was created correctly. If a problem is encountered we will assert
 	 * to prevent a corrupted package from being saved.
-	 * The main thing we check is that the payload was added as in the correct list for the correct payload storage type based on the flags 
-	 * we have for the payload.
-	 * Note that it is not expected that we will ever encounter these problems (so asserting is acceptable) and the checks could be considered
-	 * a little over cautious. We should consider just removing this check in 5.2 onwards.
+	 * The main thing we check is that the payload was added to the correct list for the correct payload storage type based on
+	 * the flags we have for the payload.
+	 * Note that it is not expected that we will ever encounter these problems (so asserting is acceptable) and the checks could
+	 * be considered a little over cautious. We should consider just removing this check in 5.2 onwards.
 	 * 
 	 * @param LinkerSave	The linker containing the package trailer builder
 	 * @param Id			The hash of the payload we want to verify
