Fix a regression that caused subobjects instanced from non-DSO archetypes (e.g. Blueprint-added components) to fail export creation on load, leading to data loss (e.g. transform resets to origin).

diff --git a/Engine/Source/Runtime/CoreUObject/Private/UObject/LinkerLoad.cpp b/Engine/Source/Runtime/CoreUObject/Private/UObject/LinkerLoad.cpp
--- a/Engine/Source/Runtime/CoreUObject/Private/UObject/LinkerLoad.cpp
+++ b/Engine/Source/Runtime/CoreUObject/Private/UObject/LinkerLoad.cpp
@@ -4874,6 +4874,18 @@ UObject* FLinkerLoad::CreateExport( int32 Index )
 			}
 		}
 
+		// Initial saves of TRACK_OBJECT_EXPORT_IS_INHERITED incorrectly considered Blueprint-added component archetypes
+		// and subobjects instanced from those archetypes as inherited instances. The flag itself is intended to denote an
+		// instanced default subobject that's based on an archetype contained within the owner's archetype's set of instanced
+		// default subobjects; that is, the subobject owner's archetype is expected to also contain a matching default subobject
+		// instance with the same type/name. However, if the instanced subobject is based on an archetype that's owned by
+		// something other than it's owner's archetype (e.g. Blueprint-added component archetypes, which are owned by the
+		// Blueprint class object), such a match would not exist, and the export would then fail to be created below as a result.
+		if (Export.bIsInheritedInstance && Template->GetOuter()->IsA<UClass>())
+		{
+			Export.bIsInheritedInstance = false;
+		}
+
 		// If this export was an instance inherited from another class, 
 		// it should not be recreated unless the archetype also has it
 		if (Export.bIsInheritedInstance)
diff --git a/Engine/Source/Runtime/CoreUObject/Private/UObject/ObjectResource.cpp b/Engine/Source/Runtime/CoreUObject/Private/UObject/ObjectResource.cpp
--- a/Engine/Source/Runtime/CoreUObject/Private/UObject/ObjectResource.cpp
+++ b/Engine/Source/Runtime/CoreUObject/Private/UObject/ObjectResource.cpp
@@ -88,7 +88,20 @@ FObjectExport::FObjectExport( UObject* InObject, bool bInNotAlwaysLoadedForEdito
 		bNotForClient = !Object->NeedsLoadForClient();
 		bNotForServer = !Object->NeedsLoadForServer();
 		bIsAsset = Object->IsAsset();
-		bIsInheritedInstance = !Object->HasAnyFlags(RF_ArchetypeObject) && Object->GetArchetype()->IsDefaultSubobject();
+
+		// Flag this export as an inherited instance if the object's archetype exists within the set
+		// of default subobjects owned by the object owner's archetype. This is used by the linker to
+		// determine whether or not the subobject should be instanced as an export on load. Note that
+		// if the archetype is owned by a different object, we treat it as a non-default subobject and
+		// thus exclude it from consideration. This is because an instanced subobject with a non-
+		// standard archetype won't find a matching instance in its owner's archetype subobject set,
+		// and thus wouldn't pass the instancing check on load. One example of a non-default instanced
+		// subobject is a Blueprint-added component, whose archetype is owned by the class object.
+		const UObject* Archetype = Object->GetArchetype();
+		if (Archetype->IsDefaultSubobject() && (Archetype->GetOuter() == Object->GetOuter()->GetArchetype()))
+		{
+			bIsInheritedInstance = true;
+		}
 	}
 }
 
@@ -164,12 +177,6 @@ void operator<<(FStructuredArchive::FSlot Slot, FObjectExport& E)
 	if (BaseArchive.UEVer() >= EUnrealEngineObjectUE5Version::TRACK_OBJECT_EXPORT_IS_INHERITED)
 	{
 		SERIALIZE_BIT_TO_RECORD(bIsInheritedInstance);
-
-		// Initial saves of TRACK_OBJECT_EXPORT_IS_INHERITED incorrectly considered nested template objects as inherited instances; fix that up here on load.
-		if (BaseArchive.IsLoading() && BaseArchive.UEVer().ToValue() == static_cast<int32>(EUnrealEngineObjectUE5Version::TRACK_OBJECT_EXPORT_IS_INHERITED))
-		{
-			E.bIsInheritedInstance &= !(!!(E.ObjectFlags & RF_ArchetypeObject));
-		}
 	}
 
 	Record << SA_VALUE(TEXT("PackageFlags"), E.PackageFlags);
