Updating FSoftObjecPtr values stored in the asset registry tags to match the current format returned by FSoftObjecPtr::ToString() (stripping class names and leaving just path names)

(reimplementing 21054544 from UE5/Main)

diff --git a/Engine/Source/Runtime/CoreUObject/Private/UObject/Obj.cpp b/Engine/Source/Runtime/CoreUObject/Private/UObject/Obj.cpp
--- a/Engine/Source/Runtime/CoreUObject/Private/UObject/Obj.cpp
+++ b/Engine/Source/Runtime/CoreUObject/Private/UObject/Obj.cpp
@@ -2074,7 +2074,7 @@ static void PostLoadAssetRegistryTagProperty(FProperty* Prop, const FAssetData&
 			if (!ExportPath.IsEmpty() && ExportPath[0] != '/' && ExportPath.FindChar('\'', ClassSeparatorIndex))
 			{
 				// Strip the class name and leave just the pathname of an object
-				FString ObjectPath = ExportPath.Mid(ClassSeparatorIndex + 1, ExportPath.Len() - ClassSeparatorIndex - 2); // -2 because we're stripping the first and last '
+				FString ObjectPath = FPackageName::ExportTextPathToObjectPath(ExportPath);
 				OutTagsAndValuesToUpdate.Add(UObject::FAssetRegistryTag(Prop->GetFName(), ObjectPath, TagType));
 			}
 		}
