Hardened FPackageName::ParseExportTextPath to better ensure that the source matches the expected format before attempting the parse

#rnx

diff --git a/Engine/Source/Runtime/Core/Public/Containers/StringView.h b/Engine/Source/Runtime/Core/Public/Containers/StringView.h
--- a/Engine/Source/Runtime/Core/Public/Containers/StringView.h
+++ b/Engine/Source/Runtime/Core/Public/Containers/StringView.h
@@ -330,6 +330,18 @@ public:
 	 */
 	inline bool FindLastChar(CharType Search, int32& OutIndex) const;
 
+	/**
+	 * Tests if index is valid, i.e. greater than or equal to zero, and less than the number of characters in the string view.
+	 *
+	 * @param Index Index to test.
+	 *
+	 * @returns True if index is valid. False otherwise.
+	 */
+	UE_NODISCARD FORCEINLINE bool IsValidIndex(SizeType Index) const
+	{
+		return Index >= 0 && Index < Len();
+	}
+
 public:
 	/**
 	 * DO NOT USE DIRECTLY
diff --git a/Engine/Source/Runtime/CoreUObject/Private/Misc/PackageName.cpp b/Engine/Source/Runtime/CoreUObject/Private/Misc/PackageName.cpp
--- a/Engine/Source/Runtime/CoreUObject/Private/Misc/PackageName.cpp
+++ b/Engine/Source/Runtime/CoreUObject/Private/Misc/PackageName.cpp
@@ -2176,30 +2176,11 @@ void FPackageName::OnCoreUObjectInitialized()
 	FLongPackagePathsSingleton::Get().OnCoreUObjectInitialized();
 }
 
-bool FPackageName::ParseExportTextPath(const FString& InExportTextPath, FString* OutClassName, FString* OutObjectPath)
-{
-	if (InExportTextPath.Split(TEXT("'"), OutClassName, OutObjectPath, ESearchCase::CaseSensitive))
-	{
-		if ( OutObjectPath )
-		{
-			FString& OutObjectPathRef = *OutObjectPath;
-			if ( OutObjectPathRef.EndsWith(TEXT("'"), ESearchCase::CaseSensitive) )
-			{
-				OutObjectPathRef.LeftChopInline(1, false);
-			}
-		}
-
-		return true;
-	}
-	
-	return false;
-}
-
 template<class T>
 bool ParseExportTextPathImpl(const T& InExportTextPath, T* OutClassName, T* OutObjectPath)
 {
 	int32 Index;
-	if (InExportTextPath.FindChar('\'', Index))
+	if (InExportTextPath.FindChar('\'', Index) && InExportTextPath.IsValidIndex(Index + 1) && InExportTextPath[InExportTextPath.Len() - 1] == '\'') // IsValidIndex checks that the found ' isn't the same one that's at the end of the string
 	{
 		if (OutClassName)
 		{
@@ -2208,8 +2189,7 @@ bool ParseExportTextPathImpl(const T& InExportTextPath, T* OutClassName, T* OutO
 
 		if (OutObjectPath)
 		{
-			*OutObjectPath = InExportTextPath.Mid(Index + 1);
-			OutObjectPath->RemoveSuffix(OutObjectPath->EndsWith('\''));
+			*OutObjectPath = InExportTextPath.Mid(Index + 1, InExportTextPath.Len() - Index - 2); // -2 because we're stripping the first and last '
 		}
 
 		return true;
@@ -2218,6 +2198,11 @@ bool ParseExportTextPathImpl(const T& InExportTextPath, T* OutClassName, T* OutO
 	return false;
 }
 
+bool FPackageName::ParseExportTextPath(const FString& InExportTextPath, FString* OutClassName, FString* OutObjectPath)
+{
+	return ParseExportTextPathImpl(InExportTextPath, OutClassName, OutObjectPath);
+}
+
 bool FPackageName::ParseExportTextPath(FWideStringView InExportTextPath, FWideStringView* OutClassName, FWideStringView* OutObjectPath)
 {
 	return ParseExportTextPathImpl(InExportTextPath, OutClassName, OutObjectPath);
