Removed usage of and deprecated Stats::IsThreadCollectingData. It was only used for FScopeCycleCounterUObject which does the checks internally already.

diff --git a/Engine/Source/Runtime/Core/Public/Stats/Stats2.h b/Engine/Source/Runtime/Core/Public/Stats/Stats2.h
--- a/Engine/Source/Runtime/Core/Public/Stats/Stats2.h
+++ b/Engine/Source/Runtime/Core/Public/Stats/Stats2.h
@@ -2354,6 +2354,7 @@ DECLARE_FLOAT_COUNTER_STAT_EXTERN(TEXT("Seconds Per Cycle"),STAT_SecondsPerCycle
 #if STATS || ENABLE_STATNAMEDEVENTS
 namespace Stats
 {
+	UE_DEPRECATED(5.2, "No replacement. Engine usage was redundant")
 	FORCEINLINE bool IsThreadCollectingData()
 	{
 #if STATS && CPUPROFILERTRACE_ENABLED
diff --git a/Engine/Source/Runtime/CoreUObject/Private/UObject/ScriptCore.cpp b/Engine/Source/Runtime/CoreUObject/Private/UObject/ScriptCore.cpp
--- a/Engine/Source/Runtime/CoreUObject/Private/UObject/ScriptCore.cpp
+++ b/Engine/Source/Runtime/CoreUObject/Private/UObject/ScriptCore.cpp
@@ -1039,14 +1039,11 @@ IMPLEMENT_VM_FUNCTION(EX_CallMath, execCallMathFunction);
 void UObject::CallFunction( FFrame& Stack, RESULT_DECL, UFunction* Function )
 {
 #if PER_FUNCTION_SCRIPT_STATS
-	const bool bShouldTrackFunction = (Stack.DepthCounter <= GMaxFunctionStatDepth) && Stats::IsThreadCollectingData();
-	FScopeCycleCounterUObject FunctionScope(bShouldTrackFunction ? Function : nullptr);
+	const bool bShouldTrackFunction = (Stack.DepthCounter <= GMaxFunctionStatDepth);
+	SCOPE_CYCLE_UOBJECT(FunctionScope, bShouldTrackFunction ? Function : nullptr);
 #endif // PER_FUNCTION_SCRIPT_STATS
 
-#if STATS || ENABLE_STATNAMEDEVENTS
-	const bool bShouldTrackObject = GVerboseScriptStats && Stats::IsThreadCollectingData();
-	FScopeCycleCounterUObject ContextScope(bShouldTrackObject ? this : nullptr);
-#endif
+	SCOPE_CYCLE_UOBJECT(ContextScope, GVerboseScriptStats ? this : nullptr);
 
 	checkSlow(Function);
 
@@ -1235,8 +1232,8 @@ void ProcessLocalFunction(UObject* Context, UFunction* Fn, FFrame& Stack, RESULT
 	else
 	{
 #if PER_FUNCTION_SCRIPT_STATS
-		const bool bShouldTrackFunction = (Stack.DepthCounter <= GMaxFunctionStatDepth) && Stats::IsThreadCollectingData();
-		FScopeCycleCounterUObject FunctionScope(bShouldTrackFunction ? Fn : nullptr);
+		const bool bShouldTrackFunction = (Stack.DepthCounter <= GMaxFunctionStatDepth);
+		SCOPE_CYCLE_UOBJECT(FunctionScope, bShouldTrackFunction ? Fn : nullptr);
 #endif // PER_FUNCTION_SCRIPT_STATS
 		ProcessScriptFunction(Context, Fn, Stack, RESULT_PARAM, ProcessLocalScriptFunction);
 	}
@@ -1980,14 +1977,10 @@ void UObject::ProcessEvent( UFunction* Function, void* Parms )
 	checkSlow((Function->ParmsSize == 0) || (Parms != NULL));
 
 #if PER_FUNCTION_SCRIPT_STATS
-	const bool bShouldTrackFunction = Stats::IsThreadCollectingData();
-	FScopeCycleCounterUObject FunctionScope(bShouldTrackFunction ? Function : nullptr);
+	SCOPE_CYCLE_UOBJECT(FunctionScope, Function);
 #endif // PER_FUNCTION_SCRIPT_STATS
 
-#if STATS || ENABLE_STATNAMEDEVENTS
-	const bool bShouldTrackObject = GVerboseScriptStats && Stats::IsThreadCollectingData();
-	FScopeCycleCounterUObject ContextScope(bShouldTrackObject ? this : nullptr);
-#endif
+	SCOPE_CYCLE_UOBJECT(ContextScope, GVerboseScriptStats ? this : nullptr);
 
 #if LIGHTWEIGHT_PROCESS_EVENT_COUNTER
 	TGuardValue<int32> PECounter(ProcessEventCounter, ProcessEventCounter + 1);
diff --git a/Engine/Source/Runtime/MovieScene/Private/EntitySystem/MovieSceneSequenceInstance.cpp b/Engine/Source/Runtime/MovieScene/Private/EntitySystem/MovieSceneSequenceInstance.cpp
--- a/Engine/Source/Runtime/MovieScene/Private/EntitySystem/MovieSceneSequenceInstance.cpp
+++ b/Engine/Source/Runtime/MovieScene/Private/EntitySystem/MovieSceneSequenceInstance.cpp
@@ -218,11 +218,7 @@ void FSequenceInstance::Start(UMovieSceneEntitySystemLinker* Linker, const FMovi
 void FSequenceInstance::Update(UMovieSceneEntitySystemLinker* Linker, const FMovieSceneContext& InContext)
 {
 	SCOPE_CYCLE_COUNTER(MovieSceneEval_SequenceInstanceUpdate);
-
-#if STATS || ENABLE_STATNAMEDEVENTS
-	const bool bShouldTrackObject = Stats::IsThreadCollectingData();
-	FScopeCycleCounterUObject ContextScope(bShouldTrackObject ? GetPlayer()->AsUObject() : nullptr);
-#endif
+	SCOPE_CYCLE_UOBJECT(ContextScope, GetPlayer()->AsUObject());
 
 	bHasEverUpdated = true;
 
diff --git a/Engine/Source/Runtime/MovieScene/Private/Evaluation/Instances/MovieSceneTrackEvaluator.cpp b/Engine/Source/Runtime/MovieScene/Private/Evaluation/Instances/MovieSceneTrackEvaluator.cpp
--- a/Engine/Source/Runtime/MovieScene/Private/Evaluation/Instances/MovieSceneTrackEvaluator.cpp
+++ b/Engine/Source/Runtime/MovieScene/Private/Evaluation/Instances/MovieSceneTrackEvaluator.cpp
@@ -108,10 +108,7 @@ void FMovieSceneTrackEvaluator::Evaluate(FMovieSceneContext Context, IMovieScene
 		return;
 	}
 
-#if STATS || ENABLE_STATNAMEDEVENTS
-	const bool bShouldTrackObject = Stats::IsThreadCollectingData();
-	FScopeCycleCounterUObject ContextScope(bShouldTrackObject ? OverrideRootSequence : nullptr);
-#endif
+	SCOPE_CYCLE_UOBJECT(ContextScope, OverrideRootSequence);
 
 	const FMovieSceneEvaluationGroup* GroupToEvaluate = SetupFrame(OverrideRootSequence, InOverrideRootID, Context);
 	if (!GroupToEvaluate)
diff --git a/Engine/Source/Runtime/UMG/Private/Animation/UMGSequenceTickManager.cpp b/Engine/Source/Runtime/UMG/Private/Animation/UMGSequenceTickManager.cpp
--- a/Engine/Source/Runtime/UMG/Private/Animation/UMGSequenceTickManager.cpp
+++ b/Engine/Source/Runtime/UMG/Private/Animation/UMGSequenceTickManager.cpp
@@ -101,11 +101,7 @@ void UUMGSequenceTickManager::TickWidgetAnimations(float DeltaSeconds)
 	// This may delay the animation update by one frame.
 
 	{
-	#if STATS || ENABLE_STATNAMEDEVENTS
-		const bool bShouldTrackObject = Stats::IsThreadCollectingData();
-		FScopeCycleCounterUObject ContextScope(bShouldTrackObject ? this : nullptr);
-	#endif
-
+		SCOPE_CYCLE_UOBJECT(ContextScope, this);
 
 		// Process animations for visible widgets
 		for (auto WidgetIter = WeakUserWidgetData.CreateIterator(); WidgetIter; ++WidgetIter)
@@ -143,9 +139,7 @@ void UUMGSequenceTickManager::TickWidgetAnimations(float DeltaSeconds)
 			}
 			else
 			{
-	#if STATS || ENABLE_STATNAMEDEVENTS
-				FScopeCycleCounterUObject WidgetContextScope(bShouldTrackObject ? UserWidget : nullptr);
-	#endif
+				SCOPE_CYCLE_UOBJECT(WidgetContextScope, UserWidget);
 
 	#if WITH_EDITOR
 				const bool bTickAnimations = !UserWidget->IsDesignTime();
