Add serialization support for cooking editor bulk data as empty during cooking.

- This is a rework of CL 20965910, to work with the newer code in editor bulkdata.
- bKeepFileDataByReference will be set to false as cooked packages will probably not be able to reference workspace domain payloads.
- When cooking we will use the legacy serialization path (i.e not to package trailer) BUT we will also not write out a payload at the moment.

diff --git a/Engine/Source/Runtime/CoreUObject/Private/Serialization/EditorBulkData.cpp b/Engine/Source/Runtime/CoreUObject/Private/Serialization/EditorBulkData.cpp
--- a/Engine/Source/Runtime/CoreUObject/Private/Serialization/EditorBulkData.cpp
+++ b/Engine/Source/Runtime/CoreUObject/Private/Serialization/EditorBulkData.cpp
@@ -801,7 +801,7 @@ void FEditorBulkData::Serialize(FArchive& Ar, UObject* Owner, bool bAllowRegiste
 			LinkerSave = Cast<FLinkerSave>(Ar.GetLinker());
 			// If we're doing a save that can refer to bulk data by reference, and our legacy data format supports it,
 			// keep any legacy data we have referenced rather than stored, to save space and avoid spending time loading it.
-			bKeepFileDataByReference = LinkerSave != nullptr && LinkerSave->bProceduralSave;
+			bKeepFileDataByReference = LinkerSave != nullptr && LinkerSave->bProceduralSave && !Ar.IsCooking();
 			if (!bKeepFileDataByReference)
 			{
 				UpdateKeyIfNeeded();
@@ -846,14 +846,20 @@ void FEditorBulkData::Serialize(FArchive& Ar, UObject* Owner, bool bAllowRegiste
 
 		if (Ar.IsSaving())
 		{
-			checkf(Ar.IsCooking() == false, TEXT("FEditorBulkData::Serialize should not be called during a cook"));
-
 			EFlags UpdatedFlags = BuildFlagsForSerialization(Ar, bKeepFileDataByReference);
 
 			// Write out required extra data if we're saving by reference
 			bool bWriteOutPayload = true;
-			if (IsReferencingByPackagePath(UpdatedFlags))
+				
+			if (EnumHasAnyFlags(UpdatedFlags, EFlags::IsCooked))
+			{
+				// If we are cooking, we currently aren't saving any payload, since they aren't supported at runtime
+				// TODO: add iostore support for editor bulkdata
+				bWriteOutPayload = false;
+			}
+			else if (EnumHasAnyFlags(UpdatedFlags, EFlags::ReferencesLegacyFile))
 			{
+				// Write out required extra data if we're saving by reference
 				if (!IsStoredInPackageTrailer(UpdatedFlags))
 				{
 					Ar << OffsetInFile;
@@ -985,7 +991,7 @@ void FEditorBulkData::Serialize(FArchive& Ar, UObject* Owner, bool bAllowRegiste
 				// But we only care if the archive has a linker! (loading from a package)
 				// TODO: Consider removing in 5.2+
 				UE_CLOG(IsDataVirtualized() && Ar.GetLinker() != nullptr, LogSerialization, Warning, TEXT("Payload in '%s' is virtualized in an older format and should be re-saved!"), *GetDebugNameFromOwner(Owner));
-				if (!IsDataVirtualized())
+				if (!IsDataVirtualized() && !EnumHasAnyFlags(Flags, EFlags::IsCooked))
 				{
 					Ar << OffsetInFile;
 				}
@@ -1029,7 +1035,7 @@ void FEditorBulkData::Serialize(FArchive& Ar, UObject* Owner, bool bAllowRegiste
 					AttachedAr->AttachBulkData(this);
 #endif //VBD_ALLOW_LINKERLOADER_ATTACHMENT
 				}
-				else
+				else if (!EnumHasAnyFlags(Flags, EFlags::IsCooked))
 				{
 					checkf(Ar.Tell() == OffsetInFile, TEXT("Attempting to load an inline payload but the offset does not match"));
 
@@ -1568,7 +1574,7 @@ void FEditorBulkData::SerializeToLegacyPath(FLinkerSave& LinkerSave, FCompressed
 	// The lambda is mutable so that PayloadToSerialize is not const (due to FArchive api not accepting const values)
 	auto SerializePayload = [this, OffsetPos, PayloadToSerialize, UpdatedFlags, Owner](FLinkerSave& LinkerSave, FArchive& ExportsArchive, FArchive& DataArchive, int64 DataStartOffset) mutable
 	{
-		checkf(ExportsArchive.IsCooking() == false, TEXT("FEditorBulkData::Serialize should not be called during a cook"));
+		checkf(ExportsArchive.IsCooking() == false, TEXT("FEditorBulkData payload should not be written during a cook"));
 
 		SerializeData(DataArchive, PayloadToSerialize, UpdatedFlags);
 
@@ -1627,7 +1633,7 @@ void FEditorBulkData::SerializeToPackageTrailer(FLinkerSave& LinkerSave, FCompre
 {
 	auto OnPayloadWritten = [this, UpdatedFlags, Owner](FLinkerSave& LinkerSave, const FPackageTrailer& Trailer) mutable
 	{
-		checkf(LinkerSave.IsCooking() == false, TEXT("FEditorBulkData::Serialize should not be called during a cook"));
+		checkf(LinkerSave.IsCooking() == false, TEXT("FEditorBulkData payload should not be written to a package trailer during a cook"));
 
 		int64 PayloadOffset = Trailer.FindPayloadOffsetInFile(PayloadContentId);
 
