Add support for signed and unsigned variants of various int vector types.
Add support for 32 and 64 bit variants of various int vector types.

IntVector2/3/4 - UPROPERTY support including auto-conversion of 32 and 64  bit variants of same signedness.
IntPoint - UPROPERTY support including auto-conversion of 32 and 64 bit variants of same signedness.
IntRect

diff --git a/Engine/Source/Editor/DetailCustomizations/Private/DetailCustomizations.cpp b/Engine/Source/Editor/DetailCustomizations/Private/DetailCustomizations.cpp
--- a/Engine/Source/Editor/DetailCustomizations/Private/DetailCustomizations.cpp
+++ b/Engine/Source/Editor/DetailCustomizations/Private/DetailCustomizations.cpp
@@ -211,6 +211,11 @@ void FDetailCustomizationsModule::RegisterPropertyTypeCustomizations()
 	REGISTER_UIMINMAX_CUSTOMIZATION(NAME_Vector3f, &FVectorStructCustomization::MakeInstance);
 	REGISTER_UIMINMAX_CUSTOMIZATION(NAME_Vector3d, &FVectorStructCustomization::MakeInstance);
 	REGISTER_UIMINMAX_CUSTOMIZATION(NAME_IntVector, &FVectorStructCustomization::MakeInstance);
+	REGISTER_UIMINMAX_CUSTOMIZATION(NAME_Int32Vector, &FVectorStructCustomization::MakeInstance);
+	REGISTER_UIMINMAX_CUSTOMIZATION(NAME_Int64Vector, &FVectorStructCustomization::MakeInstance);
+	REGISTER_UIMINMAX_CUSTOMIZATION(NAME_UintVector, &FVectorStructCustomization::MakeInstance);
+	REGISTER_UIMINMAX_CUSTOMIZATION(NAME_Uint32Vector, &FVectorStructCustomization::MakeInstance);
+	REGISTER_UIMINMAX_CUSTOMIZATION(NAME_Uint64Vector, &FVectorStructCustomization::MakeInstance);
 	REGISTER_UIMINMAX_CUSTOMIZATION(NAME_Vector4, &FVector4StructCustomization::MakeInstance);
 	REGISTER_UIMINMAX_CUSTOMIZATION(NAME_Vector4f, &FVector4StructCustomization::MakeInstance);
 	REGISTER_UIMINMAX_CUSTOMIZATION(NAME_Vector4d, &FVector4StructCustomization::MakeInstance);
@@ -218,7 +223,23 @@ void FDetailCustomizationsModule::RegisterPropertyTypeCustomizations()
 	REGISTER_UIMINMAX_CUSTOMIZATION(NAME_Vector2f, &FMathStructCustomization::MakeInstance);
 	REGISTER_UIMINMAX_CUSTOMIZATION(NAME_Vector2d, &FMathStructCustomization::MakeInstance);
 	REGISTER_UIMINMAX_CUSTOMIZATION(NAME_IntPoint, &FMathStructCustomization::MakeInstance);
+	REGISTER_UIMINMAX_CUSTOMIZATION(NAME_Int32Point, &FMathStructCustomization::MakeInstance);
+	REGISTER_UIMINMAX_CUSTOMIZATION(NAME_Int64Point, &FMathStructCustomization::MakeInstance);
+	REGISTER_UIMINMAX_CUSTOMIZATION(NAME_UintPoint, &FMathStructCustomization::MakeInstance);
+	REGISTER_UIMINMAX_CUSTOMIZATION(NAME_Uint32Point, &FMathStructCustomization::MakeInstance);
+	REGISTER_UIMINMAX_CUSTOMIZATION(NAME_Uint64Point, &FMathStructCustomization::MakeInstance);
+	REGISTER_UIMINMAX_CUSTOMIZATION(NAME_IntVector2, &FVector4StructCustomization::MakeInstance);
+	REGISTER_UIMINMAX_CUSTOMIZATION(NAME_Int32Vector2, &FVector4StructCustomization::MakeInstance);
+	REGISTER_UIMINMAX_CUSTOMIZATION(NAME_Int64Vector2, &FVector4StructCustomization::MakeInstance);
+	REGISTER_UIMINMAX_CUSTOMIZATION(NAME_UintVector2, &FVector4StructCustomization::MakeInstance);
+	REGISTER_UIMINMAX_CUSTOMIZATION(NAME_Uint32Vector2, &FVector4StructCustomization::MakeInstance);
+	REGISTER_UIMINMAX_CUSTOMIZATION(NAME_Uint64Vector2, &FVector4StructCustomization::MakeInstance);
 	REGISTER_UIMINMAX_CUSTOMIZATION(NAME_IntVector4, &FVector4StructCustomization::MakeInstance);
+	REGISTER_UIMINMAX_CUSTOMIZATION(NAME_Int32Vector4, &FVector4StructCustomization::MakeInstance);
+	REGISTER_UIMINMAX_CUSTOMIZATION(NAME_Int64Vector4, &FVector4StructCustomization::MakeInstance);
+	REGISTER_UIMINMAX_CUSTOMIZATION(NAME_UintVector4, &FVector4StructCustomization::MakeInstance);
+	REGISTER_UIMINMAX_CUSTOMIZATION(NAME_Uint32Vector4, &FVector4StructCustomization::MakeInstance);
+	REGISTER_UIMINMAX_CUSTOMIZATION(NAME_Uint64Vector4, &FVector4StructCustomization::MakeInstance);
 	REGISTER_UIMINMAX_CUSTOMIZATION(NAME_Rotator, &FRotatorStructCustomization::MakeInstance);
 	REGISTER_UIMINMAX_CUSTOMIZATION(NAME_Rotator3f, &FRotatorStructCustomization::MakeInstance);
 	REGISTER_UIMINMAX_CUSTOMIZATION(NAME_Rotator3d, &FRotatorStructCustomization::MakeInstance);
diff --git a/Engine/Source/Editor/DetailCustomizations/Private/MathStructCustomizations.cpp b/Engine/Source/Editor/DetailCustomizations/Private/MathStructCustomizations.cpp
--- a/Engine/Source/Editor/DetailCustomizations/Private/MathStructCustomizations.cpp
+++ b/Engine/Source/Editor/DetailCustomizations/Private/MathStructCustomizations.cpp
@@ -413,6 +413,21 @@ TSharedRef<SWidget> FMathStructCustomization::MakeChildWidget(
 		return MakeNumericWidget<int32>(StructurePropertyHandle, PropertyHandle);
 	}
 
+	if (PropertyClass == FInt64Property::StaticClass())
+	{
+		return MakeNumericWidget<int64>(StructurePropertyHandle, PropertyHandle);
+	}
+
+	if (PropertyClass == FUInt32Property::StaticClass())
+	{
+		return MakeNumericWidget<uint32>(StructurePropertyHandle, PropertyHandle);
+	}
+
+	if (PropertyClass == FUInt64Property::StaticClass())
+	{
+		return MakeNumericWidget<uint64>(StructurePropertyHandle, PropertyHandle);
+	}
+
 	if (PropertyClass == FByteProperty::StaticClass())
 	{
 		return MakeNumericWidget<uint8>(StructurePropertyHandle, PropertyHandle);
diff --git a/Engine/Source/Runtime/Core/Public/Math/IntPoint.h b/Engine/Source/Runtime/Core/Public/Math/IntPoint.h
--- a/Engine/Source/Runtime/Core/Public/Math/IntPoint.h
+++ b/Engine/Source/Runtime/Core/Public/Math/IntPoint.h
@@ -8,6 +8,7 @@
 #include "Containers/UnrealString.h"
 #include "Serialization/StructuredArchive.h"
 #include "Templates/TypeHash.h"
+#include "Misc/LargeWorldCoordinatesSerializer.h"
 
 namespace UE::Math
 {
@@ -87,8 +88,8 @@ struct TIntPoint
 	 */
 	template <typename OtherIntType>
 	explicit TIntPoint(TIntPoint<OtherIntType> Other)
-		: X(IntCastChecked<OtherIntType>(Other.X))
-		, Y(IntCastChecked<OtherIntType>(Other.Y))
+		: X(IntCastChecked<IntType>(Other.X))
+		, Y(IntCastChecked<IntType>(Other.Y))
 	{
 	}
 
@@ -481,17 +482,7 @@ PRAGMA_ENABLE_DEPRECATION_WARNINGS
 		return true;
 	}
 
-	/**
-	 * Serialize the point.
-	 *
-	 * @param Slot The structured archive slot to serialize into.
-	 * @return true on success, false otherwise.
-	 */
-	bool Serialize(FStructuredArchive::FSlot Slot)
-	{
-		Slot << *this;
-		return true;
-	}
+	bool SerializeFromMismatchedTag(FName StructTag, FArchive& Ar);
 };
 
 template <typename IntType>
@@ -556,8 +547,39 @@ inline FString TIntPoint<uint8>::ToString() const
 
 } //! namespace UE::Math
 
-template <> struct TIsPODType<FIntPoint>  { enum { Value = true }; };
-template <> struct TIsPODType<FUintPoint> { enum { Value = true }; };
+template <> struct TIsPODType<FInt32Point>  { enum { Value = true }; };
+template <> struct TIsPODType<FUint32Point> { enum { Value = true }; };
+
+template<> struct TIsUECoreVariant<FInt32Point>  { enum { Value = true }; };
+template<> struct TIsUECoreVariant<FUint32Point> { enum { Value = true }; };
+
+template <> struct TIsPODType<FInt64Point> { enum { Value = true }; };
+template <> struct TIsPODType<FUint64Point> { enum { Value = true }; };
+
+template<> struct TIsUECoreVariant<FInt64Point> { enum { Value = true }; };
+template<> struct TIsUECoreVariant<FUint64Point> { enum { Value = true }; };
 
-template<> struct TIsUECoreType<FIntPoint>  { enum { Value = true }; };
-template<> struct TIsUECoreType<FUintPoint> { enum { Value = true }; };
\ No newline at end of file
+
+template<>
+inline bool FInt32Point::SerializeFromMismatchedTag(FName StructTag, FArchive& Ar)
+{
+	return UE_SERIALIZE_VARIANT_FROM_MISMATCHED_TAG(Ar, IntPoint, Int32Point, Int64Point);
+}
+
+template<>
+inline bool FInt64Point::SerializeFromMismatchedTag(FName StructTag, FArchive& Ar)
+{
+	return UE_SERIALIZE_VARIANT_FROM_MISMATCHED_TAG(Ar, IntPoint, Int64Point, Int32Point);
+}
+
+template<>
+inline bool FUint32Point::SerializeFromMismatchedTag(FName StructTag, FArchive& Ar)
+{
+	return UE_SERIALIZE_VARIANT_FROM_MISMATCHED_TAG(Ar, UintPoint, Uint32Point, Uint64Point);
+}
+
+template<>
+inline bool FUint64Point::SerializeFromMismatchedTag(FName StructTag, FArchive& Ar)
+{
+	return UE_SERIALIZE_VARIANT_FROM_MISMATCHED_TAG(Ar, UintPoint, Uint64Point, Uint32Point);
+}
\ No newline at end of file
diff --git a/Engine/Source/Runtime/Core/Public/Math/IntRect.h b/Engine/Source/Runtime/Core/Public/Math/IntRect.h
--- a/Engine/Source/Runtime/Core/Public/Math/IntRect.h
+++ b/Engine/Source/Runtime/Core/Public/Math/IntRect.h
@@ -506,5 +506,14 @@ uint32 GetTypeHash(const TIntRect<IntType>& InRect)
 
 } //! namespace UE::Math
 
-template <> struct TIsUECoreType<FIntRect>  { enum { Value = true }; };
-template <> struct TIsUECoreType<FUintRect> { enum { Value = true }; };
\ No newline at end of file
+template <> struct TIsPODType<FInt32Rect> { enum { Value = true }; };
+template <> struct TIsPODType<FUint32Rect> { enum { Value = true }; };
+
+template <> struct TIsUECoreVariant<FInt32Rect>  { enum { Value = true }; };
+template <> struct TIsUECoreVariant<FUint32Rect> { enum { Value = true }; };
+
+template <> struct TIsPODType<FInt64Rect> { enum { Value = true }; };
+template <> struct TIsPODType<FUint64Rect> { enum { Value = true }; };
+
+template <> struct TIsUECoreVariant<FInt64Rect> { enum { Value = true }; };
+template <> struct TIsUECoreVariant<FUint64Rect> { enum { Value = true }; };
diff --git a/Engine/Source/Runtime/Core/Public/Math/IntVector.h b/Engine/Source/Runtime/Core/Public/Math/IntVector.h
--- a/Engine/Source/Runtime/Core/Public/Math/IntVector.h
+++ b/Engine/Source/Runtime/Core/Public/Math/IntVector.h
@@ -7,6 +7,7 @@
 #include "Math/UnrealMathUtility.h"
 #include "Containers/UnrealString.h"
 #include "Serialization/StructuredArchive.h"
+#include "Misc/LargeWorldCoordinatesSerializer.h"
 
 namespace UE::Math
 {
@@ -100,9 +101,9 @@ struct TIntVector3
 	 */
 	template <typename OtherIntType>
 	explicit TIntVector3(TIntVector3<OtherIntType> Other)
-		: X(IntCastChecked<OtherIntType>(Other.X))
-		, Y(IntCastChecked<OtherIntType>(Other.Y))
-		, Z(IntCastChecked<OtherIntType>(Other.Z))
+		: X(IntCastChecked<IntType>(Other.X))
+		, Y(IntCastChecked<IntType>(Other.Y))
+		, Z(IntCastChecked<IntType>(Other.Z))
 	{
 	}
 
@@ -469,6 +470,8 @@ PRAGMA_ENABLE_DEPRECATION_WARNINGS
 		Ar << *this;
 		return true;
 	}
+
+	bool SerializeFromMismatchedTag(FName StructTag, FArchive& Ar);
 };
 
 template <>
@@ -569,8 +572,8 @@ struct TIntVector2
 	 */
 	template <typename OtherIntType>
 	explicit TIntVector2(TIntVector2<OtherIntType> Other)
-		: X(IntCastChecked<OtherIntType>(Other.X))
-		, Y(IntCastChecked<OtherIntType>(Other.Y))
+		: X(IntCastChecked<IntType>(Other.X))
+		, Y(IntCastChecked<IntType>(Other.Y))
 	{
 	}
 
@@ -609,6 +612,14 @@ PRAGMA_ENABLE_DEPRECATION_WARNINGS
 	{
 		return Ar << Vector.X << Vector.Y;
 	}
+
+	bool Serialize(FArchive& Ar)
+	{
+		Ar << *this;
+		return true;
+	}
+
+	bool SerializeFromMismatchedTag(FName StructTag, FArchive& Ar);
 };
 
 ///////////////////////////////////////////////////////////////////////////////////////////////////
@@ -669,10 +680,10 @@ struct TIntVector4
 	 */
 	template <typename OtherIntType>
 	explicit TIntVector4(TIntVector4<OtherIntType> Other)
-		: X(IntCastChecked<OtherIntType>(Other.X))
-		, Y(IntCastChecked<OtherIntType>(Other.Y))
-		, Z(IntCastChecked<OtherIntType>(Other.Z))
-		, W(IntCastChecked<OtherIntType>(Other.W))
+		: X(IntCastChecked<IntType>(Other.X))
+		, Y(IntCastChecked<IntType>(Other.Y))
+		, Z(IntCastChecked<IntType>(Other.Z))
+		, W(IntCastChecked<IntType>(Other.W))
 	{
 	}
 
@@ -982,6 +993,8 @@ PRAGMA_ENABLE_DEPRECATION_WARNINGS
 		Ar << *this;
 		return true;
 	}
+
+	bool SerializeFromMismatchedTag(FName StructTag, FArchive& Ar);
 };
 
 /**
@@ -1025,16 +1038,103 @@ uint32 GetTypeHash(const TIntVector4<T>& Vector)
 
 } //! namespace UE::Math
 
-template <> struct TIsPODType<FIntVector2>  { enum { Value = true }; };
-template <> struct TIsPODType<FUintVector2> { enum { Value = true }; };
-template <> struct TIsPODType<FIntVector3>  { enum { Value = true }; };
-template <> struct TIsPODType<FUintVector3> { enum { Value = true }; };
-template <> struct TIsPODType<FIntVector4>  { enum { Value = true }; };
-template <> struct TIsPODType<FUintVector4> { enum { Value = true }; };
-
-template <> struct TIsUECoreType<FIntVector2>  { enum { Value = true }; };
-template <> struct TIsUECoreType<FUintVector2> { enum { Value = true }; };
-template <> struct TIsUECoreType<FIntVector3>  { enum { Value = true }; };
-template <> struct TIsUECoreType<FUintVector3> { enum { Value = true }; };
-template <> struct TIsUECoreType<FIntVector4>  { enum { Value = true }; };
-template <> struct TIsUECoreType<FUintVector4> { enum { Value = true }; };
\ No newline at end of file
+template <> struct TIsPODType<FInt32Vector2>  { enum { Value = true }; };
+template <> struct TIsPODType<FUint32Vector2> { enum { Value = true }; };
+template <> struct TIsPODType<FInt32Vector3>  { enum { Value = true }; };
+template <> struct TIsPODType<FUint32Vector3> { enum { Value = true }; };
+template <> struct TIsPODType<FInt32Vector4>  { enum { Value = true }; };
+template <> struct TIsPODType<FUint32Vector4> { enum { Value = true }; };
+
+template <> struct TIsUECoreVariant<FInt32Vector2>  { enum { Value = true }; };
+template <> struct TIsUECoreVariant<FUint32Vector2> { enum { Value = true }; };
+template <> struct TIsUECoreVariant<FInt32Vector3>  { enum { Value = true }; };
+template <> struct TIsUECoreVariant<FUint32Vector3> { enum { Value = true }; };
+template <> struct TIsUECoreVariant<FInt32Vector4>  { enum { Value = true }; };
+template <> struct TIsUECoreVariant<FUint32Vector4> { enum { Value = true }; };
+
+template <> struct TIsPODType<FInt64Vector2> { enum { Value = true }; };
+template <> struct TIsPODType<FUint64Vector2> { enum { Value = true }; };
+template <> struct TIsPODType<FInt64Vector3> { enum { Value = true }; };
+template <> struct TIsPODType<FUint64Vector3> { enum { Value = true }; };
+template <> struct TIsPODType<FInt64Vector4> { enum { Value = true }; };
+template <> struct TIsPODType<FUint64Vector4> { enum { Value = true }; };
+
+template <> struct TIsUECoreVariant<FInt64Vector2> { enum { Value = true }; };
+template <> struct TIsUECoreVariant<FUint64Vector2> { enum { Value = true }; };
+template <> struct TIsUECoreVariant<FInt64Vector3> { enum { Value = true }; };
+template <> struct TIsUECoreVariant<FUint64Vector3> { enum { Value = true }; };
+template <> struct TIsUECoreVariant<FInt64Vector4> { enum { Value = true }; };
+template <> struct TIsUECoreVariant<FUint64Vector4> { enum { Value = true }; };
+
+
+template<>
+inline bool FInt32Vector2::SerializeFromMismatchedTag(FName StructTag, FArchive& Ar)
+{
+	return UE_SERIALIZE_VARIANT_FROM_MISMATCHED_TAG(Ar, IntVector2, Int32Vector2, Int64Vector2);
+}
+
+template<>
+inline bool FInt64Vector2::SerializeFromMismatchedTag(FName StructTag, FArchive& Ar)
+{
+	return UE_SERIALIZE_VARIANT_FROM_MISMATCHED_TAG(Ar, IntVector2, Int64Vector2, Int32Vector2);
+}
+
+template<>
+inline bool FUint32Vector2::SerializeFromMismatchedTag(FName StructTag, FArchive& Ar)
+{
+	return UE_SERIALIZE_VARIANT_FROM_MISMATCHED_TAG(Ar, UintVector2, Uint32Vector2, Uint64Vector2);
+}
+
+template<>
+inline bool FUint64Vector2::SerializeFromMismatchedTag(FName StructTag, FArchive& Ar)
+{
+	return UE_SERIALIZE_VARIANT_FROM_MISMATCHED_TAG(Ar, UintVector2, Uint64Vector2, Uint32Vector2);
+}
+
+template<>
+inline bool FInt32Vector3::SerializeFromMismatchedTag(FName StructTag, FArchive& Ar)
+{
+	return UE_SERIALIZE_VARIANT_FROM_MISMATCHED_TAG(Ar, IntVector, Int32Vector, Int64Vector);
+}
+
+template<>
+inline bool FInt64Vector3::SerializeFromMismatchedTag(FName StructTag, FArchive& Ar)
+{
+	return UE_SERIALIZE_VARIANT_FROM_MISMATCHED_TAG(Ar, IntVector, Int64Vector, Int32Vector);
+}
+
+template<>
+inline bool FUint32Vector3::SerializeFromMismatchedTag(FName StructTag, FArchive& Ar)
+{
+	return UE_SERIALIZE_VARIANT_FROM_MISMATCHED_TAG(Ar, UintVector, Uint32Vector, Uint64Vector);
+}
+
+template<>
+inline bool FUint64Vector3::SerializeFromMismatchedTag(FName StructTag, FArchive& Ar)
+{
+	return UE_SERIALIZE_VARIANT_FROM_MISMATCHED_TAG(Ar, UintVector, Uint64Vector, Uint32Vector);
+}
+
+template<>
+inline bool FInt32Vector4::SerializeFromMismatchedTag(FName StructTag, FArchive& Ar)
+{
+	return UE_SERIALIZE_VARIANT_FROM_MISMATCHED_TAG(Ar, IntVector4, Int32Vector4, Int64Vector4);
+}
+
+template<>
+inline bool FInt64Vector4::SerializeFromMismatchedTag(FName StructTag, FArchive& Ar)
+{
+	return UE_SERIALIZE_VARIANT_FROM_MISMATCHED_TAG(Ar, IntVector4, Int64Vector4, Int32Vector4);
+}
+
+template<>
+inline bool FUint32Vector4::SerializeFromMismatchedTag(FName StructTag, FArchive& Ar)
+{
+	return UE_SERIALIZE_VARIANT_FROM_MISMATCHED_TAG(Ar, UintVector4, Uint32Vector4, Uint64Vector4);
+}
+
+template<>
+inline bool FUint64Vector4::SerializeFromMismatchedTag(FName StructTag, FArchive& Ar)
+{
+	return UE_SERIALIZE_VARIANT_FROM_MISMATCHED_TAG(Ar, UintVector4, Uint64Vector4, Uint32Vector4);
+}
diff --git a/Engine/Source/Runtime/Core/Public/Math/MathFwd.h b/Engine/Source/Runtime/Core/Public/Math/MathFwd.h
--- a/Engine/Source/Runtime/Core/Public/Math/MathFwd.h
+++ b/Engine/Source/Runtime/Core/Public/Math/MathFwd.h
@@ -83,22 +83,52 @@ using FRotator3f	= UE::Math::TRotator<float>;
 struct FColor;
 struct FLinearColor;
 
-using FIntPoint = UE::Math::TIntPoint<int32>;
-using FUintPoint = UE::Math::TIntPoint<uint32>;
-
-using FIntRect = UE::Math::TIntRect<int32>;
-using FUintRect = UE::Math::TIntRect<uint32>;
-
+// Int vectors
+using FIntVector2 = UE::Math::TIntVector2<int32>;
 using FIntVector3 = UE::Math::TIntVector3<int32>;
-using FUintVector3 = UE::Math::TIntVector3<uint32>;
-
-using FIntVector = FIntVector3;
+using FIntVector4 = UE::Math::TIntVector4<int32>;
+using FInt32Vector2 = UE::Math::TIntVector2<int32>;
+using FInt32Vector3 = UE::Math::TIntVector3<int32>;
+using FInt32Vector4 = UE::Math::TIntVector4<int32>;
+using FInt32Vector = FInt32Vector3;
+using FInt64Vector2 = UE::Math::TIntVector2<int64>;
+using FInt64Vector3 = UE::Math::TIntVector3<int64>;
+using FInt64Vector4 = UE::Math::TIntVector4<int64>;
+using FInt64Vector = FInt64Vector3;
 
-using FIntVector2 = UE::Math::TIntVector2<int32>;
 using FUintVector2 = UE::Math::TIntVector2<uint32>;
-
-using FIntVector4 = UE::Math::TIntVector4<int32>;
+using FUintVector3 = UE::Math::TIntVector3<uint32>;
 using FUintVector4 = UE::Math::TIntVector4<uint32>;
+using FUint32Vector2 = UE::Math::TIntVector2<uint32>;
+using FUint32Vector3 = UE::Math::TIntVector3<uint32>;
+using FUint32Vector4 = UE::Math::TIntVector4<uint32>;
+using FUint32Vector = FUint32Vector3;
+using FUint64Vector2 = UE::Math::TIntVector2<uint64>;
+using FUint64Vector3 = UE::Math::TIntVector3<uint64>;
+using FUint64Vector4 = UE::Math::TIntVector4<uint64>;
+using FUint64Vector = FUint64Vector3;
+
+using FIntVector = FInt32Vector3;
+using FUintVector = FUint32Vector3;
+
+// Int points
+using FInt32Point = UE::Math::TIntPoint<int32>;
+using FInt64Point = UE::Math::TIntPoint<int64>;
+using FUint32Point = UE::Math::TIntPoint<uint32>;
+using FUint64Point = UE::Math::TIntPoint<uint64>;
+
+using FIntPoint = FInt32Point;
+using FUintPoint = FUint32Point;
+
+// Int rects
+using FInt32Rect = UE::Math::TIntRect<int32>;
+using FInt64Rect = UE::Math::TIntRect<int64>;
+using FUint32Rect = UE::Math::TIntRect<uint32>;
+using FUint64Rect = UE::Math::TIntRect<uint64>;
+
+using FIntRect = UE::Math::TIntRect<int32>;
+using FUintRect = UE::Math::TIntRect<uint32>;
+
 
 using FBoxSphereBounds3f = UE::Math::TBoxSphereBounds<float, float>;
 using FBoxSphereBounds3d = UE::Math::TBoxSphereBounds<double, double>;
diff --git a/Engine/Source/Runtime/Core/Public/Misc/LargeWorldCoordinatesSerializer.h b/Engine/Source/Runtime/Core/Public/Misc/LargeWorldCoordinatesSerializer.h
--- a/Engine/Source/Runtime/Core/Public/Misc/LargeWorldCoordinatesSerializer.h
+++ b/Engine/Source/Runtime/Core/Public/Misc/LargeWorldCoordinatesSerializer.h
@@ -15,7 +15,7 @@ inline bool IsPreLWC(const FStructuredArchive::FSlot& Slot) { return IsPreLWC(Sl
 
 // SerializeFromMismatchedTag helper for core type use only. DO NOT USE!
 template<typename FAltType, typename FType, typename FArSlot>
-bool SerializeFromMismatchedTag(FType& Target, FName StructTag, FArSlot& ArSlot, FName BaseTag, FName ThisTag, FName AltTag)
+std::enable_if_t<std::is_floating_point_v<typename FType::FReal>, bool>  SerializeFromMismatchedTag(FType& Target, FName StructTag, FArSlot& ArSlot, FName BaseTag, FName ThisTag, FName AltTag)
 {
 	if(StructTag == ThisTag || (StructTag == BaseTag && (TIsUECoreVariant<FType, double>::Value || IsPreLWC(ArSlot))))
 	{
@@ -34,4 +34,25 @@ bool SerializeFromMismatchedTag(FType& Target, FName StructTag, FArSlot& ArSlot,
 	return false;
 }
 
+// SerializeFromMismatchedTag helper for core type use only. DO NOT USE!
+template<typename FAltType, typename FType, typename FArSlot>
+std::enable_if_t<std::is_integral_v<typename FType::IntType>, bool> SerializeFromMismatchedTag(FType& Target, FName StructTag, FArSlot& ArSlot, FName BaseTag, FName ThisTag, FName AltTag)
+{
+	if (StructTag == ThisTag || (StructTag == BaseTag && (TIsUECoreVariant<FType, int32>::Value || TIsUECoreVariant<FType, uint32>::Value)))
+	{
+		// Note: Unlike float types int retains a default of 32bits, so no conversion is necessary.
+		return Target.Serialize(ArSlot);
+	}
+	else if (StructTag == AltTag || StructTag == BaseTag)
+	{
+		// Convert from alt type
+		FAltType AsAlt;
+		const bool bResult = AsAlt.Serialize(ArSlot);
+		Target = static_cast<FType>(AsAlt);
+		return bResult;
+	}
+
+	return false;
+}
+
 } // namespace LWCSerializerPrivate
diff --git a/Engine/Source/Runtime/Core/Public/Templates/IsUECoreType.h b/Engine/Source/Runtime/Core/Public/Templates/IsUECoreType.h
--- a/Engine/Source/Runtime/Core/Public/Templates/IsUECoreType.h
+++ b/Engine/Source/Runtime/Core/Public/Templates/IsUECoreType.h
@@ -25,6 +25,12 @@ struct TIsUECoreVariant<T, TC, typename std::enable_if<TIsUECoreVariant<T>::Valu
 	enum { Value = true };
 };
 
+template <typename T, typename TC>
+struct TIsUECoreVariant<T, TC, typename std::enable_if<TIsUECoreVariant<T>::Value && std::is_same<TC, typename T::IntType>::value, void>::type>
+{
+	enum { Value = true };
+};
+
 /**
  * Traits class which tests if a type is part of the core types included in CoreMinimal.h.
  */
diff --git a/Engine/Source/Runtime/Core/Public/UObject/UnrealNames.inl b/Engine/Source/Runtime/Core/Public/UObject/UnrealNames.inl
--- a/Engine/Source/Runtime/Core/Public/UObject/UnrealNames.inl
+++ b/Engine/Source/Runtime/Core/Public/UObject/UnrealNames.inl
@@ -100,7 +100,8 @@ REGISTER_NAME(93,Box2f)
 REGISTER_NAME(94,Box2d)
 REGISTER_NAME(95,IntVector)
 REGISTER_NAME(96,IntVector4)
-
+REGISTER_NAME(97,UintVector)
+REGISTER_NAME(98,UintVector4)
 
 // Object class names.
 REGISTER_NAME(100,Object)
@@ -113,6 +114,27 @@ REGISTER_NAME(106,ScriptStruct)
 REGISTER_NAME(107,Function)
 REGISTER_NAME(108,Pawn)
 
+// Special types continued
+REGISTER_NAME(150, Int32Vector)
+REGISTER_NAME(151, Int64Vector)
+REGISTER_NAME(152, Uint32Vector)
+REGISTER_NAME(153, Uint64Vector)
+REGISTER_NAME(154, Int32Vector4)
+REGISTER_NAME(155, Int64Vector4)
+REGISTER_NAME(156, Uint32Vector4)
+REGISTER_NAME(157, Uint64Vector4)
+REGISTER_NAME(158, IntVector2)
+REGISTER_NAME(159, Int32Vector2)
+REGISTER_NAME(160, Int64Vector2)
+REGISTER_NAME(161, UintVector2)
+REGISTER_NAME(162, Uint32Vector2)
+REGISTER_NAME(163, Uint64Vector2)
+REGISTER_NAME(164, UintPoint)
+REGISTER_NAME(165, Int32Point)
+REGISTER_NAME(166, Int64Point)
+REGISTER_NAME(167, Uint32Point)
+REGISTER_NAME(168, Uint64Point)
+
 // Misc.
 REGISTER_NAME(200,State)
 REGISTER_NAME(201,TRUE)
diff --git a/Engine/Source/Runtime/CoreUObject/Private/UObject/Property.cpp b/Engine/Source/Runtime/CoreUObject/Private/UObject/Property.cpp
--- a/Engine/Source/Runtime/CoreUObject/Private/UObject/Property.cpp
+++ b/Engine/Source/Runtime/CoreUObject/Private/UObject/Property.cpp
@@ -47,8 +47,8 @@ UE_IMPLEMENT_STRUCT("/Script/CoreUObject", Vector3f);
 UE_IMPLEMENT_STRUCT("/Script/CoreUObject", Vector3d);
 UE_IMPLEMENT_STRUCT("/Script/CoreUObject", Vector);	// Aliased
 
-template<>
-struct TStructOpsTypeTraits<FIntPoint> : public TStructOpsTypeTraitsBase2<FIntPoint>
+template<typename T>
+struct TIntPointStructOpsTypeTraits : public TStructOpsTypeTraitsBase2<T>
 {
 	enum 
 	{
@@ -56,25 +56,22 @@ struct TStructOpsTypeTraits<FIntPoint> : public TStructOpsTypeTraitsBase2<FIntPo
 		WithNoInitConstructor = true,
 		WithZeroConstructor = true,
 		WithSerializer = true,
+		WithSerializeFromMismatchedTag = true,
 	};
 };
-UE_IMPLEMENT_STRUCT("/Script/CoreUObject", IntPoint);
-
-template<>
-struct TStructOpsTypeTraits<FIntVector> : public TStructOpsTypeTraitsBase2<FIntVector>
-{
-	enum
-	{
-		WithIdenticalViaEquality = true,
-		WithNoInitConstructor = true,
-		WithZeroConstructor = true,
-		WithSerializer = true,
-	};
-};
-UE_IMPLEMENT_STRUCT("/Script/CoreUObject", IntVector);
+template<> struct TStructOpsTypeTraits<FInt32Point> : public TIntPointStructOpsTypeTraits<FInt32Point> {};
+template<> struct TStructOpsTypeTraits<FInt64Point> : public TIntPointStructOpsTypeTraits<FInt64Point> {};
+template<> struct TStructOpsTypeTraits<FUint32Point> : public TIntPointStructOpsTypeTraits<FUint32Point> {};
+template<> struct TStructOpsTypeTraits<FUint64Point> : public TIntPointStructOpsTypeTraits<FUint64Point> {};
+UE_IMPLEMENT_STRUCT("/Script/CoreUObject", Int32Point);
+UE_IMPLEMENT_STRUCT("/Script/CoreUObject", Int64Point);
+UE_IMPLEMENT_STRUCT("/Script/CoreUObject", Uint32Point);
+UE_IMPLEMENT_STRUCT("/Script/CoreUObject", Uint64Point);
+UE_IMPLEMENT_STRUCT("/Script/CoreUObject", IntPoint);		// Aliased
+UE_IMPLEMENT_STRUCT("/Script/CoreUObject", UintPoint);		// Aliased
 
-template<>
-struct TStructOpsTypeTraits<FIntVector4> : public TStructOpsTypeTraitsBase2<FIntVector4>
+template<typename T>
+struct TIntVectorStructOpsTypeTraits : public TStructOpsTypeTraitsBase2<T>
 {
 	enum
 	{
@@ -82,9 +79,41 @@ struct TStructOpsTypeTraits<FIntVector4> : public TStructOpsTypeTraitsBase2<FInt
 		WithNoInitConstructor = true,
 		WithZeroConstructor = true,
 		WithSerializer = true,
+		WithSerializeFromMismatchedTag = true,
 	};
 };
-UE_IMPLEMENT_STRUCT("/Script/CoreUObject", IntVector4);
+template<> struct TStructOpsTypeTraits<FInt32Vector2> : public TIntVectorStructOpsTypeTraits<FInt32Vector2> {};
+template<> struct TStructOpsTypeTraits<FInt64Vector2> : public TIntVectorStructOpsTypeTraits<FInt64Vector2> {};
+template<> struct TStructOpsTypeTraits<FUint32Vector2> : public TIntVectorStructOpsTypeTraits<FUint32Vector2> {};
+template<> struct TStructOpsTypeTraits<FUint64Vector2> : public TIntVectorStructOpsTypeTraits<FUint64Vector2> {};
+UE_IMPLEMENT_STRUCT("/Script/CoreUObject", Int32Vector2);
+UE_IMPLEMENT_STRUCT("/Script/CoreUObject", Int64Vector2);
+UE_IMPLEMENT_STRUCT("/Script/CoreUObject", Uint32Vector2);
+UE_IMPLEMENT_STRUCT("/Script/CoreUObject", Uint64Vector2);
+UE_IMPLEMENT_STRUCT("/Script/CoreUObject", IntVector2);		// Aliased
+UE_IMPLEMENT_STRUCT("/Script/CoreUObject", UintVector2);	// Aliased
+
+template<> struct TStructOpsTypeTraits<FInt32Vector3> : public TIntVectorStructOpsTypeTraits<FInt32Vector3> {};
+template<> struct TStructOpsTypeTraits<FInt64Vector3> : public TIntVectorStructOpsTypeTraits<FInt64Vector3> {};
+template<> struct TStructOpsTypeTraits<FUint32Vector3> : public TIntVectorStructOpsTypeTraits<FUint32Vector3> {};
+template<> struct TStructOpsTypeTraits<FUint64Vector3> : public TIntVectorStructOpsTypeTraits<FUint64Vector3> {};
+UE_IMPLEMENT_STRUCT("/Script/CoreUObject", Int32Vector);
+UE_IMPLEMENT_STRUCT("/Script/CoreUObject", Int64Vector);
+UE_IMPLEMENT_STRUCT("/Script/CoreUObject", Uint32Vector);
+UE_IMPLEMENT_STRUCT("/Script/CoreUObject", Uint64Vector);
+UE_IMPLEMENT_STRUCT("/Script/CoreUObject", IntVector);		// Aliased
+UE_IMPLEMENT_STRUCT("/Script/CoreUObject", UintVector);		// Aliased
+
+template<> struct TStructOpsTypeTraits<FInt32Vector4> : public TIntVectorStructOpsTypeTraits<FInt32Vector4> {};
+template<> struct TStructOpsTypeTraits<FInt64Vector4> : public TIntVectorStructOpsTypeTraits<FInt64Vector4> {};
+template<> struct TStructOpsTypeTraits<FUint32Vector4> : public TIntVectorStructOpsTypeTraits<FUint32Vector4> {};
+template<> struct TStructOpsTypeTraits<FUint64Vector4> : public TIntVectorStructOpsTypeTraits<FUint64Vector4> {};
+UE_IMPLEMENT_STRUCT("/Script/CoreUObject", Int32Vector4);
+UE_IMPLEMENT_STRUCT("/Script/CoreUObject", Int64Vector4);
+UE_IMPLEMENT_STRUCT("/Script/CoreUObject", Uint32Vector4);
+UE_IMPLEMENT_STRUCT("/Script/CoreUObject", Uint64Vector4); 
+UE_IMPLEMENT_STRUCT("/Script/CoreUObject", IntVector4);		// Aliased
+UE_IMPLEMENT_STRUCT("/Script/CoreUObject", UintVector4);	// Aliased
 
 template<typename T>
 struct TVector2StructOpsTypeTraits : public TStructOpsTypeTraitsBase2<T>
diff --git a/Engine/Source/Runtime/CoreUObject/Public/UObject/NoExportTypes.h b/Engine/Source/Runtime/CoreUObject/Public/UObject/NoExportTypes.h
--- a/Engine/Source/Runtime/CoreUObject/Public/UObject/NoExportTypes.h
+++ b/Engine/Source/Runtime/CoreUObject/Public/UObject/NoExportTypes.h
@@ -852,7 +852,179 @@ struct FIntPoint
 
 	UPROPERTY(EditAnywhere, BlueprintReadWrite, Category=IntPoint, SaveGame)
 	int32 Y;
+};
+
+USTRUCT(immutable, noexport, IsAlwaysAccessible, HasDefaults, HasNoOpConstructor, IsCoreType)
+struct FInt32Point
+{
+	UPROPERTY(EditAnywhere, Category = IntPoint, SaveGame)
+	int32 X;
+
+	UPROPERTY(EditAnywhere, Category = IntPoint, SaveGame)
+	int32 Y;
+};
+
+USTRUCT(immutable, noexport, IsAlwaysAccessible, HasDefaults, HasNoOpConstructor, IsCoreType)
+struct FInt64Point
+{
+	UPROPERTY(EditAnywhere, Category=IntPoint, SaveGame)
+	int64 X;
+
+	UPROPERTY(EditAnywhere, Category=IntPoint, SaveGame)
+	int64 Y;
+};
+
+/**
+ * Screen coordinates.
+ * @note The full C++ class is located here: Engine\Source\Runtime\Core\Public\Math\IntPoint.h
+ */
+USTRUCT(immutable, noexport, IsAlwaysAccessible, HasDefaults, HasNoOpConstructor, IsCoreType)
+struct FUintPoint
+{
+	UPROPERTY(EditAnywhere, Category = IntPoint, SaveGame)
+	int32 X;
+
+	UPROPERTY(EditAnywhere, Category = IntPoint, SaveGame)
+	int32 Y;
+};
+
+USTRUCT(immutable, noexport, IsAlwaysAccessible, HasDefaults, HasNoOpConstructor, IsCoreType)
+struct FUint32Point
+{
+	UPROPERTY(EditAnywhere, Category = IntPoint, SaveGame)
+	int32 X;
+
+	UPROPERTY(EditAnywhere, Category = IntPoint, SaveGame)
+	int32 Y;
+};
+
+USTRUCT(immutable, noexport, IsAlwaysAccessible, HasDefaults, HasNoOpConstructor, IsCoreType)
+struct FUint64Point
+{
+	UPROPERTY(EditAnywhere, Category = IntPoint, SaveGame)
+	int64 X;
+
+	UPROPERTY(EditAnywhere, Category = IntPoint, SaveGame)
+	int64 Y;
+};
+
+/**
+ * An integer vector in 3D space.
+ * @note The full C++ class is located here: Engine\Source\Runtime\Core\Public\Math\IntVector.h
+ */
+USTRUCT(immutable, noexport, IsAlwaysAccessible, HasDefaults, HasNoOpConstructor, IsCoreType)
+struct FInt32Vector2
+{
+	UPROPERTY(EditAnywhere, Category = IntVector, SaveGame)
+	int32 X;
+
+	UPROPERTY(EditAnywhere, Category = IntVector, SaveGame)
+	int32 Y;
+};
+
+/**
+ * An integer vector in 3D space.
+ * @note The full C++ class is located here: Engine\Source\Runtime\Core\Public\Math\IntVector.h
+ */
+USTRUCT(immutable, noexport, IsAlwaysAccessible, HasDefaults, HasNoOpConstructor, IsCoreType)
+struct FInt64Vector2
+{
+	UPROPERTY(EditAnywhere, Category = IntVector, SaveGame)
+	int64 X;
+
+	UPROPERTY(EditAnywhere, Category = IntVector, SaveGame)
+	int64 Y;
+};
+
+/**
+ * An integer vector in 4D space.
+ * @note The full C++ class is located here: Engine\Source\Runtime\Core\Public\Math\IntVector.h
+ */
+USTRUCT(immutable, noexport, IsAlwaysAccessible, HasDefaults, HasNoOpConstructor, IsCoreType)
+struct FIntVector2
+{
+	UPROPERTY(EditAnywhere, Category = IntVector, SaveGame)
+	int32 X;
+
+	UPROPERTY(EditAnywhere, Category = IntVector, SaveGame)
+	int32 Y;
+};
+
+/**
+ * An integer vector in 3D space.
+ * @note The full C++ class is located here: Engine\Source\Runtime\Core\Public\Math\IntVector.h
+ */
+USTRUCT(immutable, noexport, IsAlwaysAccessible, HasDefaults, HasNoOpConstructor, IsCoreType)
+struct FUint32Vector2
+{
+	UPROPERTY(EditAnywhere, Category = IntVector, SaveGame)
+	uint32 X;
+
+	UPROPERTY(EditAnywhere, Category = IntVector, SaveGame)
+	uint32 Y;
+};
+
+/**
+ * An integer vector in 3D space.
+ * @note The full C++ class is located here: Engine\Source\Runtime\Core\Public\Math\IntVector.h
+ */
+USTRUCT(immutable, noexport, IsAlwaysAccessible, HasDefaults, HasNoOpConstructor, IsCoreType)
+struct FUint64Vector2
+{
+	UPROPERTY(EditAnywhere, Category = IntVector, SaveGame)
+	uint64 X;
+
+	UPROPERTY(EditAnywhere, Category = IntVector, SaveGame)
+	uint64 Y;
+};
+
+/**
+ * An integer vector in 4D space.
+ * @note The full C++ class is located here: Engine\Source\Runtime\Core\Public\Math\IntVector.h
+ */
+USTRUCT(immutable, noexport, IsAlwaysAccessible, HasDefaults, HasNoOpConstructor, IsCoreType)
+struct FUintVector2
+{
+	UPROPERTY(EditAnywhere, Category = IntVector, SaveGame)
+	uint32 X;
+
+	UPROPERTY(EditAnywhere, Category = IntVector, SaveGame)
+	uint32 Y;
+};
+
+
+/**
+ * An integer vector in 3D space.
+ * @note The full C++ class is located here: Engine\Source\Runtime\Core\Public\Math\IntVector.h
+ */
+USTRUCT(immutable, noexport, IsAlwaysAccessible, HasDefaults, HasNoOpConstructor, IsCoreType)
+struct FInt32Vector
+{
+	UPROPERTY(EditAnywhere, Category = IntVector, SaveGame)
+	int32 X;
+
+	UPROPERTY(EditAnywhere, Category = IntVector, SaveGame)
+	int32 Y;
 
+	UPROPERTY(EditAnywhere, Category = IntVector, SaveGame)
+	int32 Z;
+};
+
+/**
+ * An integer vector in 3D space.
+ * @note The full C++ class is located here: Engine\Source\Runtime\Core\Public\Math\IntVector.h
+ */
+USTRUCT(immutable, noexport, IsAlwaysAccessible, HasDefaults, HasNoOpConstructor, IsCoreType)
+struct FInt64Vector
+{
+	UPROPERTY(EditAnywhere, Category = IntVector, SaveGame)
+	int64 X;
+
+	UPROPERTY(EditAnywhere, Category = IntVector, SaveGame)
+	int64 Y;
+
+	UPROPERTY(EditAnywhere, Category = IntVector, SaveGame)
+	int64 Z;
 };
 
 /**
@@ -872,6 +1044,99 @@ struct FIntVector
 	int32 Z;
 };
 
+
+/**
+ * An integer vector in 3D space.
+ * @note The full C++ class is located here: Engine\Source\Runtime\Core\Public\Math\IntVector.h
+ */
+USTRUCT(immutable, noexport, IsAlwaysAccessible, HasDefaults, HasNoOpConstructor, IsCoreType)
+struct FUint32Vector
+{
+	UPROPERTY(EditAnywhere, Category = IntVector, SaveGame)
+	uint32 X;
+
+	UPROPERTY(EditAnywhere, Category = IntVector, SaveGame)
+	uint32 Y;
+
+	UPROPERTY(EditAnywhere, Category = IntVector, SaveGame)
+	uint32 Z;
+};
+
+/**
+ * An integer vector in 3D space.
+ * @note The full C++ class is located here: Engine\Source\Runtime\Core\Public\Math\IntVector.h
+ */
+USTRUCT(immutable, noexport, IsAlwaysAccessible, HasDefaults, HasNoOpConstructor, IsCoreType)
+struct FUint64Vector
+{
+	UPROPERTY(EditAnywhere, Category = IntVector, SaveGame)
+	uint64 X;
+
+	UPROPERTY(EditAnywhere, Category = IntVector, SaveGame)
+	uint64 Y;
+
+	UPROPERTY(EditAnywhere, Category = IntVector, SaveGame)
+	uint64 Z;
+};
+
+/**
+ * An integer vector in 3D space.
+ * @note The full C++ class is located here: Engine\Source\Runtime\Core\Public\Math\IntVector.h
+ */
+USTRUCT(immutable, noexport, IsAlwaysAccessible, HasDefaults, HasNoOpConstructor, IsCoreType)
+struct FUintVector
+{
+	UPROPERTY(EditAnywhere, Category=IntVector, SaveGame)
+	uint32 X;
+
+	UPROPERTY(EditAnywhere, Category=IntVector, SaveGame)
+	uint32 Y;
+
+	UPROPERTY(EditAnywhere, Category=IntVector, SaveGame)
+	uint32 Z;
+};
+
+
+/**
+ * An integer vector in 3D space.
+ * @note The full C++ class is located here: Engine\Source\Runtime\Core\Public\Math\IntVector.h
+ */
+USTRUCT(immutable, noexport, IsAlwaysAccessible, HasDefaults, HasNoOpConstructor, IsCoreType)
+struct FInt32Vector4
+{
+	UPROPERTY(EditAnywhere, Category = IntVector, SaveGame)
+	int32 X;
+
+	UPROPERTY(EditAnywhere, Category = IntVector, SaveGame)
+	int32 Y;
+
+	UPROPERTY(EditAnywhere, Category = IntVector, SaveGame)
+	int32 Z;
+
+	UPROPERTY(EditAnywhere, Category = IntVector, SaveGame)
+	int32 W;
+};
+
+/**
+ * An integer vector in 3D space.
+ * @note The full C++ class is located here: Engine\Source\Runtime\Core\Public\Math\IntVector.h
+ */
+USTRUCT(immutable, noexport, IsAlwaysAccessible, HasDefaults, HasNoOpConstructor, IsCoreType)
+struct FInt64Vector4
+{
+	UPROPERTY(EditAnywhere, Category = IntVector, SaveGame)
+	int64 X;
+
+	UPROPERTY(EditAnywhere, Category = IntVector, SaveGame)
+	int64 Y;
+
+	UPROPERTY(EditAnywhere, Category = IntVector, SaveGame)
+	int64 Z;
+
+	UPROPERTY(EditAnywhere, Category = IntVector, SaveGame)
+	int64 W;
+};
+
 /**
  * An integer vector in 4D space.
  * @note The full C++ class is located here: Engine\Source\Runtime\Core\Public\Math\IntVector.h
@@ -892,6 +1157,67 @@ struct FIntVector4
 	int32 W;
 };
 
+/**
+ * An integer vector in 3D space.
+ * @note The full C++ class is located here: Engine\Source\Runtime\Core\Public\Math\IntVector.h
+ */
+USTRUCT(immutable, noexport, IsAlwaysAccessible, HasDefaults, HasNoOpConstructor, IsCoreType)
+struct FUint32Vector4
+{
+	UPROPERTY(EditAnywhere, Category = IntVector, SaveGame)
+	uint32 X;
+
+	UPROPERTY(EditAnywhere, Category = IntVector, SaveGame)
+	uint32 Y;
+
+	UPROPERTY(EditAnywhere, Category = IntVector, SaveGame)
+	uint32 Z;
+
+	UPROPERTY(EditAnywhere, Category = IntVector, SaveGame)
+	uint32 W;
+};
+
+/**
+ * An integer vector in 3D space.
+ * @note The full C++ class is located here: Engine\Source\Runtime\Core\Public\Math\IntVector.h
+ */
+USTRUCT(immutable, noexport, IsAlwaysAccessible, HasDefaults, HasNoOpConstructor, IsCoreType)
+struct FUint64Vector4
+{
+	UPROPERTY(EditAnywhere, Category = IntVector, SaveGame)
+	uint64 X;
+
+	UPROPERTY(EditAnywhere, Category = IntVector, SaveGame)
+	uint64 Y;
+
+	UPROPERTY(EditAnywhere, Category = IntVector, SaveGame)
+	uint64 Z;
+
+	UPROPERTY(EditAnywhere, Category = IntVector, SaveGame)
+	uint64 W;
+};
+
+/**
+ * An integer vector in 4D space.
+ * @note The full C++ class is located here: Engine\Source\Runtime\Core\Public\Math\IntVector.h
+ */
+USTRUCT(immutable, noexport, IsAlwaysAccessible, HasDefaults, HasNoOpConstructor, IsCoreType)
+struct FUintVector4
+{
+	UPROPERTY(EditAnywhere, Category = IntVector4, SaveGame)
+	uint32 X;
+
+	UPROPERTY(EditAnywhere, Category = IntVector4, SaveGame)
+	uint32 Y;
+
+	UPROPERTY(EditAnywhere, Category = IntVector4, SaveGame)
+	uint32 Z;
+
+	UPROPERTY(EditAnywhere, Category = IntVector4, SaveGame)
+	uint32 W;
+};
+
+
 /**
  * Stores a color with 8 bits of precision per channel. (BGRA).
  * @note The full C++ class is located here: Engine\Source\Runtime\Core\Public\Math\Color.h
