Fix editor bulkdata bug that was preventing payloads from being loaded from disk when -DDC2TextureBuilds was in use.

#rnx

### Problem
- When CL 21089705 was submitted, we changed the behaviour to assume that it is only safe to load payloads off disk if the AttachedAr pointer is valid rather than checking for the package path being valid. This allowed us to keep the package path intact after ::DetachFromDisk was called and give better error messages.
- However not all code paths were forcing the AttachedAr pointer to be set, TornOff bulkdata instances in the bulkdata registry for example causing these code paths to fail entirely.
- What we need is a way to not allow bulkdata instances that were attached but have been detached from loading (or at least force them to use the speculative loading path that was added in 21089705) and allow any other bulkdata instance that was never attached, to load it's payload as long as the path is set.

### Fix
- Instead of asking ::IsAttachedToPackageFile we instead ask ::HasAttachedArchive as in most places we are really checking if the AttachedAr pointer is null or not before we manipulate it
- There are two places where what we are really asking is if it is safe to load the payload from disk or not, these places now use a new method ::CanLoadDataFromDisk which checks that the bulkdata has access to a valid package path AND that the transient flag EFlags::WasDetached has not bee set.
-- This means that bulk data instances that are never attached to the packages FLinkerLoad can still load off disk, although this is less safe it reverts us back to the behaviour before CL 21089705was submitted.

### Misc
- Marked the transient flags as such in the code documentation.
- Ideally these flags would not have been added to the EFlags enum and a specific ETransientFlag enum added instead but the clean up for that would be irksome at this point.

diff --git a/Engine/Source/Runtime/CoreUObject/Private/Serialization/EditorBulkData.cpp b/Engine/Source/Runtime/CoreUObject/Private/Serialization/EditorBulkData.cpp
--- a/Engine/Source/Runtime/CoreUObject/Private/Serialization/EditorBulkData.cpp
+++ b/Engine/Source/Runtime/CoreUObject/Private/Serialization/EditorBulkData.cpp
@@ -443,7 +443,7 @@ FEditorBulkData& FEditorBulkData::operator=(const FEditorBulkData& Other)
 
 FEditorBulkData::~FEditorBulkData()
 {
-	if (IsAttachedToPackageFile())
+	if (HasAttachedArchive())
 	{
 		AttachedAr->DetachBulkData(this, false);
 		AttachedAr = nullptr;
@@ -666,9 +666,9 @@ void FEditorBulkData::CreateFromBulkData(FBulkData& InBulkData, const FGuid& InG
 	Reset();
 
 #if UE_ALLOW_LINKERLOADER_ATTACHMENT
-	check(!IsAttachedToPackageFile());
+	check(!HasAttachedArchive());
 	AttachedAr = InBulkData.AttachedAr;
-	if (IsAttachedToPackageFile())
+	if (HasAttachedArchive())
 	{
 		AttachedAr->AttachBulkData(this);
 	}
@@ -1025,7 +1025,7 @@ void FEditorBulkData::Serialize(FArchive& Ar, UObject* Owner, bool bAllowRegiste
 				if (!PackagePath.IsEmpty() && CacheableArchive != nullptr)
 				{
 #if UE_ALLOW_LINKERLOADER_ATTACHMENT
-					if (IsAttachedToPackageFile())
+					if (HasAttachedArchive())
 					{
 						// TODO: Remove this when doing UE-159339
 						AttachedAr->DetachBulkData(this, false);
@@ -1128,7 +1128,7 @@ FCompressedBuffer FEditorBulkData::LoadFromDisk() const
 	{
 		return LoadFromSidecarFile();
 	}
-	else if (!IsAttachedToPackageFile())
+	else if (!CanLoadDataFromDisk())
 	{
 		if (IsReferencingOldBulkData())
 		{
@@ -1476,7 +1476,12 @@ bool FEditorBulkData::CanUnloadData() const
 	// the payload from disk via ::LoadFromPackageTrailer but we are not guaranteed success
 	// because the package file may have changed to one without the payload in it at all.
 	// Due to this we do not allow the payload to be unloaded if we are detached.
-	return IsDataVirtualized() || (!PackagePath.IsEmpty() && IsAttachedToPackageFile());
+	return IsDataVirtualized() || CanLoadDataFromDisk();
+}
+
+bool FEditorBulkData::CanLoadDataFromDisk() const
+{
+	return !PackagePath.IsEmpty() && !EnumHasAnyFlags(Flags, EFlags::WasDetached);
 }
 
 bool FEditorBulkData::IsMemoryOnlyPayload() const
@@ -1489,7 +1494,7 @@ void FEditorBulkData::Reset()
 	// Unregister rather than allowing the Registry to keep our record, since we are changing the payload
 	Unregister();
 	// Note that we do not reset the BulkDataId
-	if (IsAttachedToPackageFile())
+	if (HasAttachedArchive())
 	{
 		AttachedAr->DetachBulkData(this, false);
 		AttachedAr = nullptr;
@@ -1518,7 +1523,7 @@ void FEditorBulkData::DetachFromDisk(FArchive* Ar, bool bEnsurePayloadIsLoaded)
 	TRACE_CPUPROFILER_EVENT_SCOPE(FEditorBulkData::DetachFromDisk);
 
 	check(Ar != nullptr);
-	check(Ar == AttachedAr || !IsAttachedToPackageFile() || AttachedAr->IsProxyOf(Ar));
+	check(Ar == AttachedAr || HasAttachedArchive() == false || AttachedAr->IsProxyOf(Ar));
 
 	// If bEnsurePayloadIsLoaded is true, then we should assume that a change to the 
 	// package file is imminent and we should load the payload into memory if possible.
@@ -1555,6 +1560,8 @@ void FEditorBulkData::DetachFromDisk(FArchive* Ar, bool bEnsurePayloadIsLoaded)
 	}
 
 	AttachedAr = nullptr;	
+
+	EnumAddFlags(Flags, EFlags::WasDetached);
 }
 
 FGuid FEditorBulkData::GetIdentifier() const
@@ -1704,7 +1711,7 @@ void FEditorBulkData::UpdatePayloadImpl(FSharedBuffer&& InPayload, FIoHash&& InP
 	// Unregister before calling DetachBulkData; DetachFromDisk calls OnExitMemory which is incorrect since
 	// we are changing data rather than leaving memory
 	Unregister();
-	if (IsAttachedToPackageFile())
+	if (HasAttachedArchive())
 	{
 		AttachedAr->DetachBulkData(this, false);
 		AttachedAr = nullptr;
diff --git a/Engine/Source/Runtime/CoreUObject/Public/Serialization/EditorBulkData.h b/Engine/Source/Runtime/CoreUObject/Public/Serialization/EditorBulkData.h
--- a/Engine/Source/Runtime/CoreUObject/Public/Serialization/EditorBulkData.h
+++ b/Engine/Source/Runtime/CoreUObject/Public/Serialization/EditorBulkData.h
@@ -341,22 +341,24 @@ private:
 		DisablePayloadCompression	= 1 << 4,
 		/** The legacy file being referenced derived its key from guid and it should be replaced with a key-from-hash when saved */
 		LegacyKeyWasGuidDerived		= 1 << 5,
-		/** The Guid has been registered with the BulkDataRegistry */
+		/** (Transient) The Guid has been registered with the BulkDataRegistry */
 		HasRegistered				= 1 << 6,
-		/** The BulkData object is a copy used only to represent the id and payload; it does not communicate with the BulkDataRegistry, and will point DDC jobs toward the original BulkData */
+		/** (Transient) The BulkData object is a copy used only to represent the id and payload; it does not communicate with the BulkDataRegistry, and will point DDC jobs toward the original BulkData */
 		IsTornOff					= 1 << 7,
 		/** The bulkdata object references a payload stored in a WorkspaceDomain file  */
 		ReferencesWorkspaceDomain	= 1 << 8,
 		/** The payload is stored in a package trailer, so the bulkdata object will have to poll the trailer to find the payload offset */
 		StoredInPackageTrailer		= 1 << 9,
 		/** The bulkdata object was cooked. */
-		IsCooked					= 1 << 10
+		IsCooked					= 1 << 10,
+		/** (Transient) The package owning the bulkdata has been detached from disk and we can no longer load from it */
+		WasDetached					= 1 << 11
 	};
 
 	FRIEND_ENUM_CLASS_FLAGS(EFlags);
 
 	/** A common grouping of EFlags */
-	static constexpr EFlags TransientFlags = EFlags((uint32)EFlags::HasRegistered | (uint32)EFlags::IsTornOff);
+	static constexpr EFlags TransientFlags = EFlags((uint32)EFlags::HasRegistered | (uint32)EFlags::IsTornOff | (uint32)EFlags::WasDetached);
 
 	/** Used to control what level of error reporting we return from some methods */
 	enum ErrorVerbosity
@@ -388,6 +390,7 @@ private:
 	FCompressedBuffer PullData() const;
 
 	bool CanUnloadData() const;
+	bool CanLoadDataFromDisk() const;
 
 	void UpdateKeyIfNeeded();
 
@@ -446,10 +449,10 @@ private:
 	}
 
 	/** 
-	 * Returns true when the bulkdata has an attachment to it's package file on disk, meaning that it is safe for us to load
-	 * the payload from the file directly as we know where the payload is on disk 
+	 * Returns true when the bulkdata has an attachment to it's package file on disk, if the bulkdata later becomes detached
+	 * then EFlag::WasDetached will be set.
 	 */
-	bool IsAttachedToPackageFile() const
+	bool HasAttachedArchive() const
 	{
 		return AttachedAr != nullptr;
 	}
