BulkData should determine chunk id based on GetPackageIdToLoad(), not GetPackageId() because level instances can modify their name

diff --git a/Engine/Source/Runtime/CoreUObject/Private/Serialization/BulkData.cpp b/Engine/Source/Runtime/CoreUObject/Private/Serialization/BulkData.cpp
--- a/Engine/Source/Runtime/CoreUObject/Private/Serialization/BulkData.cpp
+++ b/Engine/Source/Runtime/CoreUObject/Private/Serialization/BulkData.cpp
@@ -1375,7 +1375,7 @@ void FBulkData::Serialize(FArchive& Ar, UObject* Owner, bool bAttemptFileMapping
 					check(IsInSeparateFile());
 					check(Package != nullptr);
 					BulkMeta.SetFlags(EBulkDataFlags(BulkMeta.GetFlags() | BULKDATA_UsesIoDispatcher));
-					BulkChunkId = FBulkDataChunkId::FromPackageId(Package->GetPackageId());
+					BulkChunkId = FBulkDataChunkId::FromPackageId(Package->GetPackageIdToLoad());
 				}
 			}
 			else
