Temporarily do not write out optional bulkdata when found in editor optional context

diff --git a/Engine/Source/Runtime/CoreUObject/Private/UObject/SavePackage/SavePackageUtilities.cpp b/Engine/Source/Runtime/CoreUObject/Private/UObject/SavePackage/SavePackageUtilities.cpp
--- a/Engine/Source/Runtime/CoreUObject/Private/UObject/SavePackage/SavePackageUtilities.cpp
+++ b/Engine/Source/Runtime/CoreUObject/Private/UObject/SavePackage/SavePackageUtilities.cpp
@@ -2101,10 +2101,9 @@ ESavePackageResult SaveBulkData(FLinkerSave* Linker, int64& InOutStartOffset, co
 					
 				PackageWriter->WriteBulkData(BulkInfo, AddSizeAndConvertToIoBuffer(BulkArchive.Get()), BulkArchive->FileRegions);
 			}
-			if (OptionalBulkArchive && OptionalBulkArchive->TotalSize())
+			// @note FH: temporarily do not handle optional bulk data into editor optional packages, proper support will be added soon
+			if (OptionalBulkArchive && OptionalBulkArchive->TotalSize() && !bIsOptionalRealm)
 			{
-				// @note FH: temporarily comment assert for not supporting optional bulk data into editor optional packages
-				//checkf(!bIsOptionalRealm, TEXT("OptionalBulkData is currently unsupported with optional package multi output for %s"), *InOuter->GetName());
 				BulkInfo.ChunkId = CreateIoChunkId(PackageId.Value(), BulkInfo.MultiOutputIndex, EIoChunkType::OptionalBulkData);
 				BulkInfo.BulkDataType = IPackageWriter::FBulkDataInfo::Optional;
 				BulkInfo.LooseFilePath = FPathViews::ChangeExtension(BulkInfo.LooseFilePath, LexToString(EPackageExtension::BulkDataOptional));
