Fix a bug where saving a umap that contains a landscape component more than once if it was previously loaded from the editor domain.

#rnx

- Before 21089705 we recorded the package path that were were saving an editor bulkdata to, but since we had no way to "attach" to the file on disk we would not load from it. The changes in 21089705 meant that we could try loading from the package, which fails in some odd ways when the editor domain is enabled.
- Adding the ::wasDetached flag when saving the package should prevent this and restore the previous behaviour.

diff --git a/Engine/Source/Runtime/CoreUObject/Private/Serialization/EditorBulkData.cpp b/Engine/Source/Runtime/CoreUObject/Private/Serialization/EditorBulkData.cpp
--- a/Engine/Source/Runtime/CoreUObject/Private/Serialization/EditorBulkData.cpp
+++ b/Engine/Source/Runtime/CoreUObject/Private/Serialization/EditorBulkData.cpp
@@ -1608,6 +1608,11 @@ void FEditorBulkData::SerializeToLegacyPath(FLinkerSave& LinkerSave, FCompressed
 					return;
 				}
 
+				// Mark the bulkdata as detached so that we don't try to load the payload from the newly saved package
+				// or unload the current payload as under the existing system it is not safe to load from a package
+				// that we are not attached to.
+				EnumAddFlags(this->Flags, EFlags::WasDetached);
+
 				this->PackagePath = InPackagePath;
 				check(!this->PackagePath.IsEmpty()); // LinkerSave guarantees a valid PackagePath if we're updating loaded path
 				this->OffsetInFile = DataStartOffset;
@@ -1665,9 +1670,11 @@ void FEditorBulkData::SerializeToPackageTrailer(FLinkerSave& LinkerSave, FCompre
 				{
 					if (PayloadOffset != INDEX_NONE)
 					{
-						// A valid payload offset means that it was saved locally to disk. We can set the path at this point
-						// but will only be able to load from it if the package is reattached to the file on disk, which is
-						// not a currently supported feature.
+						// Mark the bulkdata as detached so that we don't try to load the payload from the newly saved package
+						// or unload the current payload as under the existing system it is not safe to load from a package
+						// that we are not attached to.
+						EnumAddFlags(this->Flags, EFlags::WasDetached);
+
 						this->PackagePath = InPackagePath;
 						check(!this->PackagePath.IsEmpty()); // LinkerSave guarantees a valid PackagePath if we're updating loaded path
 					}
