Add package import is coming from to error message
#rnx

diff --git a/Engine/Source/Runtime/CoreUObject/Private/UObject/LinkerLoad.cpp b/Engine/Source/Runtime/CoreUObject/Private/UObject/LinkerLoad.cpp
--- a/Engine/Source/Runtime/CoreUObject/Private/UObject/LinkerLoad.cpp
+++ b/Engine/Source/Runtime/CoreUObject/Private/UObject/LinkerLoad.cpp
@@ -3455,8 +3455,10 @@ bool FLinkerLoad::VerifyImportInner(const int32 ImportIndex, FString& WarningSuf
 						// Found the FObjectExport for this import
 						if ((Import.ClassName != Import.SourceLinker->GetExportClassName(j)) || (Import.ClassPackage != Import.SourceLinker->GetExportClassPackage(j)))
 						{
-							UE_ASSET_LOG(LogLinker, Warning, PackagePath, TEXT("Resolved import with a different class: import class '%s.%s', package class '%s.%s'. Resave to fix."), 
-								*Import.ClassPackage.ToString(), *Import.ClassName.ToString(), *Import.SourceLinker->GetExportClassPackage(j).ToString(), *Import.SourceLinker->GetExportClassName(j).ToString());
+							UE_ASSET_LOG(LogLinker, Warning, PackagePath, TEXT("Resolved import from '%s' with a different class: import class '%s.%s', package class '%s.%s'. Resave to fix."), 
+								*Import.SourceLinker->GetPackagePath().GetPackageName(),
+								*Import.ClassPackage.ToString(), *Import.ClassName.ToString(), 
+								*Import.SourceLinker->GetExportClassPackage(j).ToString(), *Import.SourceLinker->GetExportClassName(j).ToString());
 						}
 
 						Import.SourceIndex = j;
