Multiprocess cook: Copy EDLWarning data for saved packages from CookWorkers to Director.
Add IMPCollector for generic method to pass messages from CookWorkers to Director.
#rnx

diff --git a/Engine/Source/Editor/UnrealEd/Private/CookOnTheFlyServer.cpp b/Engine/Source/Editor/UnrealEd/Private/CookOnTheFlyServer.cpp
--- a/Engine/Source/Editor/UnrealEd/Private/CookOnTheFlyServer.cpp
+++ b/Engine/Source/Editor/UnrealEd/Private/CookOnTheFlyServer.cpp
@@ -21,6 +21,7 @@
 #include "Cooker/AsyncIODelete.h"
 #include "Cooker/CookDirector.h"
 #include "Cooker/CookedSavePackageValidator.h"
+#include "Cooker/CookMPCollector.h"
 #include "Cooker/CookOnTheFlyServerInterface.h"
 #include "Cooker/CookPackageData.h"
 #include "Cooker/CookPlatformManager.h"
@@ -37,6 +38,7 @@
 #include "Cooker/WorkerRequestsLocal.h"
 #include "Cooker/WorkerRequestsRemote.h"
 #include "CookerSettings.h"
+#include "CookOnTheFlyNetServer.h"
 #include "CookPackageSplitter.h"
 #include "DerivedDataCacheInterface.h"
 #include "DistanceFieldAtlas.h"
@@ -108,6 +110,7 @@
 #include "Serialization/ArchiveUObject.h"
 #include "Serialization/ArrayReader.h"
 #include "Serialization/ArrayWriter.h"
+#include "Serialization/CompactBinaryWriter.h"
 #include "Serialization/CustomVersion.h"
 #include "Settings/LevelEditorPlaySettings.h"
 #include "Settings/ProjectPackagingSettings.h"
@@ -132,7 +135,6 @@
 #include "UObject/UObjectArray.h"
 #include "UObject/UObjectIterator.h"
 #include "ZenStoreWriter.h"
-#include "CookOnTheFlyNetServer.h"
 
 #define LOCTEXT_NAMESPACE "Cooker"
 
@@ -9651,6 +9653,41 @@ void UCookOnTheFlyServer::BeginCookEditorSystems()
 	ModifiedAssetFilenames.Empty();
 }
 
+namespace UE::Cook
+{
+
+/** MultiprocessCook collector for FEDLCookChecker data. */
+class FEDLMPCollector : public IMPCollector
+{
+public:
+	virtual FGuid GetMessageType() override { return MessageType; }
+	virtual const TCHAR* GetDebugName() override { return TEXT("FEDLMPCollector"); }
+
+	virtual void ClientTick(FClientContext& Context) override;
+	virtual void ReceiveMessage(FServerContext& Context, FCbObjectView Message) override;
+
+	static FGuid MessageType;
+};
+FGuid FEDLMPCollector::MessageType(TEXT("0164FD08F6884F6A82D2D00F8F70B182"));
+
+void FEDLMPCollector::ClientTick(FClientContext& Context)
+{
+	FCbWriter Writer;
+	bool bHasData;
+	UE::SavePackageUtilities::EDLCookInfoMoveToCompactBinaryAndClear(Writer, bHasData);
+	if (bHasData)
+	{
+		Context.AddMessage(Writer.Save().AsObject());
+	}
+}
+
+void FEDLMPCollector::ReceiveMessage(FServerContext& Context, FCbObjectView Message)
+{
+	UE::SavePackageUtilities::EDLCookInfoAppendFromCompactBinary(Message.AsFieldView());
+}
+
+}
+
 void UCookOnTheFlyServer::BeginCookEDLCookInfo(FBeginCookContext& BeginContext)
 {
 	if (IsCookingInEditor())
@@ -9665,6 +9702,14 @@ void UCookOnTheFlyServer::BeginCookEDLCookInfo(FBeginCookContext& BeginContext)
 		return;
 	}
 	UE::SavePackageUtilities::StartSavingEDLCookInfoForVerification();
+	if (CookDirector)
+	{
+		CookDirector->Register(new UE::Cook::FEDLMPCollector());
+	}
+	else if (CookWorkerClient)
+	{
+		CookWorkerClient->Register(new UE::Cook::FEDLMPCollector());
+	}
 }
 
 void UCookOnTheFlyServer::RegisterCookByTheBookDelegates()
diff --git a/Engine/Source/Editor/UnrealEd/Private/Cooker/CookDirector.cpp b/Engine/Source/Editor/UnrealEd/Private/Cooker/CookDirector.cpp
--- a/Engine/Source/Editor/UnrealEd/Private/Cooker/CookDirector.cpp
+++ b/Engine/Source/Editor/UnrealEd/Private/Cooker/CookDirector.cpp
@@ -3,6 +3,7 @@
 #include "CookDirector.h"
 
 #include "CompactBinaryTCP.h"
+#include "CookMPCollector.h"
 #include "CookPackageData.h"
 #include "CookWorkerServer.h"
 #include "CookOnTheSide/CookOnTheFlyServer.h"
@@ -227,6 +228,30 @@ void FCookDirector::ShutdownCookSession()
 	ParseDesiredNumRemoteWorkers();
 }
 
+void FCookDirector::Register(IMPCollector* Collector)
+{
+	TRefCountPtr<IMPCollector>& Existing = MessageHandlers.FindOrAdd(Collector->GetMessageType());
+	if (Existing)
+	{
+		UE_LOG(LogCook, Error, TEXT("Duplicate IMPCollectors registered. Guid: %s, Existing: %s, Registering: %s. Keeping the Existing."),
+			*Collector->GetMessageType().ToString(), Existing->GetDebugName(), Collector->GetDebugName());
+		return;
+	}
+	Existing = Collector;
+}
+
+void FCookDirector::Unregister(IMPCollector* Collector)
+{
+	TRefCountPtr<IMPCollector> Existing;
+	MessageHandlers.RemoveAndCopyValue(Collector->GetMessageType(), Existing);
+	if (Existing && Existing.GetReference() != Collector)
+	{
+		UE_LOG(LogCook, Error, TEXT("Duplicate IMPCollector during Unregister. Guid: %s, Existing: %s, Unregistering: %s. Ignoring the Unregister."),
+			*Collector->GetMessageType().ToString(), Existing->GetDebugName(), Collector->GetDebugName());
+		MessageHandlers.Add(Collector->GetMessageType(), MoveTemp(Existing));
+	}
+}
+
 void FCookDirector::SetWorkersStalled(bool bInWorkersStalled)
 {
 	if (bInWorkersStalled != bWorkersStalled)
diff --git a/Engine/Source/Editor/UnrealEd/Private/Cooker/CookDirector.h b/Engine/Source/Editor/UnrealEd/Private/Cooker/CookDirector.h
--- a/Engine/Source/Editor/UnrealEd/Private/Cooker/CookDirector.h
+++ b/Engine/Source/Editor/UnrealEd/Private/Cooker/CookDirector.h
@@ -17,6 +17,7 @@ class FCbObject;
 class FCbWriter;
 class UCookOnTheFlyServer;
 namespace UE::Cook { class FCookWorkerServer; }
+namespace UE::Cook { class IMPCollector; }
 namespace UE::Cook { struct FPackageData; }
 namespace UE::Cook { struct FWorkerId; }
 
@@ -53,6 +54,11 @@ public:
 	};
 	EShowWorker GetShowWorkerOption() const { return ShowWorkerOption; }
 
+	/** Register a Collector to receive messages of its MessageType from CookWorkers. */
+	void Register(IMPCollector* Collector);
+	/** Unegister a Collector that was registered. */
+	void Unregister(IMPCollector* Collector);
+
 private:
 	/** CookWorker connections that have not yet identified which CookWorker they are. */
 	struct FPendingConnection
@@ -98,6 +104,7 @@ private:
 private:
 	TMap<int32, TUniquePtr<FCookWorkerServer>> RemoteWorkers;
 	TMap<FCookWorkerServer*, TUniquePtr<FCookWorkerServer>> ShuttingDownWorkers;
+	TMap<FGuid, TRefCountPtr<IMPCollector>> MessageHandlers;
 	TArray<FPendingConnection> PendingConnections;
 	FString WorkerConnectAuthority;
 	UCookOnTheFlyServer& COTFS;
diff --git a/Engine/Source/Editor/UnrealEd/Private/Cooker/CookMPCollector.cpp b/Engine/Source/Editor/UnrealEd/Private/Cooker/CookMPCollector.cpp
new file mode 100644
--- /dev/null
+++ b/Engine/Source/Editor/UnrealEd/Private/Cooker/CookMPCollector.cpp
@@ -0,0 +1,59 @@
+// Copyright Epic Games, Inc. All Rights Reserved.
+
+#include "CookMPCollector.h"
+
+#include "Serialization/CompactBinary.h"
+
+namespace UE::Cook
+{
+
+static uint8 PlatformToIntImpl(const ITargetPlatform* Platform, TConstArrayView<const ITargetPlatform*> Platforms)
+{
+	int32 Result = Platforms.IndexOfByKey(Platform);
+	check(Result != INDEX_NONE);
+	uint8 NarrowResult = static_cast<uint8>(Result);
+	check(static_cast<int32>(NarrowResult) == Result);
+	return NarrowResult;
+}
+
+static const ITargetPlatform* IntToPlatformImpl(uint8 PlatformAsInt, TConstArrayView<const ITargetPlatform*> Platforms)
+{
+	int32 Index = static_cast<int32>(PlatformAsInt);
+	if (Index < Platforms.Num())
+	{
+		return Platforms[Index];
+	}
+	return nullptr;
+}
+
+void IMPCollector::FClientContext::AddMessage(FCbObject Object)
+{
+	Messages.Add(Object);
+}
+
+uint8 IMPCollector::FClientContext::PlatformToInt(const ITargetPlatform* Platform) const
+{
+	return PlatformToIntImpl(Platform, Platforms);
+}
+
+const ITargetPlatform* IMPCollector::FClientContext::IntToPlatform(uint8 PlatformAsInt) const
+{
+	return IntToPlatformImpl(PlatformAsInt, Platforms);
+}
+
+uint8 IMPCollector::FServerContext::PlatformToInt(const ITargetPlatform* Platform) const
+{
+	return PlatformToIntImpl(Platform, Platforms);
+}
+
+const ITargetPlatform* IMPCollector::FServerContext::IntToPlatform(uint8 PlatformAsInt) const
+{
+	return IntToPlatformImpl(PlatformAsInt, Platforms);
+}
+
+void IMPCollector::ReceiveMessage(FServerContext& Context, FCbObjectView Message)
+{
+
+}
+
+}
diff --git a/Engine/Source/Editor/UnrealEd/Private/Cooker/CookMPCollector.h b/Engine/Source/Editor/UnrealEd/Private/Cooker/CookMPCollector.h
new file mode 100644
--- /dev/null
+++ b/Engine/Source/Editor/UnrealEd/Private/Cooker/CookMPCollector.h
@@ -0,0 +1,70 @@
+// Copyright Epic Games, Inc. All Rights Reserved.
+
+#pragma once
+
+#include "Containers/Array.h"
+#include "Containers/ArrayView.h"
+#include "HAL/Platform.h"
+#include "Misc/Guid.h"
+#include "Templates/RefCounting.h"
+
+class FCbObject;
+class FCbObjectView;
+class FCbWriter;
+class ITargetPlatform;
+
+namespace UE::Cook { class FCookWorkerServer; }
+
+namespace UE::Cook
+{
+
+/**
+ * Interface used during cooking to send data collected from save/load on a remote CookWorker
+ * to the Director for aggregation into files saves at the end of the cook. 
+ */
+class IMPCollector : public FRefCountBase
+{
+public:
+	class FClientContext
+	{
+	public:
+		TConstArrayView<const ITargetPlatform*> GetPlatforms() const { return Platforms; }
+		bool IsFlush()  const { return bFlush; }
+
+		void AddMessage(FCbObject Object);
+
+		uint8 PlatformToInt(const ITargetPlatform* Platform) const;
+		const ITargetPlatform* IntToPlatform(uint8 PlatformAsInt) const;
+
+	private:
+		TConstArrayView<const ITargetPlatform*> Platforms;
+		TArray<FCbObject> Messages;
+		bool bFlush = false;
+
+		friend class FCookWorkerClient;
+	};
+
+	class FServerContext
+	{
+	public:
+		TConstArrayView<const ITargetPlatform*> GetPlatforms() { return Platforms; }
+
+		uint8 PlatformToInt(const ITargetPlatform* Platform) const;
+		const ITargetPlatform* IntToPlatform(uint8 PlatformAsInt) const;
+
+	private:
+		TConstArrayView<const ITargetPlatform*> Platforms;
+
+		friend class FCookWorkerServer;
+	};
+
+	virtual ~IMPCollector() {}
+
+	virtual FGuid GetMessageType() = 0;
+	virtual const TCHAR* GetDebugName() = 0;
+
+	virtual void ClientTick(FClientContext& Context) {}
+	virtual void ReceiveMessage(FServerContext& Context, FCbObjectView Message);
+};
+
+}
\ No newline at end of file
diff --git a/Engine/Source/Editor/UnrealEd/Private/Cooker/CookWorkerClient.cpp b/Engine/Source/Editor/UnrealEd/Private/Cooker/CookWorkerClient.cpp
--- a/Engine/Source/Editor/UnrealEd/Private/Cooker/CookWorkerClient.cpp
+++ b/Engine/Source/Editor/UnrealEd/Private/Cooker/CookWorkerClient.cpp
@@ -4,6 +4,7 @@
 
 #include "CompactBinaryTCP.h"
 #include "CookDirector.h"
+#include "CookMPCollector.h"
 #include "CookPackageData.h"
 #include "CookPlatformManager.h"
 #include "CookTypes.h"
@@ -63,6 +64,7 @@ void FCookWorkerClient::TickFromSchedulerThread(FTickStackData& StackData)
 		{
 			SendPendingResults();
 			PumpSendMessages();
+			TickCollectors(StackData, false /* bFlush */);
 		}
 	}
 	else
@@ -439,6 +441,7 @@ void FCookWorkerClient::PumpDisconnect(FTickStackData& StackData)
 		{
 		case EConnectStatus::WaitForDisconnect:
 		{
+			TickCollectors(StackData, true /* bFlush */);
 			// Add code here for any waiting we need to do for the local CookOnTheFlyServer to gracefully shutdown
 			SendMessage(FAbortWorkerMessage(FAbortWorkerMessage::EType::Abort));
 			SendToState(EConnectStatus::WaitForDisconnectSocketFlush);
@@ -496,7 +499,7 @@ void FCookWorkerClient::SendToState(EConnectStatus TargetStatus)
 
 void FCookWorkerClient::LogInvalidMessage(const TCHAR* MessageTypeName)
 {
-	UE_LOG(LogCook, Warning, TEXT("CookWorkerClient received invalidly formatted message for type %s from CookDirector. Ignoring it."),
+	UE_LOG(LogCook, Error, TEXT("CookWorkerClient received invalidly formatted message for type %s from CookDirector. Ignoring it."),
 		MessageTypeName);
 }
 
@@ -526,4 +529,52 @@ void FCookWorkerClient::AssignPackages(FAssignPackagesMessage& Message)
 	}
 }
 
+void FCookWorkerClient::Register(IMPCollector* Collector)
+{
+	CollectorsToTick.AddUnique(Collector);
+}
+
+void FCookWorkerClient::Unregister(IMPCollector* Collector)
+{
+	CollectorsToTick.RemoveSwap(Collector);
+}
+
+void FCookWorkerClient::TickCollectors(FTickStackData& StackData, bool bFlush)
+{
+	if (StackData.LoopStartTime < NextTickCollectorsTimeSeconds && !bFlush)
+	{
+		return;
+	}
+
+	if (!CollectorsToTick.IsEmpty())
+	{
+		IMPCollector::FClientContext Context;
+		Context.Platforms = OrderedSessionPlatforms;
+		Context.bFlush = bFlush;
+		TArray<UE::CompactBinaryTCP::FMarshalledMessage> MarshalledMessages;
+
+		for (IMPCollector* Collector : CollectorsToTick)
+		{
+			Collector->ClientTick(Context);
+			if (!Context.Messages.IsEmpty())
+			{
+				FGuid MessageType = Collector->GetMessageType();
+				for (FCbObject& Object : Context.Messages)
+				{
+					MarshalledMessages.Add({ MessageType, MoveTemp(Object) });
+				}
+				Context.Messages.Reset();
+			}
+		}
+
+		if (!MarshalledMessages.IsEmpty())
+		{
+			UE::CompactBinaryTCP::TryWritePacket(ServerSocket, SendBuffer, MoveTemp(MarshalledMessages));
+		}
+	}
+
+	constexpr float TickCollectorsPeriodSeconds = 10.f;
+	NextTickCollectorsTimeSeconds = FPlatformTime::Seconds() + TickCollectorsPeriodSeconds;
+}
+
 }
\ No newline at end of file
diff --git a/Engine/Source/Editor/UnrealEd/Private/Cooker/CookWorkerClient.h b/Engine/Source/Editor/UnrealEd/Private/Cooker/CookWorkerClient.h
--- a/Engine/Source/Editor/UnrealEd/Private/Cooker/CookWorkerClient.h
+++ b/Engine/Source/Editor/UnrealEd/Private/Cooker/CookWorkerClient.h
@@ -8,6 +8,7 @@
 #include "CookTypes.h"
 #include "IPAddress.h"
 
+namespace UE::Cook { class IMPCollector; }
 namespace UE::Cook { struct FAssignPackagesMessage; }
 namespace UE::Cook { struct FDirectorConnectionInfo; }
 namespace UE::Cook { struct FInitialConfigMessage; }
@@ -50,6 +51,11 @@ public:
 	/** Queue a message to the server that the Package was saved. Will be sent during Tick. */
 	void ReportPromoteToSaveComplete(FPackageData& PackageData);
 
+	/** Register a Collector for periodic ticking that sends messages to the Director. */
+	void Register(IMPCollector* Collector);
+	/** Unegister a Collector that was registered. */
+	void Unregister(IMPCollector* Collector);
+
 private:
 	enum class EConnectStatus
 	{
@@ -91,6 +97,8 @@ private:
 	void LogInvalidMessage(const TCHAR* MessageTypeName);
 	/** Send packages assigned from the server into the request state. */
 	void AssignPackages(FAssignPackagesMessage& Message);
+	/** Tick the registered collectors. */
+	void TickCollectors(FTickStackData& StackData, bool bFlush);
 
 
 private:
@@ -98,12 +106,14 @@ private:
 	TUniquePtr<FInitialConfigMessage> InitialConfigMessage;
 	TArray<ITargetPlatform*> OrderedSessionPlatforms;
 	TArray<FPackageRemoteResult> PendingResults;
+	TArray<TRefCountPtr<IMPCollector>> CollectorsToTick;
 	UE::CompactBinaryTCP::FSendBuffer SendBuffer;
 	UE::CompactBinaryTCP::FReceiveBuffer ReceiveBuffer;
 	FString DirectorURI;
 	UCookOnTheFlyServer& COTFS;
 	FSocket* ServerSocket = nullptr;
 	double ConnectStartTimeSeconds = 0.;
+	double NextTickCollectorsTimeSeconds = 0.;
 	EConnectStatus ConnectStatus = EConnectStatus::Uninitialized;
 	ECookMode::Type DirectorCookMode = ECookMode::CookByTheBook;
 };
diff --git a/Engine/Source/Editor/UnrealEd/Private/Cooker/CookWorkerServer.cpp b/Engine/Source/Editor/UnrealEd/Private/Cooker/CookWorkerServer.cpp
--- a/Engine/Source/Editor/UnrealEd/Private/Cooker/CookWorkerServer.cpp
+++ b/Engine/Source/Editor/UnrealEd/Private/Cooker/CookWorkerServer.cpp
@@ -5,6 +5,7 @@
 #include "Commandlets/AssetRegistryGenerator.h"
 #include "CompactBinaryTCP.h"
 #include "CookDirector.h"
+#include "CookMPCollector.h"
 #include "CookPackageData.h"
 #include "CookPlatformManager.h"
 #include "HAL/PlatformProcess.h"
@@ -427,6 +428,22 @@ void FCookWorkerServer::HandleReceiveMessages(TArray<UE::CompactBinaryTCP::FMars
 				RecordResults(ResultsMessage);
 			}
 		}
+		else
+		{
+			TRefCountPtr<IMPCollector>* Collector = Director.MessageHandlers.Find(Message.MessageType);
+			if (Collector)
+			{
+				check(*Collector);
+				IMPCollector::FServerContext Context;
+				Context.Platforms = OrderedSessionPlatforms;
+				(*Collector)->ReceiveMessage(Context, Message.Object);
+			}
+			else
+			{
+				UE_LOG(LogCook, Error, TEXT("CookWorkerServer received message of unknown type %s from CookWorker. Ignoring it."),
+					*Message.MessageType.ToString());
+			}
+		}
 	}
 }
 
@@ -509,7 +526,7 @@ void FCookWorkerServer::RecordResults(FPackageResultsMessage& Message)
 
 void FCookWorkerServer::LogInvalidMessage(const TCHAR* MessageTypeName)
 {
-	UE_LOG(LogCook, Warning, TEXT("CookWorkerClient received invalidly formatted message for type %s from CookDirector. Ignoring it."),
+	UE_LOG(LogCook, Error, TEXT("CookWorkerServer received invalidly formatted message for type %s from CookWorker. Ignoring it."),
 		MessageTypeName);
 }
 
diff --git a/Engine/Source/Runtime/CoreUObject/Private/UObject/SavePackage/SaveContext.h b/Engine/Source/Runtime/CoreUObject/Private/UObject/SavePackage/SaveContext.h
--- a/Engine/Source/Runtime/CoreUObject/Private/UObject/SavePackage/SaveContext.h
+++ b/Engine/Source/Runtime/CoreUObject/Private/UObject/SavePackage/SaveContext.h
@@ -628,12 +628,12 @@ public:
 		SerializeContext = InContext;
 	}
 
-	FEDLCookChecker* GetEDLCookChecker() const
+	FEDLCookCheckerThreadState* GetEDLCookChecker() const
 	{
 		return EDLCookChecker;
 	}
 
-	void SetEDLCookChecker(FEDLCookChecker* InCookChecker)
+	void SetEDLCookChecker(FEDLCookCheckerThreadState* InCookChecker)
 	{
 		EDLCookChecker = InCookChecker;
 	}
@@ -1018,7 +1018,7 @@ private:
 	FCanSkipEditorReferencedPackagesWhenCooking SkipEditorRefCookingSetting;
 
 	// Pointer to the EDLCookChecker associated with this context
-	FEDLCookChecker* EDLCookChecker = nullptr;
+	FEDLCookCheckerThreadState* EDLCookChecker = nullptr;
 
 	// Matching any mark in ExcludedObjectMarks indicates that an object should be excluded from being either an import or an export for this save
 	const EObjectMark ExcludedObjectMarks;
diff --git a/Engine/Source/Runtime/CoreUObject/Private/UObject/SavePackage/SavePackageUtilities.cpp b/Engine/Source/Runtime/CoreUObject/Private/UObject/SavePackage/SavePackageUtilities.cpp
--- a/Engine/Source/Runtime/CoreUObject/Private/UObject/SavePackage/SavePackageUtilities.cpp
+++ b/Engine/Source/Runtime/CoreUObject/Private/UObject/SavePackage/SavePackageUtilities.cpp
@@ -13,10 +13,13 @@
 #include "Misc/ConfigCacheIni.h"
 #include "Misc/PathViews.h"
 #include "Misc/Paths.h"
+#include "Misc/ScopeExit.h"
 #include "Misc/ScopedSlowTask.h"
 #include "SaveContext.h"
 #include "Serialization/BulkData.h"
 #include "Serialization/EditorBulkData.h"
+#include "Serialization/CompactBinarySerialization.h"
+#include "Serialization/CompactBinaryWriter.h"
 #include "Serialization/LargeMemoryWriter.h"
 #include "Serialization/PackageWriter.h"
 #include "Tasks/Task.h"
@@ -1126,19 +1129,12 @@ void FEDLCookChecker::FEDLNodeData::Merge(FEDLCookChecker::FEDLNodeData&& Other)
 	ImportingPackagesSorted.SetNum(Algo::Unique(ImportingPackagesSorted), true /* bAllowShrinking */);
 }
 
-FEDLCookChecker::FEDLCookChecker(EInternalConstruct)
-	: bIsActive(false)
+FEDLCookCheckerThreadState::FEDLCookCheckerThreadState()
 {
+	Checker.SetActiveIfNeeded();
 
-}
-
-FEDLCookChecker::FEDLCookChecker()
-	:FEDLCookChecker(EInternalConstruct::Type)
-{
-	SetActiveIfNeeded();
-
-	FScopeLock CookCheckerInstanceLock(&CookCheckerInstanceCritical);
-	CookCheckerInstances.Add(this);
+	FScopeLock CookCheckerInstanceLock(&FEDLCookChecker::CookCheckerInstanceCritical);
+	FEDLCookChecker::CookCheckerInstances.Add(&Checker);
 }
 
 void FEDLCookChecker::SetActiveIfNeeded()
@@ -1368,24 +1364,28 @@ void FEDLCookChecker::Merge(FEDLCookChecker&& Other)
 	}
 }
 
-void FEDLCookChecker::Verify(bool bFullReferencesExpected)
+FEDLCookChecker FEDLCookChecker::AccumulateAndClear()
 {
-	check(!GIsSavingPackage);
-
-	FEDLCookChecker Accumulator(EInternalConstruct::Type);
+	FEDLCookChecker Accumulator;
 
+	FScopeLock CookCheckerInstanceLock(&CookCheckerInstanceCritical);
+	for (FEDLCookChecker* Checker : CookCheckerInstances)
 	{
-		FScopeLock CookCheckerInstanceLock(&CookCheckerInstanceCritical);
-		for (FEDLCookChecker* Checker : CookCheckerInstances)
+		if (Checker->bIsActive)
 		{
-			if (Checker->bIsActive)
-			{
-				Accumulator.bIsActive = true;
-				Accumulator.Merge(MoveTemp(*Checker));
-			}
+			Accumulator.bIsActive = true;
+			Accumulator.Merge(MoveTemp(*Checker));
 			Checker->Reset();
-		}			
+			Checker->bIsActive = true;
+		}
 	}
+	return Accumulator;
+}
+
+void FEDLCookChecker::Verify(bool bFullReferencesExpected)
+{
+	check(!GIsSavingPackage);
+	FEDLCookChecker Accumulator = AccumulateAndClear();
 
 	if (Accumulator.bIsActive)
 	{
@@ -1442,6 +1442,151 @@ void FEDLCookChecker::Verify(bool bFullReferencesExpected)
 	}
 }
 
+void FEDLCookChecker::MoveToCompactBinaryAndClear(FCbWriter& Writer, bool& bOutHasData)
+{
+	bOutHasData = false;
+	FEDLCookChecker Accumulator = AccumulateAndClear();
+	if (!Accumulator.bIsActive)
+	{
+		return;
+	}
+	if (Accumulator.Nodes.IsEmpty() && Accumulator.NodePrereqs.IsEmpty() && Accumulator.PackagesWithUnknownExports.IsEmpty())
+	{
+		return;
+	}
+	bOutHasData = true;
+
+	Accumulator.WriteToCompactBinary(Writer);
+}
+
+bool FEDLCookChecker::AppendFromCompactBinary(FCbFieldView Field)
+{
+	FEDLCookChecker Instance;
+	if (!Instance.ReadFromCompactBinary(Field))
+	{
+		return false;
+	}
+	FEDLCookChecker& CurrentChecker = FEDLCookCheckerThreadState::Get().Checker;
+	CurrentChecker.Merge(MoveTemp(Instance));
+	return true;
+}
+
+void FEDLCookChecker::WriteToCompactBinary(FCbWriter& Writer)
+{
+	Writer.BeginObject();
+	{
+		Writer.BeginArray("Nodes");
+		for (const FEDLNodeData& Node : Nodes)
+		{
+			Writer << Node.Name;
+			Writer << Node.ImportingPackagesSorted;
+			Writer << Node.ParentID;
+			Writer << static_cast<uint8>(Node.ObjectEvent);
+			Writer << Node.bIsExport;
+		}
+		Writer.EndArray();
+		Writer.BeginArray("NodePrereqs");
+		for (const TPair<FEDLNodeID, FEDLNodeID>& Pair : NodePrereqs)
+		{
+			Writer << static_cast<uint32>(Pair.Key);
+			Writer << static_cast<uint32>(Pair.Value);
+		}
+		Writer.EndArray();
+		Writer.BeginArray("PackagesWithUnknownExports");
+		for (FName PackageName : PackagesWithUnknownExports)
+		{
+			Writer << PackageName;
+		}
+		Writer.EndArray();
+	}
+	Writer.EndObject();
+}
+
+
+bool FEDLCookChecker::ReadFromCompactBinary(FCbFieldView Field)
+{
+	Reset();
+
+	bool bSuccess = false;
+	ON_SCOPE_EXIT
+	{
+		if (!bSuccess)
+		{
+			Reset();
+		}
+	};
+
+	FCbFieldView NodesField = Field["Nodes"];
+	Nodes.Reserve(NodesField.AsArrayView().Num() / 5);
+	if (NodesField.HasError())
+	{
+		return false;
+	}
+	FCbFieldViewIterator NodeIter = NodesField.CreateViewIterator();
+	while (NodeIter)
+	{
+		FEDLNodeID NodeID = Nodes.Num();
+		FEDLNodeData& Node = Nodes.Emplace_GetRef();
+		Node.ID = NodeID;
+		if (!LoadFromCompactBinary(NodeIter, Node.Name)) { return false; }
+		++NodeIter;
+		if (!LoadFromCompactBinary(NodeIter, Node.ImportingPackagesSorted)) { return false; }
+		++NodeIter;
+		if (!LoadFromCompactBinary(NodeIter, Node.ParentID)) { return false; }
+		++NodeIter;
+		uint8 LocalObjectEvent;
+		if (!LoadFromCompactBinary(NodeIter, LocalObjectEvent)) { return false; }
+		if (LocalObjectEvent > static_cast<uint8>(EObjectEvent::Max)) { return false; }
+		Node.ObjectEvent = static_cast<EObjectEvent>(LocalObjectEvent);
+		++NodeIter;
+		if (!LoadFromCompactBinary(NodeIter, Node.bIsExport)) { return false; }
+		++NodeIter;
+	}
+
+	FCbFieldView PrereqsField = Field["NodePrereqs"];
+	NodePrereqs.Reserve(PrereqsField.AsArrayView().Num() / 2);
+	if (PrereqsField.HasError())
+	{
+		return false;
+	}
+	FCbFieldViewIterator PrereqsIter = PrereqsField.CreateViewIterator();
+	while (PrereqsIter)
+	{
+		uint32 Key;
+		uint32 Value;
+		if (!LoadFromCompactBinary(PrereqsIter, Key)) { return false; }
+		++PrereqsIter;
+		if (!LoadFromCompactBinary(PrereqsIter, Value)) { return false; }
+		++PrereqsIter;
+		NodePrereqs.Add(static_cast<FEDLNodeID>(Key), static_cast<FEDLNodeID>(Value));
+	}
+
+	FCbFieldView PackagesWithUnknownExportsField = Field["PackagesWithUnknownExports"];
+	PackagesWithUnknownExports.Reserve(PackagesWithUnknownExportsField.AsArrayView().Num());
+	if (PackagesWithUnknownExportsField.HasError())
+	{
+		return false;
+	}
+	for (FCbFieldView PackageNameField : PackagesWithUnknownExportsField)
+	{
+		FName PackageName;
+		if (!LoadFromCompactBinary(PackageNameField, PackageName))
+		{
+			return false;
+		}
+		PackagesWithUnknownExports.Add(PackageName);
+	}
+
+	for (const FEDLNodeData& Node : Nodes)
+	{
+		NodeHashToNodeID.Add(Node.GetNodeHash(*this), Node.ID);
+	}
+	bIsActive = !Nodes.IsEmpty() || !NodePrereqs.IsEmpty() || !PackagesWithUnknownExports.IsEmpty();
+
+	bSuccess = true;
+	return true;
+}
+
 FCriticalSection FEDLCookChecker::CookCheckerInstanceCritical;
 TArray<FEDLCookChecker*> FEDLCookChecker::CookCheckerInstances;
 
@@ -1460,7 +1605,17 @@ void VerifyEDLCookInfo(bool bFullReferencesExpected)
 
 void EDLCookInfoAddIterativelySkippedPackage(FName LongPackageName)
 {
-	FEDLCookChecker::Get().AddPackageWithUnknownExports(LongPackageName);
+	FEDLCookCheckerThreadState::Get().AddPackageWithUnknownExports(LongPackageName);
+}
+
+void EDLCookInfoMoveToCompactBinaryAndClear(FCbWriter& Writer, bool& bOutHasData)
+{
+	FEDLCookChecker::MoveToCompactBinaryAndClear(Writer, bOutHasData);
+}
+
+bool EDLCookInfoAppendFromCompactBinary(FCbFieldView Field)
+{
+	return FEDLCookChecker::AppendFromCompactBinary(Field);
 }
 
 }
diff --git a/Engine/Source/Runtime/CoreUObject/Private/UObject/SavePackage/SavePackageUtilities.h b/Engine/Source/Runtime/CoreUObject/Private/UObject/SavePackage/SavePackageUtilities.h
--- a/Engine/Source/Runtime/CoreUObject/Private/UObject/SavePackage/SavePackageUtilities.h
+++ b/Engine/Source/Runtime/CoreUObject/Private/UObject/SavePackage/SavePackageUtilities.h
@@ -17,6 +17,8 @@
 
 // This file contains private utilities shared by UPackage::Save and UPackage::Save2 
 
+class FCbFieldView;
+class FCbWriter;
 class FPackagePath;
 class FSaveContext;
 class FSavePackageContext;
@@ -167,10 +169,12 @@ private:
 	TMap<UObject*, FString>			ObjectToFullNameMap;
 };
 
+struct FEDLCookCheckerThreadState;
+
 /**
  * Helper struct used during cooking to validate EDL dependencies
  */
-struct FEDLCookChecker : public TThreadSingleton<FEDLCookChecker>
+struct FEDLCookChecker
 {
 	void SetActiveIfNeeded();
 
@@ -183,8 +187,14 @@ struct FEDLCookChecker : public TThreadSingleton<FEDLCookChecker>
 
 	static void StartSavingEDLCookInfoForVerification();
 	static void Verify(bool bFullReferencesExpected);
+	static void MoveToCompactBinaryAndClear(FCbWriter& Writer, bool& bOutHasData);
+	static bool AppendFromCompactBinary(FCbFieldView Field);
 
 private:
+	static FEDLCookChecker AccumulateAndClear();
+	void WriteToCompactBinary(FCbWriter& Writer);
+	bool ReadFromCompactBinary(FCbFieldView Field);
+
 	typedef uint32 FEDLNodeID;
 	static const FEDLNodeID NodeIDInvalid = static_cast<FEDLNodeID>(-1);
 
@@ -193,7 +203,8 @@ public: // FEDLNodeHash is public only so that GetTypeHash can be defined
 	enum class EObjectEvent : uint8
 	{
 		Create,
-		Serialize
+		Serialize,
+		Max = Serialize,
 	};
 
 	/**
@@ -262,6 +273,7 @@ private:
 		/** True if the UObject represented by this node has been exported by a SavePackage call; used to verify that the imports requested by packages are present somewhere in the cook. */
 		bool bIsExport;
 
+		FEDLNodeData() { /* Fields are uninitialized */ }
 		FEDLNodeData(FEDLNodeID InID, FEDLNodeID InParentID, FName InName, EObjectEvent InObjectEvent);
 		FEDLNodeData(FEDLNodeID InID, FEDLNodeID InParentID, FName InName, FEDLNodeData&& Other);
 		FEDLNodeHash GetNodeHash(const FEDLCookChecker& Owner) const;
@@ -272,14 +284,6 @@ private:
 		void Merge(FEDLNodeData&& Other);
 	};
 
-	enum class EInternalConstruct
-	{
-		Type
-	};
-
-	FEDLCookChecker();
-	FEDLCookChecker(EInternalConstruct);
-
 	FEDLNodeID FindOrAddNode(const FEDLNodeHash& NodeLookup);
 	FEDLNodeID FindOrAddNode(FEDLNodeData&& NodeData, const FEDLCookChecker& OldOwnerOfNode, FEDLNodeID ParentIDInThis, bool& bNew);
 	FEDLNodeID FindNode(const FEDLNodeHash& NodeHash);
@@ -289,7 +293,7 @@ private:
 
 	/**
 	 * All the FEDLNodeDatas that have been created for this checker. These are allocated as elements of an array rather than pointers to reduce cputime and
-	 * memory due to many small allocations, and to provide index-based identifiers. Nodes are not deleted during the lifetime of the checker.
+	 * memory due to many small allocations, and to provide index-based identifiers. Nodes are not deleted until the checker is reset.
 	 */
 	TArray<FEDLNodeData> Nodes;
 	/** A map to lookup the node for a UObject or for the corresponding node in another thread's FEDLCookChecker. */
@@ -302,13 +306,41 @@ private:
 	 */
 	TSet<FName> PackagesWithUnknownExports;
 	/** True if the EDLCookChecker should be active; it is turned off if the runtime will not be using EDL. */
-	bool bIsActive;
+	bool bIsActive = false;
 
 	/** When cooking with concurrent saving, each thread has its own FEDLCookChecker, and these are merged after the cook is complete. */
 	static FCriticalSection CookCheckerInstanceCritical;
 	static TArray<FEDLCookChecker*> CookCheckerInstances;
 
-	friend TThreadSingleton<FEDLCookChecker>;
+	friend FEDLCookCheckerThreadState;
+};
+
+/** Per-thread accessor for writing EDL dependencies to global FEDLCookChecker storage. */
+struct FEDLCookCheckerThreadState : public TThreadSingleton<FEDLCookCheckerThreadState>
+{
+	FEDLCookCheckerThreadState();
+
+	void AddImport(UObject* Import, UPackage* ImportingPackage)
+	{
+		Checker.AddImport(Import, ImportingPackage);
+	}
+	void AddExport(UObject* Export)
+	{
+		Checker.AddExport(Export);
+	}
+	void AddArc(UObject* DepObject, bool bDepIsSerialize, UObject* Export, bool bExportIsSerialize)
+	{
+		Checker.AddArc(DepObject, bDepIsSerialize, Export, bExportIsSerialize);
+	}
+	void AddPackageWithUnknownExports(FName LongPackageName)
+	{
+		Checker.AddPackageWithUnknownExports(LongPackageName);
+	}
+
+private:
+	FEDLCookChecker Checker;
+	friend TThreadSingleton<FEDLCookCheckerThreadState>;
+	friend FEDLCookChecker;
 };
 
 #if WITH_EDITORONLY_DATA
diff --git a/Engine/Source/Runtime/CoreUObject/Private/UObject/SavePackage2.cpp b/Engine/Source/Runtime/CoreUObject/Private/UObject/SavePackage2.cpp
--- a/Engine/Source/Runtime/CoreUObject/Private/UObject/SavePackage2.cpp
+++ b/Engine/Source/Runtime/CoreUObject/Private/UObject/SavePackage2.cpp
@@ -582,7 +582,7 @@ ESavePackageResult ValidateExports(FSaveContext& SaveContext)
 	if (SaveContext.IsCooking())
 	{
 		// Add the exports for the cook checker
-		if (FEDLCookChecker* EDLCookChecker = SaveContext.GetEDLCookChecker())
+		if (FEDLCookCheckerThreadState* EDLCookChecker = SaveContext.GetEDLCookChecker())
 		{
 			// the package isn't actually in the export map, but that is ok, we add it as export anyway for error checking
 			EDLCookChecker->AddExport(SaveContext.GetPackage());
@@ -834,7 +834,7 @@ ESavePackageResult ValidateImports(FSaveContext& SaveContext)
 	if (SaveContext.IsCooking() && SaveContext.GetCurrentHarvestingRealm() != ESaveRealm::Optional)
 	{
 		// Now that imports are validated add them to the cook checker if available
-		if (FEDLCookChecker* EDLCookChecker = SaveContext.GetEDLCookChecker())
+		if (FEDLCookCheckerThreadState* EDLCookChecker = SaveContext.GetEDLCookChecker())
 		{
 			for (UObject* Import : SaveContext.GetImports())
 			{
@@ -1472,7 +1472,7 @@ void SavePreloadDependencies(FStructuredArchive::FRecord& StructuredArchiveRoot,
 				IncludeIndexAsDependency(CreateBeforeCreateDependencies, Export.SuperIndex);
 			}
 			// Currently do not validate the optional context with the edl checker
-			FEDLCookChecker* EDLCookChecker = SaveContext.GetCurrentHarvestingRealm() != ESaveRealm::Optional ? SaveContext.GetEDLCookChecker() : nullptr;
+			FEDLCookCheckerThreadState* EDLCookChecker = SaveContext.GetCurrentHarvestingRealm() != ESaveRealm::Optional ? SaveContext.GetEDLCookChecker() : nullptr;
 			auto AddArcForDepChecking = [&Linker, &Export, EDLCookChecker](bool bExportIsSerialize, FPackageIndex Dep, bool bDepIsSerialize)
 			{
 				check(Export.Object);
@@ -2459,7 +2459,7 @@ ESavePackageResult InnerSave(FSaveContext& SaveContext)
 {
 	TRefCountPtr<FUObjectSerializeContext> SerializeContext(FUObjectThreadContext::Get().GetSerializeContext());
 	SaveContext.SetSerializeContext(SerializeContext);
-	SaveContext.SetEDLCookChecker(&FEDLCookChecker::Get());
+	SaveContext.SetEDLCookChecker(&FEDLCookCheckerThreadState::Get());
 
 	// Create slow task dialog if needed
 	const int32 TotalSaveSteps = 4;
diff --git a/Engine/Source/Runtime/CoreUObject/Public/UObject/SavePackage.h b/Engine/Source/Runtime/CoreUObject/Public/UObject/SavePackage.h
--- a/Engine/Source/Runtime/CoreUObject/Public/UObject/SavePackage.h
+++ b/Engine/Source/Runtime/CoreUObject/Public/UObject/SavePackage.h
@@ -30,6 +30,8 @@ class UObject;
 
 class FArchive;
 class FArchiveDiffMap;
+class FCbFieldView;
+class FCbWriter;
 class FIoBuffer;
 class FOutputDevice;
 class FPackagePath;
@@ -211,6 +213,9 @@ namespace UE::SavePackageUtilities
 	COREUOBJECT_API void StartSavingEDLCookInfoForVerification();
 	COREUOBJECT_API void VerifyEDLCookInfo(bool bFullReferencesExpected = true);
 	COREUOBJECT_API void EDLCookInfoAddIterativelySkippedPackage(FName LongPackageName);
+	COREUOBJECT_API void EDLCookInfoMoveToCompactBinaryAndClear(FCbWriter& Writer, bool& bOutHasData);
+	COREUOBJECT_API bool EDLCookInfoAppendFromCompactBinary(FCbFieldView Field);
+
 
 #if WITH_EDITOR
 	/** void FAddResaveOnDemandPackage(FName SystemName, FName PackageName); */
