#UE Add import name to message about type mismatch

diff --git a/Engine/Source/Runtime/CoreUObject/Private/UObject/LinkerLoad.cpp b/Engine/Source/Runtime/CoreUObject/Private/UObject/LinkerLoad.cpp
--- a/Engine/Source/Runtime/CoreUObject/Private/UObject/LinkerLoad.cpp
+++ b/Engine/Source/Runtime/CoreUObject/Private/UObject/LinkerLoad.cpp
@@ -3455,7 +3455,8 @@ bool FLinkerLoad::VerifyImportInner(const int32 ImportIndex, FString& WarningSuf
 						// Found the FObjectExport for this import
 						if ((Import.ClassName != Import.SourceLinker->GetExportClassName(j)) || (Import.ClassPackage != Import.SourceLinker->GetExportClassPackage(j)))
 						{
-							UE_ASSET_LOG(LogLinker, Warning, PackagePath, TEXT("Resolved import from '%s' with a different class: import class '%s.%s', package class '%s.%s'. Resave to fix."), 
+							UE_ASSET_LOG(LogLinker, Warning, PackagePath, TEXT("Resolved import with name '%s' from '%s' with a different class: import class '%s.%s', package class '%s.%s'. Resave to fix."), 
+								*Import.ObjectName.ToString(),
 								*Import.SourceLinker->GetPackagePath().GetPackageName(),
 								*Import.ClassPackage.ToString(), *Import.ClassName.ToString(), 
 								*Import.SourceLinker->GetExportClassPackage(j).ToString(), *Import.SourceLinker->GetExportClassName(j).ToString());
