Sequencer: Fix PropertyTrack redirect

diff --git a/Engine/Source/Runtime/CoreUObject/Private/UObject/CoreRedirects.cpp b/Engine/Source/Runtime/CoreUObject/Private/UObject/CoreRedirects.cpp
--- a/Engine/Source/Runtime/CoreUObject/Private/UObject/CoreRedirects.cpp
+++ b/Engine/Source/Runtime/CoreUObject/Private/UObject/CoreRedirects.cpp
@@ -2041,7 +2041,7 @@ static void RegisterNativeRedirects49(TArray<FCoreRedirect>& Redirects)
 	CLASS_REDIRECT("MovieSceneFloatTrack", "/Script/MovieSceneTracks.MovieSceneFloatTrack");
 	CLASS_REDIRECT("MovieSceneParticleSection", "/Script/MovieSceneTracks.MovieSceneParticleSection");
 	CLASS_REDIRECT("MovieSceneParticleTrack", "/Script/MovieSceneTracks.MovieSceneParticleTrack");
-	CLASS_REDIRECT("MovieScenePropertyTrack", "/Script/MovieScene.MovieScenePropertyTrack");
+	CLASS_REDIRECT("MovieScenePropertyTrack", "/Script/MovieSceneTracks.MovieScenePropertyTrack");
 	CLASS_REDIRECT("MovieSceneSection", "/Script/MovieScene.MovieSceneSection");
 	CLASS_REDIRECT("MovieSceneTrack", "/Script/MovieScene.MovieSceneTrack");
 	CLASS_REDIRECT("MovieSceneVectorSection", "/Script/MovieSceneTracks.MovieSceneVectorSection");
