Protect AllEnumNames which can be accessed from multiple threads at the same time (i.e. async loading thread and game-thread)

diff --git a/Engine/Source/Runtime/CoreUObject/Private/UObject/Enum.cpp b/Engine/Source/Runtime/CoreUObject/Private/UObject/Enum.cpp
--- a/Engine/Source/Runtime/CoreUObject/Private/UObject/Enum.cpp
+++ b/Engine/Source/Runtime/CoreUObject/Private/UObject/Enum.cpp
@@ -20,6 +20,7 @@ DEFINE_LOG_CATEGORY_STATIC(LogEnum, Log, All);
 	UEnum implementation.
 -----------------------------------------------------------------------------*/
 
+FRWLock UEnum::AllEnumNamesLock;
 TMap<FName, UEnum*> UEnum::AllEnumNames;
 
 UEnum::UEnum(const FObjectInitializer& ObjectInitializer)
@@ -260,6 +261,7 @@ bool UEnum::IsValidEnumName(FName InName) const
 
 void UEnum::AddNamesToMasterList()
 {
+	FWriteScopeLock ScopeLock(AllEnumNamesLock);
 	for (TPair<FName, int64> Kvp : Names)
 	{
 		UEnum* Enum = AllEnumNames.FindRef(Kvp.Key);
@@ -276,6 +278,7 @@ void UEnum::AddNamesToMasterList()
 
 void UEnum::RemoveNamesFromMasterList()
 {
+	FWriteScopeLock ScopeLock(AllEnumNamesLock);
 	for (TPair<FName, int64> Kvp : Names)
 	{
 		UEnum* Enum = AllEnumNames.FindRef(Kvp.Key);
diff --git a/Engine/Source/Runtime/CoreUObject/Public/UObject/Class.h b/Engine/Source/Runtime/CoreUObject/Public/UObject/Class.h
--- a/Engine/Source/Runtime/CoreUObject/Public/UObject/Class.h
+++ b/Engine/Source/Runtime/CoreUObject/Public/UObject/Class.h
@@ -2054,6 +2054,7 @@ public:
 	 */
 	static int64 LookupEnumName(FName TestName, UEnum** FoundEnum = nullptr)
 	{
+		FReadScopeLock ScopeLock(AllEnumNamesLock);
 		UEnum* TheEnum = AllEnumNames.FindRef(TestName);
 		if (FoundEnum != nullptr)
 		{
@@ -2330,6 +2331,9 @@ protected:
 	/** pointer to function used to look up the enum's display name. Currently only assigned for UEnums generated for nativized blueprints */
 	FEnumDisplayNameFn EnumDisplayNameFn;
 
+	/** lock to be taken when accessing AllEnumNames */
+	static FRWLock AllEnumNamesLock;
+
 	/** global list of all value names used by all enums in memory, used for property text import */
 	static TMap<FName, UEnum*> AllEnumNames;
 
