AsyncLoading2: Remove old and unused code for additional worker threads

diff --git a/Engine/Source/Runtime/CoreUObject/Private/Serialization/AsyncLoading2.cpp b/Engine/Source/Runtime/CoreUObject/Private/Serialization/AsyncLoading2.cpp
--- a/Engine/Source/Runtime/CoreUObject/Private/Serialization/AsyncLoading2.cpp
+++ b/Engine/Source/Runtime/CoreUObject/Private/Serialization/AsyncLoading2.cpp
@@ -2183,59 +2183,6 @@ struct FScopedAsyncPackageEvent2
 	~FScopedAsyncPackageEvent2();
 };
 
-class FAsyncLoadingThreadWorker : private FRunnable
-{
-public:
-	FAsyncLoadingThreadWorker(FAsyncLoadEventGraphAllocator& InGraphAllocator, FAsyncLoadEventQueue2& InEventQueue, FIoDispatcher& InIoDispatcher, FZenaphore& InZenaphore, TAtomic<int32>& InActiveWorkersCount)
-		: Zenaphore(InZenaphore)
-		, EventQueue(InEventQueue)
-		, GraphAllocator(InGraphAllocator)
-		, IoDispatcher(InIoDispatcher)
-		, ActiveWorkersCount(InActiveWorkersCount)
-	{
-	}
-
-	void StartThread();
-	
-	void StopThread()
-	{
-		bStopRequested = true;
-		bSuspendRequested = true;
-		Zenaphore.NotifyAll();
-	}
-	
-	void SuspendThread()
-	{
-		bSuspendRequested = true;
-		Zenaphore.NotifyAll();
-	}
-	
-	void ResumeThread()
-	{
-		bSuspendRequested = false;
-	}
-	
-	int32 GetThreadId() const
-	{
-		return ThreadId;
-	}
-
-private:
-	virtual bool Init() override { return true; }
-	virtual uint32 Run() override;
-	virtual void Stop() override {};
-
-	FZenaphore& Zenaphore;
-	FAsyncLoadEventQueue2& EventQueue;
-	FAsyncLoadEventGraphAllocator& GraphAllocator;
-	FIoDispatcher& IoDispatcher;
-	TAtomic<int32>& ActiveWorkersCount;
-	FRunnableThread* Thread = nullptr;
-	TAtomic<bool> bStopRequested { false };
-	TAtomic<bool> bSuspendRequested { false };
-	int32 ThreadId = 0;
-};
-
 class FAsyncLoadingThread2 final
 	: public FRunnable
 	, public IAsyncPackageLoader
@@ -2250,9 +2197,6 @@ private:
 	FRunnableThread* Thread;
 	TAtomic<bool> bStopRequested { false };
 	TAtomic<bool> bSuspendRequested { false };
-	TArray<FAsyncLoadingThreadWorker> Workers;
-	TAtomic<int32> ActiveWorkersCount { 0 };
-	bool bWorkersSuspended = false;
 	bool bHasRegisteredAllScriptObjects = false;
 	/** [ASYNC/GAME THREAD] true if the async thread is actually started. We don't start it until after we boot because the boot process on the game thread can create objects that are also being created by the loader */
 	bool bThreadStarted = false;
@@ -2350,7 +2294,6 @@ public:
 
 	/** [EDL] Event queue */
 	FZenaphore AltZenaphore;
-	TArray<FZenaphore> WorkerZenaphores;
 	FAsyncLoadEventGraphAllocator GraphAllocator;
 	FAsyncLoadEventQueue2 EventQueue;
 	FAsyncLoadEventQueue2 MainThreadEventQueue;
@@ -2401,16 +2344,6 @@ public:
 			{
 				return true;
 			}
-			else
-			{
-				for (const FAsyncLoadingThreadWorker& Worker : Workers)
-				{
-					if (CurrentThreadId == Worker.GetThreadId())
-					{
-						return true;
-					}
-				}
-			}
 			return false;
 		}
 		else
@@ -2586,9 +2519,6 @@ private:
 
 	void OnLeakedPackageRename(UPackage* Package);
 
-	void SuspendWorkers();
-	void ResumeWorkers();
-
 	void FinalizeInitialLoad();
 
 	void RemoveUnreachableObjects(const FUnreachablePublicExports& PublicExports, const FUnreachablePackages& Packages);
@@ -3353,72 +3283,6 @@ FScopedAsyncPackageEvent2::~FScopedAsyncPackageEvent2()
 #endif
 }
 
-void FAsyncLoadingThreadWorker::StartThread()
-{
-	LLM_SCOPE(ELLMTag::AsyncLoading);
-	UE::Trace::ThreadGroupBegin(TEXT("AsyncLoading"));
-	Thread = FRunnableThread::Create(this, TEXT("FAsyncLoadingThreadWorker"), 0, TPri_Normal);
-	ThreadId = Thread->GetThreadID();
-	UE::Trace::ThreadGroupEnd();
-}
-
-uint32 FAsyncLoadingThreadWorker::Run()
-{
-	LLM_SCOPE(ELLMTag::AsyncLoading);
-
-	FPlatformProcess::SetThreadAffinityMask(FPlatformAffinity::GetAsyncLoadingThreadMask());
-	FMemory::SetupTLSCachesOnCurrentThread();
-
-	FAsyncLoadingThreadState2::Create(GraphAllocator, IoDispatcher);
-
-	FZenaphoreWaiter Waiter(Zenaphore, TEXT("WaitForEvents"));
-
-	FAsyncLoadingThreadState2& ThreadState = *FAsyncLoadingThreadState2::Get();
-
-	bool bSuspended = false;
-	while (!bStopRequested)
-	{
-		if (bSuspended)
-		{
-			if (!bSuspendRequested.Load(EMemoryOrder::SequentiallyConsistent))
-			{
-				bSuspended = false;
-			}
-			else
-			{
-				FPlatformProcess::Sleep(0.001f);
-			}
-		}
-		else
-		{
-			bool bDidSomething = false;
-			{
-				FGCScopeGuard GCGuard;
-				TRACE_CPUPROFILER_EVENT_SCOPE(AsyncLoadingTime);
-				++ActiveWorkersCount;
-				do
-				{
-					bDidSomething = EventQueue.PopAndExecute(ThreadState);
-					
-					if (bSuspendRequested.Load(EMemoryOrder::Relaxed))
-					{
-						bSuspended = true;
-						bDidSomething = true;
-						break;
-					}
-				} while (bDidSomething);
-				--ActiveWorkersCount;
-			}
-			if (!bDidSomething)
-			{
-				ThreadState.ProcessDeferredFrees();
-				Waiter.Wait();
-			}
-		}
-	}
-	return 0;
-}
-
 FUObjectSerializeContext* FAsyncPackage2::GetSerializeContext()
 {
 	return FUObjectThreadContext::Get().GetSerializeContext();
@@ -5554,38 +5418,6 @@ bool FAsyncLoadingThread2::Init()
 	return true;
 }
 
-void FAsyncLoadingThread2::SuspendWorkers()
-{
-	if (bWorkersSuspended)
-	{
-		return;
-	}
-	TRACE_CPUPROFILER_EVENT_SCOPE(SuspendWorkers);
-	for (FAsyncLoadingThreadWorker& Worker : Workers)
-	{
-		Worker.SuspendThread();
-	}
-	while (ActiveWorkersCount > 0)
-	{
-		FPlatformProcess::SleepNoStats(0);
-	}
-	bWorkersSuspended = true;
-}
-
-void FAsyncLoadingThread2::ResumeWorkers()
-{
-	if (!bWorkersSuspended)
-	{
-		return;
-	}
-	TRACE_CPUPROFILER_EVENT_SCOPE(ResumeWorkers);
-	for (FAsyncLoadingThreadWorker& Worker : Workers)
-	{
-		Worker.ResumeThread();
-	}
-	bWorkersSuspended = false;
-}
-
 void FAsyncLoadingThread2::FinalizeInitialLoad()
 {
 	TRACE_CPUPROFILER_EVENT_SCOPE(FinalizeInitialLoad);
@@ -5624,7 +5456,6 @@ uint32 FAsyncLoadingThread2::Run()
 			{
 				ThreadResumedEvent->Trigger();
 				bIsSuspended = false;
-				ResumeWorkers();
 			}
 			else
 			{
@@ -5689,7 +5520,6 @@ uint32 FAsyncLoadingThread2::Run()
 
 					if (bShouldSuspend || bSuspendRequested.Load(EMemoryOrder::Relaxed) || IsGarbageCollectionWaiting())
 					{
-						SuspendWorkers();
 						ThreadSuspendedEvent->Trigger();
 						bIsSuspended = true;
 						bDidSomething = true;
@@ -5774,10 +5604,6 @@ EAsyncPackageState::Type FAsyncLoadingThread2::TickAsyncThreadFromGameThread(FAs
 
 void FAsyncLoadingThread2::Stop()
 {
-	for (FAsyncLoadingThreadWorker& Worker : Workers)
-	{
-		Worker.StopThread();
-	}
 	bSuspendRequested = true;
 	bStopRequested = true;
 	AltZenaphore.NotifyAll();
