Fixed issue with LiveCoding reload where a fatal error would be generated even if a class had not changed.

Improved logging to also show the object type being reloaded.

diff --git a/Engine/Source/Runtime/CoreUObject/Private/UObject/DeferredRegistry.h b/Engine/Source/Runtime/CoreUObject/Private/UObject/DeferredRegistry.h
--- a/Engine/Source/Runtime/CoreUObject/Private/UObject/DeferredRegistry.h
+++ b/Engine/Source/Runtime/CoreUObject/Private/UObject/DeferredRegistry.h
@@ -44,6 +44,13 @@ public:
 	using TType = typename T::TType;
 	using TVersion = typename T::TVersion;
 
+	enum class AddResult
+	{
+		New,
+		ExistingNoChange,
+		ExistingChanged,
+	};
+
 private:
 
 	using FPackageAndNameKey = TTuple<FName, FName>;
@@ -81,7 +88,7 @@ public:
 	/// <param name="InInfo">Persistent information about the object</param>
 	/// <param name="InVersion">Version information for this incarnation of the object</param>
 	/// <param name="DeprevatedFieldInfo">Deprecated factory object for the inner register for UClass.  Either InInnerRegister or this value must be specified for UClass</param>
-	bool AddRegistration(TType* (*InOuterRegister)(), TType* (*InInnerRegister)(), const TCHAR* InPackageName, const TCHAR* InName, TInfo& InInfo, const TVersion& InVersion, FFieldCompiledInInfo* DeprecatedFieldInfo)
+	AddResult AddRegistration(TType* (*InOuterRegister)(), TType* (*InInnerRegister)(), const TCHAR* InPackageName, const TCHAR* InName, TInfo& InInfo, const TVersion& InVersion, FFieldCompiledInInfo* DeprecatedFieldInfo)
 	{
 #if WITH_RELOAD
 		const FPackageAndNameKey Key = FPackageAndNameKey{ InPackageName, InName };
@@ -117,10 +124,10 @@ public:
 		{
 			Registrations.Add(FRegistrant{ InOuterRegister, InInnerRegister, InPackageName, &InInfo, OldSingleton, DeprecatedFieldInfo, bHasChanged });
 		}
-		return ExistingInfo != nullptr;
+		return ExistingInfo == nullptr ? AddResult::New : (bHasChanged ? AddResult::ExistingChanged : AddResult::ExistingNoChange);
 #else
 		Registrations.Add(FRegistrant{ InOuterRegister, InInnerRegister, InPackageName, &InInfo, DeprecatedFieldInfo });
-		return false;
+		return AddResult::New;
 #endif
 	}
 
@@ -253,7 +260,7 @@ public:
 			{
 				if (Registrant.bHasChanged && Registrant.OldSingleton != nullptr)
 				{
-					UE_LOG(LogClass, Log, TEXT("%s Reload."), *Registrant.OldSingleton->GetName());
+					UE_LOG(LogClass, Log, TEXT("%s %s Reload."), GetRegistryName(), *Registrant.OldSingleton->GetName());
 
 					// Reset the cached class construct info.
 					Registrant.Info->InnerSingleton = nullptr;
@@ -329,6 +336,10 @@ public:
 		return Registry;
 	}
 
+#if WITH_RELOAD
+	const TCHAR* GetRegistryName();
+#endif
+
 private:
 	TArray<FRegistrant> Registrations;
 	TMap<FPackageAndNameKey, TUniquePtr<TInfo>> CompatInfoMap; // Can be removed after deprecated API removed
@@ -360,6 +371,13 @@ inline void TDeferredRegistry<FClassRegistrationInfo>::UpdateSingletons(FClassRe
 	NewInfo.OuterSingleton = nullptr;
 }
 
+#if WITH_RELOAD
+template <> inline const TCHAR* TDeferredRegistry<FClassRegistrationInfo>::GetRegistryName() { return TEXT("UClass"); }
+template <> inline const TCHAR* TDeferredRegistry<FEnumRegistrationInfo>::GetRegistryName() { return TEXT("UEnum"); }
+template <> inline const TCHAR* TDeferredRegistry<FStructRegistrationInfo>::GetRegistryName() { return TEXT("UScriptStruct"); }
+template <> inline const TCHAR* TDeferredRegistry<FPackageRegistrationInfo>::GetRegistryName() { return TEXT("UPackage"); }
+#endif
+
 using FClassDeferredRegistry = TDeferredRegistry<FClassRegistrationInfo>;
 using FEnumDeferredRegistry = TDeferredRegistry<FEnumRegistrationInfo>;
 using FStructDeferredRegistry = TDeferredRegistry<FStructRegistrationInfo>;
diff --git a/Engine/Source/Runtime/CoreUObject/Private/UObject/UObjectBase.cpp b/Engine/Source/Runtime/CoreUObject/Private/UObject/UObjectBase.cpp
--- a/Engine/Source/Runtime/CoreUObject/Private/UObject/UObjectBase.cpp
+++ b/Engine/Source/Runtime/CoreUObject/Private/UObject/UObjectBase.cpp
@@ -652,12 +652,12 @@ void RegisterCompiledInInfo(class UClass* (*InOuterRegister)(), class UClass* (*
 {
 	check(InOuterRegister);
 	check(InInnerRegister);
-	bool bExisting = FClassDeferredRegistry::Get().AddRegistration(InOuterRegister, InInnerRegister, InPackageName, InName, InInfo, InVersionInfo, nullptr);
+	FClassDeferredRegistry::AddResult result = FClassDeferredRegistry::Get().AddRegistration(InOuterRegister, InInnerRegister, InPackageName, InName, InInfo, InVersionInfo, nullptr);
 #if WITH_RELOAD
-	if (bExisting && !IsReloadActive())
+	if (result == FClassDeferredRegistry::AddResult::ExistingChanged && !IsReloadActive())
 	{
 		// Class exists, this can only happen during hot-reload or live coding
-		UE_LOG(LogUObjectBase, Fatal, TEXT("Trying to recreate class '%s' outside of hot reload and live coding!"), InName);
+		UE_LOG(LogUObjectBase, Fatal, TEXT("Trying to recreate changed class '%s' outside of hot reload and live coding!"), InName);
 	}
 #endif
 	FString NoPrefix(UObjectBase::RemoveClassPrefix(InName));
