Removing uses of master/slave found in CoreUObject with base folder Runtime and User Defined Enums

diff --git a/Engine/Source/Runtime/CoreUObject/Private/Blueprint/BlueprintSupport.cpp b/Engine/Source/Runtime/CoreUObject/Private/Blueprint/BlueprintSupport.cpp
--- a/Engine/Source/Runtime/CoreUObject/Private/Blueprint/BlueprintSupport.cpp
+++ b/Engine/Source/Runtime/CoreUObject/Private/Blueprint/BlueprintSupport.cpp
@@ -447,22 +447,22 @@ void FBlueprintSupport::ValidateNoExternalRefsToSkeletons()
  ******************************************************************************/
 
 #if WITH_EDITOR
-UClass* FScopedClassDependencyGather::BatchMasterClass = NULL;
+UClass* FScopedClassDependencyGather::BatchAuthorityClass = NULL;
 TArray<UClass*> FScopedClassDependencyGather::BatchClassDependencies;
 
 FScopedClassDependencyGather::FScopedClassDependencyGather(UClass* ClassToGather, FUObjectSerializeContext* InLoadContext)
-	: bMasterClass(false)
+	: bAuthoritativeClass(false)
 	, LoadContext(InLoadContext)
 {
 	// Do NOT track duplication dependencies, as these are intermediate products that we don't care about
 	if( !GIsDuplicatingClassForReinstancing )
 	{
-		if( BatchMasterClass == NULL )
+		if( BatchAuthorityClass == NULL )
 		{
-			// If there is no current dependency master, register this class as the master, and reset the array
-			BatchMasterClass = ClassToGather;
+			// If there is no current dependency authority, register this class as the authority, and reset the array
+			BatchAuthorityClass = ClassToGather;
 			BatchClassDependencies.Empty();
-			bMasterClass = true;
+			bAuthoritativeClass = true;
 		}
 		else
 		{
@@ -476,11 +476,11 @@ FScopedClassDependencyGather::~FScopedClassDependencyGather()
 {
 	// If this gatherer was the initial gatherer for the current scope, process 
 	// dependencies (unless compiling on load is explicitly disabled)
-	if( bMasterClass )
+	if( bAuthoritativeClass )
 	{
 		auto DependencyIter = BatchClassDependencies.CreateIterator();
 		// implemented as a lambda, to prevent duplicated code between 
-		// BatchMasterClass and BatchClassDependencies entries
+		// BatchAuthorityClass and BatchClassDependencies entries
 		auto RecompileClassLambda = [&DependencyIter](UClass* Class, FUObjectSerializeContext* InLoadContext)
 		{
 			Class->ConditionalRecompileClass(InLoadContext);
@@ -510,9 +510,9 @@ FScopedClassDependencyGather::~FScopedClassDependencyGather()
 			}
 		};
 
-		BatchMasterClass->ConditionalRecompileClass(LoadContext);
+		BatchAuthorityClass->ConditionalRecompileClass(LoadContext);
 
-		BatchMasterClass = NULL;
+		BatchAuthorityClass = NULL;
 	}
 }
 
diff --git a/Engine/Source/Runtime/CoreUObject/Private/UObject/DeferredRegistry.cpp b/Engine/Source/Runtime/CoreUObject/Private/UObject/DeferredRegistry.cpp
--- a/Engine/Source/Runtime/CoreUObject/Private/UObject/DeferredRegistry.cpp
+++ b/Engine/Source/Runtime/CoreUObject/Private/UObject/DeferredRegistry.cpp
@@ -23,7 +23,7 @@ void ReloadProcessObject(UEnum* Enum, const TCHAR* RenamePrefix)
 	const FName OldRename = MakeUniqueObjectName(GetTransientPackage(), Enum->GetClass(), *FString::Printf(TEXT("%s_%s"), RenamePrefix, *Enum->GetName()));
 	Enum->Rename(*OldRename.ToString(), GetTransientPackage());
 
-	Enum->RemoveNamesFromMasterList();
+	Enum->RemoveNamesFromPrimaryList();
 }
 
 void ReloadProcessObject(UClass* Class, const TCHAR* RenamePrefix)
@@ -47,7 +47,7 @@ void ReloadProcessObject(UClass* Class, const TCHAR* RenamePrefix)
 	{
 		if (auto Enum = dynamic_cast<UEnum*>(ClassSubobject))
 		{
-			Enum->RemoveNamesFromMasterList();
+			Enum->RemoveNamesFromPrimaryList();
 		}
 	}
 
diff --git a/Engine/Source/Runtime/CoreUObject/Private/UObject/Enum.cpp b/Engine/Source/Runtime/CoreUObject/Private/UObject/Enum.cpp
--- a/Engine/Source/Runtime/CoreUObject/Private/UObject/Enum.cpp
+++ b/Engine/Source/Runtime/CoreUObject/Private/UObject/Enum.cpp
@@ -91,13 +91,13 @@ void UEnum::Serialize( FArchive& Ar )
 			// Rename enum names to reflect new class.
 			RenameNamesAfterDuplication();
 		}
-		AddNamesToMasterList();
+		AddNamesToPrimaryList();
 	}
 }
 
 void UEnum::BeginDestroy()
 {
-	RemoveNamesFromMasterList();
+	RemoveNamesFromPrimaryList();
 
 	Super::BeginDestroy();
 }
@@ -276,6 +276,22 @@ void UEnum::AddNamesToMasterList()
 	}
 }
 
+void UEnum::AddNamesToPrimaryList()
+{
+	for (TPair<FName, int64> Kvp : Names)
+	{
+		UEnum* Enum = AllEnumNames.FindRef(Kvp.Key);
+		if (Enum == nullptr || Enum->HasAnyFlags(RF_NewerVersionExists))
+		{
+			AllEnumNames.Add(Kvp.Key, this);
+		}
+		else if (Enum != this && Enum->GetOutermost() != GetTransientPackage())
+		{
+			UE_LOG(LogEnum, Warning, TEXT("Enum name collision: '%s' is in both '%s' and '%s'"), *Kvp.Key.ToString(), *GetPathName(), *Enum->GetPathName());
+		}
+	}
+}
+
 void UEnum::RemoveNamesFromMasterList()
 {
 	FWriteScopeLock ScopeLock(AllEnumNamesLock);
@@ -289,6 +305,18 @@ void UEnum::RemoveNamesFromMasterList()
 	}
 }
 
+void UEnum::RemoveNamesFromPrimaryList()
+{
+	for (TPair<FName, int64> Kvp : Names)
+	{
+		UEnum* Enum = AllEnumNames.FindRef(Kvp.Key);
+		if (Enum == this)
+		{
+			AllEnumNames.Remove(Kvp.Key);
+		}
+	}
+}
+
 FString UEnum::GenerateEnumPrefix() const
 {
 	FString Prefix;
@@ -599,7 +627,7 @@ bool UEnum::SetEnums(TArray<TPair<FName, int64>>& InNames, UEnum::ECppForm InCpp
 {
 	if (Names.Num() > 0)
 	{
-		RemoveNamesFromMasterList();
+		RemoveNamesFromPrimaryList();
 	}
 	Names     = InNames;
 	CppForm   = InCppForm;
@@ -619,7 +647,7 @@ bool UEnum::SetEnums(TArray<TPair<FName, int64>>& InNames, UEnum::ECppForm InCpp
 			Names.Emplace(MaxEnumItem, GetMaxEnumValue() + 1);
 		}
 	}
-	AddNamesToMasterList();
+	AddNamesToPrimaryList();
 
 	return true;
 }
diff --git a/Engine/Source/Runtime/CoreUObject/Private/UObject/FindStronglyConnected.cpp b/Engine/Source/Runtime/CoreUObject/Private/UObject/FindStronglyConnected.cpp
--- a/Engine/Source/Runtime/CoreUObject/Private/UObject/FindStronglyConnected.cpp
+++ b/Engine/Source/Runtime/CoreUObject/Private/UObject/FindStronglyConnected.cpp
@@ -84,7 +84,7 @@ void FFindStronglyConnected::FindAllCycles()
 	UE_LOG(LogObj, Log, TEXT("Finding simple cycles"));
 	Stack.Empty();
 	NodeIndex.Empty();
-	MasterIndex = 1;
+	PrimaryIndex = 1;
 
 	for (int32 Index = 0; Index < Components.Num(); Index++)
 	{
@@ -151,10 +151,10 @@ void FFindStronglyConnected::StrongConnect( UObject* Node )
 FFindStronglyConnected::NodeInfo* FFindStronglyConnected::StrongConnectInner( UObject* Node )
 {
 	NodeInfo NewNode;
-	NewNode.IndexValue = MasterIndex;
-	NewNode.LowIndex = MasterIndex;
+	NewNode.IndexValue = PrimaryIndex;
+	NewNode.LowIndex = PrimaryIndex;
 	NewNode.InStack = true;
-	MasterIndex++;
+	PrimaryIndex++;
 	Stack.Push(Node);
 	NodeInfo* CurrentIndex = &NodeIndex.Add(Node, NewNode);
 
diff --git a/Engine/Source/Runtime/CoreUObject/Private/UObject/FindStronglyConnected.h b/Engine/Source/Runtime/CoreUObject/Private/UObject/FindStronglyConnected.h
--- a/Engine/Source/Runtime/CoreUObject/Private/UObject/FindStronglyConnected.h
+++ b/Engine/Source/Runtime/CoreUObject/Private/UObject/FindStronglyConnected.h
@@ -20,14 +20,14 @@ struct FFindStronglyConnected
 	TSet<UObject*> PermanentObjects;
 	TArray<UObject*> TempObjects;
 	TMap<UObject*, NodeInfo> NodeIndex;
-	int32 MasterIndex;
+	int32 PrimaryIndex;
 	TArray<UObject *> Stack;
 
 	TArray<TArray<UObject *> > Components;
 	TArray<TArray<UObject *> > SimpleCycles;
 
 	FFindStronglyConnected()
-		: MasterIndex(1)
+		: PrimaryIndex(1)
 	{
 	}
 
diff --git a/Engine/Source/Runtime/CoreUObject/Private/UObject/UObjectArray.cpp b/Engine/Source/Runtime/CoreUObject/Private/UObject/UObjectArray.cpp
--- a/Engine/Source/Runtime/CoreUObject/Private/UObject/UObjectArray.cpp
+++ b/Engine/Source/Runtime/CoreUObject/Private/UObject/UObjectArray.cpp
@@ -82,7 +82,7 @@ FUObjectArray::FUObjectArray()
 , ObjLastNonGCIndex(INDEX_NONE)
 , MaxObjectsNotConsideredByGC(0)
 , OpenForDisregardForGC(!HACK_HEADER_GENERATOR)
-, MasterSerialNumber(START_SERIAL_NUMBER)
+, PrimarySerialNumber(START_SERIAL_NUMBER)
 {
 	GCoreObjectArrayForDebugVisualizers = &GUObjectArray.ObjObjects;
 }
@@ -398,7 +398,7 @@ int32 FUObjectArray::AllocateSerialNumber(int32 Index)
 	int32 SerialNumber = *SerialNumberPtr;
 	if (!SerialNumber)
 	{
-		SerialNumber = MasterSerialNumber.Increment();
+		SerialNumber = PrimarySerialNumber.Increment();
 		UE_CLOG(SerialNumber <= START_SERIAL_NUMBER, LogUObjectArray, Fatal, TEXT("UObject serial numbers overflowed (trying to allocate serial number %d)."), SerialNumber);
 		int32 ValueWas = FPlatformAtomics::InterlockedCompareExchange((int32*)SerialNumberPtr, SerialNumber, 0);
 		if (ValueWas != 0)
diff --git a/Engine/Source/Runtime/CoreUObject/Public/Blueprint/BlueprintSupport.h b/Engine/Source/Runtime/CoreUObject/Public/Blueprint/BlueprintSupport.h
--- a/Engine/Source/Runtime/CoreUObject/Public/Blueprint/BlueprintSupport.h
+++ b/Engine/Source/Runtime/CoreUObject/Public/Blueprint/BlueprintSupport.h
@@ -160,7 +160,7 @@ private:
 /**
  * This is a helper struct that allows us to gather all previously unloaded class dependencies of a UClass
  * The first time we create a new UClass object in FLinkerLoad::CreateExport(), we register it as a dependency
- * master.  Any subsequent UClasses that are created for the first time during the preload of that class are
+ * authority. Any subsequent UClasses that are created for the first time during the preload of that class are
  * added to the list as potential cyclic referencers.  We then step over the list at the end of the load, and
  * recompile any classes that may depend on each other a second time to ensure that that functions and properties
  * are properly resolved
@@ -181,11 +181,11 @@ public:
 	static TArray<UClass*> const& GetCachedDependencies();
 
 private:
-	/** Whether or not this dependency gather is the dependency master, and thus should process all dependencies in the destructor */
-	bool bMasterClass;	
+	/** Whether or not this dependency gather is the dependency authoritative class, and thus should process all dependencies in the destructor */
+	bool bAuthoritativeClass;
 
 	/** The current class that is gathering potential dependencies in this scope */
-	static UClass* BatchMasterClass;
+	static UClass* BatchAuthorityClass;
 
 	/** List of dependencies (i.e. UClasses that have been newly instantiated) in the scope of this dependency gather */
 	static TArray<UClass*> BatchClassDependencies;
diff --git a/Engine/Source/Runtime/CoreUObject/Public/UObject/Class.h b/Engine/Source/Runtime/CoreUObject/Public/UObject/Class.h
--- a/Engine/Source/Runtime/CoreUObject/Public/UObject/Class.h
+++ b/Engine/Source/Runtime/CoreUObject/Public/UObject/Class.h
@@ -1999,8 +1999,12 @@ public:
 	bool IsValidEnumName(FName InName) const;
 
 	/** Removes the Names in this enum from the master AllEnumNames list */
+	UE_DEPRECATED(5.1, "RemoveNamesFromMasterList is deprecated, please use RemoveNamesFromPrimaryList instead.")
 	void RemoveNamesFromMasterList();
 
+	/** Removes the Names in this enum from the primary AllEnumNames list */
+	void RemoveNamesFromPrimaryList();
+
 	/** Try to update an out-of-date enum index after an enum changes at runtime */
 	virtual int64 ResolveEnumerator(FArchive& Ar, int64 EnumeratorIndex) const;
 
@@ -2338,8 +2342,12 @@ protected:
 	static TMap<FName, UEnum*> AllEnumNames;
 
 	/** adds the Names in this enum to the master AllEnumNames list */
+	UE_DEPRECATED(5.1, "AddNamesToMasterList is deprecated, please use AddNamesToPrimaryList instead.")
 	void AddNamesToMasterList();
 
+	/** adds the Names in this enum to the primary AllEnumNames list */
+	void AddNamesToPrimaryList();
+
 private:
 	FORCEINLINE static FString GetValueAsString_Internal( const TCHAR* EnumPath, const int64 EnumeratorValue)
 	{
diff --git a/Engine/Source/Runtime/CoreUObject/Public/UObject/UObjectAnnotation.h b/Engine/Source/Runtime/CoreUObject/Public/UObject/UObjectAnnotation.h
--- a/Engine/Source/Runtime/CoreUObject/Public/UObject/UObjectAnnotation.h
+++ b/Engine/Source/Runtime/CoreUObject/Public/UObject/UObjectAnnotation.h
@@ -522,7 +522,7 @@ class FUObjectAnnotationChunked : public FUObjectArray::FUObjectDeleteListener
 	};
 
 
-	/** Master table to chunks of pointers **/
+	/** Primary table to chunks of pointers **/
 	TArray<TAnnotationChunk> Chunks;
 	/** Number of elements we currently have **/
 	int32 NumAnnotations;
diff --git a/Engine/Source/Runtime/CoreUObject/Public/UObject/UObjectArray.h b/Engine/Source/Runtime/CoreUObject/Public/UObject/UObjectArray.h
--- a/Engine/Source/Runtime/CoreUObject/Public/UObject/UObjectArray.h
+++ b/Engine/Source/Runtime/CoreUObject/Public/UObject/UObjectArray.h
@@ -230,7 +230,7 @@ struct FUObjectItem
 */
 class FFixedUObjectArray
 {
-	/** Static master table to chunks of pointers **/
+	/** Static primary table to chunks of pointers **/
 	FUObjectItem* Objects;
 	/** Number of elements we currently have **/
 	int32 MaxElements;
@@ -369,7 +369,7 @@ class FChunkedFixedUObjectArray
 		NumElementsPerChunk = 64 * 1024,
 	};
 
-	/** Master table to chunks of pointers **/
+	/** Primary table to chunks of pointers **/
 	FUObjectItem** Objects;
 	/** If requested, a contiguous memory where all objects are allocated **/
 	FUObjectItem* PreAllocatedObjects;
@@ -1103,8 +1103,8 @@ private:
 	FCriticalSection UObjectDeleteListenersCritical;
 #endif
 
-	/** Current master serial number **/
-	FThreadSafeCounter	MasterSerialNumber;
+	/** Current primary serial number **/
+	FThreadSafeCounter	PrimarySerialNumber;
 
 public:
 
diff --git a/Engine/Source/Runtime/Engine/Private/UserDefinedEnum.cpp b/Engine/Source/Runtime/Engine/Private/UserDefinedEnum.cpp
--- a/Engine/Source/Runtime/Engine/Private/UserDefinedEnum.cpp
+++ b/Engine/Source/Runtime/Engine/Private/UserDefinedEnum.cpp
@@ -219,7 +219,7 @@ bool UUserDefinedEnum::SetEnums(TArray<TPair<FName, int64>>& InNames, ECppForm I
 	ensure(bAddMaxKeyIfMissing);
 	if (Names.Num() > 0)
 	{
-		RemoveNamesFromMasterList();
+		RemoveNamesFromPrimaryList();
 	}
 	Names = InNames;
 	CppForm = InCppForm;
@@ -238,7 +238,7 @@ bool UUserDefinedEnum::SetEnums(TArray<TPair<FName, int64>>& InNames, ECppForm I
 		{
 			int64 MaxEnumValue = (InNames.Num() == 0)? 0 : GetMaxEnumValue() + 1;
 			Names.Emplace(MaxEnumItem, MaxEnumValue);
-			AddNamesToMasterList();
+			AddNamesToPrimaryList();
 			return true;
 		}
 	}
