Reduce verbosity of "LogStreaming: Display: NotifyRegistrationEvent:" from Display to Log to clean up stdout of various engine tools.

#rnx

diff --git a/Engine/Source/Runtime/CoreUObject/Private/Serialization/AsyncPackageLoader.cpp b/Engine/Source/Runtime/CoreUObject/Private/Serialization/AsyncPackageLoader.cpp
--- a/Engine/Source/Runtime/CoreUObject/Private/Serialization/AsyncPackageLoader.cpp
+++ b/Engine/Source/Runtime/CoreUObject/Private/Serialization/AsyncPackageLoader.cpp
@@ -48,7 +48,7 @@ public:
 
 	void Replay(IAsyncPackageLoader& PackageLoader)
 	{
-		UE_LOG(LogStreaming, Display, TEXT("NotifyRegistrationEvent: Replay %d entries"), RecordedEvents.Num());
+		UE_LOG(LogStreaming, Log, TEXT("NotifyRegistrationEvent: Replay %d entries"), RecordedEvents.Num());
 		for (const FRecordedEvent& Event : RecordedEvents)
 		{
 			PackageLoader.NotifyRegistrationEvent(*Event.PackageName, *Event.Name, Event.NotifyRegistrationType, Event.NotifyRegistrationPhase, Event.RegisterFunc, Event.bDynamic);
