Added missing includes

#rnx

diff --git a/Engine/Source/Runtime/CoreUObject/Public/Misc/TransactionObjectEvent.h b/Engine/Source/Runtime/CoreUObject/Public/Misc/TransactionObjectEvent.h
--- a/Engine/Source/Runtime/CoreUObject/Public/Misc/TransactionObjectEvent.h
+++ b/Engine/Source/Runtime/CoreUObject/Public/Misc/TransactionObjectEvent.h
@@ -4,6 +4,8 @@
 
 #include "CoreTypes.h"
 #include "Misc/Guid.h"
+#include "UObject/NameTypes.h"
+#include "Templates/SharedPointer.h"
 
 class ITransactionObjectAnnotation;
 
