AsyncLoading2: Move GC-related importstore logic from GT callback to ALT resume phase.
Make import store more reliable by not relying on Object->HasAnyFlags(RF_WasLoaded | RF_Public) filtering (We can take the hit of checking all unreachable objects on the ALT)
The logic can be forced to run on the GT again with s.RemoveUnreachableObjectsOnGT=1, which in debug and development builds enables extra verifications.

Return actual number of destroyed objects from FAsyncPurge::GetObjectsDestroyedSinceLastMarkPhase to make these logs correct:
LogGarbage: 2.118 ms for incrementally purging unreachable objects (FinishDestroyed: 16176, Destroyed: 1571 / 16176)
LogGarbage: 2.303 ms for incrementally purging unreachable objects (FinishDestroyed: 16176, Destroyed: 1971 / 16176)

#rnx

diff --git a/Engine/Source/Runtime/CoreUObject/Private/Serialization/AsyncLoading2.cpp b/Engine/Source/Runtime/CoreUObject/Private/Serialization/AsyncLoading2.cpp
--- a/Engine/Source/Runtime/CoreUObject/Private/Serialization/AsyncLoading2.cpp
+++ b/Engine/Source/Runtime/CoreUObject/Private/Serialization/AsyncLoading2.cpp
@@ -193,6 +193,10 @@ void FindAllRuntimeScriptPackages(TArray<UPackage*>& OutPackages)
 #define ALT2_VERIFY_RECURSIVE_LOADS !(WITH_IOSTORE_IN_EDITOR) && DO_CHECK
 #endif
 
+#ifndef ALT2_VERIFY_UNREACHABLE_OBJECTS
+#define ALT2_VERIFY_UNREACHABLE_OBJECTS DO_CHECK
+#endif
+
 #ifndef ALT2_LOG_VERBOSE
 #define ALT2_LOG_VERBOSE DO_CHECK
 #endif
@@ -202,11 +206,6 @@ static FString GAsyncLoading2_DebugPackageNamesString;
 static TSet<FPackageId> GAsyncLoading2_VerbosePackageIds;
 static FString GAsyncLoading2_VerbosePackageNamesString;
 static int32 GAsyncLoading2_VerboseLogFilter = 2; //None=0,Filter=1,All=2
-#if WITH_EDITOR
-static constexpr EObjectFlags GlobalImportObjectConditionalFlags = RF_Public;
-#else
-static constexpr EObjectFlags GlobalImportObjectConditionalFlags = RF_Public | RF_WasLoaded; // Assumes that we don't do any renaming of objects clearing the RF_WasLoaded flag
-#endif
 #if !UE_BUILD_SHIPPING
 static void ParsePackageNames(const FString& PackageNamesString, TSet<FPackageId>& PackageIds)
 {
@@ -265,6 +264,12 @@ if (GAsyncLoading2_DebugPackageIds.Contains((UPackage)->GetPackageId())) \
 	UE_DEBUG_BREAK(); \
 }
 
+#define UE_ASYNC_PACKAGEID_DEBUG(PackageId) \
+if (GAsyncLoading2_DebugPackageIds.Contains(PackageId)) \
+{ \
+	UE_DEBUG_BREAK(); \
+}
+
 // The ELogVerbosity::VerbosityMask is used to silence PVS,
 // using constexpr gave the same warning, and the disable comment can can't be used in a macro: //-V501 
 // warning V501: There are identical sub-expressions 'ELogVerbosity::Verbose' to the left and to the right of the '<' operator.
@@ -297,6 +302,14 @@ if ((Condition)) \
 #define UE_ASYNC_PACKAGE_CLOG_VERBOSE(Condition, Verbosity, PackageDesc, LogDesc, Format, ...)
 #endif
 
+static bool GRemoveUnreachableObjectsOnGT = false;
+static FAutoConsoleVariableRef CVarRemoveUnreachableObjectsOnGT(
+	TEXT("s.RemoveUnreachableObjectsOnGT"),
+	GRemoveUnreachableObjectsOnGT,
+	TEXT("Force running removal of unreachable objects on the Game Thread (slower, but enables extra verification in debug and development builds."),
+	ECVF_Default
+	);
+
 CSV_DECLARE_CATEGORY_MODULE_EXTERN(CORE_API, Basic);
 CSV_DECLARE_CATEGORY_MODULE_EXTERN(CORE_API, FileIO);
 
@@ -450,11 +463,17 @@ struct FAsyncPackageDesc2
 	}
 };
 
-// Note: RemoveUnreachableObjects could move from GT to ALT by removing the debug raw pointers here,
-// the tradeoff would be increased complexity and more restricted debug and log possibilities.
-using FUnreachablePublicExport = TPair<int32, UObject*>;
-using FUnreachablePackages = TArray<UPackage*>;
-using FUnreachablePublicExports = TArray<FUnreachablePublicExport>;
+struct FUnreachableObject
+{
+	FPackageId PackageId;
+	int32 ObjectIndex = -1;
+	FName ObjectName;
+#if ALT2_VERIFY_UNREACHABLE_OBJECTS
+	UObject* DebugObject = nullptr;
+#endif
+};
+
+using FUnreachableObjects = TArray<FUnreachableObject>;
 
 class FLoadedPackageRef
 {
@@ -830,7 +849,7 @@ public:
 	inline FLoadedPackageRef& FindPackageRefChecked(FPackageId PackageId)
 	{
 		FLoadedPackageRef* PackageRef = FindPackageRef(PackageId);
-		UE_CLOG(!PackageRef, LogStreaming, Fatal, TEXT("FindPackageRefChecked: Package with id %d has been deleted"), PackageId.ValueForDebugging());
+		UE_CLOG(!PackageRef, LogStreaming, Fatal, TEXT("FindPackageRefChecked: Package with id 0x%llX has been deleted"), PackageId.ValueForDebugging());
 		return *PackageRef;
 	}
 
@@ -900,34 +919,36 @@ public:
 	}
 #endif
 
-	void RemovePackages(const FUnreachablePackages& UnreachablePackages)
+	void RemovePackages(const FUnreachableObjects& ObjectsToRemove)
 	{
-		check(FGCCSyncObject::Get().IsGCLocked() || !IsAsyncLoading());
-
-		const int32 PackageCount = UnreachablePackages.Num();
-		for (UPackage* Package : UnreachablePackages)
+		TRACE_CPUPROFILER_EVENT_SCOPE(RemovePackages);
+		for (const FUnreachableObject& Item : ObjectsToRemove)
 		{
-			UE_ASYNC_UPACKAGE_DEBUG(Package);
-			if (Package->CanBeImported())
+			const FPackageId& PackageId = Item.PackageId;
+			if (PackageId.IsValid())
 			{
+				UE_ASYNC_PACKAGEID_DEBUG(PackageId);
+
+#if ALT2_VERIFY_UNREACHABLE_OBJECTS
 				FLoadedPackageRef PackageRef;
-				bool bRemoved = Packages.RemoveAndCopyValue(Package->GetPackageId(), PackageRef);
+				bool bRemoved = Packages.RemoveAndCopyValue(PackageId, PackageRef);
 				if (bRemoved)
 				{
 					if (PackageRef.RefCount > 0)
 					{
+						FString PackageName = Item.ObjectName.ToString();
 						UE_LOG(LogStreaming, Error,
-							TEXT("RemovePackage: %s %s (0x%llX) - with (ObjectFlags=%x, InternalObjectFlags=%x) - ")
+							TEXT("RemovePackage: %s (0x%llX) - ")
 							TEXT("Package destroyed while still being referenced, RefCount %d > 0."),
-							*Package->GetName(),
-							*Package->GetLoadedPath().GetDebugName(), Package->GetPackageId().Value(),
-							Package->GetFlags(), Package->GetInternalFlags(), PackageRef.RefCount);
-						checkf(false, TEXT("Package %s destroyed with RefCount"), *Package->GetName());
+							*PackageName,
+							PackageId.Value(),
+							PackageRef.RefCount);
+						checkf(false, TEXT("Package %s destroyed with RefCount"), *PackageName);
 					}
-#if DO_CHECK
 					PackageRef.VerifyAllPublicExportsRemoved();
-#endif
 				}
+#endif
+				Packages.Remove(PackageId);
 			}
 		}
 	}
@@ -962,30 +983,27 @@ public:
 		return ObjectIndexToPublicExport.Num();
 	}
 
-	void RemovePublicExports(const FUnreachablePublicExports& PublicExports)
+	void RemovePublicExports(const FUnreachableObjects& ObjectsToRemove)
 	{
+		TRACE_CPUPROFILER_EVENT_SCOPE(RemovePublicExports);
+
 		TArray<FPublicExportKey> PublicExportKeys;
-		PublicExportKeys.Reserve(PublicExports.Num());
+		PublicExportKeys.Reserve(ObjectsToRemove.Num());
 
-		for (const FUnreachablePublicExport& Item : PublicExports)
+		for (const FUnreachableObject& Item : ObjectsToRemove)
 		{
-			int32 ObjectIndex = Item.Key;
+			int32 ObjectIndex = Item.ObjectIndex;
+			check(ObjectIndex >= 0)
+
 			FPublicExportKey PublicExportKey;
 			if (ObjectIndexToPublicExport.RemoveAndCopyValue(ObjectIndex, PublicExportKey))
 			{
 				PublicExportKeys.Emplace(PublicExportKey);
-#if DO_CHECK
-				{
-					UObject* GCObject = Item.Value;
 
-					checkf(GCObject->HasAllFlags(GlobalImportObjectConditionalFlags),
-						TEXT("The serialized GC Object '%s' with (ObjectFlags=%x, InternalObjectFlags=%x)and id 0x%llX:0x%llX is currently missing RF_WasLoaded or RF_Public. ")
-						TEXT("The flags must have been altered incorrectly by a higher level system after the object was loaded. ")
-						TEXT("This may cause memory corruption."),
-						*GCObject->GetFullName(),
-						GCObject->GetFlags(),
-						GCObject->GetInternalFlags(),
-						PublicExportKey.GetPackageId().Value(), PublicExportKey.GetExportHash());
+#if ALT2_VERIFY_UNREACHABLE_OBJECTS
+				if (GRemoveUnreachableObjectsOnGT)
+				{
+					UObject* GCObject = Item.DebugObject;
 
 					UObject* ExistingObject = FindPublicExportObjectUnchecked(PublicExportKey);
 					checkf(ExistingObject,
@@ -1021,10 +1039,8 @@ public:
 				LastPackageId = PackageId;
 				PackageRef = LoadedPackageStore.FindPackageRef(PackageId);
 			}
-			check(PackageRef)
-			{
-				PackageRef->RemovePublicExport(PublicExportKey.GetExportHash());
-			}
+			check(PackageRef);
+			PackageRef->RemovePublicExport(PublicExportKey.GetExportHash());
 		}
 	}
 
@@ -2278,6 +2294,10 @@ private:
 	TMap<UClass*, TArray<FEventLoadNode2*>> PendingCDOs;
 	TArray<UClass*> PendingCDOsRecursiveStack;
 
+	/** [ASYNC/GAME THREAD] Unreachable objects from last NotifyUnreachableObjects callback from GC.
+	 * Protected by FGCScopeGuard */
+	FUnreachableObjects UnreachableObjects;
+
 	uint32 ConditionalBeginPostLoadTick = 0;
 	uint32 ConditionalFinishLoadingTick = 0;
 
@@ -2521,7 +2541,7 @@ private:
 
 	void FinalizeInitialLoad();
 
-	void RemoveUnreachableObjects(const FUnreachablePublicExports& PublicExports, const FUnreachablePackages& Packages);
+	void RemoveUnreachableObjects(FUnreachableObjects& ObjectsToRemove);
 
 	void ProcessPendingCDOs()
 	{
@@ -3475,7 +3495,7 @@ void FAsyncPackage2::ImportPackagesRecursiveInner(FIoBatch& IoBatch, FPackageSto
 
 					ForEachObjectWithOuter(UncookedPackage, [this, ImportedPackageId](UObject* Object)
 					{
-						if (Object->HasAllFlags(GlobalImportObjectConditionalFlags))
+						if (Object->HasAllFlags(RF_Public))
 						{
 							checkf(!Object->HasAnyInternalFlags(EInternalObjectFlags::LoaderImport), TEXT("%s"), *Object->GetFullName());
 							Object->SetInternalFlags(EInternalObjectFlags::LoaderImport);
@@ -4299,14 +4319,13 @@ void FAsyncPackage2::EventDrivenCreateExport(const FAsyncPackageHeaderData& Head
 	check(Object);
 	EInternalObjectFlags FlagsToSet = EInternalObjectFlags::Async;
 	
-	if (Desc.bCanBeImported && Export.PublicExportHash &&
-		(Object->HasAnyFlags(RF_Public) || WITH_IOSTORE_IN_EDITOR))
+	if (Desc.bCanBeImported && Export.PublicExportHash)
 	{
 		FlagsToSet |= EInternalObjectFlags::LoaderImport;
 		ImportStore.StoreGlobalObject(Desc.UPackageId, Export.PublicExportHash, Object);
 
-		UE_ASYNC_PACKAGE_LOG_VERBOSE(VeryVerbose, Desc, TEXT("CreateExport"),
-			TEXT("Created public export %s. Tracked as 0x%llX:0x%llX"), *Object->GetPathName(), Desc.UPackageId.Value(), Export.PublicExportHash);
+		UE_ASYNC_PACKAGE_LOG_VERBOSE(VeryVerbose, Desc, TEXT("CreateExport"), TEXT("Created %s export %s. Tracked as 0x%llX:0x%llX"),
+			Object->HasAnyFlags(RF_Public) ? TEXT("public") : TEXT("private"), *Object->GetPathName(), Desc.UPackageId.Value(), Export.PublicExportHash);
 	}
 	else
 	{
@@ -5447,132 +5466,163 @@ uint32 FAsyncLoadingThread2::Run()
 	FinalizeInitialLoad();
 
 	FZenaphoreWaiter Waiter(AltZenaphore, TEXT("WaitForEvents"));
-	bool bIsSuspended = false;
+	enum class EMainState : uint8
+	{
+		Suspended,
+		Loading,
+		Waiting,
+	};
+	EMainState PreviousState = EMainState::Loading;
+	EMainState CurrentState = EMainState::Loading;
 	while (!bStopRequested)
 	{
-		if (bIsSuspended)
+		if (CurrentState == EMainState::Suspended)
 		{
-			if (!bSuspendRequested.Load(EMemoryOrder::SequentiallyConsistent) && !IsGarbageCollectionWaiting())
-			{
-				ThreadResumedEvent->Trigger();
-				bIsSuspended = false;
-			}
-			else
+			// suspended, sleep until loading can be resumed
+			while (!bStopRequested)
 			{
+				if (!bSuspendRequested.Load(EMemoryOrder::SequentiallyConsistent) && !IsGarbageCollectionWaiting())
+				{
+					ThreadResumedEvent->Trigger();
+					CurrentState = EMainState::Loading;
+					break;
+				}
+
 				FPlatformProcess::Sleep(0.001f);
 			}
 		}
-		else
+		else if (CurrentState == EMainState::Waiting)
 		{
-			bool bDidSomething = false;
+			// no packages in flight and waiting for new load package requests,
+			// or done serializing and waiting for deferred deletes of packages being postloaded
+			Waiter.Wait();
+			CurrentState = EMainState::Loading;
+		}
+		else if (CurrentState == EMainState::Loading)
+		{
+			TRACE_CPUPROFILER_EVENT_SCOPE(AsyncLoadingTime);
+
+			bool bShouldSuspend = false;
+			bool bShouldWaitForExternalReads = false;
+			while (!bStopRequested)
 			{
-				FGCScopeGuard GCGuard;
-				TRACE_CPUPROFILER_EVENT_SCOPE(AsyncLoadingTime);
-				do
+				if (bShouldSuspend || bSuspendRequested.Load(EMemoryOrder::Relaxed) || IsGarbageCollectionWaiting())
+				{
+					TRACE_CPUPROFILER_EVENT_SCOPE(SuspendAsyncLoading);
+					ThreadSuspendedEvent->Trigger();
+					CurrentState = EMainState::Suspended;
+					break;
+				}
+
 				{
-					bDidSomething = false;
+					FGCScopeGuard GCGuard;
+					RemoveUnreachableObjects(UnreachableObjects);
+
+					if (bShouldWaitForExternalReads)
+					{
+						TRACE_CPUPROFILER_EVENT_SCOPE(WaitingForExternalReads);
+						FAsyncPackage2* Package = nullptr;
+						ExternalReadQueue.Dequeue(Package);
+						check(Package);
+						EAsyncPackageState::Type Result = Package->ProcessExternalReads(FAsyncPackage2::ExternalReadAction_Wait);
+						check(Result == EAsyncPackageState::Complete);
+						bShouldWaitForExternalReads = false;
+						continue;
+					}
+
 
 					if (QueuedPackagesCounter || !PendingPackages.IsEmpty())
 					{
 						if (CreateAsyncPackagesFromQueue(ThreadState))
 						{
-							bDidSomething = true;
+							// Fall through to FAsyncLoadEventQueue2 processing unless we need to suspend
+							if (bSuspendRequested.Load(EMemoryOrder::Relaxed) || IsGarbageCollectionWaiting())
+							{
+								bShouldSuspend = true;
+								continue;
+							}
 						}
 					}
 
-					bool bShouldSuspend = false;
-					bool bPopped = false;
-					do 
+					// do as much event queue processing as we possibly can
 					{
-						bPopped = false;
+						bool bDidSomething = false;
+						bool bPopped = false;
+						do 
 						{
-							FAsyncPackage2** Package = ExternalReadQueue.Peek();
-							if (Package != nullptr)
+							bPopped = false;
+							for (FAsyncLoadEventQueue2* Queue : AltEventQueues)
 							{
-								TRACE_CPUPROFILER_EVENT_SCOPE(PollExternalReads);
-								check(*Package);
-								EAsyncPackageState::Type Result = (*Package)->ProcessExternalReads(FAsyncPackage2::ExternalReadAction_Poll);
-								if (Result == EAsyncPackageState::Complete)
+								if (Queue->PopAndExecute(ThreadState))
 								{
-									ExternalReadQueue.Dequeue();
 									bPopped = true;
 									bDidSomething = true;
 								}
+								if (bSuspendRequested.Load(EMemoryOrder::Relaxed) || IsGarbageCollectionWaiting())
+								{
+									bShouldSuspend = true;
+									bDidSomething = true;
+									bPopped = false;
+									break;
+								}
 							}
-						}
+						} while (bPopped);
 
-						for (FAsyncLoadEventQueue2* Queue : AltEventQueues)
+						if (bDidSomething)
 						{
-							if (Queue->PopAndExecute(ThreadState))
-							{
-								bPopped = true;
-								bDidSomething = true;
-							}
+							continue;
+						}
+					}
 
-							if (bSuspendRequested.Load(EMemoryOrder::Relaxed) || IsGarbageCollectionWaiting())
+					{
+						FAsyncPackage2** Package = ExternalReadQueue.Peek();
+						if (Package != nullptr)
+						{
+							TRACE_CPUPROFILER_EVENT_SCOPE(PollExternalReads);
+							check(*Package);
+							EAsyncPackageState::Type Result = (*Package)->ProcessExternalReads(FAsyncPackage2::ExternalReadAction_Poll);
+							if (Result == EAsyncPackageState::Complete)
 							{
-								bShouldSuspend = true;
-								bPopped = false;
-								break;
+								ExternalReadQueue.Dequeue();
+								continue;
 							}
 						}
-					} while (bPopped);
+					}
 
-					if (bShouldSuspend || bSuspendRequested.Load(EMemoryOrder::Relaxed) || IsGarbageCollectionWaiting())
+					if (ThreadState.HasDeferredFrees())
 					{
-						ThreadSuspendedEvent->Trigger();
-						bIsSuspended = true;
-						bDidSomething = true;
-						break;
+						ThreadState.ProcessDeferredFrees();
+						continue;
 					}
 
-				} while (bDidSomething);
-			}
-
-			if (!bDidSomething)
-			{
-				if (ThreadState.HasDeferredFrees())
-				{
-					TRACE_CPUPROFILER_EVENT_SCOPE(AsyncLoadingTime);
-					ThreadState.ProcessDeferredFrees();
-					bDidSomething = true;
-				}
-
-				if (!DeferredDeletePackages.IsEmpty())
-				{
-					FGCScopeGuard GCGuard;
-					TRACE_CPUPROFILER_EVENT_SCOPE(AsyncLoadingTime);
-					FAsyncPackage2* Package = nullptr;
-					int32 Count = 0;
-					while (++Count <= 100 && DeferredDeletePackages.Dequeue(Package))
+					if (!DeferredDeletePackages.IsEmpty())
 					{
-						DeleteAsyncPackage(Package);
+						FAsyncPackage2* Package = nullptr;
+						int32 Count = 0;
+						while (++Count <= 100 && DeferredDeletePackages.Dequeue(Package))
+						{
+							DeleteAsyncPackage(Package);
+						}
+						continue;
 					}
-					bDidSomething = true;
-				}
-			}
+				} // release FGCScopeGuard
 
-			if (!bDidSomething)
-			{
-				FAsyncPackage2* Package = nullptr;
 				if (PendingIoRequestsCounter.Load() > 0)
 				{
-					TRACE_CPUPROFILER_EVENT_SCOPE(AsyncLoadingTime);
 					TRACE_CPUPROFILER_EVENT_SCOPE(WaitingForIo);
 					Waiter.Wait();
+					continue;
 				}
-				else if (ExternalReadQueue.Dequeue(Package))
-				{
-					TRACE_CPUPROFILER_EVENT_SCOPE(AsyncLoadingTime);
-					TRACE_CPUPROFILER_EVENT_SCOPE(WaitingForExternalReads);
 
-					EAsyncPackageState::Type Result = Package->ProcessExternalReads(FAsyncPackage2::ExternalReadAction_Wait);
-					check(Result == EAsyncPackageState::Complete);
-				}
-				else
+				if (!ExternalReadQueue.IsEmpty())
 				{
-					Waiter.Wait();
+					bShouldWaitForExternalReads = true;
+					continue;
 				}
+
+				// no async loading work left to do for now
+				CurrentState = EMainState::Waiting;
+				break;
 			}
 		}
 	}
@@ -5617,6 +5667,7 @@ void FAsyncLoadingThread2::CancelLoading()
 
 void FAsyncLoadingThread2::SuspendLoading()
 {
+	TRACE_CPUPROFILER_EVENT_SCOPE(SuspendLoading);
 	UE_CLOG(!IsInGameThread() || IsInSlateThread(), LogStreaming, Fatal, TEXT("Async loading can only be suspended from the main thread"));
 	if (!bSuspendRequested)
 	{
@@ -5632,6 +5683,7 @@ void FAsyncLoadingThread2::SuspendLoading()
 
 void FAsyncLoadingThread2::ResumeLoading()
 {
+	TRACE_CPUPROFILER_EVENT_SCOPE(ResumeLoading);
 	check(IsInGameThread() && !IsInSlateThread());
 	if (bSuspendRequested)
 	{
@@ -5660,6 +5712,8 @@ float FAsyncLoadingThread2::GetAsyncLoadPercentage(const FName& PackageName)
 #if ALT2_VERIFY_ASYNC_FLAGS
 static void VerifyLoadFlagsWhenFinishedLoading()
 {
+	TRACE_CPUPROFILER_EVENT_SCOPE(VerifyLoadFlagsWhenFinishedLoading);
+
 	const EInternalObjectFlags AsyncFlags =
 		EInternalObjectFlags::Async | EInternalObjectFlags::AsyncLoading;
 
@@ -5719,40 +5773,42 @@ static void VerifyLoadFlagsWhenFinishedLoading()
 #endif
 
 FORCENOINLINE static void FilterUnreachableObjects(
-	const TArrayView<FUObjectItem*>& UnreachableObjects,
-	FUnreachablePublicExports& PublicExports,
-	FUnreachablePackages& Packages)
+	TArrayView<FUObjectItem*> UnreachableObjectItems,
+	FUnreachableObjects& OutUnreachableObjects)
 {
 	TRACE_CPUPROFILER_EVENT_SCOPE(FilterUnreachableObjects);
 
-	PublicExports.Reserve(UnreachableObjects.Num());
-	Packages.Reserve(UnreachableObjects.Num());
+	OutUnreachableObjects.SetNum(UnreachableObjectItems.Num());
 
-	for (FUObjectItem* ObjectItem : UnreachableObjects)
+	ParallelFor(TEXT("FilterUnreachableObjects"), UnreachableObjectItems.Num(), 2048, [&UnreachableObjectItems, &OutUnreachableObjects](int32 Index)
 	{
-		UObject* Object = static_cast<UObject*>(ObjectItem->Object);
-		// Including all objects is slow,
-		// but allows RemovePublicExports to check for serialized public exports that have unintentionally lost their flags
-		if (DO_CHECK || Object->HasAllFlags(GlobalImportObjectConditionalFlags))
+		UObject* Object = static_cast<UObject*>(UnreachableObjectItems[Index]->Object);
+
+		FUnreachableObject& Item = OutUnreachableObjects[Index];
+		Item.ObjectIndex = GUObjectArray.ObjectToIndex(Object);
+		Item.ObjectName = Object->GetFName();
+#if ALT2_VERIFY_UNREACHABLE_OBJECTS
+		if (GRemoveUnreachableObjectsOnGT)
 		{
-			if (Object->GetOuter())
-			{
-				PublicExports.Emplace(GUObjectArray.ObjectToIndex(Object), Object);
-			}
-			else
+			Item.DebugObject = Object;
+		}
+#endif
+		if (!Object->GetOuter())
+		{
+			UPackage* Package = static_cast<UPackage*>(Object);
+			if (Package->bCanBeImported)
 			{
-				UPackage* Package = static_cast<UPackage*>(Object);
-				Packages.Emplace(Package);
+				Item.PackageId = Package->GetPackageId();
 			}
 		}
-	}
+	});
 }
 
 void FAsyncLoadingThread2::OnLeakedPackageRename(UPackage* Package)
 {
 	check(IsInGameThread());
 	
-	if(!FGCCSyncObject::Get().IsGCLocked())
+	if (!FGCCSyncObject::Get().IsGCLocked())
 	{
 		// Flush async loading so that we can be sure nothing is modifying the stores, and that nothing is depending on this package
 		FlushAsyncLoading(); 
@@ -5760,36 +5816,54 @@ void FAsyncLoadingThread2::OnLeakedPackageRename(UPackage* Package)
 
 	// Code such as LoadMap or LevelStreaming has renamed a loaded package which was detected as leaking so that we can load another copy of it.
 	// We should not have any loading happening at present, so we can remove these objects from our stores 
-	FUnreachablePublicExports ObjectsToRemove;
-	FUnreachablePackages PackagesToRemove;
-	PackagesToRemove.Add(Package);
-	ForEachObjectWithOuter(Package, [&ObjectsToRemove](UObject* Obj) {
-		if (DO_CHECK || Obj->HasAllFlags(GlobalImportObjectConditionalFlags))
-		{
-			ObjectsToRemove.Emplace(GUObjectArray.ObjectToIndex(Obj), Obj);
-		}
+	TArray<FUObjectItem*> LeakedObjectItems;
+	LeakedObjectItems.Emplace(GUObjectArray.ObjectToObjectItem(Package));
+	ForEachObjectWithOuter(Package, [&LeakedObjectItems](UObject* Obj) {
+		LeakedObjectItems.Emplace(GUObjectArray.ObjectToObjectItem(Obj));
 	});
 
-	RemoveUnreachableObjects(ObjectsToRemove, PackagesToRemove);
+	FUnreachableObjects LeakedUnreachableObjects;
+	FilterUnreachableObjects(LeakedObjectItems, LeakedUnreachableObjects);
+	RemoveUnreachableObjects(LeakedUnreachableObjects);
 }
 
-void FAsyncLoadingThread2::RemoveUnreachableObjects(const FUnreachablePublicExports& PublicExports, const FUnreachablePackages& Packages)
+void FAsyncLoadingThread2::RemoveUnreachableObjects(FUnreachableObjects& ObjectsToRemove)
 {
-	TRACE_CPUPROFILER_EVENT_SCOPE(RemoveUnreachableObjects);
-
-	if (PublicExports.Num() > 0)
+	if (ObjectsToRemove.Num() == 0)
 	{
-		// TRACE_CPUPROFILER_EVENT_SCOPE(RemovePublicExports);
-		GlobalImportStore.RemovePublicExports(PublicExports);
-	}
-	if (Packages.Num() > 0)
-	{
-		// TRACE_CPUPROFILER_EVENT_SCOPE(RemovePackages);
-		LoadedPackageStore.RemovePackages(Packages);
+		return;
 	}
+
+	TRACE_CPUPROFILER_EVENT_SCOPE(RemoveUnreachableObjects);
+
+	const int32 ObjectCount = ObjectsToRemove.Num();
+
+	const int32 OldLoadedPackageCount = LoadedPackageStore.NumTracked();
+	const int32 OldPublicExportCount = GlobalImportStore.GetStoredPublicExportsCount();
+
+	const double StartTime = FPlatformTime::Seconds();
+
+	const int32 NewLoadedPackageCount = LoadedPackageStore.NumTracked();
+	const int32 NewPublicExportCount = GlobalImportStore.GetStoredPublicExportsCount();
+	const int32 RemovedLoadedPackageCount = OldLoadedPackageCount - NewLoadedPackageCount;
+	const int32 RemovedPublicExportCount = OldPublicExportCount - NewPublicExportCount;
+
+	GlobalImportStore.RemovePublicExports(ObjectsToRemove);
+	LoadedPackageStore.RemovePackages(ObjectsToRemove);
+	ObjectsToRemove.Reset();
+
+	const double StopTime = FPlatformTime::Seconds();
+	UE_LOG(LogStreaming, Log,
+		TEXT("%.3f ms for processing %d objects in RemoveUnreachableObjects(Queued=%d, Async=%d). ")
+		TEXT("Removed %d (%d->%d) packages and %d (%d->%d) public exports."),
+		(StopTime - StartTime) * 1000,
+		ObjectCount,
+		GetNumQueuedPackages(), GetNumAsyncPackages(),
+		RemovedLoadedPackageCount, OldLoadedPackageCount, NewLoadedPackageCount,
+		RemovedPublicExportCount, OldPublicExportCount, NewPublicExportCount);
 }
 
-void FAsyncLoadingThread2::NotifyUnreachableObjects(const TArrayView<FUObjectItem*>& UnreachableObjects)
+void FAsyncLoadingThread2::NotifyUnreachableObjects(const TArrayView<FUObjectItem*>& UnreachableObjectItems)
 {
 	TRACE_CPUPROFILER_EVENT_SCOPE(NotifyUnreachableObjects);
 
@@ -5798,46 +5872,11 @@ void FAsyncLoadingThread2::NotifyUnreachableObjects(const TArrayView<FUObjectIte
 		return;
 	}
 
-	const double StartTime = FPlatformTime::Seconds();
+	// unreachable objects from last GC should typically have been processed already,
+	// if not handle them here before adding new ones
+	RemoveUnreachableObjects(UnreachableObjects);
 
-	FUnreachablePackages Packages;
-	FUnreachablePublicExports PublicExports;
-	FilterUnreachableObjects(UnreachableObjects, PublicExports, Packages);
-
-	const int32 PackageCount = Packages.Num();
-	const int32 PublicExportCount = PublicExports.Num();
-	if (PackageCount > 0 || PublicExportCount > 0)
-	{
-		const int32 OldLoadedPackageCount = LoadedPackageStore.NumTracked();
-		const int32 OldPublicExportCount = GlobalImportStore.GetStoredPublicExportsCount();
-
-		const double RemoveStartTime = FPlatformTime::Seconds();
-		RemoveUnreachableObjects(PublicExports, Packages);
-
-		const int32 NewLoadedPackageCount = LoadedPackageStore.NumTracked();
-		const int32 NewPublicExportCount = GlobalImportStore.GetStoredPublicExportsCount();
-		const int32 RemovedLoadedPackageCount = OldLoadedPackageCount - NewLoadedPackageCount;
-		const int32 RemovedPublicExportCount = OldPublicExportCount - NewPublicExportCount;
-
-		const double StopTime = FPlatformTime::Seconds();
-		UE_LOG(LogStreaming, Display,
-			TEXT("%.3f ms (%.3f+%.3f) ms for processing %d/%d objects in NotifyUnreachableObjects( Queued=%d, Async=%d). ")
-			TEXT("Removed %d/%d (%d->%d tracked) packages and %d/%d (%d->%d tracked) public exports."),
-			(StopTime - StartTime) * 1000,
-			(RemoveStartTime - StartTime) * 1000,
-			(StopTime - RemoveStartTime) * 1000,
-			PublicExportCount + PackageCount, UnreachableObjects.Num(),
-			GetNumQueuedPackages(), GetNumAsyncPackages(),
-			RemovedLoadedPackageCount, PackageCount, OldLoadedPackageCount, NewLoadedPackageCount,
-			RemovedPublicExportCount, PublicExportCount, OldPublicExportCount, NewPublicExportCount);
-	}
-	else
-	{
-		UE_LOG(LogStreaming, Display, TEXT("%.3f ms for skipping %d objects in NotifyUnreachableObjects (Queued=%d, Async=%d)."),
-			(FPlatformTime::Seconds() - StartTime) * 1000,
-			UnreachableObjects.Num(),
-			GetNumQueuedPackages(), GetNumAsyncPackages());
-	}
+	FilterUnreachableObjects(UnreachableObjectItems, UnreachableObjects);
 
 #if ALT2_VERIFY_ASYNC_FLAGS
 	if (!IsAsyncLoading())
@@ -5846,6 +5885,14 @@ void FAsyncLoadingThread2::NotifyUnreachableObjects(const TArrayView<FUObjectIte
 		VerifyLoadFlagsWhenFinishedLoading();
 	}
 #endif
+
+	if (GRemoveUnreachableObjectsOnGT)
+	{
+		RemoveUnreachableObjects(UnreachableObjects);
+	}
+
+	// wake up ALT to remove unreachable objects
+	AltZenaphore.NotifyAll();
 }
 
 /**
diff --git a/Engine/Source/Runtime/CoreUObject/Private/UObject/GarbageCollection.cpp b/Engine/Source/Runtime/CoreUObject/Private/UObject/GarbageCollection.cpp
--- a/Engine/Source/Runtime/CoreUObject/Private/UObject/GarbageCollection.cpp
+++ b/Engine/Source/Runtime/CoreUObject/Private/UObject/GarbageCollection.cpp
@@ -604,7 +604,7 @@ public:
 	/** Returns the number of objects already destroyed */
 	int32 GetObjectsDestroyedSinceLastMarkPhase() const
 	{
-		return ObjectsDestroyedSinceLastMarkPhase;
+		return ObjectsDestroyedSinceLastMarkPhase - NumObjectsToDestroyOnGameThread + NumObjectsDestroyedOnGameThread;
 	}
 
 	/** Resets the number of objects already destroyed */
@@ -2385,7 +2385,6 @@ void CollectGarbageInternal(EObjectFlags KeepFlags, bool bPerformFullPurge)
 			ArrayPool.DumpGarbageReferencers(AllArrays);
 
 			GatherUnreachableObjects(!(Options & EFastReferenceCollectorOptions::Parallel));
-			NotifyUnreachableObjects(GUnreachableObjects);
 
 			// This needs to happen after GatherUnreachableObjects since it can mark more objects as unreachable
 			ArrayPool.ClearWeakReferences(AllArrays);
@@ -2407,6 +2406,8 @@ void CollectGarbageInternal(EObjectFlags KeepFlags, bool bPerformFullPurge)
 
 			// Make sure nothing will be using potentially freed arrays
 			AllArrays.Empty();
+
+			NotifyUnreachableObjects(GUnreachableObjects);
 		}
 		if (GGCLockBehavior == FGCLockBehavior::Default)
 		{
