AssetDataTagMap: Fix FMapHandle::FindValue to perform a case-insensitive find when the the input FName is numberless. Change the runtime type for Numberless names to be FDisplayNameEntryId, which has a faster lookup of the comparision index than an FNameEntryId storing the display index.
#rnx

diff --git a/Engine/Source/Runtime/Core/Public/UObject/NameBatchSerialization.h b/Engine/Source/Runtime/Core/Public/UObject/NameBatchSerialization.h
--- a/Engine/Source/Runtime/Core/Public/UObject/NameBatchSerialization.h
+++ b/Engine/Source/Runtime/Core/Public/UObject/NameBatchSerialization.h
@@ -10,55 +10,6 @@
 
 //////////////////////////////////////////////////////////////////////////
 
-/** 
- * Serialization util that optimizes WITH_CASE_PRESERVING_NAME-loading by reducing comparison id lookups
- * 
- * Stores 32-bit display entry id with an unused bit to indicate if FName::GetComparisonIdFromDisplayId lookup is needed.
- * 
- * Note that only display entries should be saved to make output deterministic.
- */
-class FDisplayNameEntryId
-{
-public:
-	explicit FDisplayNameEntryId(FName Name) : FDisplayNameEntryId(Name.GetDisplayIndex(), Name.GetComparisonIndex()) {}
-	FORCEINLINE FName ToName(uint32 Number) const { return FName(GetComparisonId(), GetDisplayId(), Number); }
-
-private:
-#if WITH_CASE_PRESERVING_NAME
-	static constexpr uint32 DifferentIdsFlag = 1u << 31;
-	static constexpr uint32 DisplayIdMask = ~DifferentIdsFlag;
-
-	uint32 Value = 0;
-
-	FDisplayNameEntryId(FNameEntryId Id, FNameEntryId CmpId) : Value(Id.ToUnstableInt() | (Id != CmpId) * DifferentIdsFlag) {}
-	FORCEINLINE bool SameIds() const { return (Value & DifferentIdsFlag) == 0; }
-	FORCEINLINE FNameEntryId GetDisplayId() const	 { return FNameEntryId::FromUnstableInt(Value & DisplayIdMask); }
-	FORCEINLINE FNameEntryId GetComparisonId() const { return SameIds() ? GetDisplayId() : FName::GetComparisonIdFromDisplayId(GetDisplayId()); }
-	friend bool operator==(FDisplayNameEntryId A, FDisplayNameEntryId B) { return A.Value == B.Value; }
-#else
-	FNameEntryId Id;
-
-	FDisplayNameEntryId(FNameEntryId InId, FNameEntryId) : Id(InId) {}
-	FORCEINLINE FNameEntryId GetDisplayId() const	 { return Id; }
-	FORCEINLINE FNameEntryId GetComparisonId() const { return Id; }
-	friend bool operator==(FDisplayNameEntryId A, FDisplayNameEntryId B) { return A.Id == B.Id; }
-#endif
-	friend bool operator==(FNameEntryId A, FDisplayNameEntryId B) { return A == B.GetDisplayId(); }
-	friend bool operator==(FDisplayNameEntryId A, FNameEntryId B) { return A.GetDisplayId() == B; }
-	friend uint32 GetTypeHash(FDisplayNameEntryId InId) { return GetTypeHash(InId.GetDisplayId()); }
-
-public: // Internal functions for batch serialization code - intentionally lacking CORE_API
-	static FDisplayNameEntryId FromComparisonId(FNameEntryId ComparisonId);
-	FNameEntryId ToDisplayId() const;
-	void SetLoadedComparisonId(FNameEntryId ComparisonId); // Called first
-#if WITH_CASE_PRESERVING_NAME
-	void SetLoadedDifferentDisplayId(FNameEntryId DisplayId); // Called second if display id differs
-	FNameEntryId GetLoadedComparisonId() const; // Get the already loaded comparison id
-#endif
-};
-
-//////////////////////////////////////////////////////////////////////////
-
 #if ALLOW_NAME_BATCH_SAVING
 // Save display entries in given order to a name blob and a versioned hash blob.
 CORE_API void SaveNameBatch(TArrayView<const FDisplayNameEntryId> Names, TArray<uint8>& OutNameData, TArray<uint8>& OutHashData);
diff --git a/Engine/Source/Runtime/Core/Public/UObject/NameTypes.h b/Engine/Source/Runtime/Core/Public/UObject/NameTypes.h
--- a/Engine/Source/Runtime/Core/Public/UObject/NameTypes.h
+++ b/Engine/Source/Runtime/Core/Public/UObject/NameTypes.h
@@ -1807,6 +1807,54 @@ public:
 
 };
 
+/**
+ * Serialization util that optimizes WITH_CASE_PRESERVING_NAME-loading by reducing comparison id lookups
+ *
+ * Stores 32-bit display entry id with an unused bit to indicate if FName::GetComparisonIdFromDisplayId lookup is needed.
+ *
+ * Note that only display entries should be saved to make output deterministic.
+ */
+class FDisplayNameEntryId
+{
+public:
+	FDisplayNameEntryId() : FDisplayNameEntryId(FName()) {}
+	explicit FDisplayNameEntryId(FName Name) : FDisplayNameEntryId(Name.GetDisplayIndex(), Name.GetComparisonIndex()) {}
+	FORCEINLINE FName ToName(uint32 Number) const { return FName(GetComparisonId(), GetDisplayId(), Number); }
+
+private:
+#if WITH_CASE_PRESERVING_NAME
+	static constexpr uint32 DifferentIdsFlag = 1u << 31;
+	static constexpr uint32 DisplayIdMask = ~DifferentIdsFlag;
+
+	uint32 Value = 0;
+
+	FDisplayNameEntryId(FNameEntryId Id, FNameEntryId CmpId) : Value(Id.ToUnstableInt() | (Id != CmpId) * DifferentIdsFlag) {}
+	FORCEINLINE bool SameIds() const { return (Value & DifferentIdsFlag) == 0; }
+	FORCEINLINE FNameEntryId GetDisplayId() const { return FNameEntryId::FromUnstableInt(Value & DisplayIdMask); }
+	FORCEINLINE FNameEntryId GetComparisonId() const { return SameIds() ? GetDisplayId() : FName::GetComparisonIdFromDisplayId(GetDisplayId()); }
+	friend bool operator==(FDisplayNameEntryId A, FDisplayNameEntryId B) { return A.Value == B.Value; }
+#else
+	FNameEntryId Id;
+
+	FDisplayNameEntryId(FNameEntryId InId, FNameEntryId) : Id(InId) {}
+	FORCEINLINE FNameEntryId GetDisplayId() const { return Id; }
+	FORCEINLINE FNameEntryId GetComparisonId() const { return Id; }
+	friend bool operator==(FDisplayNameEntryId A, FDisplayNameEntryId B) { return A.Id == B.Id; }
+#endif
+	friend bool operator==(FNameEntryId A, FDisplayNameEntryId B) { return A == B.GetDisplayId(); }
+	friend bool operator==(FDisplayNameEntryId A, FNameEntryId B) { return A.GetDisplayId() == B; }
+	friend uint32 GetTypeHash(FDisplayNameEntryId InId) { return GetTypeHash(InId.GetDisplayId()); }
+
+public: // Internal functions for batch serialization code - intentionally lacking CORE_API
+	static FDisplayNameEntryId FromComparisonId(FNameEntryId ComparisonId);
+	FNameEntryId ToDisplayId() const;
+	void SetLoadedComparisonId(FNameEntryId ComparisonId); // Called first
+#if WITH_CASE_PRESERVING_NAME
+	void SetLoadedDifferentDisplayId(FNameEntryId DisplayId); // Called second if display id differs
+	FNameEntryId GetLoadedComparisonId() const; // Get the already loaded comparison id
+#endif
+};
+
 /**
  * A string builder with inline storage for FNames.
  */
diff --git a/Engine/Source/Runtime/CoreUObject/Private/AssetRegistry/AssetDataTagMap.cpp b/Engine/Source/Runtime/CoreUObject/Private/AssetRegistry/AssetDataTagMap.cpp
--- a/Engine/Source/Runtime/CoreUObject/Private/AssetRegistry/AssetDataTagMap.cpp
+++ b/Engine/Source/Runtime/CoreUObject/Private/AssetRegistry/AssetDataTagMap.cpp
@@ -274,6 +274,16 @@ namespace FixedTagPrivate
 	{
 		return FName::CreateFromDisplayId(EntryId, NAME_NO_NUMBER_INTERNAL);
 	}
+
+	static FName MakeNumberedName(FDisplayNameEntryId EntryId)
+	{
+		return EntryId.ToName(NAME_NO_NUMBER_INTERNAL);
+	}
+
+	static bool EqualsInsensitive(FDisplayNameEntryId A, FDisplayNameEntryId B)
+	{
+		return MakeNumberedName(A) == MakeNumberedName(B);
+	}
 	
 	static FNameEntryId MakeNumberlessDisplayName(FName Name)
 	{
@@ -410,7 +420,7 @@ namespace FixedTagPrivate
 		{
 		case EValueType::AnsiString:			return FString(Store.GetAnsiString(Index));
 		case EValueType::WideString:			return FString(Store.GetWideString(Index));
-		case EValueType::NumberlessName:		return FName::GetEntry(Store.NumberlessNames[Index])->GetPlainNameString();
+		case EValueType::NumberlessName:		return MakeNumberedName(Store.NumberlessNames[Index]).GetPlainNameString();
 		case EValueType::Name:					return Store.Names[Index].ToString();
 		case EValueType::NumberlessExportPath:	return Store.NumberlessExportPaths[Index].ToString();
 		case EValueType::ExportPath:			return Store.ExportPaths[Index].ToString();
@@ -466,7 +476,7 @@ namespace FixedTagPrivate
 		{
 		case EValueType::AnsiString:			return FAssetRegistryExportPath(Store.GetAnsiString(Index));
 		case EValueType::WideString:			return FAssetRegistryExportPath(Store.GetWideString(Index));
-		case EValueType::NumberlessName:		return MakeNumberedPath(*FName::GetEntry(Store.NumberlessNames[Index]));
+		case EValueType::NumberlessName:		return MakeNumberedPath(MakeNumberedName(Store.NumberlessNames[Index]));
 		case EValueType::Name:					return MakeNumberedPath(Store.Names[Index]);
 		case EValueType::NumberlessExportPath:	return MakeNumberedPath(Store.NumberlessExportPaths[Index]);
 		case EValueType::ExportPath:			return Store.ExportPaths[Index];
@@ -599,7 +609,7 @@ namespace FixedTagPrivate
 		{
 			for (const FNumberlessPair& Pair : GetNumberlessView())
 			{
-				if (Key.GetComparisonIndex() == Pair.Key)
+				if (Key.GetComparisonIndex() == MakeNumberedName(Pair.Key).GetComparisonIndex())
 				{
 					return &Pair.Value;
 				}
@@ -617,7 +627,7 @@ namespace FixedTagPrivate
 	static FNumberlessPair MakeNumberlessPair(FNumberedPair Pair)
 	{
 		check(Pair.Key.GetNumber() == NAME_NO_NUMBER_INTERNAL);
-		return {Pair.Key.GetComparisonIndex(), Pair.Value};
+		return {FDisplayNameEntryId(Pair.Key), Pair.Value};
 	}
 
 	FNumberedPair FMapHandle::At(uint32 Index) const
@@ -862,7 +872,7 @@ namespace FixedTagPrivate
 		{
 			FName Name = Value.AsName();
 			return IsNumberless(Name)
-				? FValueId{EValueType::NumberlessName,			Index(NumberlessNameIndices, MakeNumberlessDisplayName(Name))}
+				? FValueId{EValueType::NumberlessName,			Index(NumberlessNameIndices, FDisplayNameEntryId(Name))}
 				: FValueId{EValueType::Name,					Index(NameIndices, Name)};
 		}
 		else if (Options.StoreAsPath.Contains(Key))
@@ -988,14 +998,14 @@ namespace FixedTagPrivate
 			return Item;
 		}
 		
-		void SaveItem(FNameEntryId NameId)
+		void SaveItem(FDisplayNameEntryId NameId)
 		{
 			SaveItem(MakeNumberedName(NameId));
 		}
 
-		template<> FNameEntryId LoadItem<FNameEntryId>()
+		template<> FDisplayNameEntryId LoadItem<FDisplayNameEntryId>()
 		{
-			return MakeNumberlessDisplayName(LoadItem<FName>());
+			return FDisplayNameEntryId(LoadItem<FName>());
 		}
 		
 		void SaveItem(FNumberlessExportPath Path)
@@ -1494,7 +1504,7 @@ namespace FixedTagPrivate
 		{
 			for (FNumberlessPair Pair : Fixed.GetNumberlessView())
 			{
-				const FString* LooseValue = Loose->Find(FName::CreateFromDisplayId(Pair.Key, 0));
+				const FString* LooseValue = Loose->Find(MakeNumberedName(Pair.Key));
 				if (!LooseValue || FAssetTagValueRef(Fixed.StoreIndex, Pair.Value) != *LooseValue)
 				{
 					return false;
@@ -1762,7 +1772,7 @@ bool FStoreTest::RunTest(const FString& Parameters)
 		TestEqual("String values excludes",		CountOccurences(Data.AnsiStrings, "NameValue"), 0);
 		TestEqual("String value deduplication",	CountOccurences(Data.AnsiStrings, "SameSame"), 1);
 		TestEqual("Wide characters",			FStringView(Data.WideStrings.GetData(), Data.WideStrings.Num() - 1), FStringView(TEXT("Wide\x00DF")));
-		TestEqual("Numberless name values",		Data.NumberlessNames, {MakeNumberlessDisplayName("NameValue")});
+		TestEqual("Numberless name values",		Data.NumberlessNames, {FDisplayNameEntryId("NameValue")});
 		TestEqual("Numbered name values",		Data.Names, {FName("NameValue_0")});
 		TestEqual("Numberless path values",		Data.NumberlessExportPaths.Num(), 3); // C\'P.O\', P.O, O  
 		TestEqual("Numbered path values",		Data.ExportPaths.Num(), 7); // NumPath[0-7] values
diff --git a/Engine/Source/Runtime/CoreUObject/Public/AssetRegistry/AssetDataTagMap.h b/Engine/Source/Runtime/CoreUObject/Public/AssetRegistry/AssetDataTagMap.h
--- a/Engine/Source/Runtime/CoreUObject/Public/AssetRegistry/AssetDataTagMap.h
+++ b/Engine/Source/Runtime/CoreUObject/Public/AssetRegistry/AssetDataTagMap.h
@@ -106,7 +106,7 @@ namespace FixedTagPrivate
 
 	struct FNumberlessPair
 	{
-		FNameEntryId Key;
+		FDisplayNameEntryId Key;
 		FValueId Value;
 	};
 
@@ -158,7 +158,7 @@ namespace FixedTagPrivate
 			{
 				for (FNumberlessPair Pair : GetNumberlessView())
 				{
-					Fn(FNumberedPair{FName::CreateFromDisplayId(Pair.Key, 0), Pair.Value});
+					Fn(FNumberedPair{Pair.Key.ToName(NAME_NO_NUMBER_INTERNAL), Pair.Value});
 				}
 			}
 			else
diff --git a/Engine/Source/Runtime/CoreUObject/Public/AssetRegistry/AssetDataTagMapSerializationDetails.h b/Engine/Source/Runtime/CoreUObject/Public/AssetRegistry/AssetDataTagMapSerializationDetails.h
--- a/Engine/Source/Runtime/CoreUObject/Public/AssetRegistry/AssetDataTagMapSerializationDetails.h
+++ b/Engine/Source/Runtime/CoreUObject/Public/AssetRegistry/AssetDataTagMapSerializationDetails.h
@@ -51,7 +51,7 @@ namespace FixedTagPrivate
 		TArrayView<ANSICHAR> AnsiStrings;
 		TArrayView<uint32> WideStringOffsets;
 		TArrayView<WIDECHAR> WideStrings;
-		TArrayView<FNameEntryId> NumberlessNames;
+		TArrayView<FDisplayNameEntryId> NumberlessNames;
 		TArrayView<FName> Names;
 		TArrayView<FNumberlessExportPath> NumberlessExportPaths;
 		TArrayView<FAssetRegistryExportPath> ExportPaths;
@@ -103,7 +103,7 @@ namespace FixedTagPrivate
 		TArray<ANSICHAR> AnsiStrings;
 		TArray<uint32> WideStringOffsets;
 		TArray<WIDECHAR> WideStrings;
-		TArray<FNameEntryId> NumberlessNames;
+		TArray<FDisplayNameEntryId> NumberlessNames;
 		TArray<FName> Names;
 		TArray<FNumberlessExportPath> NumberlessExportPaths;
 		TArray<FAssetRegistryExportPath> ExportPaths;
@@ -208,7 +208,7 @@ namespace FixedTagPrivate
 		const FOptions Options;
 		FStringIndexer AnsiStrings;
 		FStringIndexer WideStrings;
-		TMap<FNameEntryId, uint32> NumberlessNameIndices;
+		TMap<FDisplayNameEntryId, uint32> NumberlessNameIndices;
 		TMap<FName, uint32, FDefaultSetAllocator, FCaseSensitiveFuncs<uint32>> NameIndices;
 		TMap<FNumberlessExportPath, uint32, FDefaultSetAllocator, FCaseSensitiveFuncs<uint32>> NumberlessExportPathIndices;
 		TMap<FAssetRegistryExportPath, uint32, FDefaultSetAllocator, FCaseSensitiveFuncs<uint32>> ExportPathIndices;
