Fixed issue where LiveCoding could remove properties from a class that contained delegates.

diff --git a/Engine/Source/Runtime/CoreUObject/Private/UObject/UObjectGlobals.cpp b/Engine/Source/Runtime/CoreUObject/Private/UObject/UObjectGlobals.cpp
--- a/Engine/Source/Runtime/CoreUObject/Private/UObject/UObjectGlobals.cpp
+++ b/Engine/Source/Runtime/CoreUObject/Private/UObject/UObjectGlobals.cpp
@@ -5292,12 +5292,31 @@ namespace UECodeGen_Private
 			return;
 		}
 
+		FName FuncName(UTF8_TO_TCHAR(Params.NameUTF8));
+
+#if WITH_LIVE_CODING
+		// When a package is patched, it might reference a function in a class.  When this happens, the existing UFunction
+		// object gets reused but the UField's Next pointer gets nulled out.  This ends up terminating the function list
+		// for the class.  To work around this issue, cache the next pointer and then restore it after the new instance
+		// is created.  Only do this if we reuse the current instance.
+		UField* PrevFunctionNextField = nullptr;
+		UFunction* PrevFunction = nullptr;
+		if (UObject* PrevObject = StaticFindObjectFastInternal( /*Class=*/ nullptr, Outer, FuncName, true))
+		{
+			PrevFunction = Cast<UFunction>(PrevObject);
+			if (PrevFunction != nullptr)
+			{
+				PrevFunctionNextField = PrevFunction->Next;
+			}
+		}
+#endif
+
 		UFunction* NewFunction;
 		if (Params.FunctionFlags & FUNC_Delegate)
 		{
 			if (Params.OwningClassName == nullptr)
 			{
-				NewFunction = new (EC_InternalUseOnlyConstructor, Outer, UTF8_TO_TCHAR(Params.NameUTF8), Params.ObjectFlags) UDelegateFunction(
+				NewFunction = new (EC_InternalUseOnlyConstructor, Outer, FuncName, Params.ObjectFlags) UDelegateFunction(
 					FObjectInitializer(),
 					Super,
 					Params.FunctionFlags,
@@ -5306,7 +5325,7 @@ namespace UECodeGen_Private
 			}
 			else
 			{
-				USparseDelegateFunction* NewSparseFunction = new (EC_InternalUseOnlyConstructor, Outer, UTF8_TO_TCHAR(Params.NameUTF8), Params.ObjectFlags) USparseDelegateFunction(
+				USparseDelegateFunction* NewSparseFunction = new (EC_InternalUseOnlyConstructor, Outer, FuncName, Params.ObjectFlags) USparseDelegateFunction(
 					FObjectInitializer(),
 					Super,
 					Params.FunctionFlags,
@@ -5319,7 +5338,7 @@ namespace UECodeGen_Private
 		}
 		else
 		{
-			NewFunction = new (EC_InternalUseOnlyConstructor, Outer, UTF8_TO_TCHAR(Params.NameUTF8), Params.ObjectFlags) UFunction(
+			NewFunction = new (EC_InternalUseOnlyConstructor, Outer, FuncName, Params.ObjectFlags) UFunction(
 				FObjectInitializer(),
 				Super,
 				Params.FunctionFlags,
@@ -5330,6 +5349,10 @@ namespace UECodeGen_Private
 
 #if WITH_LIVE_CODING
 		NewFunction->SingletonPtr = SingletonPtr;
+		if (NewFunction == PrevFunction)
+		{
+			NewFunction->Next = PrevFunctionNextField;
+		}
 #endif
 
 #if WITH_METADATA
