InvalidEnumValues added to ObjectMacros.h - functionality itself present since 5.0

diff --git a/Engine/Source/Runtime/CoreUObject/Public/UObject/ObjectMacros.h b/Engine/Source/Runtime/CoreUObject/Public/UObject/ObjectMacros.h
--- a/Engine/Source/Runtime/CoreUObject/Public/UObject/ObjectMacros.h
+++ b/Engine/Source/Runtime/CoreUObject/Public/UObject/ObjectMacros.h
@@ -1441,6 +1441,9 @@ namespace UM
 
 		/// [PropertyMetadata] Used for enum properties to define the subset of valid values as a comma-separated string; values outside of this subset are not made available.
 		ValidEnumValues,
+
+	    /// [PropertyMetadata] Used for enum properties to define the subset of invalid values as a comma-separated string; values within this subset are not made available.
+        InvalidEnumValues,
 	};
 
 	// Metadata usable in UPROPERTY for customizing the behavior of Persona and UMG
