AsyncLoading2: Make sure to call RemoveUnreachableObjects also when running without the Async Loading Thread (follow up fix for 21353595).
Increase log verbosity back to Display for RemoveUnreachableObjects.
Remove unnecessary FGCScopeGuard from the TickAsyncThreadFromGameThread.

#rnx

diff --git a/Engine/Source/Runtime/CoreUObject/Private/Serialization/AsyncLoading2.cpp b/Engine/Source/Runtime/CoreUObject/Private/Serialization/AsyncLoading2.cpp
--- a/Engine/Source/Runtime/CoreUObject/Private/Serialization/AsyncLoading2.cpp
+++ b/Engine/Source/Runtime/CoreUObject/Private/Serialization/AsyncLoading2.cpp
@@ -302,11 +302,12 @@ if ((Condition)) \
 #define UE_ASYNC_PACKAGE_CLOG_VERBOSE(Condition, Verbosity, PackageDesc, LogDesc, Format, ...)
 #endif
 
-static bool GRemoveUnreachableObjectsOnGT = false;
-static FAutoConsoleVariableRef CVarRemoveUnreachableObjectsOnGT(
-	TEXT("s.RemoveUnreachableObjectsOnGT"),
-	GRemoveUnreachableObjectsOnGT,
-	TEXT("Force running removal of unreachable objects on the Game Thread (slower, but enables extra verification in debug and development builds."),
+static bool GGRemoveUnreachableObjectsFromGCNotifyOnGT = false;
+static FAutoConsoleVariableRef CVarGRemoveUnreachableObjectsFromGCNotifyOnGT(
+	TEXT("s.GRemoveUnreachableObjectsFromGCNotifyOnGT"),
+	GGRemoveUnreachableObjectsFromGCNotifyOnGT,
+	TEXT("Force running removal of unreachable objects from the Garbage Collection Notify callback on the Game Thread. "
+		"This also enables extra verification in debug and development builds (slow)."),
 	ECVF_Default
 	);
 
@@ -1001,7 +1002,7 @@ public:
 				PublicExportKeys.Emplace(PublicExportKey);
 
 #if ALT2_VERIFY_UNREACHABLE_OBJECTS
-				if (GRemoveUnreachableObjectsOnGT)
+				if (GGRemoveUnreachableObjectsFromGCNotifyOnGT)
 				{
 					UObject* GCObject = Item.DebugObject;
 
@@ -4942,6 +4943,7 @@ EAsyncPackageState::Type FAsyncLoadingThread2::ProcessAsyncLoadingFromGameThread
 	FAsyncLoadingTickScope2 InAsyncLoadingTick(*this);
 	uint32 LoopIterations = 0;
 
+	RemoveUnreachableObjects(UnreachableObjects);
 	while (true)
 	{
 		do 
@@ -5643,7 +5645,6 @@ EAsyncPackageState::Type FAsyncLoadingThread2::TickAsyncThreadFromGameThread(FAs
 		}
 		else
 		{
-			FGCScopeGuard GCGuard;
 			Result = ProcessAsyncLoadingFromGameThread(ThreadState, ProcessedRequests);
 			bDidSomething = bDidSomething || ProcessedRequests > 0;
 		}
@@ -5788,7 +5789,7 @@ FORCENOINLINE static void FilterUnreachableObjects(
 		Item.ObjectIndex = GUObjectArray.ObjectToIndex(Object);
 		Item.ObjectName = Object->GetFName();
 #if ALT2_VERIFY_UNREACHABLE_OBJECTS
-		if (GRemoveUnreachableObjectsOnGT)
+		if (GGRemoveUnreachableObjectsFromGCNotifyOnGT)
 		{
 			Item.DebugObject = Object;
 		}
@@ -5853,7 +5854,7 @@ void FAsyncLoadingThread2::RemoveUnreachableObjects(FUnreachableObjects& Objects
 	ObjectsToRemove.Reset();
 
 	const double StopTime = FPlatformTime::Seconds();
-	UE_LOG(LogStreaming, Log,
+	UE_LOG(LogStreaming, Display,
 		TEXT("%.3f ms for processing %d objects in RemoveUnreachableObjects(Queued=%d, Async=%d). ")
 		TEXT("Removed %d (%d->%d) packages and %d (%d->%d) public exports."),
 		(StopTime - StartTime) * 1000,
@@ -5886,7 +5887,7 @@ void FAsyncLoadingThread2::NotifyUnreachableObjects(const TArrayView<FUObjectIte
 	}
 #endif
 
-	if (GRemoveUnreachableObjectsOnGT)
+	if (GGRemoveUnreachableObjectsFromGCNotifyOnGT)
 	{
 		RemoveUnreachableObjects(UnreachableObjects);
 	}
