Ensure that FObjectPathId constructors from strings are permitted in all configurations that build unit tests.

diff --git a/Engine/Source/Runtime/CoreUObject/Private/UObject/ObjectPathId.cpp b/Engine/Source/Runtime/CoreUObject/Private/UObject/ObjectPathId.cpp
--- a/Engine/Source/Runtime/CoreUObject/Private/UObject/ObjectPathId.cpp
+++ b/Engine/Source/Runtime/CoreUObject/Private/UObject/ObjectPathId.cpp
@@ -231,7 +231,7 @@ FName FObjectPathId::MakeImportPathIdAndPackageName(const FObjectImport& Import,
 	return NamePathProducer.GetPackageName();
 }
 
-#if !(UE_BUILD_SHIPPING || UE_BUILD_TEST)
+#if WITH_LOW_LEVEL_TESTS
 template <typename CharType>
 struct TStringViewNamePathProducer
 {
diff --git a/Engine/Source/Runtime/CoreUObject/Public/UObject/ObjectPathId.h b/Engine/Source/Runtime/CoreUObject/Public/UObject/ObjectPathId.h
--- a/Engine/Source/Runtime/CoreUObject/Public/UObject/ObjectPathId.h
+++ b/Engine/Source/Runtime/CoreUObject/Public/UObject/ObjectPathId.h
@@ -70,7 +70,7 @@ public:
 
 	COREUOBJECT_API static FName MakeImportPathIdAndPackageName(const FObjectImport& Import, const FLinkerTables& LinkerTables, FObjectPathId& OutPathId);
 
-#if !(UE_BUILD_SHIPPING || UE_BUILD_TEST)
+#if WITH_LOW_LEVEL_TESTS
 	// Should only be used by unit tests
 	explicit FObjectPathId(FWideStringView StringPath);
 	explicit FObjectPathId(FAnsiStringView StringPath);
