Reverted unnecessary additions by IWYU

- Array.h is not needed by CompactBinarySerialization.h because the use of TArray is in a template function.
- SharedBuffer.h is not needed by DerivedData.h because its types are only used as return types or function parameters.
- FString and FCompositeBuffer are both declared in Fwd.h headers and do not need duplicate forward declarations.

#rnx

diff --git a/Engine/Source/Runtime/Core/Public/Serialization/CompactBinarySerialization.h b/Engine/Source/Runtime/Core/Public/Serialization/CompactBinarySerialization.h
--- a/Engine/Source/Runtime/Core/Public/Serialization/CompactBinarySerialization.h
+++ b/Engine/Source/Runtime/Core/Public/Serialization/CompactBinarySerialization.h
@@ -2,7 +2,6 @@
 
 #pragma once
 
-#include "Containers/Array.h"
 #include "Containers/ContainersFwd.h"
 #include "Containers/StringFwd.h"
 #include "CoreTypes.h"
@@ -15,7 +14,6 @@
 
 class FArchive;
 class FName;
-class FString;
 struct FGuid;
 
 ///////////////////////////////////////////////////////////////////////////////////////////////////////////////////////
diff --git a/Engine/Source/Runtime/CoreUObject/Public/Serialization/DerivedData.h b/Engine/Source/Runtime/CoreUObject/Public/Serialization/DerivedData.h
--- a/Engine/Source/Runtime/CoreUObject/Public/Serialization/DerivedData.h
+++ b/Engine/Source/Runtime/CoreUObject/Public/Serialization/DerivedData.h
@@ -7,23 +7,20 @@
 #include "Math/NumericLimits.h"
 #include "Memory/MemoryFwd.h"
 #include "Memory/MemoryView.h"
-#include "Memory/SharedBuffer.h"
 #include "Misc/EnumClassFlags.h"
 #include "Templates/Function.h"
 #include "Templates/PimplPtr.h"
 
-class FCompositeBuffer;
-
 #define UE_API COREUOBJECT_API
 
 class FArchive;
 class FCompressedBuffer;
 class UObject;
+struct FIoHash;
+
 namespace UE::DerivedData { struct FCacheKey; }
 namespace UE::DerivedData { struct FValueId; }
 namespace UE::DerivedData { template <typename CharType> class TSharedString; }
-struct FIoHash;
-
 namespace UE::DerivedData { using FSharedString = TSharedString<TCHAR>; }
 namespace UE::DerivedData::Private { class FEditorData; }
 namespace UE::DerivedData::Private { class FIoResponse; }
