AsyncLoading2: Improve the initial load logic from cl 20643384.
Replace FinalizeInitialLoad with CompleteRegistrationEvent and use the same code path with/without the async loading thread.
Strip the FindAllScriptObjects logic from test and shipping builds.
Optimize initial load NotifyRegistrationEvent handling by passing the finished object instead of looking it up by name.

#rnx

diff --git a/Engine/Source/Runtime/CoreUObject/Private/Serialization/AsyncLoading.cpp b/Engine/Source/Runtime/CoreUObject/Private/Serialization/AsyncLoading.cpp
--- a/Engine/Source/Runtime/CoreUObject/Private/Serialization/AsyncLoading.cpp
+++ b/Engine/Source/Runtime/CoreUObject/Private/Serialization/AsyncLoading.cpp
@@ -5521,7 +5521,7 @@ void FAsyncLoadingThread::NotifyConstructedDuringAsyncLoading(UObject* Object, b
 	}
 }
 
-void FAsyncLoadingThread::NotifyRegistrationEvent(const TCHAR* PackageName, const TCHAR* Name, ENotifyRegistrationType NotifyRegistrationType, ENotifyRegistrationPhase NotifyRegistrationPhase, UObject* (*InRegister)(), bool InbDynamic)
+void FAsyncLoadingThread::NotifyRegistrationEvent(const TCHAR* PackageName, const TCHAR* Name, ENotifyRegistrationType NotifyRegistrationType, ENotifyRegistrationPhase NotifyRegistrationPhase, UObject* (*InRegister)(), bool InbDynamic, UObject* FinishedObject)
 {
 	EDLBootNotificationManager.NotifyRegistrationEvent(PackageName, Name, NotifyRegistrationType, NotifyRegistrationPhase, InRegister, InbDynamic);
 }
diff --git a/Engine/Source/Runtime/CoreUObject/Private/Serialization/AsyncLoading2.cpp b/Engine/Source/Runtime/CoreUObject/Private/Serialization/AsyncLoading2.cpp
--- a/Engine/Source/Runtime/CoreUObject/Private/Serialization/AsyncLoading2.cpp
+++ b/Engine/Source/Runtime/CoreUObject/Private/Serialization/AsyncLoading2.cpp
@@ -1108,34 +1108,18 @@ public:
 		ObjectIndexToPublicExport.Add(ObjectIndex, Key);
 	}
 
-	void FindAllScriptObjects();
+	void RegistrationComplete();
 
-	void AddScriptObject(FStringView PackageName, FStringView Name)
+	void AddScriptObject(FStringView PackageName, FStringView Name, UObject* Object)
 	{
 		TStringBuilder<FName::StringBufferSize> FullName;
 		FPathViews::Append(FullName, PackageName);
 		FPathViews::Append(FullName, Name);
 		FPackageObjectIndex GlobalImportIndex = FPackageObjectIndex::FromScriptPath(FullName);
 
-		UObject* Package = StaticFindObjectFastInternal(UPackage::StaticClass(), nullptr, FName(PackageName), true);
-		// check(Package);
-		if (!ensureMsgf(Package, TEXT("Failed to find package when registering script object '%.*s/%.*s'"),
-			PackageName.Len(), PackageName.GetData(), Name.Len(), Name.GetData()))
-		{
-			return;
-		}
-
-		UObject* Object = StaticFindObjectFastInternal(nullptr, Package, FName(Name));
-		// check(Object);
-		if (!ensureMsgf(Object, TEXT("Failed to find object when registering script object '%.*s/%.*s'"),
-			PackageName.Len(), PackageName.GetData(), Name.Len(), Name.GetData()))
-		{
-			return;
-		}
-
 #if WITH_EDITOR
 		FPackageObjectIndex PackageGlobalImportIndex = FPackageObjectIndex::FromScriptPath(PackageName);
-		ScriptObjects.Add(PackageGlobalImportIndex, Package);
+		ScriptObjects.Add(PackageGlobalImportIndex, Object->GetOutermost());
 #endif
 		ScriptObjects.Add(GlobalImportIndex, Object);
 
@@ -2383,7 +2367,7 @@ public:
 
 	virtual void NotifyUnreachableObjects(const TArrayView<FUObjectItem*>& UnreachableObjects) override;
 
-	virtual void NotifyRegistrationEvent(const TCHAR* PackageName, const TCHAR* Name, ENotifyRegistrationType NotifyRegistrationType, ENotifyRegistrationPhase NotifyRegistrationPhase, UObject* (*InRegister)(), bool InbDynamic) override;
+	virtual void NotifyRegistrationEvent(const TCHAR* PackageName, const TCHAR* Name, ENotifyRegistrationType NotifyRegistrationType, ENotifyRegistrationPhase NotifyRegistrationPhase, UObject* (*InRegister)(), bool InbDynamic, UObject* FinishedObject) override;
 
 	virtual void NotifyRegistrationComplete() override;
 
@@ -2540,8 +2524,6 @@ private:
 
 	void OnLeakedPackageRename(UPackage* Package);
 
-	void FinalizeInitialLoad();
-
 	void RemoveUnreachableObjects(FUnreachableObjects& ObjectsToRemove);
 
 	void ProcessPendingCDOs()
@@ -3404,44 +3386,51 @@ void FAsyncPackage2::SetupScriptDependencies()
 }
 
 
-void FGlobalImportStore::FindAllScriptObjects()
+void FGlobalImportStore::RegistrationComplete()
 {
-	TRACE_CPUPROFILER_EVENT_SCOPE(FindAllScriptObjects);
-	TStringBuilder<FName::StringBufferSize> Name;
-	TArray<UPackage*> ScriptPackages;
-	TArray<UObject*> Objects;
-	FindAllRuntimeScriptPackages(ScriptPackages);
-
-	for (UPackage* Package : ScriptPackages)
+#if DO_CHECK
 	{
-#if WITH_EDITOR
-		Name.Reset();
-		Package->GetPathName(nullptr, Name);
-		FPackageObjectIndex PackageGlobalImportIndex = FPackageObjectIndex::FromScriptPath(Name);
-		ScriptObjects.Add(PackageGlobalImportIndex, Package);
-#endif
+		TRACE_CPUPROFILER_EVENT_SCOPE(FindAllScriptObjectsDebug);
+		TStringBuilder<FName::StringBufferSize> Name;
+		TArray<UPackage*> ScriptPackages;
+		TArray<UObject*> Objects;
+		FindAllRuntimeScriptPackages(ScriptPackages);
 
-		Objects.Reset();
-		GetObjectsWithOuter(Package, Objects, /*bIncludeNestedObjects*/true);
-		for (UObject* Object : Objects)
+		for (UPackage* Package : ScriptPackages)
 		{
-			if (Object->HasAnyFlags(RF_Public))
+#if WITH_EDITOR
+			Name.Reset();
+			Package->GetPathName(nullptr, Name);
+			FPackageObjectIndex PackageGlobalImportIndex = FPackageObjectIndex::FromScriptPath(Name);
+			if (!ScriptObjects.Contains(PackageGlobalImportIndex))
 			{
-				Name.Reset();
-				Object->GetPathName(nullptr, Name);
-				FPackageObjectIndex GlobalImportIndex = FPackageObjectIndex::FromScriptPath(Name);
-				if (!ScriptObjects.Contains(GlobalImportIndex))
-				{
-					ScriptObjects.Add(GlobalImportIndex, Object);
-#if !WITH_EDITOR
-					UE_LOG(LogStreaming, Warning, TEXT("Found additional public script object 0x%016llX: %s\r\n"),
-						GlobalImportIndex.Value(),
-						*Object->GetFullName());
+				ScriptObjects.Add(PackageGlobalImportIndex, Package);
+				ensureMsgf(false, TEXT("Script package %s (0x%016llX) is missing a NotifyRegistrationEvent from the initial load phase."),
+					*Package->GetFullName(),
+					PackageGlobalImportIndex.Value());
+			}
 #endif
+			Objects.Reset();
+			GetObjectsWithOuter(Package, Objects, /*bIncludeNestedObjects*/true);
+			for (UObject* Object : Objects)
+			{
+				if (Object->HasAnyFlags(RF_Public))
+				{
+					Name.Reset();
+					Object->GetPathName(nullptr, Name);
+					FPackageObjectIndex GlobalImportIndex = FPackageObjectIndex::FromScriptPath(Name);
+					if (!ScriptObjects.Contains(GlobalImportIndex))
+					{
+						ScriptObjects.Add(GlobalImportIndex, Object);
+						ensureMsgf(false, TEXT("Script object %s (0x%016llX) is missing a NotifyRegistrationEvent from the initial load phase."),
+							*Object->GetFullName(),
+							GlobalImportIndex.Value());
+					}
 				}
 			}
 		}
 	}
+#endif
 	ScriptObjects.Shrink();
 }
 
@@ -5415,11 +5404,7 @@ void FAsyncLoadingThread2::StartThread()
 	check(PendingCDOsRecursiveStack.Num() == 0);
 	PendingCDOsRecursiveStack.Empty();
 
-	if (!FAsyncLoadingThreadSettings::Get().bAsyncLoadingThreadEnabled)
-	{
-		FinalizeInitialLoad();
-	}
-	else if (!Thread)
+	if (!Thread)
 	{
 		UE_LOG(LogStreaming, Log, TEXT("Starting Async Loading Thread."));
 		bThreadStarted = true;
@@ -5439,17 +5424,6 @@ bool FAsyncLoadingThread2::Init()
 	return true;
 }
 
-void FAsyncLoadingThread2::FinalizeInitialLoad()
-{
-	TRACE_CPUPROFILER_EVENT_SCOPE(FinalizeInitialLoad);
-	GlobalImportStore.FindAllScriptObjects(); // for verification only
-	bHasRegisteredAllScriptObjects = true;
-
-	UE_LOG(LogStreaming, Display,
-		TEXT("AsyncLoading2 - InitialLoad Finalized: Registered %d public script object entries (%.2f KB)"),
-		GlobalImportStore.GetStoredScriptObjectsCount(), (float)GlobalImportStore.GetStoredScriptObjectsAllocatedSize() / 1024.f);
-}
-
 uint32 FAsyncLoadingThread2::Run()
 {
 	LLM_SCOPE(ELLMTag::AsyncLoading);
@@ -5465,8 +5439,6 @@ uint32 FAsyncLoadingThread2::Run()
 
 	FAsyncLoadingThreadState2& ThreadState = *FAsyncLoadingThreadState2::Get();
 	
-	FinalizeInitialLoad();
-
 	FZenaphoreWaiter Waiter(AltZenaphore, TEXT("WaitForEvents"));
 	enum class EMainState : uint8
 	{
@@ -5928,17 +5900,24 @@ void FAsyncLoadingThread2::NotifyRegistrationEvent(
 	ENotifyRegistrationType NotifyRegistrationType,
 	ENotifyRegistrationPhase NotifyRegistrationPhase,
 	UObject* (*InRegister)(),
-	bool InbDynamic)
+	bool InbDynamic,
+	UObject* FinishedObject)
 {
 	if (NotifyRegistrationPhase == ENotifyRegistrationPhase::NRP_Finished)
 	{
-		GlobalImportStore.AddScriptObject(PackageName, Name);
+		ensureMsgf(FinishedObject, TEXT("FinishedObject was not provided by NotifyRegistrationEvent when called with ENotifyRegistrationPhase::NRP_Finished, see call stack for offending code."));
+		GlobalImportStore.AddScriptObject(PackageName, Name, FinishedObject);
 	}
 }
 
 void FAsyncLoadingThread2::NotifyRegistrationComplete()
 {
+	GlobalImportStore.RegistrationComplete();
+	bHasRegisteredAllScriptObjects = true;
 
+	UE_LOG(LogStreaming, Display,
+		TEXT("AsyncLoading2 - NotifyRegistrationComplete: Registered %d public script object entries (%.2f KB)"),
+		GlobalImportStore.GetStoredScriptObjectsCount(), (float)GlobalImportStore.GetStoredScriptObjectsAllocatedSize() / 1024.f);
 }
 
 /*-----------------------------------------------------------------------------
diff --git a/Engine/Source/Runtime/CoreUObject/Private/Serialization/AsyncLoadingThread.h b/Engine/Source/Runtime/CoreUObject/Private/Serialization/AsyncLoadingThread.h
--- a/Engine/Source/Runtime/CoreUObject/Private/Serialization/AsyncLoadingThread.h
+++ b/Engine/Source/Runtime/CoreUObject/Private/Serialization/AsyncLoadingThread.h
@@ -340,7 +340,7 @@ public:
 
 	void FireCompletedCompiledInImport(void* AsyncPacakge, FPackageIndex Import);
 
-	void NotifyRegistrationEvent(const TCHAR* PackageName, const TCHAR* Name, ENotifyRegistrationType NotifyRegistrationType, ENotifyRegistrationPhase NotifyRegistrationPhase, UObject* (*InRegister)(), bool InbDynamic) override;
+	void NotifyRegistrationEvent(const TCHAR* PackageName, const TCHAR* Name, ENotifyRegistrationType NotifyRegistrationType, ENotifyRegistrationPhase NotifyRegistrationPhase, UObject* (*InRegister)(), bool InbDynamic, UObject* FinishedObject) override;
 
 	void NotifyRegistrationComplete() override;
 
@@ -756,4 +756,4 @@ public:
 #endif
 
 	static FAsyncLoadingThread* Instance;
-};
\ No newline at end of file
+};
diff --git a/Engine/Source/Runtime/CoreUObject/Private/Serialization/AsyncPackageLoader.cpp b/Engine/Source/Runtime/CoreUObject/Private/Serialization/AsyncPackageLoader.cpp
--- a/Engine/Source/Runtime/CoreUObject/Private/Serialization/AsyncPackageLoader.cpp
+++ b/Engine/Source/Runtime/CoreUObject/Private/Serialization/AsyncPackageLoader.cpp
@@ -41,9 +41,9 @@ public:
 		RecordedEvents.Reserve(16388);
 	}
 
-	void NotifyRegistrationEvent(const TCHAR* PackageName, const TCHAR* Name, ENotifyRegistrationType NotifyRegistrationType, ENotifyRegistrationPhase NotifyRegistrationPhase, UObject* (*RegisterFunc)(), bool bDynamic)
+	void NotifyRegistrationEvent(const TCHAR* PackageName, const TCHAR* Name, ENotifyRegistrationType NotifyRegistrationType, ENotifyRegistrationPhase NotifyRegistrationPhase, UObject* (*RegisterFunc)(), bool bDynamic, UObject* FinishedObject)
 	{
-		RecordedEvents.Add({ PackageName, Name, NotifyRegistrationType, NotifyRegistrationPhase, RegisterFunc, bDynamic });
+		RecordedEvents.Add({ PackageName, Name, NotifyRegistrationType, NotifyRegistrationPhase, RegisterFunc, FinishedObject, bDynamic });
 	}
 
 	void Replay(IAsyncPackageLoader& PackageLoader)
@@ -51,7 +51,7 @@ public:
 		UE_LOG(LogStreaming, Log, TEXT("NotifyRegistrationEvent: Replay %d entries"), RecordedEvents.Num());
 		for (const FRecordedEvent& Event : RecordedEvents)
 		{
-			PackageLoader.NotifyRegistrationEvent(*Event.PackageName, *Event.Name, Event.NotifyRegistrationType, Event.NotifyRegistrationPhase, Event.RegisterFunc, Event.bDynamic);
+			PackageLoader.NotifyRegistrationEvent(*Event.PackageName, *Event.Name, Event.NotifyRegistrationType, Event.NotifyRegistrationPhase, Event.RegisterFunc, Event.bDynamic, Event.FinishedObject);
 		}
 		RecordedEvents.Empty();
 	}
@@ -64,6 +64,7 @@ private:
 		ENotifyRegistrationType NotifyRegistrationType;
 		ENotifyRegistrationPhase NotifyRegistrationPhase;
 		UObject* (*RegisterFunc)();
+		UObject* FinishedObject;
 		bool bDynamic;
 	};
 
@@ -639,22 +640,23 @@ float GetAsyncLoadPercentage(const FName& PackageName)
 	return GetAsyncPackageLoader().GetAsyncLoadPercentage(PackageName);
 }
 
-void NotifyRegistrationEvent(const TCHAR* PackageName, const TCHAR* Name, ENotifyRegistrationType NotifyRegistrationType, ENotifyRegistrationPhase NotifyRegistrationPhase, UObject *(*InRegister)(), bool InbDynamic)
+void NotifyRegistrationEvent(const TCHAR* PackageName, const TCHAR* Name, ENotifyRegistrationType NotifyRegistrationType, ENotifyRegistrationPhase NotifyRegistrationPhase, UObject *(*InRegister)(), bool InbDynamic, UObject* FinishedObject)
 {
 	LLM_SCOPE(ELLMTag::AsyncLoading);
 
 	if (GPackageLoader)
 	{
-		GPackageLoader->NotifyRegistrationEvent(PackageName, Name, NotifyRegistrationType, NotifyRegistrationPhase, InRegister, InbDynamic);
+		GPackageLoader->NotifyRegistrationEvent(PackageName, Name, NotifyRegistrationType, NotifyRegistrationPhase, InRegister, InbDynamic, FinishedObject);
 	}
 	else
 	{
-		GetEarlyRegistrationEventsRecorder().NotifyRegistrationEvent(PackageName, Name, NotifyRegistrationType, NotifyRegistrationPhase, InRegister, InbDynamic);
+		GetEarlyRegistrationEventsRecorder().NotifyRegistrationEvent(PackageName, Name, NotifyRegistrationType, NotifyRegistrationPhase, InRegister, InbDynamic, FinishedObject);
 	}
 }
 
 void NotifyRegistrationComplete()
 {
+	TRACE_CPUPROFILER_EVENT_SCOPE(NotifyRegistrationComplete);
 	LLM_SCOPE(ELLMTag::AsyncLoading);
 	GPackageLoader->NotifyRegistrationComplete();
 	FlushAsyncLoading();
diff --git a/Engine/Source/Runtime/CoreUObject/Private/Serialization/EditorPackageLoader.cpp b/Engine/Source/Runtime/CoreUObject/Private/Serialization/EditorPackageLoader.cpp
--- a/Engine/Source/Runtime/CoreUObject/Private/Serialization/EditorPackageLoader.cpp
+++ b/Engine/Source/Runtime/CoreUObject/Private/Serialization/EditorPackageLoader.cpp
@@ -199,15 +199,16 @@ public:
 		ENotifyRegistrationType NotifyRegistrationType,
 		ENotifyRegistrationPhase NotifyRegistrationPhase,
 		UObject* (*InRegister)(),
-		bool InbDynamic) override
+		bool InbDynamic,
+		UObject* FinishedObject) override
 	{
 		if (UncookedPackageLoader)
 		{
-			UncookedPackageLoader->NotifyRegistrationEvent(PackageName, Name, NotifyRegistrationType, NotifyRegistrationPhase, InRegister, InbDynamic);
+			UncookedPackageLoader->NotifyRegistrationEvent(PackageName, Name, NotifyRegistrationType, NotifyRegistrationPhase, InRegister, InbDynamic, FinishedObject);
 		}
 		if (CookedPackageLoader)
 		{
-			CookedPackageLoader->NotifyRegistrationEvent(PackageName, Name, NotifyRegistrationType, NotifyRegistrationPhase, InRegister, InbDynamic);
+			CookedPackageLoader->NotifyRegistrationEvent(PackageName, Name, NotifyRegistrationType, NotifyRegistrationPhase, InRegister, InbDynamic, FinishedObject);
 		}
 	}
 
diff --git a/Engine/Source/Runtime/CoreUObject/Private/UObject/Class.cpp b/Engine/Source/Runtime/CoreUObject/Private/UObject/Class.cpp
--- a/Engine/Source/Runtime/CoreUObject/Private/UObject/Class.cpp
+++ b/Engine/Source/Runtime/CoreUObject/Private/UObject/Class.cpp
@@ -4140,7 +4140,7 @@ UObject* UClass::CreateDefaultObject()
 					TCHAR CDOName[FName::StringBufferSize];
 					GetOutermost()->GetFName().ToString(PackageName);
 					GetDefaultObjectName().ToString(CDOName);
-					NotifyRegistrationEvent(PackageName, CDOName, ENotifyRegistrationType::NRT_ClassCDO, ENotifyRegistrationPhase::NRP_Finished);
+					NotifyRegistrationEvent(PackageName, CDOName, ENotifyRegistrationType::NRT_ClassCDO, ENotifyRegistrationPhase::NRP_Finished, nullptr, false, ClassDefaultObject);
 				}
 				ClassDefaultObject->PostCDOContruct();
 			}
diff --git a/Engine/Source/Runtime/CoreUObject/Private/UObject/UObjectBase.cpp b/Engine/Source/Runtime/CoreUObject/Private/UObject/UObjectBase.cpp
--- a/Engine/Source/Runtime/CoreUObject/Private/UObject/UObjectBase.cpp
+++ b/Engine/Source/Runtime/CoreUObject/Private/UObject/UObjectBase.cpp
@@ -590,7 +590,7 @@ void UObjectCompiledInDeferStruct(class UScriptStruct* (*InRegister)(), const TC
 class UScriptStruct *GetStaticStruct(class UScriptStruct *(*InRegister)(), UObject* StructOuter, const TCHAR* StructName)
 {
 	UScriptStruct *Result = (*InRegister)();
-	NotifyRegistrationEvent(*StructOuter->GetOutermost()->GetName(), StructName, ENotifyRegistrationType::NRT_Struct, ENotifyRegistrationPhase::NRP_Finished);
+	NotifyRegistrationEvent(*StructOuter->GetOutermost()->GetName(), StructName, ENotifyRegistrationType::NRT_Struct, ENotifyRegistrationPhase::NRP_Finished, nullptr, false, Result);
 	return Result;
 }
 
@@ -614,7 +614,7 @@ void UObjectCompiledInDeferEnum(class UEnum* (*InRegister)(), const TCHAR* Packa
 class UEnum *GetStaticEnum(class UEnum *(*InRegister)(), UObject* EnumOuter, const TCHAR* EnumName)
 {
 	UEnum *Result = (*InRegister)();
-	NotifyRegistrationEvent(*EnumOuter->GetOutermost()->GetName(), EnumName, ENotifyRegistrationType::NRT_Enum, ENotifyRegistrationPhase::NRP_Finished);
+	NotifyRegistrationEvent(*EnumOuter->GetOutermost()->GetName(), EnumName, ENotifyRegistrationType::NRT_Enum, ENotifyRegistrationPhase::NRP_Finished, nullptr, false, Result);
 	return Result;
 }
 
@@ -822,10 +822,6 @@ static void UObjectLoadAllCompiledInDefaultProperties()
 			{
 				UE_LOG(LogUObjectBootstrap, Verbose, TEXT("UObjectLoadAllCompiledInDefaultProperties After Registrant %s %s"), PackageName, *Class.GetName());
 
-				TCHAR ClassName[FName::StringBufferSize];
-				Class.GetFName().ToString(ClassName);
-				NotifyRegistrationEvent(PackageName, ClassName, ENotifyRegistrationType::NRT_Class, ENotifyRegistrationPhase::NRP_Finished);
-
 				if (Class.GetOutermost()->GetFName() == GLongCoreUObjectPackageName)
 				{
 					NewClassesInCoreUObject.Add(&Class);
@@ -839,6 +835,27 @@ static void UObjectLoadAllCompiledInDefaultProperties()
 					NewClasses.Add(&Class);
 				}
 			}); 
+
+		auto NotifyClassFinishedRegistrationEvents = [](TArray<UClass*>& Classes)
+		{
+			for (UClass* Class : Classes)
+			{
+				TCHAR PackageName[FName::StringBufferSize];
+				TCHAR ClassName[FName::StringBufferSize];
+				Class->GetOutermost()->GetFName().ToString(PackageName);
+				Class->GetFName().ToString(ClassName);
+				NotifyRegistrationEvent(PackageName, ClassName, ENotifyRegistrationType::NRT_Class, ENotifyRegistrationPhase::NRP_Finished, nullptr, false, Class);
+			}
+		};
+
+		// notify async loader of all new classes before creating the class default objects
+		{
+			SCOPED_BOOT_TIMING("NotifyClassFinishedRegistrationEvents");
+			NotifyClassFinishedRegistrationEvents(NewClassesInCoreUObject);
+			NotifyClassFinishedRegistrationEvents(NewClassesInEngine);
+			NotifyClassFinishedRegistrationEvents(NewClasses);
+		}
+
 		{
 			SCOPED_BOOT_TIMING("CoreUObject Classes");
 			for (UClass* Class : NewClassesInCoreUObject) // we do these first because we assume these never trigger loads
diff --git a/Engine/Source/Runtime/CoreUObject/Private/UObject/UObjectGlobals.cpp b/Engine/Source/Runtime/CoreUObject/Private/UObject/UObjectGlobals.cpp
--- a/Engine/Source/Runtime/CoreUObject/Private/UObject/UObjectGlobals.cpp
+++ b/Engine/Source/Runtime/CoreUObject/Private/UObject/UObjectGlobals.cpp
@@ -5496,7 +5496,7 @@ namespace UECodeGen_Private
 				// Notify loader of new top level noexport objects like UScriptStruct, UDelegateFunction and USparseDelegateFunction
 				TCHAR ObjectName[FName::StringBufferSize];
 				Object->GetFName().ToString(ObjectName);
-				NotifyRegistrationEvent(PackageName, ObjectName, ENotifyRegistrationType::NRT_NoExportObject, ENotifyRegistrationPhase::NRP_Finished);
+				NotifyRegistrationEvent(PackageName, ObjectName, ENotifyRegistrationType::NRT_NoExportObject, ENotifyRegistrationPhase::NRP_Finished, nullptr, false, Object);
 			}
 		}
 #if WITH_RELOAD
diff --git a/Engine/Source/Runtime/CoreUObject/Public/Serialization/AsyncPackageLoader.h b/Engine/Source/Runtime/CoreUObject/Public/Serialization/AsyncPackageLoader.h
--- a/Engine/Source/Runtime/CoreUObject/Public/Serialization/AsyncPackageLoader.h
+++ b/Engine/Source/Runtime/CoreUObject/Public/Serialization/AsyncPackageLoader.h
@@ -221,7 +221,7 @@ public:
 
 	virtual void NotifyUnreachableObjects(const TArrayView<FUObjectItem*>& UnreachableObjects) = 0;
 
-	virtual void NotifyRegistrationEvent(const TCHAR* PackageName, const TCHAR* Name, ENotifyRegistrationType NotifyRegistrationType, ENotifyRegistrationPhase NotifyRegistrationPhase, UObject* (*InRegister)(), bool InbDynamic) = 0;
+	virtual void NotifyRegistrationEvent(const TCHAR* PackageName, const TCHAR* Name, ENotifyRegistrationType NotifyRegistrationType, ENotifyRegistrationPhase NotifyRegistrationPhase, UObject* (*InRegister)(), bool InbDynamic, UObject* FinishedObject) = 0;
 
 	virtual void NotifyRegistrationComplete() = 0;
 
diff --git a/Engine/Source/Runtime/CoreUObject/Public/UObject/LinkerLoad.h b/Engine/Source/Runtime/CoreUObject/Public/UObject/LinkerLoad.h
--- a/Engine/Source/Runtime/CoreUObject/Public/UObject/LinkerLoad.h
+++ b/Engine/Source/Runtime/CoreUObject/Public/UObject/LinkerLoad.h
@@ -1403,5 +1403,5 @@ enum class ENotifyRegistrationPhase
 	NRP_Finished,
 };
 
-COREUOBJECT_API void NotifyRegistrationEvent(const TCHAR* PackageName, const TCHAR* Name, ENotifyRegistrationType NotifyRegistrationType, ENotifyRegistrationPhase NotifyRegistrationPhase, UObject *(*InRegister)() = nullptr, bool InbDynamic = false);
+COREUOBJECT_API void NotifyRegistrationEvent(const TCHAR* PackageName, const TCHAR* Name, ENotifyRegistrationType NotifyRegistrationType, ENotifyRegistrationPhase NotifyRegistrationPhase, UObject *(*InRegister)() = nullptr, bool InbDynamic = false, UObject* FinishedObject = nullptr);
 COREUOBJECT_API void NotifyRegistrationComplete();
