AsyncLoading2: Only start the thread when it should be enabled (bug fix for 21402146)

#rnx

diff --git a/Engine/Source/Runtime/CoreUObject/Private/Serialization/AsyncLoading2.cpp b/Engine/Source/Runtime/CoreUObject/Private/Serialization/AsyncLoading2.cpp
--- a/Engine/Source/Runtime/CoreUObject/Private/Serialization/AsyncLoading2.cpp
+++ b/Engine/Source/Runtime/CoreUObject/Private/Serialization/AsyncLoading2.cpp
@@ -5404,7 +5404,7 @@ void FAsyncLoadingThread2::StartThread()
 	check(PendingCDOsRecursiveStack.Num() == 0);
 	PendingCDOsRecursiveStack.Empty();
 
-	if (!Thread)
+	if (FAsyncLoadingThreadSettings::Get().bAsyncLoadingThreadEnabled && !Thread)
 	{
 		UE_LOG(LogStreaming, Log, TEXT("Starting Async Loading Thread."));
 		bThreadStarted = true;
