Fix inverted check in FNumericProperty::CanHoldValue

Reported through UDN #00459500

diff --git a/Engine/Source/Runtime/CoreUObject/Public/UObject/UnrealType.h b/Engine/Source/Runtime/CoreUObject/Public/UObject/UnrealType.h
--- a/Engine/Source/Runtime/CoreUObject/Public/UObject/UnrealType.h
+++ b/Engine/Source/Runtime/CoreUObject/Public/UObject/UnrealType.h
@@ -1601,7 +1601,7 @@ PRAGMA_ENABLE_DEPRECATION_WARNINGS
 	template <typename T>
 	bool CanHoldValue(T Value) const
 	{
-		if (!TIsFloatingPoint<T>::Value)
+		if (TIsFloatingPoint<T>::Value)
 		{
 			//@TODO: FLOATPRECISION: This feels wrong, it might be losing precision before it tests to see if it's going to lose precision...
 			return CanHoldDoubleValueInternal((double)Value);
