Removing old unreal type registration methods.

diff --git a/Engine/Source/Runtime/CoreUObject/Private/UObject/DeferredRegistry.h b/Engine/Source/Runtime/CoreUObject/Private/UObject/DeferredRegistry.h
--- a/Engine/Source/Runtime/CoreUObject/Private/UObject/DeferredRegistry.h
+++ b/Engine/Source/Runtime/CoreUObject/Private/UObject/DeferredRegistry.h
@@ -34,8 +34,6 @@ void ReloadProcessObject(UClass* Object, const TCHAR* RenamePrefix);
 void ReloadProcessObject(UPackage* Object, const TCHAR* RenamePrefix);
 #endif
 
-PRAGMA_DISABLE_DEPRECATION_WARNINGS
-
 template <typename T>
 class TDeferredRegistry
 {
@@ -69,7 +67,6 @@ public:
 #if WITH_RELOAD
 		TType* OldSingleton;
 #endif
-		FFieldCompiledInInfo* DeprecatedFieldInfo;
 
 #if WITH_RELOAD
 		bool bHasChanged;
@@ -87,8 +84,7 @@ public:
 	/// <param name="InName">Name of the object</param>
 	/// <param name="InInfo">Persistent information about the object</param>
 	/// <param name="InVersion">Version information for this incarnation of the object</param>
-	/// <param name="DeprevatedFieldInfo">Deprecated factory object for the inner register for UClass.  Either InInnerRegister or this value must be specified for UClass</param>
-	AddResult AddRegistration(TType* (*InOuterRegister)(), TType* (*InInnerRegister)(), const TCHAR* InPackageName, const TCHAR* InName, TInfo& InInfo, const TVersion& InVersion, FFieldCompiledInInfo* DeprecatedFieldInfo)
+	AddResult AddRegistration(TType* (*InOuterRegister)(), TType* (*InInnerRegister)(), const TCHAR* InPackageName, const TCHAR* InName, TInfo& InInfo, const TVersion& InVersion)
 	{
 #if WITH_RELOAD
 		const FPackageAndNameKey Key = FPackageAndNameKey{ InPackageName, InName };
@@ -122,11 +118,11 @@ public:
 		}
 		if (bAdd)
 		{
-			Registrations.Add(FRegistrant{ InOuterRegister, InInnerRegister, InPackageName, &InInfo, OldSingleton, DeprecatedFieldInfo, bHasChanged });
+			Registrations.Add(FRegistrant{ InOuterRegister, InInnerRegister, InPackageName, &InInfo, OldSingleton, bHasChanged });
 		}
 		return ExistingInfo == nullptr ? AddResult::New : (bHasChanged ? AddResult::ExistingChanged : AddResult::ExistingNoChange);
 #else
-		Registrations.Add(FRegistrant{ InOuterRegister, InInnerRegister, InPackageName, &InInfo, DeprecatedFieldInfo });
+		Registrations.Add(FRegistrant{ InOuterRegister, InInnerRegister, InPackageName, &InInfo });
 		return AddResult::New;
 #endif
 	}
@@ -313,20 +309,6 @@ public:
 		NewInfo.OuterSingleton = OldInfo.OuterSingleton;
 	}
 
-	/**
-	* For the given package and name, return a registration information structure that can be
-	* used for compatibility with older API.  Will be removed after deprecated API removed.
-	*/
-	TInfo& MakeDeprecatedInfo(const TCHAR* InPackageName, const TCHAR* InName)
-	{
-		TUniquePtr<TInfo>& Info = CompatInfoMap.FindOrAdd(FPackageAndNameKey{ InPackageName, InName });
-		if (!Info.IsValid())
-		{
-			Info.Reset(new TInfo);
-		}
-		return *Info.Get();
-	}
-
 	/**
 	* Return the registry singleton
 	*/
@@ -342,7 +324,6 @@ public:
 
 private:
 	TArray<FRegistrant> Registrations;
-	TMap<FPackageAndNameKey, TUniquePtr<TInfo>> CompatInfoMap; // Can be removed after deprecated API removed
 	int32 ProcessedRegistrations = 0;
 
 #if WITH_RELOAD
@@ -354,10 +335,6 @@ private:
 template <>
 inline typename TDeferredRegistry<FClassRegistrationInfo>::TType* TDeferredRegistry<FClassRegistrationInfo>::InnerRegister(const TDeferredRegistry<FClassRegistrationInfo>::FRegistrant& Registrant)
 {
-	if (Registrant.DeprecatedFieldInfo)
-	{
-		return Registrant.DeprecatedFieldInfo->Register();
-	}
 	LLM_SCOPE(ELLMTag::UObject);
 	return Registrant.InnerRegisterFn();
 }
@@ -382,5 +359,3 @@ using FClassDeferredRegistry = TDeferredRegistry<FClassRegistrationInfo>;
 using FEnumDeferredRegistry = TDeferredRegistry<FEnumRegistrationInfo>;
 using FStructDeferredRegistry = TDeferredRegistry<FStructRegistrationInfo>;
 using FPackageDeferredRegistry = TDeferredRegistry<FPackageRegistrationInfo>;
-
-PRAGMA_ENABLE_DEPRECATION_WARNINGS
diff --git a/Engine/Source/Runtime/CoreUObject/Private/UObject/UObjectBase.cpp b/Engine/Source/Runtime/CoreUObject/Private/UObject/UObjectBase.cpp
--- a/Engine/Source/Runtime/CoreUObject/Private/UObject/UObjectBase.cpp
+++ b/Engine/Source/Runtime/CoreUObject/Private/UObject/UObjectBase.cpp
@@ -575,18 +575,10 @@ void UObjectForceRegistration(UObjectBase* Object, bool bCheckForModuleRelease)
 void RegisterCompiledInInfo(class UScriptStruct* (*InOuterRegister)(), const TCHAR* InPackageName, const TCHAR* InName, FStructRegistrationInfo& InInfo, const FStructReloadVersionInfo& InVersionInfo)
 {
 	check(InOuterRegister);
-	FStructDeferredRegistry::Get().AddRegistration(InOuterRegister, nullptr, InPackageName, InName, InInfo, InVersionInfo, nullptr);
+	FStructDeferredRegistry::Get().AddRegistration(InOuterRegister, nullptr, InPackageName, InName, InInfo, InVersionInfo);
 	NotifyRegistrationEvent(InPackageName, InName, ENotifyRegistrationType::NRT_Struct, ENotifyRegistrationPhase::NRP_Added, (UObject * (*)())(InOuterRegister), false);
 }
 
-// [DEPRECATED] - Legacy implementation that will be removed later. Also note that dynamic struct types are no longer supported here.
-void UObjectCompiledInDeferStruct(class UScriptStruct* (*InRegister)(), const TCHAR* PackageName, const TCHAR* ObjectName, bool /*bDynamic*/, const TCHAR* /*DynamicPathName*/)
-{
-	FStructDeferredRegistry& Registry = FStructDeferredRegistry::Get();
-	FStructRegistrationInfo& Info = Registry.MakeDeprecatedInfo(PackageName, ObjectName);
-	Registry.AddRegistration(nullptr, InRegister, PackageName, ObjectName, Info, FStructReloadVersionInfo{}, nullptr);
-}
-
 class UScriptStruct *GetStaticStruct(class UScriptStruct *(*InRegister)(), UObject* StructOuter, const TCHAR* StructName)
 {
 	UScriptStruct *Result = (*InRegister)();
@@ -599,18 +591,10 @@ class UScriptStruct *GetStaticStruct(class UScriptStruct *(*InRegister)(), UObje
 void RegisterCompiledInInfo(class UEnum* (*InOuterRegister)(), const TCHAR* InPackageName, const TCHAR* InName, FEnumRegistrationInfo& InInfo, const FEnumReloadVersionInfo& InVersionInfo)
 {
 	check(InOuterRegister);
-	FEnumDeferredRegistry::Get().AddRegistration(InOuterRegister, nullptr, InPackageName, InName, InInfo, InVersionInfo, nullptr);
+	FEnumDeferredRegistry::Get().AddRegistration(InOuterRegister, nullptr, InPackageName, InName, InInfo, InVersionInfo);
 	NotifyRegistrationEvent(InPackageName, InName, ENotifyRegistrationType::NRT_Enum, ENotifyRegistrationPhase::NRP_Added, (UObject * (*)())(InOuterRegister), false);
 }
 
-// [DEPRECATED] - Legacy implementation that will be removed later. Also note that dynamic enum types are no longer supported here.
-void UObjectCompiledInDeferEnum(class UEnum* (*InRegister)(), const TCHAR* PackageName, const TCHAR* ObjectName, bool /*bDynamic*/, const TCHAR* /*DynamicPathName*/)
-{
-	FEnumDeferredRegistry& Registry = FEnumDeferredRegistry::Get();
-	FEnumRegistrationInfo& Info = Registry.MakeDeprecatedInfo(PackageName, ObjectName);
-	Registry.AddRegistration(InRegister, InRegister, PackageName, ObjectName, Info, FEnumReloadVersionInfo{}, nullptr);
-}
-
 class UEnum *GetStaticEnum(class UEnum *(*InRegister)(), UObject* EnumOuter, const TCHAR* EnumName)
 {
 	UEnum *Result = (*InRegister)();
@@ -652,7 +636,7 @@ void RegisterCompiledInInfo(class UClass* (*InOuterRegister)(), class UClass* (*
 {
 	check(InOuterRegister);
 	check(InInnerRegister);
-	FClassDeferredRegistry::AddResult result = FClassDeferredRegistry::Get().AddRegistration(InOuterRegister, InInnerRegister, InPackageName, InName, InInfo, InVersionInfo, nullptr);
+	FClassDeferredRegistry::AddResult result = FClassDeferredRegistry::Get().AddRegistration(InOuterRegister, InInnerRegister, InPackageName, InName, InInfo, InVersionInfo);
 #if WITH_RELOAD
 	if (result == FClassDeferredRegistry::AddResult::ExistingChanged && !IsReloadActive())
 	{
@@ -665,49 +649,13 @@ void RegisterCompiledInInfo(class UClass* (*InOuterRegister)(), class UClass* (*
 	NotifyRegistrationEvent(InPackageName, *(FString(DEFAULT_OBJECT_PREFIX) + NoPrefix), ENotifyRegistrationType::NRT_ClassCDO, ENotifyRegistrationPhase::NRP_Added, (UObject * (*)())(InOuterRegister), false);
 }
 
-PRAGMA_DISABLE_DEPRECATION_WARNINGS
-void UClassCompiledInDefer(FFieldCompiledInInfo* ClassInfo, const TCHAR* Name, SIZE_T ClassSize, uint32 Crc)
-{
-	FClassDeferredRegistry& Registry = FClassDeferredRegistry::Get();
-	FClassRegistrationInfo& Info = Registry.MakeDeprecatedInfo(ClassInfo->ClassPackage(), Name);
-	Registry.AddRegistration(nullptr, nullptr, ClassInfo->ClassPackage(), Name, Info, CONSTRUCT_RELOAD_VERSION_INFO(FClassReloadVersionInfo, ClassSize, Crc), ClassInfo);
-}
-PRAGMA_ENABLE_DEPRECATION_WARNINGS
-
-// [DEPRECATED] - Legacy implementation that will be removed later. Also note that dynamic class types are no longer supported here.
-void UObjectCompiledInDefer(UClass *(*InRegister)(), UClass *(*InStaticClass)(), const TCHAR* Name, const TCHAR* PackageName, bool /*bDynamic*/, const TCHAR* /*DynamicPathName*/, void (*InInitSearchableValues)(TMap<FName, FName>&))
-{
-	// Either add all classes if not hot-reloading, or those which have changed
-	FClassDeferredRegistry& ClassRegistry = FClassDeferredRegistry::Get();
-	FClassRegistrationInfo& Info = ClassRegistry.MakeDeprecatedInfo(PackageName, Name);
-
-	// This is slow, but is deprecated.
-	for (FClassDeferredRegistry::FRegistrant& Registrant : ClassRegistry.GetRegistrations())
-	{
-		if (Registrant.Info == &Info)
-		{
-#if WITH_RELOAD
-			if (Registrant.bHasChanged)
-#endif
-			{
-				// This will probably need to be moved
-				FString NoPrefix(UObjectBase::RemoveClassPrefix(Name));
-				NotifyRegistrationEvent(PackageName, *NoPrefix, ENotifyRegistrationType::NRT_Class, ENotifyRegistrationPhase::NRP_Added, (UObject * (*)())(InRegister), false);
-				NotifyRegistrationEvent(PackageName, *(FString(DEFAULT_OBJECT_PREFIX) + NoPrefix), ENotifyRegistrationType::NRT_ClassCDO, ENotifyRegistrationPhase::NRP_Added, (UObject * (*)())(InRegister), false);
-			}
-			Registrant.OuterRegisterFn = InRegister;
-			break;
-		}
-	}
-}
-
 // UPackage registration
 
 void RegisterCompiledInInfo(UPackage* (*InOuterRegister)(), const TCHAR* InPackageName, FPackageRegistrationInfo& InInfo, const FPackageReloadVersionInfo& InVersionInfo)
 {
 #if WITH_RELOAD
 	check(InOuterRegister);
-	FPackageDeferredRegistry::Get().AddRegistration(reinterpret_cast<class UPackage* (*)()>(InOuterRegister), nullptr, TEXT(""), InPackageName, InInfo, InVersionInfo, nullptr);
+	FPackageDeferredRegistry::Get().AddRegistration(reinterpret_cast<class UPackage* (*)()>(InOuterRegister), nullptr, TEXT(""), InPackageName, InInfo, InVersionInfo);
 #endif
 }
 
diff --git a/Engine/Source/Runtime/CoreUObject/Public/UObject/UObjectBase.h b/Engine/Source/Runtime/CoreUObject/Public/UObject/UObjectBase.h
--- a/Engine/Source/Runtime/CoreUObject/Public/UObject/UObjectBase.h
+++ b/Engine/Source/Runtime/CoreUObject/Public/UObject/UObjectBase.h
@@ -368,12 +368,6 @@ COREUOBJECT_API void RegisterCompiledInInfo(class UScriptStruct* (*InOuterRegist
  */
 COREUOBJECT_API class UScriptStruct* GetStaticStruct(class UScriptStruct* (*InRegister)(), UObject* StructOuter, const TCHAR* StructName);
 
-UE_DEPRECATED(5.0, "GetStaticStruct with size and hash has been deprecated, use the version without the size and hash.")
-inline class UScriptStruct* GetStaticStruct(class UScriptStruct* (*InRegister)(), UObject* StructOuter, const TCHAR* StructName, SIZE_T Size, uint32 Hash)
-{
-	return GetStaticStruct(InRegister, StructOuter, StructName);
-}
-
 /**
  * Reload version information for enumerations
  */
@@ -437,115 +431,6 @@ COREUOBJECT_API void RegisterCompiledInInfo(UPackage* (*InOuterRegister)(), cons
  */
 COREUOBJECT_API void RegisterCompiledInInfo(const TCHAR* PackageName, const FClassRegisterCompiledInInfo* ClassInfo, size_t NumClassInfo, const FStructRegisterCompiledInInfo* StructInfo, size_t NumStructInfo, const FEnumRegisterCompiledInInfo* EnumInfo, size_t NumEnumInfo);
 
-/**
- * Base class for deferred native class registration
- */
-struct UE_DEPRECATED(5.0, "FFieldCompiledInInfo has been deprecated, use the RegistrationInfo structures below.") FFieldCompiledInInfo
-{
-	FFieldCompiledInInfo(SIZE_T InClassSize, uint32 InCrc)
-		: Size(InClassSize)
-		, Crc(InCrc)
-		, OldClass(nullptr)
-		, bHasChanged(false)
-	{
-	}
-
-	/** Registers the native class (constructs a UClass object) */
-	virtual UClass* Register() const = 0;
-
-	/** Return the package the class belongs in */
-	virtual const TCHAR* ClassPackage() const = 0;
-
-	/** Size of the class */
-	SIZE_T Size;
-	/** CRC of the generated code for this class */
-	uint32 Crc;
-	/** Old UClass object */
-	UClass* OldClass;
-	/** True if this class has changed after hot-reload (or new class) */
-	bool bHasChanged;
-};
-
-/**
-* Adds a class to deferred registration queue.
-*/
-PRAGMA_DISABLE_DEPRECATION_WARNINGS
-UE_DEPRECATED(5.0, "UClassCompiledInDefer has been deprecated, use RegisterCompiledInInfo.")
-COREUOBJECT_API void UClassCompiledInDefer(FFieldCompiledInInfo* Class, const TCHAR* Name, SIZE_T ClassSize, uint32 Crc);
-PRAGMA_ENABLE_DEPRECATION_WARNINGS
-
-/**
- * Specialized version of the deferred class registration structure.
- */
-PRAGMA_DISABLE_DEPRECATION_WARNINGS
-template <typename TClass>
-struct UE_DEPRECATED(5.0, "TClassCompiledInDefer has been deprecated, use FRegisterCompiledInInfo.") TClassCompiledInDefer : public FFieldCompiledInInfo
-{
-	TClassCompiledInDefer(const TCHAR* InName, SIZE_T InClassSize, uint32 InCrc)
-	: FFieldCompiledInInfo(InClassSize, InCrc)
-	{
-		UClassCompiledInDefer(this, InName, InClassSize, InCrc);
-	}
-	virtual UClass* Register() const override
-	{
-        LLM_SCOPE(ELLMTag::UObject);
-		return TClass::StaticClass();
-	}
-	virtual const TCHAR* ClassPackage() const override
-	{
-		return TClass::StaticPackage();
-	}
-};
-PRAGMA_ENABLE_DEPRECATION_WARNINGS
-
-/**
- * Stashes the singleton function that builds a compiled in class. Later, this is executed.
- */
-UE_DEPRECATED(5.0, "UObjectCompiledInDefer has been deprecated, use RegisterCompiledInInfo.")
-COREUOBJECT_API void UObjectCompiledInDefer(class UClass *(*InRegister)(), class UClass *(*InStaticClass)(), const TCHAR* Name, const TCHAR* PackageName, bool bDynamic, const TCHAR* DynamicPathName, void (*InInitSearchableValues)(TMap<FName, FName>&));
-
-struct UE_DEPRECATED(5.0, "FCompiledInDefer has been deprecated, use FRegisterCompiledInInfo.") FCompiledInDefer
-{
-	FCompiledInDefer(class UClass *(*InRegister)(), class UClass *(*InStaticClass)(), const TCHAR* PackageName, const TCHAR* Name, bool bDynamic, const TCHAR* DynamicPackageName = nullptr, const TCHAR* DynamicPathName = nullptr, void (*InInitSearchableValues)(TMap<FName, FName>&) = nullptr)
-	{
-		PRAGMA_DISABLE_DEPRECATION_WARNINGS
-		UObjectCompiledInDefer(InRegister, InStaticClass, Name, PackageName, bDynamic, DynamicPathName, InInitSearchableValues);
-		PRAGMA_ENABLE_DEPRECATION_WARNINGS
-	}
-};
-
-/**
- * Stashes the singleton function that builds a compiled in struct (StaticStruct). Later, this is executed.
- */
-UE_DEPRECATED(5.0, "UObjectCompiledInDeferStruct has been deprecated, use RegisterCompiledInInfo.")
-COREUOBJECT_API void UObjectCompiledInDeferStruct(class UScriptStruct *(*InRegister)(), const TCHAR* PackageName, const TCHAR* ObjectName, bool bDynamic, const TCHAR* DynamicPathName);
-
-struct UE_DEPRECATED(5.0, "FCompiledInDeferStruct has been deprecated, use FRegisterCompiledInInfo.") FCompiledInDeferStruct
-{
-	FCompiledInDeferStruct(class UScriptStruct *(*InRegister)(), const TCHAR* PackageName, const TCHAR* Name, bool bDynamic, const TCHAR* DynamicPackageName, const TCHAR* DynamicPathName)
-	{
-		PRAGMA_DISABLE_DEPRECATION_WARNINGS
-		UObjectCompiledInDeferStruct(InRegister, PackageName, Name, bDynamic, DynamicPathName);
-		PRAGMA_ENABLE_DEPRECATION_WARNINGS
-	}
-};
-
-/**
- * Stashes the singleton function that builds a compiled in enum. Later, this is executed.
- */
-UE_DEPRECATED(5.0, "UObjectCompiledInDeferEnum has been deprecated, use RegisterCompiledInInfo.")
-COREUOBJECT_API void UObjectCompiledInDeferEnum(class UEnum *(*InRegister)(), const TCHAR* PackageName, const TCHAR* ObjectName, bool bDynamic, const TCHAR* DynamicPathName);
-
-struct UE_DEPRECATED(5.0, "FCompiledInDeferEnum has been deprecated, use FRegisterCompiledInInfo.") FCompiledInDeferEnum
-{
-	FCompiledInDeferEnum(class UEnum *(*InRegister)(), const TCHAR* PackageName, const TCHAR* Name, bool bDynamic, const TCHAR* DynamicPackageName, const TCHAR* DynamicPathName)
-	{
-		PRAGMA_DISABLE_DEPRECATION_WARNINGS
-		UObjectCompiledInDeferEnum(InRegister, PackageName, Name, bDynamic, DynamicPathName);
-		PRAGMA_ENABLE_DEPRECATION_WARNINGS
-	}
-};
-
 // @todo: BP2CPP_remove
 /** Called during HotReload to hook up an existing structure */
 UE_DEPRECATED(5.0, "This API is no longer in use and will be removed.")
