[Core] make tests work in Test config
added ifdef to FObjectPathId to enabled for WITH_LOW_LEVEL_TESTS
changed ObjectHandleTest to use test data to avoid having to working with cooked, non-cooked or editor only data

#p4v-preflight-copy 21321644

diff --git a/Engine/Source/Runtime/CoreUObject/Tests/ObjectHandleTest.cpp b/Engine/Source/Runtime/CoreUObject/Tests/ObjectHandleTest.cpp
--- a/Engine/Source/Runtime/CoreUObject/Tests/ObjectHandleTest.cpp
+++ b/Engine/Source/Runtime/CoreUObject/Tests/ObjectHandleTest.cpp
@@ -158,35 +158,44 @@ TEST_CASE_METHOD(FObjectHandleTestBase, "CoreUObject::FObjectHandle::Pointer Beh
 	ObjectRefMetrics.TestNumReads(TEXT("NumReads should be incremented by one after a resolve attempt on a pointer handle"),1);
 }
 
-TEST_CASE_METHOD(FObjectHandleTestBase, "CoreUObject::FObjectHandle::Resolve Engine Content Target", "[CoreUObject][ObjectHandle][.Engine]")
+TEST_CASE_METHOD(FObjectHandleTestBase, "CoreUObject::FObjectHandle::Resolve Engine Content Target", "[CoreUObject][ObjectHandle]")
 {
-	// Confirm we successfully resolve a correct reference to engine content
-	TestResolvableNonNull("/Engine/EngineResources/DefaultTexture", "DefaultTexture");
-
-	// @TODO: OBJPTR: These assets aren't in a standard cook of EngineTest, so avoid testing them when using cooked content.  Should look for other assets to use instead.
-	if (!FPlatformProperties::RequiresCookedData())
-	{
-		// Confirm we successfully resolve a correct reference to a subobject in engine content
-		TestResolvableNonNull("/Engine/FunctionalTesting/Blueprints/AITesting_MoveGoal", "AITesting_MoveGoal.EventGraph.K2Node_VariableGet_142", nullptr, nullptr, true);
-
-		// Attempt to load something that uses a User Defined Enum
-		TestResolvableNonNull("/Engine/ArtTools/RenderToTexture/Macros/RenderToTextureMacros", "RenderToTextureMacros:Array to HLSL Float Array.K2Node_Select_1", nullptr, nullptr, true);
-	}
+	const FName TestPackageName(TEXT("/Engine/Test/ObjectPtrDefaultSerialize/Transient"));
+	UPackage* TestPackage = NewObject<UPackage>(nullptr, TestPackageName, RF_Transient);
+	TestPackage->AddToRoot();
+	UObject* TestSoftObject = NewObject<UMetaData>(TestPackage, TEXT("DefaultSerializeObject"));
+	UObject* TestSubObject = NewObject<UMetaData>(TestSoftObject, TEXT("SubObject"));
+	ON_SCOPE_EXIT{
+		TestPackage->RemoveFromRoot();
+	};
+
+	// Confirm we successfully resolve a correct reference to a subobject
+	TestResolvableNonNull("/Engine/Test/ObjectPtrDefaultSerialize/Transient", "DefaultSerializeObject.SubObject", nullptr, nullptr, true);
+
+	TestResolvableNonNull("/Engine/Test/ObjectPtrDefaultSerialize/Transient", "DefaultSerializeObject");
 }
 
 TEST_CASE_METHOD(FObjectHandleTestBase, "CoreUObject::FObjectHandle::Resolve Non Existent Target", "[CoreUObject][ObjectHandle]")
 {
 	// Confirm we don't successfully resolve an incorrect reference to engine content
 	TestResolveFailure("/Engine/EngineResources/NonExistentPackageName_0", "DefaultTexture");
-	TestResolveFailure("/Engine/EngineResources/DefaultTexture", "NonExistentObject_0");
+
+	const FName TestPackageName(TEXT("/Engine/Test/ObjectPtrDefaultSerialize/Transient"));
+	UPackage* TestPackage = NewObject<UPackage>(nullptr, TestPackageName, RF_Transient);
+	TestPackage->AddToRoot();
+	UObject* TestSoftObject = NewObject<UMetaData>(TestPackage, TEXT("DefaultSerializeObject"));
+	ON_SCOPE_EXIT{
+		TestPackage->RemoveFromRoot();
+	};
+
+	TestResolveFailure("/Engine/Test/ObjectPtrDefaultSerialize/Transient", "DefaultSerializeObject_DoesNotExist");
 }
 
 
-TEST_CASE_METHOD(FObjectHandleTestBase, "CoreUObject::FObjectHandle::Resolve Script Target", "[CoreUObject][ObjectHandle][.Engine]")
+TEST_CASE_METHOD(FObjectHandleTestBase, "CoreUObject::FObjectHandle::Resolve Script Target", "[CoreUObject][ObjectHandle]")
 {
 	// Confirm we successfully resolve a correct reference to engine content
-	TestResolvableNonNull("/Script/Engine", "Default__Actor");
-	TestResolvableNonNull("/Script/Engine", "DefaultPawn");
+	TestResolvableNonNull("/Script/CoreUObject", "MetaData");
 }
 
 #if UE_WITH_OBJECT_HANDLE_LATE_RESOLVE
