Don't try and remove objects in renamed packages from the loaded package store if the package is not in there.

diff --git a/Engine/Source/Runtime/CoreUObject/Private/Serialization/AsyncLoading2.cpp b/Engine/Source/Runtime/CoreUObject/Private/Serialization/AsyncLoading2.cpp
--- a/Engine/Source/Runtime/CoreUObject/Private/Serialization/AsyncLoading2.cpp
+++ b/Engine/Source/Runtime/CoreUObject/Private/Serialization/AsyncLoading2.cpp
@@ -5787,6 +5787,12 @@ void FAsyncLoadingThread2::OnLeakedPackageRename(UPackage* Package)
 		FlushAsyncLoading(); 
 	}
 
+	FLoadedPackageRef* PackageRef = LoadedPackageStore.FindPackageRef(Package->GetPackageId());
+	if (PackageRef == nullptr) 
+	{
+		return;
+	}
+
 	// Code such as LoadMap or LevelStreaming has renamed a loaded package which was detected as leaking so that we can load another copy of it.
 	// We should not have any loading happening at present, so we can remove these objects from our stores 
 	TArray<FUObjectItem*> LeakedObjectItems;
