Trying to catch mismatched object / internal flags when restoring flags in FScopedObjectFlagMarker::RestoreObjectFlags()

diff --git a/Engine/Source/Runtime/CoreUObject/Private/UObject/UObjectGlobals.cpp b/Engine/Source/Runtime/CoreUObject/Private/UObject/UObjectGlobals.cpp
--- a/Engine/Source/Runtime/CoreUObject/Private/UObject/UObjectGlobals.cpp
+++ b/Engine/Source/Runtime/CoreUObject/Private/UObject/UObjectGlobals.cpp
@@ -4210,16 +4210,18 @@ void FScopedObjectFlagMarker::RestoreObjectFlags()
 		UObject* Object = It.Key();
 		FStoredObjectFlags& PreviousObjectFlags = It.Value();
 
-		// clear all flags, frist clear the mirrored flags as we don't allow clearing them through ClearFlags
+		// clear all flags, first clear the mirrored flags as we don't allow clearing them through ClearFlags
 		Object->ClearGarbage(); // The currently mirrored flags are mutually exclusive and this will take care of both
 		Object->ClearFlags(RF_AllFlags);
 		Object->ClearInternalFlags(EInternalObjectFlags::AllFlags);
 
 		// then reset the ones that were originally set
-		if (!!(PreviousObjectFlags.InternalFlags & EInternalObjectFlags::MirroredFlags))
+		if (!!(PreviousObjectFlags.InternalFlags & EInternalObjectFlags::MirroredFlags) || !!(PreviousObjectFlags.Flags & RF_InternalMirroredFlags))
 		{
+			// Note that once an object is marked as Garbage (both in object and internal flags) it can't be marked as PendingKill and vice versa
 PRAGMA_DISABLE_DEPRECATION_WARNINGS
-			checkf(!!(PreviousObjectFlags.Flags & RF_InternalMirroredFlags), TEXT("Object %s had internal mirrored flag set but it was not matched in object flags"), *Object->GetFullName());
+			checkf(!!(PreviousObjectFlags.Flags & RF_InternalMirroredFlags), TEXT("%s had internal mirrored flag set but it was not matched in object flags"), *Object->GetFullName());
+			checkf(!!(PreviousObjectFlags.InternalFlags & EInternalObjectFlags::MirroredFlags), TEXT("%s had object mirrored flag set but it was not matched in internal flags"), *Object->GetFullName());
 PRAGMA_ENABLE_DEPRECATION_WARNINGS
 			Object->MarkAsGarbage();
 		}
