Add extra explanation of deduplicated garbage refs to summary.

diff --git a/Engine/Source/Runtime/CoreUObject/Private/UObject/GarbageCollection.cpp b/Engine/Source/Runtime/CoreUObject/Private/UObject/GarbageCollection.cpp
--- a/Engine/Source/Runtime/CoreUObject/Private/UObject/GarbageCollection.cpp
+++ b/Engine/Source/Runtime/CoreUObject/Private/UObject/GarbageCollection.cpp
@@ -2560,7 +2560,7 @@ void FGCArrayPool::DumpGarbageReferencers(TArray<FGCArrayStruct*>& AllArrays)
 				ArrayStruct->GarbageReferences.Reset();
 			}
 		}
-		UE_CLOG(TotalGarbageReferences > 0, LogGarbage, Log, TEXT("Reported %d/%d garbage references in %f ms."), ReportedGarbageReferences, TotalGarbageReferences, (FPlatformTime::Seconds() - StartTime) * 1000);
+		UE_CLOG(TotalGarbageReferences > 0, LogGarbage, Log, TEXT("Reported %d/%d garbage references in %f ms. (The rest were omitted as duplicate referencing class/property.)"), ReportedGarbageReferences, TotalGarbageReferences, (FPlatformTime::Seconds() - StartTime) * 1000);
 	}
 #endif
 }
