Bulk Data - deprecate StartAsyncLoading/CreateStreamingRequestForRange

- deprecated and replaced the usage of StartAsyncLoading and CreateStreamingRequestForRange with a new batch API for requesting bulk data. The
new API fits better with the I/O dispatcher API and does not assume that the bulk data resides in the same I/O chunk. The new API also
removes state handling away from the bulk data instances.

diff --git a/Engine/Plugins/Runtime/HairStrands/Source/HairStrandsCore/Private/GroomResources.cpp b/Engine/Plugins/Runtime/HairStrands/Source/HairStrandsCore/Private/GroomResources.cpp
--- a/Engine/Plugins/Runtime/HairStrands/Source/HairStrandsCore/Private/GroomResources.cpp
+++ b/Engine/Plugins/Runtime/HairStrands/Source/HairStrandsCore/Private/GroomResources.cpp
@@ -14,6 +14,7 @@
 #include "Serialization/LargeMemoryWriter.h"
 #include "Serialization/MemoryReader.h"
 #include "Serialization/MemoryWriter.h"
+#include "Serialization/BulkData.h"
 #include "UObject/PhysicsObjectVersion.h"
 #include "UObject/ReleaseObjectVersion.h"
 #include "NiagaraSystem.h"
@@ -540,11 +541,6 @@ void FHairCommonResource::StreamInLODData(int32 LODIndex)
 	}
 }
 
-void AsyncLoadHairBulkData(FByteBulkData& InData, bool& bIsLoading)
-{
-	bIsLoading = InData.StartAsyncLoading() || bIsLoading;
-}
-
 /////////////////////////////////////////////////////////////////////////////////////////
 
 void FHairCardIndexBuffer::InitRHI()
@@ -785,22 +781,28 @@ FHairStrandsRestResource::FHairStrandsRestResource(FHairStrandsBulkData& InBulkD
 
 bool FHairStrandsRestResource::InternalIsDataLoaded()
 {
-	bool bIsLoading = false;
+	if (BulkDataRequest.IsNone())
+	{
+		auto Builder = FBulkDataRequest::StreamToInstance();
+		Builder.Read(BulkData.Positions);
+		Builder.Read(BulkData.Attributes0);
+		Builder.Read(BulkData.Attributes1);
 
-	AsyncLoadHairBulkData(BulkData.Positions, bIsLoading);
-	AsyncLoadHairBulkData(BulkData.Attributes0, bIsLoading);
-	AsyncLoadHairBulkData(BulkData.Attributes1, bIsLoading);
+		if (!!(BulkData.Flags & FHairStrandsBulkData::DataFlags_HasMaterialData))
+		{
+			Builder.Read(BulkData.Materials);
+		}
 
-	if (!!(BulkData.Flags & FHairStrandsBulkData::DataFlags_HasMaterialData))
-	{
-		AsyncLoadHairBulkData(BulkData.Materials, bIsLoading);
+		BulkDataRequest = Builder.Issue();
 	}
-
-	return !bIsLoading;
+	
+	return BulkDataRequest.IsCompleted();
 }
 
 void FHairStrandsRestResource::InternalAllocate(FRDGBuilder& GraphBuilder)
 {
+	BulkDataRequest = FBulkDataRequest();
+
 	const uint32 PointCount = BulkData.PointCount;
 
 	// 1. Lock data, which force the loading data from files (on non-editor build/cooked data). These data are then uploaded to the GPU
@@ -1018,16 +1020,24 @@ FHairStrandsClusterCullingResource::FHairStrandsClusterCullingResource(FHairStra
 
 bool FHairStrandsClusterCullingResource::InternalIsDataLoaded()
 {
-	bool bIsLoading = false;
-	AsyncLoadHairBulkData(BulkData.PackedClusterInfos, bIsLoading);
-	AsyncLoadHairBulkData(BulkData.ClusterLODInfos, bIsLoading);
-	AsyncLoadHairBulkData(BulkData.VertexToClusterIds, bIsLoading);
-	AsyncLoadHairBulkData(BulkData.ClusterVertexIds, bIsLoading);
-	return !bIsLoading;
+	if (BulkDataRequest.IsNone())
+	{
+		BulkDataRequest = FBulkDataRequest::StreamToInstance()
+			.Read(BulkData.PackedClusterInfos)
+			.Read(BulkData.ClusterLODInfos)
+			.Read(BulkData.VertexToClusterIds)
+			.Read(BulkData.ClusterVertexIds)
+			.Issue();
+	}
+
+	return BulkDataRequest.IsCompleted();
 }
 
 void FHairStrandsClusterCullingResource::InternalAllocate(FRDGBuilder& GraphBuilder)
 {
+	check(BulkDataRequest.IsNone() == false)
+	BulkDataRequest = FBulkDataRequest();
+
 	if (GHairStrandsBulkData_Validation > 0)
 	{
 		BulkData.Validate(false);
@@ -1046,6 +1056,7 @@ void FHairStrandsClusterCullingResource::InternalRelease()
 	ClusterLODInfoBuffer.Release();
 	ClusterVertexIdBuffer.Release();
 	VertexToClusterIdBuffer.Release();
+	BulkDataRequest = FBulkDataRequest();
 }
 
 /////////////////////////////////////////////////////////////////////////////////////////
@@ -1069,16 +1080,24 @@ void FHairStrandsRestRootResource::PopulateFromRootData()
 		LOD.Status = FLOD::EStatus::Invalid;
 		LOD.SampleCount = MeshProjectionLOD.SampleCount;
 	}
+	LODRequests.SetNum(BulkData.MeshProjectionLODs.Num());
 }
 
 bool FHairStrandsRestRootResource::InternalIsDataLoaded()
 {
-	bool bIsLoading = false;
-	if (BulkData.PointCount > 0)
+	if (BulkData.PointCount == 0)
 	{
-		AsyncLoadHairBulkData(BulkData.VertexToCurveIndexBuffer, bIsLoading);
+		return true;
 	}
-	return !bIsLoading;
+
+	if (BulkDataRequest.IsNone())
+	{
+		BulkDataRequest = FBulkDataRequest::StreamToInstance()
+			.Read(BulkData.VertexToCurveIndexBuffer)
+			.Issue();
+	}
+
+	return BulkDataRequest.IsCompleted();
 }
 
 bool FHairStrandsRestRootResource::InternalIsLODDataLoaded(int32 LODIndex)
@@ -1088,26 +1107,38 @@ bool FHairStrandsRestRootResource::InternalIsLODDataLoaded(int32 LODIndex)
 	check(LODs.Num() == BulkData.MeshProjectionLODs.Num());
 	if (LODIndex >= 0 && LODIndex < LODs.Num())
 	{
-		FHairStrandsRootBulkData::FMeshProjectionLOD& CPUData = BulkData.MeshProjectionLODs[LODIndex];
-		const bool bHasValidCPUData = CPUData.RootBarycentricBuffer.GetBulkDataSize() > 0;
-		if (bHasValidCPUData)
+		FBulkDataRequest& LODRequest = LODRequests[LODIndex];
+		if (LODRequest.IsNone())
 		{
-			AsyncLoadHairBulkData(CPUData.RootBarycentricBuffer, bIsLoading);
-			AsyncLoadHairBulkData(CPUData.RootToUniqueTriangleIndexBuffer, bIsLoading);
-			AsyncLoadHairBulkData(CPUData.UniqueTriangleIndexBuffer, bIsLoading);
+			auto NewRequest = FBulkDataRequest::StreamToInstance();
+			FHairStrandsRootBulkData::FMeshProjectionLOD& CPUData = BulkData.MeshProjectionLODs[LODIndex];
+			const bool bHasValidCPUData = CPUData.RootBarycentricBuffer.GetBulkDataSize() > 0;
+			if (bHasValidCPUData)
+			{
+				NewRequest.Read(CPUData.RootBarycentricBuffer);
+				NewRequest.Read(CPUData.RootToUniqueTriangleIndexBuffer);
+				NewRequest.Read(CPUData.UniqueTriangleIndexBuffer);
 
-			AsyncLoadHairBulkData(CPUData.RestUniqueTrianglePosition0Buffer, bIsLoading);
-			AsyncLoadHairBulkData(CPUData.RestUniqueTrianglePosition1Buffer, bIsLoading);
-			AsyncLoadHairBulkData(CPUData.RestUniqueTrianglePosition2Buffer, bIsLoading);
+				NewRequest.Read(CPUData.RestUniqueTrianglePosition0Buffer);
+				NewRequest.Read(CPUData.RestUniqueTrianglePosition1Buffer);
+				NewRequest.Read(CPUData.RestUniqueTrianglePosition2Buffer);
+			}
+			
+			const bool bHasValidCPUWeights = CPUData.MeshSampleIndicesBuffer.GetBulkDataSize() > 0;
+			if (bHasValidCPUWeights)
+			{
+				NewRequest.Read(CPUData.MeshInterpolationWeightsBuffer);
+				NewRequest.Read(CPUData.MeshSampleIndicesBuffer);
+				NewRequest.Read(CPUData.RestSamplePositionsBuffer);
+			}
+			
+			if (bHasValidCPUData || bHasValidCPUWeights)
+			{
+				LODRequest = NewRequest.Issue();
+			}
 		}
 
-		const bool bHasValidCPUWeights = CPUData.MeshSampleIndicesBuffer.GetBulkDataSize() > 0;
-		if (bHasValidCPUWeights)
-		{
-			AsyncLoadHairBulkData(CPUData.MeshInterpolationWeightsBuffer, bIsLoading);
-			AsyncLoadHairBulkData(CPUData.MeshSampleIndicesBuffer, bIsLoading);
-			AsyncLoadHairBulkData(CPUData.RestSamplePositionsBuffer, bIsLoading);
-		}		
+		bIsLoading = LODRequest.IsCompleted() == false;
 	}
 
 	return !bIsLoading;
@@ -1142,6 +1173,8 @@ void FHairStrandsRestRootResource::InternalAllocateLOD(FRDGBuilder& GraphBuilder
 		{
 			return;
 		}
+		
+		LODRequests[LODIndex] = FBulkDataRequest();
 
 		FHairStrandsRootBulkData::FMeshProjectionLOD& CPUData = BulkData.MeshProjectionLODs[LODIndex];
 		const bool bHasValidCPUData = CPUData.RootBarycentricBuffer.GetBulkDataSize() > 0;
@@ -1209,6 +1242,7 @@ void FHairStrandsRestRootResource::InternalRelease()
 		GPUData.RestSamplePositionsBuffer.Release();
 	}
 	LODs.Empty();
+	LODRequests.Empty();
 }
 
 /////////////////////////////////////////////////////////////////////////////////////////
@@ -1476,20 +1510,23 @@ FHairStrandsInterpolationResource::FHairStrandsInterpolationResource(FHairStrand
 
 bool FHairStrandsInterpolationResource::InternalIsDataLoaded()
 {
-	bool bIsLoading = false;
-	const bool bUseSingleGuide = !!(BulkData.Flags & FHairStrandsInterpolationBulkData::DataFlags_HasSingleGuideData);
-	if (bUseSingleGuide)
-	{
-		AsyncLoadHairBulkData(BulkData.Interpolation, bIsLoading);
-	}
-	else
+	if (BulkDataRequest.IsNone())
 	{
-		AsyncLoadHairBulkData(BulkData.Interpolation0, bIsLoading);
-		AsyncLoadHairBulkData(BulkData.Interpolation1, bIsLoading);
+		auto NewRequest = FBulkDataRequest::StreamToInstance();
+		const bool bUseSingleGuide = !!(BulkData.Flags & FHairStrandsInterpolationBulkData::DataFlags_HasSingleGuideData);
+		if (bUseSingleGuide)
+		{
+			NewRequest.Read(BulkData.Interpolation);
+		}
+		else
+		{
+			NewRequest.Read(BulkData.Interpolation0);
+			NewRequest.Read(BulkData.Interpolation1);
+		}
+		BulkDataRequest = NewRequest.Read(BulkData.SimRootPointIndex).Issue();
 	}
-	AsyncLoadHairBulkData(BulkData.SimRootPointIndex, bIsLoading);
 
-	return !bIsLoading;
+	return BulkDataRequest.IsCompleted();
 }
 
 void FHairStrandsInterpolationResource::InternalAllocate(FRDGBuilder& GraphBuilder)
@@ -1513,6 +1550,7 @@ void FHairStrandsInterpolationResource::InternalRelease()
 	Interpolation0Buffer.Release();
 	Interpolation1Buffer.Release();
 	SimRootPointIndexBuffer.Release();
+	BulkDataRequest = FBulkDataRequest();
 }
 
 /////////////////////////////////////////////////////////////////////////////////////////
diff --git a/Engine/Plugins/Runtime/HairStrands/Source/HairStrandsCore/Public/GroomResources.h b/Engine/Plugins/Runtime/HairStrands/Source/HairStrandsCore/Public/GroomResources.h
--- a/Engine/Plugins/Runtime/HairStrands/Source/HairStrandsCore/Public/GroomResources.h
+++ b/Engine/Plugins/Runtime/HairStrands/Source/HairStrandsCore/Public/GroomResources.h
@@ -177,9 +177,15 @@ struct FHairStrandsRestRootResource : public FHairCommonResource
 	/* Store the hair projection information for each mesh LOD */
 	TArray<FLOD> LODs;
 
+	/* LOD bulk data requests */
+	TArray<FBulkDataRequest> LODRequests;
+
 	/* Store CPU data for root info & root binding */
 	FHairStrandsRootBulkData& BulkData;
 
+	/* Bulk data request handle */
+	FBulkDataRequest BulkDataRequest;
+
 	/* Type of curves */
 	const EHairStrandsResourcesType CurveType;
 };
@@ -338,6 +344,9 @@ struct FHairStrandsRestResource : public FHairCommonResource
 	/* Reference to the hair strands render data */
 	FHairStrandsBulkData& BulkData;
 
+	/* Handle to bulk data request */
+	FBulkDataRequest BulkDataRequest;
+
 	/* Type of curves */
 	const EHairStrandsResourcesType CurveType;
 
@@ -446,6 +455,9 @@ struct FHairStrandsClusterCullingResource : public FHairCommonResource
 	FRDGExternalBuffer ClusterVertexIdBuffer;
 
 	FHairStrandsClusterCullingBulkData& BulkData;
+
+	/* Handle to bulk data request */
+	FBulkDataRequest BulkDataRequest;
 };
 
 struct FHairStrandsInterpolationResource : public FHairCommonResource
@@ -481,6 +493,9 @@ struct FHairStrandsInterpolationResource : public FHairCommonResource
 
 	/* Reference to the hair strands interpolation render data */
 	FHairStrandsInterpolationBulkData& BulkData;
+
+	/* Handle to bulk data request */
+	FBulkDataRequest BulkDataRequest;
 };
 
 #if RHI_RAYTRACING
@@ -842,4 +857,4 @@ struct FHairMeshesDeformedResource : public FHairCommonResource
 ////////////////////////////////////////////////////////////////////////////////////////////////////////////////////////
 // Debug data (used for debug visalization but also for texture generation)
 void CreateHairStrandsDebugDatas(const FHairStrandsDatas& InData, FHairStrandsDebugDatas& Out);
-void CreateHairStrandsDebugResources(class FRDGBuilder& GraphBuilder, const FHairStrandsDebugDatas* In, FHairStrandsDebugDatas::FResources* Out);
\ No newline at end of file
+void CreateHairStrandsDebugResources(class FRDGBuilder& GraphBuilder, const FHairStrandsDebugDatas* In, FHairStrandsDebugDatas::FResources* Out);
diff --git a/Engine/Source/Runtime/Core/Public/IO/IoDispatcher.h b/Engine/Source/Runtime/Core/Public/IO/IoDispatcher.h
--- a/Engine/Source/Runtime/Core/Public/IO/IoDispatcher.h
+++ b/Engine/Source/Runtime/Core/Public/IO/IoDispatcher.h
@@ -801,6 +801,12 @@ public:
 		: RequestedOffset(InOffset)
 		, RequestedSize(InSize)
 	{ }
+	
+	FIoReadOptions(uint64 InOffset, uint64 InSize, void* InTargetVa)
+		: RequestedOffset(InOffset)
+		, RequestedSize(InSize)
+		, TargetVa(InTargetVa)
+	{ }
 
 	~FIoReadOptions() = default;
 
diff --git a/Engine/Source/Runtime/CoreUObject/Private/Serialization/BulkData.cpp b/Engine/Source/Runtime/CoreUObject/Private/Serialization/BulkData.cpp
--- a/Engine/Source/Runtime/CoreUObject/Private/Serialization/BulkData.cpp
+++ b/Engine/Source/Runtime/CoreUObject/Private/Serialization/BulkData.cpp
@@ -241,8 +241,8 @@ FBulkDataChunkId::FBulkDataChunkId()
 {
 }
 
-FBulkDataChunkId::FBulkDataChunkId(TUniquePtr<FImpl>&& InImplPtr)
-	: ImplPtr(MoveTemp(InImplPtr))
+FBulkDataChunkId::FBulkDataChunkId(TPimplPtr<FImpl>&& InImpl)
+	: Impl(MoveTemp(InImpl))
 {
 }
 
@@ -257,11 +257,11 @@ FBulkDataChunkId::FBulkDataChunkId(const FBulkDataChunkId& Other)
 
 FBulkDataChunkId& FBulkDataChunkId::operator=(const FBulkDataChunkId& Other)
 {
-	ImplPtr.Reset();
+	Impl.Reset();
 
-	if (Other.ImplPtr)
+	if (Other.Impl)
 	{
-		ImplPtr = MakeUnique<FImpl>(Other.ImplPtr->PathOrId);
+		Impl = MakePimpl<FImpl>(Other.Impl->PathOrId);
 	}
 
 	return *this;
@@ -269,9 +269,9 @@ FBulkDataChunkId& FBulkDataChunkId::operator=(const FBulkDataChunkId& Other)
 
 FPackageId FBulkDataChunkId::GetPackageId() const
 {
-	if (ImplPtr)
+	if (Impl)
 	{
-		if (const FPackageId* Id = ImplPtr->PathOrId.TryGet<FPackageId>())
+		if (const FPackageId* Id = Impl->PathOrId.TryGet<FPackageId>())
 		{
 			return *Id;
 		}
@@ -282,9 +282,9 @@ FPackageId FBulkDataChunkId::GetPackageId() const
 
 const FPackagePath& FBulkDataChunkId::GetPackagePath() const
 {
-	if (ImplPtr)
+	if (Impl)
 	{
-		if (const FPackagePath* Path = ImplPtr->PathOrId.TryGet<FPackagePath>())
+		if (const FPackagePath* Path = Impl->PathOrId.TryGet<FPackagePath>())
 		{
 			return *Path;
 		}
@@ -296,14 +296,14 @@ const FPackagePath& FBulkDataChunkId::GetPackagePath() const
 
 FIoFilenameHash FBulkDataChunkId::GetIoFilenameHash(EBulkDataFlags BulkDataFlags) const
 {
-	if (ImplPtr)
+	if (Impl)
 	{
-		if (const FPackageId* Id = ImplPtr->PathOrId.TryGet<FPackageId>())
+		if (const FPackageId* Id = Impl->PathOrId.TryGet<FPackageId>())
 		{
 			return MakeIoFilenameHash(CreateBulkDataIoChunkId(UE::BulkData::Private::FBulkMetaData(BulkDataFlags), *Id));
 		}
 
-		if (const FPackagePath* PackagePath = ImplPtr->PathOrId.TryGet<FPackagePath>())
+		if (const FPackagePath* PackagePath = Impl->PathOrId.TryGet<FPackagePath>())
 		{
 			return MakeIoFilenameHash(*PackagePath);
 		}
@@ -314,14 +314,14 @@ FIoFilenameHash FBulkDataChunkId::GetIoFilenameHash(EBulkDataFlags BulkDataFlags
 
 FString FBulkDataChunkId::ToDebugString() const
 {
-	if (ImplPtr)
+	if (Impl)
 	{
-		if (const FPackageId* Id = ImplPtr->PathOrId.TryGet<FPackageId>())
+		if (const FPackageId* Id = Impl->PathOrId.TryGet<FPackageId>())
 		{
 			return FString::Printf(TEXT("0x%llX"), Id->Value());
 		}
 
-		if (const FPackagePath* PackagePath = ImplPtr->PathOrId.TryGet<FPackagePath>())
+		if (const FPackagePath* PackagePath = Impl->PathOrId.TryGet<FPackagePath>())
 		{
 			return PackagePath->GetDebugName();
 		}
@@ -332,12 +332,12 @@ FString FBulkDataChunkId::ToDebugString() const
 
 FBulkDataChunkId FBulkDataChunkId::FromPackagePath(const FPackagePath& PackagePath)
 {
-	return FBulkDataChunkId(MakeUnique<FImpl>(PackagePath));
+	return FBulkDataChunkId(MakePimpl<FImpl>(PackagePath));
 }
 
 FBulkDataChunkId FBulkDataChunkId::FromPackageId(const FPackageId& PackageId)
 {
-	return FBulkDataChunkId(MakeUnique<FImpl>(PackageId));
+	return FBulkDataChunkId(MakePimpl<FImpl>(PackageId));
 }
 
 } // namespace UE::BulkData
diff --git a/Engine/Source/Runtime/CoreUObject/Private/Serialization/BulkDataStreaming.cpp b/Engine/Source/Runtime/CoreUObject/Private/Serialization/BulkDataStreaming.cpp
--- a/Engine/Source/Runtime/CoreUObject/Private/Serialization/BulkDataStreaming.cpp
+++ b/Engine/Source/Runtime/CoreUObject/Private/Serialization/BulkDataStreaming.cpp
@@ -1,8 +1,10 @@
 // Copyright Epic Games, Inc. All Rights Reserved.
+#include "Serialization/BulkData.h"
+#include "Algo/AllOf.h"
 #include "Async/MappedFileHandle.h"
 #include "Experimental/Async/LazyEvent.h"
 #include "HAL/CriticalSection.h"
-#include "Serialization/BulkData.h"
+#include "Misc/Timespan.h"
 #include "Serialization/MemoryReader.h"
 #include "UObject/PackageResourceManager.h"
 
@@ -767,3 +769,451 @@ void FlushAsyncLoad(FBulkData* Owner)
 }
 
 } // namespace UE::BulkData
+
+//////////////////////////////////////////////////////////////////////////////
+
+class FBulkDataRequest::FBase
+{
+public:
+	virtual ~FBase() = default;
+
+	virtual void Cancel() = 0;
+
+	virtual bool Issue(TArrayView<FReadRequest> Requests, EAsyncIOPriorityAndFlags Priority) = 0;
+
+	bool Wait(uint32 Milliseconds)
+	{
+		return DoneEvent.Wait(Milliseconds);
+	}
+
+	inline FBulkDataRequest::EStatus GetStatus() const
+	{
+		return static_cast<FBulkDataRequest::EStatus>(Status.load(std::memory_order_consume));
+	}
+
+protected:
+
+	FBase(FCompletionCallback&& Callback)
+		: DoneEvent(EEventMode::ManualReset)
+		, CompletionCallback(MoveTemp(Callback))
+	{
+	}
+
+	inline void SetStatus(FBulkDataRequest::EStatus InStatus)
+	{
+		Status.store(uint32(InStatus), std::memory_order_release); 
+	}
+
+	inline void CompleteRequest(FBulkDataRequest::EStatus CompletionStatus)
+	{
+		SetStatus(CompletionStatus);
+
+		if (CompletionCallback)
+		{
+			CompletionCallback(CompletionStatus);
+		}
+
+		DoneEvent.Trigger();
+	}
+
+	std::atomic<uint32>	Status;
+	UE::FLazyEvent DoneEvent;
+	FCompletionCallback CompletionCallback;
+};
+
+//////////////////////////////////////////////////////////////////////////////
+
+class FBulkDataRequest::FReadChunk final : public FBulkDataRequest::FBase
+{
+public:
+	FReadChunk(FCompletionCallback&& Callback)
+		: FBase(MoveTemp(Callback))
+	{
+	}
+
+	virtual ~FReadChunk()
+	{ 
+		Cancel();
+		Wait(MAX_uint32);
+	}
+
+	virtual void Cancel() override
+	{
+		if (GetStatus() == FBulkDataRequest::EStatus::Pending)
+		{
+			for (FIoRequest& Request : PendingRequests)
+			{
+				Request.Cancel();
+			}
+		}
+	}
+
+	virtual bool Issue(TArrayView<FReadRequest> Requests, EAsyncIOPriorityAndFlags Priority) override
+	{
+		check(Requests.Num() > 0);
+		
+		PendingRequests.Reserve(Requests.Num());
+
+		SetStatus(EStatus::Pending);
+
+		FIoBatch Batch = FIoDispatcher::Get().NewBatch();
+		const int32 IoDispatcherPriority = ConvertToIoDispatcherPriority(Priority);
+
+		for (FReadRequest& Request : Requests)
+		{
+			const FIoChunkId ChunkId = CreateBulkDataIoChunkId(Request.BulkMeta, Request.BulkChunkId.GetPackageId());
+			PendingRequests.Add(Batch.Read(ChunkId, FIoReadOptions(Request.Offset, Request.Size, Request.TargetVa), IoDispatcherPriority));
+		}
+
+		Batch.IssueWithCallback([this]()
+		{
+			FBulkDataRequest::EStatus BatchStatus = FBulkDataRequest::EStatus::Ok;
+
+			for (const FIoRequest& Request : PendingRequests)
+			{
+				if (EIoErrorCode ErrorCode = Request.Status().GetErrorCode(); ErrorCode != EIoErrorCode::Ok)
+				{
+					BatchStatus = ErrorCode == EIoErrorCode::Cancelled
+						? FBulkDataRequest::EStatus::Cancelled
+						: FBulkDataRequest::EStatus::Error;
+
+					break;
+				}
+			}
+			
+			PendingRequests.Empty();
+			CompleteRequest(BatchStatus);
+		});
+
+		return true;
+	}
+
+private:
+
+	TArray<FIoRequest, TInlineAllocator<4>> PendingRequests;
+};
+
+//////////////////////////////////////////////////////////////////////////////
+
+class FBulkDataRequest::FReadFile final : public FBulkDataRequest::FBase
+{
+public:
+	FReadFile(FCompletionCallback&& Callback)
+		: FBase(MoveTemp(Callback))
+	{
+	}
+
+	virtual ~FReadFile()
+	{ 
+		Cancel();
+		Wait(MAX_uint32);
+	}
+
+	virtual void Cancel() override
+	{
+		if (GetStatus() == FBulkDataRequest::EStatus::Pending)
+		{
+			for (TUniquePtr<IAsyncReadRequest>& Request : PendingRequests)
+			{
+				Request->Cancel();
+			}
+		}
+	}
+
+	virtual bool Issue(TArrayView<FReadRequest> Requests, EAsyncIOPriorityAndFlags Priority) override
+	{
+		using namespace UE::BulkData::Private;
+
+		check(Requests.Num() > 0);
+		
+		SetStatus(EStatus::Pending);
+		
+		PendingRequestCount = Requests.Num();
+		PendingStatus = FBulkDataRequest::EStatus::Ok;
+		
+		for (FReadRequest& Request : Requests)
+		{
+			IAsyncReadFileHandle* FileHandle = OpenFile(Request.BulkMeta, Request.BulkChunkId);
+
+			if (FileHandle == nullptr)
+			{
+				return false;
+			}
+
+			FAsyncFileCallBack ReadFileCallback = [this](bool bWasCancelled, IAsyncReadRequest* ReadRequest)
+			{
+				if (bWasCancelled)
+				{
+					PendingStatus = FBulkDataRequest::EStatus::Cancelled;
+				}
+
+				if (1 == PendingRequestCount.fetch_sub(1))
+				{
+					CompleteRequest(PendingStatus);
+				}
+			};
+
+			if (ReadFile(*FileHandle, Request.Offset, Request.Size, Priority, &ReadFileCallback, Request.TargetVa) == false)
+			{
+				SetStatus(EStatus::Error);
+
+				return false;
+			}
+		}
+
+		return true;
+	}
+
+private:
+
+	IAsyncReadFileHandle* OpenFile(const UE::BulkData::Private::FBulkMetaData& BulkMeta, const UE::BulkData::Private::FBulkDataChunkId& BulkChunkId)
+	{
+		const FPackagePath Path = BulkChunkId.GetPackagePath();
+
+		TUniquePtr<IAsyncReadFileHandle>& FileHandle = PathToFileHandle.FindOrAdd(Path.GetPackageFName());
+
+		if (FileHandle.IsValid() == false)
+		{
+			IPackageResourceManager& ResourceMgr = IPackageResourceManager::Get();
+
+			const bool bExternalResource = BulkMeta.HasAllFlags(static_cast<EBulkDataFlags>(BULKDATA_PayloadInSeperateFile | BULKDATA_WorkspaceDomainPayload));
+
+			if (bExternalResource)
+			{
+				FileHandle = ResourceMgr.OpenAsyncReadExternalResource(EPackageExternalResource::WorkspaceDomainFile, Path.GetPackageName()).Handle;
+			}
+			else
+			{
+				const EPackageSegment Segment = GetPackageSegmentFromFlags(BulkMeta);
+				FileHandle = ResourceMgr.OpenAsyncReadPackage(Path, Segment).Handle;
+			}
+		}
+
+		return FileHandle.Get();
+	}
+
+	bool ReadFile(IAsyncReadFileHandle& FileHandle, int64 Offset, int64 Size, EAsyncIOPriorityAndFlags Priority, FAsyncFileCallBack* Callback, void* Dst)
+	{
+		TUniquePtr<IAsyncReadRequest> ReadRequest(FileHandle.ReadRequest(Offset, Size, Priority, Callback, reinterpret_cast<uint8*>(Dst)));
+
+		if (ReadRequest.IsValid())
+		{
+			PendingRequests.Add(MoveTemp(ReadRequest));
+
+			return true;
+		}
+
+		return false;
+	}
+
+	using FPathToFileHandle = TMap<FName, TUniquePtr<IAsyncReadFileHandle>>;
+	using FReadRequests = TArray<TUniquePtr<IAsyncReadRequest>>;
+	
+	FPathToFileHandle PathToFileHandle;
+	FReadRequests PendingRequests;
+	FBulkDataRequest::EStatus PendingStatus;
+	std::atomic<int32> PendingRequestCount;
+};
+
+//////////////////////////////////////////////////////////////////////////////
+
+void FBulkDataRequest::FDeleter::operator()(FBase* InImpl)
+{
+	delete InImpl;
+}
+
+FBulkDataRequest::FBulkDataRequest(FBulkDataRequest::FBasePtr&& InImpl)
+	: Impl(MoveTemp(InImpl))
+{
+}
+
+FBulkDataRequest::EStatus FBulkDataRequest::GetStatus() const
+{
+	return Impl.IsValid() ? Impl->GetStatus() : EStatus::None;
+}
+
+bool FBulkDataRequest::IsNone() const
+{
+	return FBulkDataRequest::EStatus::None == GetStatus();
+}
+
+bool FBulkDataRequest::IsPending() const
+{
+	return FBulkDataRequest::EStatus::Pending == GetStatus();
+}
+
+bool FBulkDataRequest::IsOk() const
+{
+	return FBulkDataRequest::EStatus::Ok == GetStatus();
+}
+
+bool FBulkDataRequest::IsCompleted() const
+{
+	const uint32 Status = static_cast<uint32>(GetStatus());
+	return Status > static_cast<uint32>(EStatus::Pending);
+}
+
+void FBulkDataRequest::Wait()
+{
+	check(Impl);
+	Impl->Wait(MAX_uint32);
+}
+
+bool FBulkDataRequest::WaitFor(uint32 Milliseconds)
+{
+	check(Impl);
+	return Impl->Wait(Milliseconds);
+}
+
+bool FBulkDataRequest::WaitFor(const FTimespan& WaitTime)
+{
+	return WaitFor((uint32)FMath::Clamp<int64>(WaitTime.GetTicks() / ETimespan::TicksPerMillisecond, 0, MAX_uint32));
+}
+
+void FBulkDataRequest::Cancel()
+{
+	check(Impl);
+	Impl->Cancel();
+}
+
+void FBulkDataRequest::FRequestBuilder::AddRequest(
+	const UE::BulkData::Private::FBulkMetaData& BulkMeta,
+	const UE::BulkData::Private::FBulkDataChunkId& BulkChunkId,
+	uint64 Offset,
+	uint64 Size,
+	void* TargetVa)
+{
+	FReadRequest& Request = Requests.AddDefaulted_GetRef();
+
+	Request.BulkMeta	= BulkMeta;
+	Request.BulkChunkId = BulkChunkId;
+	Request.Offset		= Offset;
+	Request.Size		= Size;
+	Request.TargetVa	= TargetVa;
+
+	TotalRequestSize += Size;
+}
+
+FBulkDataRequest::FScatterGather& FBulkDataRequest::FScatterGather::Read(FBulkData& BulkData, uint64 Offset, uint64 Size)
+{
+	const uint64 BulkOffset	= uint64(BulkData.GetBulkDataOffsetInFile());
+	const uint64 BulkSize	= uint64(BulkData.GetBulkDataSize());
+
+	AddRequest(BulkData.BulkMeta, BulkData.BulkChunkId, BulkOffset + Offset, FMath::Min(BulkSize, Size), nullptr);
+
+	return *this;
+}
+
+FBulkDataRequest FBulkDataRequest::FScatterGather::Issue(FIoBuffer& Dst, FBulkDataRequest::FCompletionCallback&& Callback, EAsyncIOPriorityAndFlags Priority)
+{
+	if (Requests.IsEmpty())
+	{
+		return FBulkDataRequest();
+	}
+
+	check(TotalRequestSize > 0);
+
+	Dst = FIoBuffer(TotalRequestSize);
+	
+	FMutableMemoryView DstView = Dst.GetMutableView();
+	for (FReadRequest& Request : Requests)
+	{
+		Request.TargetVa = DstView.GetData();
+		DstView.RightChopInline(Request.Size);
+	}
+
+	return FBulkDataRequest::Issue(Requests, Priority, MoveTemp(Callback));
+}
+
+FBulkDataRequest::FStreamToInstance& FBulkDataRequest::FStreamToInstance::Read(FBulkData& BulkData)
+{
+	const uint64 BulkOffset	= uint64(BulkData.GetBulkDataOffsetInFile());
+	const uint64 BulkSize	= uint64(BulkData.GetBulkDataSize());
+	
+	AddRequest(BulkData.BulkMeta, BulkData.BulkChunkId, BulkOffset, BulkSize, nullptr);
+	Instances.Add(&BulkData);
+
+	return *this;
+}
+
+FBulkDataRequest FBulkDataRequest::FStreamToInstance::Issue(EAsyncIOPriorityAndFlags Priority)
+{
+	if (Requests.IsEmpty())
+	{
+		return FBulkDataRequest();
+	}
+
+	check(TotalRequestSize > 0);
+
+	FIoBuffer Dst = FIoBuffer(TotalRequestSize);
+	FMutableMemoryView DstView = Dst.GetMutableView();
+
+	for (FReadRequest& Request : Requests)
+	{
+		Request.TargetVa = DstView.GetData();
+		DstView.RightChopInline(Request.Size);
+	}
+
+	return FBulkDataRequest::Issue(
+		Requests,
+		Priority,
+		[Instances = MoveTemp(Instances), Dst = MoveTemp(Dst)](FBulkDataRequest::EStatus Status)
+		{
+			if (FBulkDataRequest::EStatus::Ok != Status)
+			{
+				return;
+			}
+
+			FMemoryReaderView Ar(Dst.GetView(), true);
+			for (FBulkData* BulkData : Instances)
+			{
+				const int64 BulkOffset	= BulkData->GetBulkDataOffsetInFile();
+				const int64 BulkSize	= BulkData->GetBulkDataSize();
+				void* Data				= BulkData->ReallocateData(BulkSize);
+
+				Ar.Seek(BulkOffset);
+				BulkData->SerializeBulkData(Ar, Data, BulkSize, EBulkDataFlags(BulkData->GetBulkDataFlags()));
+			}
+		});
+}
+
+FBulkDataRequest::FScatterGather FBulkDataRequest::ScatterGather()
+{
+	return FBulkDataRequest::FScatterGather();
+}
+
+FBulkDataRequest::FStreamToInstance FBulkDataRequest::StreamToInstance()
+{
+	return FBulkDataRequest::FStreamToInstance();
+}
+
+FBulkDataRequest FBulkDataRequest::Issue(TArrayView<FReadRequest> Requests, EAsyncIOPriorityAndFlags Priority, FBulkDataRequest::FCompletionCallback&& Callback)
+{
+	if (Requests.IsEmpty())
+	{
+		return FBulkDataRequest();
+	}
+
+	check(Algo::AllOf(Requests, [](const FReadRequest& Req)
+	{ 
+		return Req.BulkChunkId.IsValid() && Req.TargetVa != nullptr;
+	}));
+
+	FBasePtr NewRequest;
+
+	if (FPackageId PackageId = Requests[0].BulkChunkId.GetPackageId(); PackageId.IsValid())
+	{
+		NewRequest.Reset(new FBulkDataRequest::FReadChunk(MoveTemp(Callback)));
+	}
+	else
+	{
+		NewRequest.Reset(new FBulkDataRequest::FReadFile(MoveTemp(Callback)));
+	}
+
+	if (NewRequest->Issue(Requests, Priority) == false)
+	{
+		return FBulkDataRequest();
+	}
+
+	return FBulkDataRequest(MoveTemp(NewRequest));
+}
diff --git a/Engine/Source/Runtime/CoreUObject/Public/Serialization/BulkData.h b/Engine/Source/Runtime/CoreUObject/Public/Serialization/BulkData.h
--- a/Engine/Source/Runtime/CoreUObject/Public/Serialization/BulkData.h
+++ b/Engine/Source/Runtime/CoreUObject/Public/Serialization/BulkData.h
@@ -28,6 +28,7 @@
 #include "Templates/Function.h"
 #include "Templates/IsPODType.h"
 #include "Templates/UniquePtr.h"
+#include "Templates/PimplPtr.h"
 #include "UObject/NameTypes.h"
 
 class FLinkerLoad;
@@ -35,6 +36,7 @@ class FOutputDevice;
 class IAsyncReadFileHandle;
 class IAsyncReadRequest;
 class UObject;
+struct FTimespan;
 namespace UE { namespace Serialization { class FEditorBulkData; } }
 
 #if WITH_EDITOR == 0 && WITH_EDITORONLY_DATA == 0
@@ -499,17 +501,17 @@ class FBulkDataChunkId
 public:
 	COREUOBJECT_API FBulkDataChunkId();
 
-	FBulkDataChunkId(const FBulkDataChunkId& Other);
+	COREUOBJECT_API FBulkDataChunkId(const FBulkDataChunkId& Other);
 
-	FBulkDataChunkId(FBulkDataChunkId&&) = default;
+	COREUOBJECT_API FBulkDataChunkId(FBulkDataChunkId&&) = default;
 
 	COREUOBJECT_API ~FBulkDataChunkId();
 
-	FBulkDataChunkId& operator=(const FBulkDataChunkId& Other);
+	COREUOBJECT_API FBulkDataChunkId& operator=(const FBulkDataChunkId& Other);
 
-	FBulkDataChunkId& operator=(FBulkDataChunkId&&) = default;
+	COREUOBJECT_API FBulkDataChunkId& operator=(FBulkDataChunkId&&) = default;
 	
-	bool IsValid() const { return ImplPtr.IsValid(); }
+	bool IsValid() const { return Impl.IsValid(); }
 
 	FPackageId GetPackageId() const;
 	
@@ -526,9 +528,9 @@ public:
 private:
 	struct FImpl;
 
-	FBulkDataChunkId(TUniquePtr<FImpl>&& InImplPtr);
+	FBulkDataChunkId(TPimplPtr<FImpl>&& InImpl);
 
-	TUniquePtr<FImpl> ImplPtr;
+	TPimplPtr<FImpl> Impl;
 };
 
 /** Returns an I/O chunk ID to be used when loading from I/O store. */
@@ -585,6 +587,7 @@ public:
 	friend class FLinkerLoad;
 	friend class FExportArchive;
 	friend class UE::Serialization::FEditorBulkData; // To allow access to AttachedAr
+	friend class FBulkDataRequest;
 
 	using BulkDataRangeArray = TArray<FBulkData*, TInlineAllocator<8>>;
 
@@ -911,6 +914,7 @@ public:
 	* @return True if an asynchronous loading operation is in progress by the time that the method returns
 	* and false if the data is already loaded or cannot be loaded from disk.
 	*/
+	UE_DEPRECATED(5.1, "Use FBulkDataRequest or CreateStreamingRequest instead")
 	bool StartAsyncLoading();
 	
 	/**
@@ -1039,6 +1043,7 @@ public:
 	 * @param CompleteCallback	Called from an arbitrary thread when the request is complete. Can be nullptr, if non-null, must remain valid until it is called. It will always be called.
 	 * @return					A request for the read. This is owned by the caller and must be deleted by the caller.
 	**/
+	UE_DEPRECATED(5.1, "Use FBulkDataRequest instead")
 	static IBulkDataIORequest* CreateStreamingRequestForRange(const BulkDataRangeArray& RangeArray, EAsyncIOPriorityAndFlags Priority, FBulkDataIORequestCallBack* CompleteCallback);
 
 	/** Enable the given flags in the given accumulator variable. */
@@ -1326,3 +1331,139 @@ public:
 	COREUOBJECT_API void Serialize(FArchive& Ar, UObject* Owner, const TArray<FName>* FormatsToSave = nullptr, bool bSingleUse = true, uint16 InAlignment = DEFAULT_ALIGNMENT, bool bInline = true, bool bMapped = false);
 	COREUOBJECT_API void SerializeAttemptMappedLoad(FArchive& Ar, UObject* Owner);
 };
+
+/** Handle to a bulk data I/O request. */
+class COREUOBJECT_API FBulkDataRequest
+{
+	class FBase;
+	class FReadChunk;
+	class FReadFile;
+
+	struct FReadRequest
+	{
+		UE::BulkData::Private::FBulkMetaData BulkMeta;
+		UE::BulkData::Private::FBulkDataChunkId BulkChunkId;
+		uint64 Offset = 0;
+		uint64 Size = ~uint64(0);
+		void* TargetVa = nullptr;
+	};
+
+public:
+
+	/** Bulk data request status. */
+	enum class EStatus : uint32
+	{
+		/** The request hasn't been issued. */
+		None,
+		/** The request is pending. */
+		Pending,
+		/** The request has been completed successfully. */
+		Ok,
+		/** The request was cancelled. */
+		Cancelled,
+		/** An error occured while issuing the request. */
+		Error
+	};
+
+	/** Completion callback. */
+	using FCompletionCallback = TFunction<void(EStatus)>;
+
+	/** Constructs a new handle to bulk data request. */
+	FBulkDataRequest() = default;
+
+	/** Destructor, cancels and waits for any pending requests. */
+	~FBulkDataRequest() = default;
+
+	/** Moves ownership from an invalid or pending request. */
+	FBulkDataRequest(FBulkDataRequest&&) = default;
+
+	/** Moves ownership from an invalid or pending request. */
+	FBulkDataRequest& operator=(FBulkDataRequest&&) = default;
+
+	/** Returns current status of the request. */
+	EStatus GetStatus() const;
+
+	/** Returns whether the request is associated with a pending or completed request. */
+	bool IsNone() const;
+
+	/** Returns whether the request is pending. */
+	bool IsPending() const;
+
+	/** Returns whether the request completed successfully. */
+	bool IsOk() const;
+
+	/** Returns whether the request has been completed. */
+	bool IsCompleted() const;
+
+	/** Blocks the calling thread until the request is completed. */
+	void Wait();
+
+	/** Waits the specified amount of time in milliseconds for the request to be completed. */
+	bool WaitFor(uint32 Milliseconds);
+
+	/** Waits the specified amount of time for the request to be completed. */
+	bool WaitFor(const FTimespan& WaitTime);
+
+	/** Cancel the pending request. */
+	void Cancel();
+
+	/** Base class for building bulk data I/O requests. */
+	class COREUOBJECT_API FRequestBuilder
+	{
+	protected:
+		FRequestBuilder() = default;
+
+		void AddRequest(
+			const UE::BulkData::Private::FBulkMetaData& BulkMeta,
+			const UE::BulkData::Private::FBulkDataChunkId& BulkChunkId,
+			uint64 Offset,
+			uint64 Size,
+			void* TargetVa);
+
+		TArray<FReadRequest> Requests;
+		uint64 TotalRequestSize = 0;
+	};
+
+	/** Scatter/gather bulk data I/O reads into a single I/O buffer. */
+	class COREUOBJECT_API FScatterGather : public FRequestBuilder
+	{
+	public:
+		FScatterGather& Read(FBulkData& BulkData, uint64 Offset = 0, uint64 Size = MAX_uint64);
+		FBulkDataRequest Issue(FIoBuffer& Dst, FCompletionCallback&& Callback, EAsyncIOPriorityAndFlags Priority = AIOP_BelowNormal);
+	};
+
+	/** Read the entire bulk data and copy the result into the bulk data instances. */
+	class COREUOBJECT_API FStreamToInstance : public FRequestBuilder
+	{
+	public:
+		FStreamToInstance& Read(FBulkData& BulkData);
+		FBulkDataRequest Issue(EAsyncIOPriorityAndFlags Priority = AIOP_BelowNormal);
+
+	private:
+		TArray<FBulkData*> Instances;
+	};
+
+	/** Returns a request builder that reads one or more bulk data into a single I/O buffer. */
+	static FScatterGather ScatterGather();
+
+	/** Returns a request builder that reads and copies the I/O result into one or more bulk data instances. */
+	static FStreamToInstance StreamToInstance();
+
+private:
+
+	static FBulkDataRequest Issue(TArrayView<FReadRequest> Requests, EAsyncIOPriorityAndFlags Priority, FBulkDataRequest::FCompletionCallback&& Callback);
+
+	struct FDeleter
+	{
+		void operator()(FBase* Impl);
+	};
+
+	using FBasePtr = TUniquePtr<FBase, FDeleter>;
+
+	FBulkDataRequest(const FBulkDataRequest&) = delete;
+	FBulkDataRequest& operator=(const FBulkDataRequest&) = delete;
+
+	FBulkDataRequest(FBasePtr&& InImpl);
+
+	FBasePtr Impl;
+};
diff --git a/Engine/Source/Runtime/Engine/Private/Streaming/SkeletalMeshUpdate.cpp b/Engine/Source/Runtime/Engine/Private/Streaming/SkeletalMeshUpdate.cpp
--- a/Engine/Source/Runtime/Engine/Private/Streaming/SkeletalMeshUpdate.cpp
+++ b/Engine/Source/Runtime/Engine/Private/Streaming/SkeletalMeshUpdate.cpp
@@ -398,7 +398,6 @@ void FSkeletalMeshStreamIn_IO::FCancelIORequestsTask::DoWork()
 
 FSkeletalMeshStreamIn_IO::FSkeletalMeshStreamIn_IO(const USkeletalMesh* InMesh, bool bHighPrio)
 	: FSkeletalMeshStreamIn(InMesh)
-	, IORequest(nullptr)
 	, bHighPrioIORequest(bHighPrio)
 {}
 
@@ -408,7 +407,7 @@ void FSkeletalMeshStreamIn_IO::Abort()
 	{
 		FSkeletalMeshStreamIn::Abort();
 
-		if (IORequest != nullptr)
+		if (BulkDataRequest.IsPending())
 		{
 			// Prevent the update from being considered done before this is finished.
 			// By checking that it was not already cancelled, we make sure this doesn't get called twice.
@@ -417,65 +416,56 @@ void FSkeletalMeshStreamIn_IO::Abort()
 	}
 }
 
-void FSkeletalMeshStreamIn_IO::SetAsyncFileCallback(const FContext& Context)
-{
-	AsyncFileCallback = [this, Context](bool bWasCancelled, IBulkDataIORequest*)
-	{
-		// At this point task synchronization would hold the number of pending requests.
-		TaskSynchronization.Decrement();
-
-		if (bWasCancelled)
-		{
-			// If IO requests was cancelled but the streaming request wasn't, this is an IO error.
-			if (!bIsCancelled)
-			{
-				bFailedOnIOError = true;
-			}
-			MarkAsCancelled();
-		}
-
-#if !UE_BUILD_SHIPPING
-		// On some platforms the IO is too fast to test cancelation requests timing issues.
-		if (FRenderAssetStreamingSettings::ExtraIOLatency > 0 && TaskSynchronization.GetValue() == 0)
-		{
-			FPlatformProcess::Sleep(FRenderAssetStreamingSettings::ExtraIOLatency * .001f); // Slow down the streaming.
-		}
-#endif
-
-		// The tick here is intended to schedule the success or cancel callback.
-		// Using TT_None ensure gets which could create a dead lock.
-		Tick(FSkeletalMeshUpdate::TT_None);
-	};
-}
-
 void FSkeletalMeshStreamIn_IO::SetIORequest(const FContext& Context)
 {
 	if (IsCancelled())
 	{
 		return;
 	}
-	check(!IORequest && PendingFirstLODIdx < CurrentFirstLODIdx);
+	check(BulkDataRequest.IsNone() && PendingFirstLODIdx < CurrentFirstLODIdx);
 
 	const USkeletalMesh* Mesh = Context.Mesh;
 	FSkeletalMeshRenderData* RenderData = Context.RenderData;
 	if (Mesh && RenderData)
 	{
-		SetAsyncFileCallback(Context);
-
+		auto ScatterGather = FBulkDataRequest::ScatterGather();
 		FBulkData::BulkDataRangeArray BulkDataArray;
 		for (int32 Index = PendingFirstLODIdx; Index < CurrentFirstLODIdx; ++Index)
 		{
-			BulkDataArray.Push(&Context.LODResourcesView[Index]->StreamingBulkData);
+			ScatterGather.Read(Context.LODResourcesView[Index]->StreamingBulkData);
 		}
 
 		// Increment as we push the request. If a request complete immediately, then it will call the callback
 		// but that won't do anything because the tick would not try to acquire the lock since it is already locked.
 		TaskSynchronization.Increment();
 
-		IORequest = FBulkData::CreateStreamingRequestForRange(
-			BulkDataArray,
-			bHighPrioIORequest ? AIOP_BelowNormal : AIOP_Low,
-			&AsyncFileCallback);
+		BulkDataRequest = ScatterGather.Issue(BulkData, [this](FBulkDataRequest::EStatus Status)
+		{
+			// At this point task synchronization would hold the number of pending requests.
+			TaskSynchronization.Decrement();
+
+			if (FBulkDataRequest::EStatus::Cancelled == Status)
+			{
+				// If IO requests was cancelled but the streaming request wasn't, this is an IO error.
+				if (!bIsCancelled)
+				{
+					bFailedOnIOError = true;
+				}
+				MarkAsCancelled();
+			}
+
+#if !UE_BUILD_SHIPPING
+			// On some platforms the IO is too fast to test cancelation requests timing issues.
+			if (FRenderAssetStreamingSettings::ExtraIOLatency > 0 && TaskSynchronization.GetValue() == 0)
+			{
+				FPlatformProcess::Sleep(FRenderAssetStreamingSettings::ExtraIOLatency * .001f); // Slow down the streaming.
+			}
+#endif
+			// The tick here is intended to schedule the success or cancel callback.
+			// Using TT_None ensure gets which could create a dead lock.
+			Tick(FSkeletalMeshUpdate::TT_None);
+		},
+		bHighPrioIORequest ? AIOP_BelowNormal : AIOP_Low);
 	}
 	else
 	{
@@ -485,17 +475,14 @@ void FSkeletalMeshStreamIn_IO::SetIORequest(const FContext& Context)
 
 void FSkeletalMeshStreamIn_IO::ClearIORequest(const FContext& Context)
 {
-	if (IORequest != nullptr)
+	if (BulkDataRequest.IsPending())
 	{
-		// If clearing requests not yet completed, cancel and wait.
-		if (!IORequest->PollCompletion())
-		{
-			IORequest->Cancel();
-			IORequest->WaitCompletion();
-		}
-		delete IORequest;
-		IORequest = nullptr;
+		BulkDataRequest.Cancel();
+		BulkDataRequest.Wait();
 	}
+	
+	BulkDataRequest = FBulkDataRequest();
+	BulkData = FIoBuffer();
 }
 
 void FSkeletalMeshStreamIn_IO::ReportIOError(const FContext& Context)
@@ -522,10 +509,9 @@ void FSkeletalMeshStreamIn_IO::SerializeLODData(const FContext& Context)
 	FSkeletalMeshRenderData* RenderData = Context.RenderData;
 	if (!IsCancelled() && Mesh && RenderData)
 	{
-		check(IORequest->GetSize() >= 0 && IORequest->GetSize() <= TNumericLimits<uint32>::Max());
+		check(BulkData.GetSize() >= 0 && BulkData.GetSize() <= TNumericLimits<uint32>::Max());
 
-		TArrayView<uint8> Data(IORequest->GetReadResults(), IORequest->GetSize());
-		FMemoryReaderView Ar(Data, true);
+		FMemoryReaderView Ar(BulkData.GetView(), true);
 		for (int32 LODIndex = PendingFirstLODIdx; LODIndex < CurrentFirstLODIdx; ++LODIndex)
 		{
 			FSkeletalMeshLODRenderData& LODResource = *Context.LODResourcesView[LODIndex];
@@ -535,7 +521,7 @@ void FSkeletalMeshStreamIn_IO::SerializeLODData(const FContext& Context)
 			LODResource.SerializeStreamedData(Ar, const_cast<USkeletalMesh*>(Mesh), LODIndex + Context.AssetLODBias, DummyStripFlags, bNeedsCPUAccess, bForceKeepCPUResources);
 		}
 
-		FMemory::Free(Data.GetData()); // Free the memory we took ownership of via IORequest->GetReadResults()
+		BulkData = FIoBuffer();
 	}
 }
 
@@ -547,10 +533,9 @@ void FSkeletalMeshStreamIn_IO::Cancel(const FContext& Context)
 
 void FSkeletalMeshStreamIn_IO::CancelIORequest()
 {
-	if (IORequest)
+	if (BulkDataRequest.IsPending())
 	{
-		// Calling cancel will trigger the SetAsyncFileCallback() which will also try a tick but will fail.
-		IORequest->Cancel();
+		BulkDataRequest.Cancel();
 	}
 }
 
diff --git a/Engine/Source/Runtime/Engine/Private/Streaming/SkeletalMeshUpdate.h b/Engine/Source/Runtime/Engine/Private/Streaming/SkeletalMeshUpdate.h
--- a/Engine/Source/Runtime/Engine/Private/Streaming/SkeletalMeshUpdate.h
+++ b/Engine/Source/Runtime/Engine/Private/Streaming/SkeletalMeshUpdate.h
@@ -171,9 +171,6 @@ protected:
 	typedef FAutoDeleteAsyncTask<FCancelIORequestsTask> FAsyncCancelIORequestsTask;
 	friend class FCancelIORequestsTask;
 
-	/** Set a callback called when IORequest is completed or cancelled */
-	void SetAsyncFileCallback(const FContext& Context);
-
 	/** Create a new async IO request to read in LOD data */
 	void SetIORequest(const FContext& Context);
 
@@ -192,8 +189,12 @@ protected:
 	/** Called by FAsyncCancelIORequestsTask to cancel inflight IO request if any */
 	void CancelIORequest();
 
-	class IBulkDataIORequest* IORequest;
-	FBulkDataIORequestCallBack AsyncFileCallback;
+	/** Handle to bulk data I/O request */
+	FBulkDataRequest BulkDataRequest;
+
+	/** Bulk data I/O buffer */
+	FIoBuffer BulkData;
+
 	bool bHighPrioIORequest;
 
 	// Whether an IO error was detected (when files do not exists).
@@ -251,4 +252,4 @@ private:
 
 typedef TSkeletalMeshStreamIn_DDC<true> FSkeletalMeshStreamIn_DDC_RenderThread;
 typedef TSkeletalMeshStreamIn_DDC<false> FSkeletalMeshStreamIn_DDC_Async;
-#endif
\ No newline at end of file
+#endif
diff --git a/Engine/Source/Runtime/Engine/Private/Streaming/StaticMeshUpdate.cpp b/Engine/Source/Runtime/Engine/Private/Streaming/StaticMeshUpdate.cpp
--- a/Engine/Source/Runtime/Engine/Private/Streaming/StaticMeshUpdate.cpp
+++ b/Engine/Source/Runtime/Engine/Private/Streaming/StaticMeshUpdate.cpp
@@ -433,7 +433,6 @@ void FStaticMeshStreamIn_IO::FCancelIORequestsTask::DoWork()
 
 FStaticMeshStreamIn_IO::FStaticMeshStreamIn_IO(const UStaticMesh* InMesh, bool bHighPrio)
 	: FStaticMeshStreamIn(InMesh)
-	, IORequest(nullptr)
 	, bHighPrioIORequest(bHighPrio)
 {}
 
@@ -443,7 +442,7 @@ void FStaticMeshStreamIn_IO::Abort()
 	{
 		FStaticMeshStreamIn::Abort();
 
-		if (IORequest != nullptr)
+		if (BulkDataRequest.IsPending())
 		{
 			// Prevent the update from being considered done before this is finished.
 			// By checking that it was not already cancelled, we make sure this doesn't get called twice.
@@ -452,37 +451,6 @@ void FStaticMeshStreamIn_IO::Abort()
 	}
 }
 
-void FStaticMeshStreamIn_IO::SetAsyncFileCallback(const FContext& Context)
-{
-	AsyncFileCallback = [this, Context](bool bWasCancelled, IBulkDataIORequest*)
-	{
-		// At this point task synchronization would hold the number of pending requests.
-		TaskSynchronization.Decrement();
-
-		if (bWasCancelled)
-		{
-			// If IO requests was cancelled but the streaming request wasn't, this is an IO error.
-			if (!bIsCancelled)
-			{
-				bFailedOnIOError = true;
-			}
-			MarkAsCancelled();
-		}
-
-#if !UE_BUILD_SHIPPING
-		// On some platforms the IO is too fast to test cancelation requests timing issues.
-		if (FRenderAssetStreamingSettings::ExtraIOLatency > 0 && TaskSynchronization.GetValue() == 0)
-		{
-			FPlatformProcess::Sleep(FRenderAssetStreamingSettings::ExtraIOLatency * .001f); // Slow down the streaming.
-		}
-#endif
-
-		// The tick here is intended to schedule the success or cancel callback.
-		// Using TT_None ensure gets which could create a dead lock.
-		Tick(FStaticMeshUpdate::TT_None);
-	};
-}
-
 void FStaticMeshStreamIn_IO::SetIORequest(const FContext& Context)
 {
 	if (IsCancelled())
@@ -490,28 +458,48 @@ void FStaticMeshStreamIn_IO::SetIORequest(const FContext& Context)
 		return;
 	}
 
-	check(!IORequest && PendingFirstLODIdx < CurrentFirstLODIdx);
+	check(BulkDataRequest.IsNone() && PendingFirstLODIdx < CurrentFirstLODIdx);
 
 	const UStaticMesh* Mesh = Context.Mesh;
 	FStaticMeshRenderData* RenderData = Context.RenderData;
 	if (Mesh && RenderData)
 	{
-		SetAsyncFileCallback(Context);
-
-		FBulkData::BulkDataRangeArray BulkDataArray;
+		auto ScatterGather = FBulkDataRequest::ScatterGather();	
 		for (int32 Index = PendingFirstLODIdx; Index < CurrentFirstLODIdx; ++Index)
 		{
-			BulkDataArray.Push(&Context.LODResourcesView[Index]->StreamingBulkData);
+			ScatterGather.Read(Context.LODResourcesView[Index]->StreamingBulkData);
 		}
-
+		
 		// Increment as we push the request. If a request complete immediately, then it will call the callback
 		// but that won't do anything because the tick would not try to acquire the lock since it is already locked.
 		TaskSynchronization.Increment();
+		
+		BulkDataRequest = ScatterGather.Issue(BulkData, [this](FBulkDataRequest::EStatus Status)
+		{
+			TaskSynchronization.Decrement();
 
-		IORequest = FBulkData::CreateStreamingRequestForRange(
-			BulkDataArray,
-			bHighPrioIORequest ? AIOP_BelowNormal : AIOP_Low,
-			&AsyncFileCallback);
+			if (FBulkDataRequest::EStatus::Cancelled == Status)
+			{
+				// If IO requests was cancelled but the streaming request wasn't, this is an IO error.
+				if (!bIsCancelled)
+				{
+					bFailedOnIOError = true;
+				}
+				MarkAsCancelled();
+			}
+
+#if !UE_BUILD_SHIPPING
+			// On some platforms the IO is too fast to test cancelation requests timing issues.
+			if (FRenderAssetStreamingSettings::ExtraIOLatency > 0 && TaskSynchronization.GetValue() == 0)
+			{
+				FPlatformProcess::Sleep(FRenderAssetStreamingSettings::ExtraIOLatency * .001f); // Slow down the streaming.
+			}
+#endif
+			// The tick here is intended to schedule the success or cancel callback.
+			// Using TT_None ensure gets which could create a dead lock.
+			Tick(FStaticMeshUpdate::TT_None);
+		},
+		bHighPrioIORequest ? AIOP_BelowNormal : AIOP_Low);
 	}
 	else
 	{
@@ -521,17 +509,14 @@ void FStaticMeshStreamIn_IO::SetIORequest(const FContext& Context)
 
 void FStaticMeshStreamIn_IO::ClearIORequest(const FContext& Context)
 {
-	if (IORequest != nullptr)
+	if (BulkDataRequest.IsPending())
 	{
-		// If clearing requests not yet completed, cancel and wait.
-		if (!IORequest->PollCompletion())
-		{
-			IORequest->Cancel();
-			IORequest->WaitCompletion();
-		}
-		delete IORequest;
-		IORequest = nullptr;
+		BulkDataRequest.Cancel();
+		BulkDataRequest.Wait();
 	}
+	
+	BulkDataRequest = FBulkDataRequest();
+	BulkData = FIoBuffer();
 }
 
 void FStaticMeshStreamIn_IO::ReportIOError(const FContext& Context)
@@ -558,11 +543,9 @@ void FStaticMeshStreamIn_IO::SerializeLODData(const FContext& Context)
 	FStaticMeshRenderData* RenderData = Context.RenderData;
 	if (!IsCancelled() && Mesh && RenderData)
 	{
-		check(IORequest->GetSize() >= 0 && IORequest->GetSize() <= TNumericLimits<uint32>::Max());
-
-		TArrayView<uint8> Data(IORequest->GetReadResults(), IORequest->GetSize());
+		check(BulkData.GetSize() >= 0 && BulkData.GetSize() <= TNumericLimits<uint32>::Max());
 
-		FMemoryReaderView Ar(Data, true);
+		FMemoryReaderView Ar(BulkData.GetView(), true);
 		for (int32 LODIdx = PendingFirstLODIdx; LODIdx < CurrentFirstLODIdx; ++LODIdx)
 		{
 			FStaticMeshLODResources& LODResource = *Context.LODResourcesView[LODIdx];
@@ -571,8 +554,8 @@ void FStaticMeshStreamIn_IO::SerializeLODData(const FContext& Context)
 			LODResource.SerializeBuffers(Ar, const_cast<UStaticMesh*>(Mesh), DummyStripFlags, DummyBuffersSize);
 			check(DummyBuffersSize.CalcBuffersSize() == LODResource.BuffersSize);
 		}
-		
-		FMemory::Free(Data.GetData());// Free the memory we took ownership of via IORequest->GetReadResults()
+
+		BulkData = FIoBuffer();
 	}
 }
 
@@ -584,10 +567,9 @@ void FStaticMeshStreamIn_IO::Cancel(const FContext& Context)
 
 void FStaticMeshStreamIn_IO::CancelIORequest()
 {
-	if (IORequest)
+	if (BulkDataRequest.IsPending())
 	{
-		// Calling cancel will trigger the SetAsyncFileCallback() which will also try a tick but will fail.
-		IORequest->Cancel();
+		BulkDataRequest.Cancel();
 	}
 }
 
diff --git a/Engine/Source/Runtime/Engine/Private/Streaming/StaticMeshUpdate.h b/Engine/Source/Runtime/Engine/Private/Streaming/StaticMeshUpdate.h
--- a/Engine/Source/Runtime/Engine/Private/Streaming/StaticMeshUpdate.h
+++ b/Engine/Source/Runtime/Engine/Private/Streaming/StaticMeshUpdate.h
@@ -169,9 +169,6 @@ protected:
 	typedef FAutoDeleteAsyncTask<FCancelIORequestsTask> FAsyncCancelIORequestsTask;
 	friend class FCancelIORequestsTask;
 
-		/** Set a callback called when IORequest is completed or cancelled */
-	void SetAsyncFileCallback(const FContext& Context);
-
 	/** Create a new async IO request to read in LOD data */
 	void SetIORequest(const FContext& Context);
 
@@ -190,8 +187,12 @@ protected:
 	/** Called by FAsyncCancelIORequestsTask to cancel inflight IO request if any */
 	void CancelIORequest();
 
-	class IBulkDataIORequest* IORequest;
-	FBulkDataIORequestCallBack AsyncFileCallback;
+	/** Handle to bulk data request */;
+	FBulkDataRequest BulkDataRequest;
+
+	/** Bulk data I/O buffer */
+	FIoBuffer BulkData;
+
 	bool bHighPrioIORequest;
 
 	// Whether an IO error was detected (when files do not exists).
@@ -249,4 +250,4 @@ private:
 
 typedef TStaticMeshStreamIn_DDC<true> FStaticMeshStreamIn_DDC_RenderThread;
 typedef TStaticMeshStreamIn_DDC<false> FStaticMeshStreamIn_DDC_Async;
-#endif
\ No newline at end of file
+#endif
