Remove automatically included exclude object flags from StaticFindObjectFastInternal calls

diff --git a/Engine/Source/Runtime/CoreUObject/Private/UObject/UObjectHash.cpp b/Engine/Source/Runtime/CoreUObject/Private/UObject/UObjectHash.cpp
--- a/Engine/Source/Runtime/CoreUObject/Private/UObject/UObjectHash.cpp
+++ b/Engine/Source/Runtime/CoreUObject/Private/UObject/UObjectHash.cpp
@@ -809,7 +809,7 @@ UObject* StaticFindObjectFastInternal(const UClass* ObjectClass, const UObject*
 
 	// If they specified an outer use that during the hashing
 	FUObjectHashTables& ThreadHash = FUObjectHashTables::Get();
-	UObject* Result = StaticFindObjectFastInternalThreadSafe(ThreadHash, ObjectClass, ObjectPackage, ObjectName, bExactClass, bAnyPackage, ExcludeFlags | RF_NewerVersionExists, ExclusiveInternalFlags);
+	UObject* Result = StaticFindObjectFastInternalThreadSafe(ThreadHash, ObjectClass, ObjectPackage, ObjectName, bExactClass, bAnyPackage, ExcludeFlags, ExclusiveInternalFlags);
 	return Result;
 }
 
@@ -821,7 +821,7 @@ UObject* StaticFindObjectFastInternal(const UClass* ObjectClass, const UObject*
 
 	// If they specified an outer use that during the hashing
 	FUObjectHashTables& ThreadHash = FUObjectHashTables::Get();
-	UObject* Result = StaticFindObjectFastInternalThreadSafe(ThreadHash, ObjectClass, ObjectPackage, ObjectName, bExactClass, /*bAnyPackage =*/ false, ExcludeFlags | RF_NewerVersionExists, ExclusiveInternalFlags);
+	UObject* Result = StaticFindObjectFastInternalThreadSafe(ThreadHash, ObjectClass, ObjectPackage, ObjectName, bExactClass, /*bAnyPackage =*/ false, ExcludeFlags, ExclusiveInternalFlags);
 	return Result;
 }
 
