AsyncLoading2: Fix potential race between OnLeakedPackageRename on GT and RemoveUnreachableObjects on ALT (relates to 21353595 and 21426054)

#rnx

diff --git a/Engine/Source/Runtime/CoreUObject/Private/Serialization/AsyncLoading2.cpp b/Engine/Source/Runtime/CoreUObject/Private/Serialization/AsyncLoading2.cpp
--- a/Engine/Source/Runtime/CoreUObject/Private/Serialization/AsyncLoading2.cpp
+++ b/Engine/Source/Runtime/CoreUObject/Private/Serialization/AsyncLoading2.cpp
@@ -2279,8 +2279,8 @@ private:
 	TMap<UClass*, TArray<FEventLoadNode2*>> PendingCDOs;
 	TArray<UClass*> PendingCDOsRecursiveStack;
 
-	/** [ASYNC/GAME THREAD] Unreachable objects from last NotifyUnreachableObjects callback from GC.
-	 * Protected by FGCScopeGuard */
+	/** [ASYNC/GAME THREAD] Unreachable objects from last NotifyUnreachableObjects callback from GC. */
+	FCriticalSection UnreachableObjectsCritical;
 	FUnreachableObjects UnreachableObjects;
 
 	uint32 ConditionalBeginPostLoadTick = 0;
@@ -5490,7 +5490,10 @@ uint32 FAsyncLoadingThread2::Run()
 
 				{
 					FGCScopeGuard GCGuard;
-					RemoveUnreachableObjects(UnreachableObjects);
+					{
+						FScopeLock UnreachableObjectsLock(&UnreachableObjectsCritical);
+						RemoveUnreachableObjects(UnreachableObjects);
+					}
 
 					if (bShouldWaitForExternalReads)
 					{
@@ -5779,6 +5782,7 @@ FORCENOINLINE static void FilterUnreachableObjects(
 
 void FAsyncLoadingThread2::OnLeakedPackageRename(UPackage* Package)
 {
+	TRACE_CPUPROFILER_EVENT_SCOPE(OnLeakedPackageRename);
 	check(IsInGameThread());
 	
 	if (!FGCCSyncObject::Get().IsGCLocked())
@@ -5793,6 +5797,12 @@ void FAsyncLoadingThread2::OnLeakedPackageRename(UPackage* Package)
 		return;
 	}
 
+	FScopeLock UnreachableObjectsLock(&UnreachableObjectsCritical);
+
+	// unreachable objects from last GC should typically have been processed already,
+	// if not handle them here before processing new ones
+	RemoveUnreachableObjects(UnreachableObjects);
+
 	// If a package that can be imported was leaked and renamed,
 	// then it must exist in the loaded package store at this point since it is normally only trimmed during GC.
 	FLoadedPackageRef* PackageRef = LoadedPackageStore.FindPackageRef(Package->GetPackageId());
@@ -5863,6 +5873,8 @@ void FAsyncLoadingThread2::NotifyUnreachableObjects(const TArrayView<FUObjectIte
 		return;
 	}
 
+	FScopeLock UnreachableObjectsLock(&UnreachableObjectsCritical);
+
 	// unreachable objects from last GC should typically have been processed already,
 	// if not handle them here before adding new ones
 	RemoveUnreachableObjects(UnreachableObjects);
