Fix const correctness when constructing bulk data I/O request

#rnx

diff --git a/Engine/Source/Runtime/CoreUObject/Private/Serialization/BulkDataStreaming.cpp b/Engine/Source/Runtime/CoreUObject/Private/Serialization/BulkDataStreaming.cpp
--- a/Engine/Source/Runtime/CoreUObject/Private/Serialization/BulkDataStreaming.cpp
+++ b/Engine/Source/Runtime/CoreUObject/Private/Serialization/BulkDataStreaming.cpp
@@ -1094,7 +1094,7 @@ void FBulkDataRequest::FRequestBuilder::AddRequest(
 	TotalRequestSize += Size;
 }
 
-FBulkDataRequest::FScatterGather& FBulkDataRequest::FScatterGather::Read(FBulkData& BulkData, uint64 Offset, uint64 Size)
+FBulkDataRequest::FScatterGather& FBulkDataRequest::FScatterGather::Read(const FBulkData& BulkData, uint64 Offset, uint64 Size)
 {
 	const uint64 BulkOffset	= uint64(BulkData.GetBulkDataOffsetInFile());
 	const uint64 BulkSize	= uint64(BulkData.GetBulkDataSize());
diff --git a/Engine/Source/Runtime/CoreUObject/Public/Serialization/BulkData.h b/Engine/Source/Runtime/CoreUObject/Public/Serialization/BulkData.h
--- a/Engine/Source/Runtime/CoreUObject/Public/Serialization/BulkData.h
+++ b/Engine/Source/Runtime/CoreUObject/Public/Serialization/BulkData.h
@@ -1428,7 +1428,7 @@ public:
 	class COREUOBJECT_API FScatterGather : public FRequestBuilder
 	{
 	public:
-		FScatterGather& Read(FBulkData& BulkData, uint64 Offset = 0, uint64 Size = MAX_uint64);
+		FScatterGather& Read(const FBulkData& BulkData, uint64 Offset = 0, uint64 Size = MAX_uint64);
 		FBulkDataRequest Issue(FIoBuffer& Dst, FCompletionCallback&& Callback, EAsyncIOPriorityAndFlags Priority = AIOP_BelowNormal);
 	};
 
