Bulk Data - skip stream to instances that is already loaded

#rnx

diff --git a/Engine/Source/Runtime/CoreUObject/Private/Serialization/BulkDataStreaming.cpp b/Engine/Source/Runtime/CoreUObject/Private/Serialization/BulkDataStreaming.cpp
--- a/Engine/Source/Runtime/CoreUObject/Private/Serialization/BulkDataStreaming.cpp
+++ b/Engine/Source/Runtime/CoreUObject/Private/Serialization/BulkDataStreaming.cpp
@@ -775,11 +775,22 @@ void FlushAsyncLoad(FBulkData* Owner)
 class FBulkDataRequest::FBase
 {
 public:
+	FBase(FCompletionCallback&& Callback)
+		: DoneEvent(EEventMode::ManualReset)
+		, CompletionCallback(MoveTemp(Callback))
+	{
+	}
+
 	virtual ~FBase() = default;
 
-	virtual void Cancel() = 0;
+	virtual void Cancel()
+	{
+	}
 
-	virtual bool Issue(TArrayView<FReadRequest> Requests, EAsyncIOPriorityAndFlags Priority) = 0;
+	virtual bool Issue(TArrayView<FReadRequest> Requests, EAsyncIOPriorityAndFlags Priority)
+	{
+		return true;
+	}
 
 	bool Wait(uint32 Milliseconds)
 	{
@@ -791,14 +802,13 @@ public:
 		return static_cast<FBulkDataRequest::EStatus>(Status.load(std::memory_order_consume));
 	}
 
-protected:
-
-	FBase(FCompletionCallback&& Callback)
-		: DoneEvent(EEventMode::ManualReset)
-		, CompletionCallback(MoveTemp(Callback))
+	void SetCompleted(FBulkDataRequest::EStatus CompletionStatus)
 	{
+		CompleteRequest(CompletionStatus);
 	}
 
+protected:
+
 	inline void SetStatus(FBulkDataRequest::EStatus InStatus)
 	{
 		Status.store(uint32(InStatus), std::memory_order_release); 
@@ -806,13 +816,12 @@ protected:
 
 	inline void CompleteRequest(FBulkDataRequest::EStatus CompletionStatus)
 	{
-		SetStatus(CompletionStatus);
-
 		if (CompletionCallback)
 		{
 			CompletionCallback(CompletionStatus);
 		}
 
+		SetStatus(CompletionStatus);
 		DoneEvent.Trigger();
 	}
 
@@ -1127,17 +1136,28 @@ FBulkDataRequest FBulkDataRequest::FScatterGather::Issue(FIoBuffer& Dst, FBulkDa
 
 FBulkDataRequest::FStreamToInstance& FBulkDataRequest::FStreamToInstance::Read(FBulkData& BulkData)
 {
-	const uint64 BulkOffset	= uint64(BulkData.GetBulkDataOffsetInFile());
-	const uint64 BulkSize	= uint64(BulkData.GetBulkDataSize());
-	
-	AddRequest(BulkData.BulkMeta, BulkData.BulkChunkId, BulkOffset, BulkSize, nullptr);
-	Instances.Add(&BulkData);
+	if (BulkData.IsBulkDataLoaded() == false)
+	{
+		const uint64 BulkOffset	= uint64(BulkData.GetBulkDataOffsetInFile());
+		const uint64 BulkSize	= uint64(BulkData.GetBulkDataSize());
+		
+		AddRequest(BulkData.BulkMeta, BulkData.BulkChunkId, BulkOffset, BulkSize, nullptr);
+		Instances.Add(&BulkData);
+	}
+
+	NumRequested++;
 
 	return *this;
 }
 
 FBulkDataRequest FBulkDataRequest::FStreamToInstance::Issue(EAsyncIOPriorityAndFlags Priority)
 {
+	// Complete the request if all instances are already loaded 
+	if (NumRequested > 0 && Requests.IsEmpty())
+	{
+		return FBulkDataRequest::FromCompletionStatus(FBulkDataRequest::EStatus::Ok, FBulkDataRequest::FCompletionCallback());
+	}
+
 	if (Requests.IsEmpty())
 	{
 		return FBulkDataRequest();
@@ -1217,3 +1237,13 @@ FBulkDataRequest FBulkDataRequest::Issue(TArrayView<FReadRequest> Requests, EAsy
 
 	return FBulkDataRequest(MoveTemp(NewRequest));
 }
+
+FBulkDataRequest FBulkDataRequest::FromCompletionStatus(FBulkDataRequest::EStatus CompletionStatus, FBulkDataRequest::FCompletionCallback&& Callback)
+{
+	check(CompletionStatus != FBulkDataRequest::EStatus::None && CompletionStatus != FBulkDataRequest::EStatus::Pending);
+
+	FBasePtr Completed(new FBulkDataRequest::FBase(MoveTemp(Callback)));
+	Completed->SetCompleted(CompletionStatus);
+
+	return FBulkDataRequest(MoveTemp(Completed));
+}
diff --git a/Engine/Source/Runtime/CoreUObject/Public/Serialization/BulkData.h b/Engine/Source/Runtime/CoreUObject/Public/Serialization/BulkData.h
--- a/Engine/Source/Runtime/CoreUObject/Public/Serialization/BulkData.h
+++ b/Engine/Source/Runtime/CoreUObject/Public/Serialization/BulkData.h
@@ -1437,10 +1437,12 @@ public:
 	{
 	public:
 		FStreamToInstance& Read(FBulkData& BulkData);
+		/** Returns a handle to a pending or completed request depending on if the bulk data has been loaded or not. */
 		FBulkDataRequest Issue(EAsyncIOPriorityAndFlags Priority = AIOP_BelowNormal);
 
 	private:
 		TArray<FBulkData*> Instances;
+		int32 NumRequested = 0;
 	};
 
 	/** Returns a request builder that reads one or more bulk data into a single I/O buffer. */
@@ -1452,6 +1454,7 @@ public:
 private:
 
 	static FBulkDataRequest Issue(TArrayView<FReadRequest> Requests, EAsyncIOPriorityAndFlags Priority, FBulkDataRequest::FCompletionCallback&& Callback);
+	static FBulkDataRequest FromCompletionStatus(EStatus CompletionStatus, FBulkDataRequest::FCompletionCallback&& Callback);
 
 	struct FDeleter
 	{
