Workaround for VS2022 compiler issue in NotifyUnreachableObjects

diff --git a/Engine/Source/Runtime/CoreUObject/Private/UObject/GarbageCollection.cpp b/Engine/Source/Runtime/CoreUObject/Private/UObject/GarbageCollection.cpp
--- a/Engine/Source/Runtime/CoreUObject/Private/UObject/GarbageCollection.cpp
+++ b/Engine/Source/Runtime/CoreUObject/Private/UObject/GarbageCollection.cpp
@@ -1450,6 +1450,16 @@ public:
 
 						LocalObjectsToSerialize.Add(Object);
 					}
+					// Cluster objects 
+					else if (bWithClusters && ObjectItem->GetOwnerIndex() > 0)
+					{
+						// treat cluster objects with FastKeepFlags the same way as if they are in the root set
+						if (ObjectItem->HasAnyFlags(FastKeepFlags))
+						{
+							KeepClusterRefsList.Push(ObjectItem);
+							LocalObjectsToSerialize.Add(Object);
+						}
+					}
 					// Regular objects or cluster root objects
 					else if (!bWithClusters || ObjectItem->GetOwnerIndex() <= 0)
 					{
@@ -1490,17 +1500,7 @@ public:
 						{
 							ObjectItem->SetFlags(EInternalObjectFlags::Unreachable);
 						}
-					}
-					// Cluster objects 
-					else if (bWithClusters && ObjectItem->GetOwnerIndex() > 0)
-					{
-						// treat cluster objects with FastKeepFlags the same way as if they are in the root set
-						if (ObjectItem->HasAnyFlags(FastKeepFlags))
-						{
-							KeepClusterRefsList.Push(ObjectItem);
-							LocalObjectsToSerialize.Add(Object);
-						}
-					}
+					}					
 				}
 			}
 
