Ignoring engine version checks when loading package for diff

diff --git a/Engine/Source/Runtime/CoreUObject/Private/UObject/LinkerLoad.cpp b/Engine/Source/Runtime/CoreUObject/Private/UObject/LinkerLoad.cpp
--- a/Engine/Source/Runtime/CoreUObject/Private/UObject/LinkerLoad.cpp
+++ b/Engine/Source/Runtime/CoreUObject/Private/UObject/LinkerLoad.cpp
@@ -1066,7 +1066,7 @@ FLinkerLoad::ELinkerStatus FLinkerLoad::CreateLoader(
 			FAsyncArchive* AsyncArchive = new FAsyncArchive(GetPackagePath(), this,
 				GEventDrivenLoaderEnabled ? Forward<TFunction<void()>>(InSummaryReadyCallback) : TFunction<void()>([]() {}));
 			Loader = AsyncArchive; // We're only allowed to delete any FAsyncArchive with this->DestroyLoader
-			bLoaderNeedsEngineVersionChecks = AsyncArchive->NeedsEngineVersionChecks();
+			bLoaderNeedsEngineVersionChecks = !(LoadFlags & LOAD_DisableEngineVersionChecks) && AsyncArchive->NeedsEngineVersionChecks();
 			if (AsyncArchive->IsError())
 			{
 				bool bRetryWithNormalArchive = AsyncArchive->GetLoadError() == FAsyncArchive::ELoadError::UnsupportedFormat;
@@ -1093,7 +1093,7 @@ FLinkerLoad::ELinkerStatus FLinkerLoad::CreateLoader(
 				OpenResult = IPackageResourceManager::Get().OpenReadPackage(GetPackagePath());
 			}
 			Loader = OpenResult.Archive.Release();
-			bLoaderNeedsEngineVersionChecks = OpenResult.bNeedsEngineVersionChecks;
+			bLoaderNeedsEngineVersionChecks = !(LoadFlags & LOAD_DisableEngineVersionChecks) && OpenResult.bNeedsEngineVersionChecks;
 			if (!Loader || Loader->IsError())
 			{
 				if (Loader)
diff --git a/Engine/Source/Runtime/CoreUObject/Public/UObject/ObjectMacros.h b/Engine/Source/Runtime/CoreUObject/Public/UObject/ObjectMacros.h
--- a/Engine/Source/Runtime/CoreUObject/Public/UObject/ObjectMacros.h
+++ b/Engine/Source/Runtime/CoreUObject/Public/UObject/ObjectMacros.h
@@ -78,6 +78,7 @@ enum ELoadFlags
 	LOAD_DeferDependencyLoads       = 0x00100000,   ///< Do not load external (blueprint) dependencies (instead, track them for deferred loading)
 	LOAD_ForFileDiff				= 0x00200000,	///< Load the package (not for diffing in the editor), instead verify at the two packages serialized output are the same, if they are not then debug break so that you can get the callstack and object information
 	LOAD_DisableCompileOnLoad		= 0x00400000,	///< Prevent this load call from running compile on load for the loaded blueprint (intentionally not recursive, dependencies will still compile on load)
+	LOAD_DisableEngineVersionChecks = 0x00800000,   ///< Prevent this load call from running engine version checks
 };
 
 /** Flags for saving objects/packages, passed into UPackage::SavePackage() as a uint32 */
