Add -DebugCoreRedirects (and -FullDebugCoreRedirects) command that writes to the log when reading and applying core redirects
This can be used to track down issues with new redirects or clean out broken old redirects
Add new CoreRedirects log category and clean up comments

diff --git a/Engine/Source/Runtime/CoreUObject/Private/UObject/CoreRedirects.cpp b/Engine/Source/Runtime/CoreUObject/Private/UObject/CoreRedirects.cpp
--- a/Engine/Source/Runtime/CoreUObject/Private/UObject/CoreRedirects.cpp
+++ b/Engine/Source/Runtime/CoreUObject/Private/UObject/CoreRedirects.cpp
@@ -5,6 +5,7 @@
 #include "UObject/Linker.h"
 #include "UObject/Package.h"
 #include "UObject/PropertyHelper.h"
+#include "UObject/TopLevelAssetPath.h"
 
 #include "GenericPlatform/GenericPlatformFile.h"
 #include "Misc/AutomationTest.h"
@@ -15,6 +16,8 @@
 #include "Serialization/DeferredMessageLog.h"
 #include "Templates/Casts.h"
 
+DEFINE_LOG_CATEGORY(LogCoreRedirects);
+
 #if !defined UE_WITH_CORE_REDIRECTS
 #	define UE_WITH_CORE_REDIRECTS 1
 #endif
@@ -476,26 +479,15 @@ FCoreRedirectObjectName FCoreRedirect::RedirectName(const FCoreRedirectObjectNam
 
 bool FCoreRedirect::IdenticalMatchRules(const FCoreRedirect& Other) const
 {
+	// All types now use the full path
 	ECoreRedirectFlags TypeFlags = RedirectFlags & ECoreRedirectFlags::Type_AllMask;
-	if (TypeFlags == ECoreRedirectFlags::Type_Struct)
-	{
-		// Struct remappings are requested based solely on the Name - not considering package and outer -
-		// so they are identical if only the name matches
-		return RedirectFlags == Other.RedirectFlags && OldName.ObjectName == Other.OldName.ObjectName;
-	}
-	else
-	{
-		// Enum remappings are requested in some locations based solely on the Name - not considering package and outer -
-		// but they are also requested from UEnum::GetIndexByNameString -> GetValueRedirects by full name
-		// Return IdenticalMatchRules=true only if the full name matches.
-		// TODO: Change FCoreRedirects::AddSingleRedirect to remove the package and outer name so that we can return IdenticalMatchRules=true for partial name enum matches
-
-		return RedirectFlags == Other.RedirectFlags && OldName == Other.OldName;
-	}
+	return RedirectFlags == Other.RedirectFlags && OldName == Other.OldName;
 }
 
 
 bool FCoreRedirects::bInitialized = false;
+bool FCoreRedirects::bInDebugMode = false;
+bool FCoreRedirects::bValidatedOnce = false;
 #if WITH_COREREDIRECTS_MULTITHREAD_WARNING
 bool FCoreRedirects::bIsInMultithreadedPhase = false;
 #endif
@@ -549,6 +541,23 @@ void FCoreRedirects::Initialize()
 	}
 	bInitialized = true;
 
+#if !UE_BUILD_SHIPPING && !NO_LOGGING
+	if (FParse::Param(FCommandLine::Get(), TEXT("FullDebugCoreRedirects")))
+	{
+		// Enable debug mode and set to maximum verbosity
+		bInDebugMode = true;
+		LogCoreRedirects.SetVerbosity(ELogVerbosity::VeryVerbose);
+		FCoreDelegates::OnFEngineLoopInitComplete.AddStatic(ValidateAllRedirects);
+	}
+	else if (FParse::Param(FCommandLine::Get(), TEXT("DebugCoreRedirects")))
+	{
+		// Enable debug mode and increase log levels but don't show every message
+		bInDebugMode = true;
+		LogCoreRedirects.SetVerbosity(ELogVerbosity::Verbose);
+		FCoreDelegates::OnFEngineLoopInitComplete.AddStatic(ValidateAllRedirects);
+	}
+#endif
+
 #if WITH_COREREDIRECTS_MULTITHREAD_WARNING
 	// Setting IsInMultithreadedPhase has to occur after LoadModule("AssetRegistry") (which can write to FCoreRedirects) and
 	// before the first package is queued onto the async loading thread (which reads from FCoreRedirects multithreaded). OnPostEngineInit is in that window.
@@ -568,8 +577,7 @@ void FCoreRedirects::Initialize()
 
 	// Prepopulate RedirectTypeMap entries that some threads write to after the engine goes multi-threaded.
 	// Most RedirectTypeMap entries are written to only from InitUObject's call to ReadRedirectsFromIni, and at that point the Engine is single-threaded.
-	// Currently the only entries that can be written to later, and therefore can be written to from multiple threads are:
-	//     KnownMissing Packages.
+	// Known missing packages and plugin loads can add entries to existing lists but will not add brand new types.
 	// Taking advantage of this, we treat the list of Key/Value pairs of RedirectTypeMap as immutable and read from it without synchronization.
 	// Note that the values for those written-during-multithreading entries need to be synchronized; it is only the list of Key/Value pairs that is immutable.
 #if WITH_COREREDIRECTS_MULTITHREAD_WARNING
@@ -618,7 +626,7 @@ bool FCoreRedirects::RedirectNameAndValues(ECoreRedirectFlags Type, const FCoreR
 					{
 						if ((*FoundValueRedirect)->ValueChanges.OrderIndependentCompareEqual(Redirect->ValueChanges) == false)
 						{
-							UE_LOG(LogLinker, Error, TEXT("RedirectNameAndValues(%s) found multiple conflicting value redirects, %s and %s!"), *OldObjectName.ToString(), *(*FoundValueRedirect)->OldName.ToString(), *Redirect->OldName.ToString());
+							UE_LOG(LogCoreRedirects, Error, TEXT("RedirectNameAndValues(%s) found multiple conflicting value redirects, %s and %s!"), *OldObjectName.ToString(), *(*FoundValueRedirect)->OldName.ToString(), *Redirect->OldName.ToString());
 						}
 					}
 					else
@@ -633,6 +641,9 @@ bool FCoreRedirects::RedirectNameAndValues(ECoreRedirectFlags Type, const FCoreR
 		}
 	}
 
+	UE_CLOG(bInDebugMode && NewObjectName != OldObjectName, LogCoreRedirects, Verbose, 
+		TEXT("RedirectNameAndValues(%s) replaced by %s"), *OldObjectName.ToString(), *NewObjectName.ToString());
+
 	return NewObjectName != OldObjectName;
 }
 
@@ -656,6 +667,8 @@ const TMap<FString, FString>* FCoreRedirects::GetValueRedirects(ECoreRedirectFla
 
 	if (FoundRedirect && FoundRedirect->ValueChanges.Num() > 0)
 	{
+		UE_CLOG(bInDebugMode, LogCoreRedirects, VeryVerbose, TEXT("GetValueRedirects found %d matches for %s"), FoundRedirect->ValueChanges.Num(), *OldObjectName.ToString());
+
 		return &FoundRedirect->ValueChanges;
 	}
 
@@ -746,6 +759,8 @@ bool FCoreRedirects::FindPreviousNames(ECoreRedirectFlags SearchFlags, const FCo
 		}
 	}
 
+	UE_CLOG(bFound && bInDebugMode, LogCoreRedirects, VeryVerbose, TEXT("FindPreviousNames found %d previous names for %s"), PreviousNames.Num(), *NewObjectName.ToString());
+
 	return bFound;
 }
 
@@ -829,7 +844,7 @@ bool FCoreRedirects::RunTests()
 
 	TArray<FRedirectTest> Tests;
 
-	UE_LOG(LogLinker, Log, TEXT("Running FCoreRedirect Tests"));
+	UE_LOG(LogCoreRedirects, Log, TEXT("Running FCoreRedirect Tests"));
 
 	// Package-specific property rename and package rename apply
 	Tests.Emplace(TEXT("/Game/PackageSpecific.Class:Property"), TEXT("/Game/PackageSpecific.Class:Property4"), ECoreRedirectFlags::Type_Property);
@@ -856,7 +871,7 @@ bool FCoreRedirects::RunTests()
 		if (NewName.ToString() != Test.NewName)
 		{
 			bSuccess = false;
-			UE_LOG(LogLinker, Error, TEXT("FCoreRedirect Test Failed: %s to %s, should be %s!"), *OldName.ToString(), *NewName.ToString(), *Test.NewName);
+			UE_LOG(LogCoreRedirects, Error, TEXT("FCoreRedirect Test Failed: %s to %s, should be %s!"), *OldName.ToString(), *NewName.ToString(), *Test.NewName);
 		}
 	}
 
@@ -868,26 +883,26 @@ bool FCoreRedirects::RunTests()
 	if (OldNames.Num() != 1 || OldNames[0].ToString() != TEXT("/Game/PackageOther.Class"))
 	{
 		bSuccess = false;
-		UE_LOG(LogLinker, Error, TEXT("FCoreRedirect Test Failed: ReverseLookup!"));
+		UE_LOG(LogCoreRedirects, Error, TEXT("FCoreRedirect Test Failed: ReverseLookup!"));
 	}
 
 	// Check removed
 	if (!IsKnownMissing(ECoreRedirectFlags::Type_Package, FCoreRedirectObjectName(TEXT("/Game/RemovedPackage"))))
 	{
 		bSuccess = false;
-		UE_LOG(LogLinker, Error, TEXT("FCoreRedirect Test Failed: /Game/RemovedPackage should be removed!"));
+		UE_LOG(LogCoreRedirects, Error, TEXT("FCoreRedirect Test Failed: /Game/RemovedPackage should be removed!"));
 	}
 
 	if (!IsKnownMissing(ECoreRedirectFlags::Type_Package, FCoreRedirectObjectName(TEXT("/Game/MissingLoadPackage"))))
 	{
 		bSuccess = false;
-		UE_LOG(LogLinker, Error, TEXT("FCoreRedirect Test Failed: /Game/MissingLoadPackage should be removed!"));
+		UE_LOG(LogCoreRedirects, Error, TEXT("FCoreRedirect Test Failed: /Game/MissingLoadPackage should be removed!"));
 	}
 
 	if (IsKnownMissing(ECoreRedirectFlags::Type_Package, FCoreRedirectObjectName(TEXT("/Game/NotRemovedPackage"))))
 	{
 		bSuccess = false;
-		UE_LOG(LogLinker, Error, TEXT("FCoreRedirect Test Failed: /Game/NotRemovedPackage should be removed!"));
+		UE_LOG(LogCoreRedirects, Error, TEXT("FCoreRedirect Test Failed: /Game/NotRemovedPackage should be removed!"));
 	}
 
 	AddKnownMissing(ECoreRedirectFlags::Type_Package, FCoreRedirectObjectName(TEXT("/Game/NotRemovedMissingLoad")), ECoreRedirectFlags::Option_MissingLoad);
@@ -895,7 +910,7 @@ bool FCoreRedirects::RunTests()
 	if (!IsKnownMissing(ECoreRedirectFlags::Type_Package, FCoreRedirectObjectName(TEXT("/Game/NotRemovedMissingLoad"))))
 	{
 		bSuccess = false;
-		UE_LOG(LogLinker, Error, TEXT("FCoreRedirect Test Failed: /Game/NotRemovedMissingLoad should be removed now!"));
+		UE_LOG(LogCoreRedirects, Error, TEXT("FCoreRedirect Test Failed: /Game/NotRemovedMissingLoad should be removed now!"));
 	}
 
 	RemoveKnownMissing(ECoreRedirectFlags::Type_Package, FCoreRedirectObjectName(TEXT("/Game/NotRemovedMissingLoad")), ECoreRedirectFlags::None);
@@ -903,7 +918,7 @@ bool FCoreRedirects::RunTests()
 	if (!IsKnownMissing(ECoreRedirectFlags::Type_Package, FCoreRedirectObjectName(TEXT("/Game/NotRemovedMissingLoad"))))
 	{
 		bSuccess = false;
-		UE_LOG(LogLinker, Error, TEXT("FCoreRedirect Test Failed: RemoveKnownMissing of /Game/NotRemovedMissingLoad but with bIsMissingLoad=false should not have removed the redirect!"));
+		UE_LOG(LogCoreRedirects, Error, TEXT("FCoreRedirect Test Failed: RemoveKnownMissing of /Game/NotRemovedMissingLoad but with bIsMissingLoad=false should not have removed the redirect!"));
 	}
 
 	RemoveKnownMissing(ECoreRedirectFlags::Type_Package, FCoreRedirectObjectName(TEXT("/Game/NotRemovedMissingLoad")), ECoreRedirectFlags::Option_MissingLoad);
@@ -911,7 +926,7 @@ bool FCoreRedirects::RunTests()
 	if (IsKnownMissing(ECoreRedirectFlags::Type_Package, FCoreRedirectObjectName(TEXT("/Game/NotRemovedMissingLoad"))))
 	{
 		bSuccess = false;
-		UE_LOG(LogLinker, Error, TEXT("FCoreRedirect Test Failed: /Game/NotRemovedMissingLoad should no longer be removed!"));
+		UE_LOG(LogCoreRedirects, Error, TEXT("FCoreRedirect Test Failed: /Game/NotRemovedMissingLoad should no longer be removed!"));
 	}
 
 	AddKnownMissing(ECoreRedirectFlags::Type_Package, FCoreRedirectObjectName(TEXT("/Game/NotRemovedPackage")), ECoreRedirectFlags::None);
@@ -919,7 +934,7 @@ bool FCoreRedirects::RunTests()
 	if (!IsKnownMissing(ECoreRedirectFlags::Type_Package, FCoreRedirectObjectName(TEXT("/Game/NotRemovedPackage"))))
 	{
 		bSuccess = false;
-		UE_LOG(LogLinker, Error, TEXT("FCoreRedirect Test Failed: /Game/NotRemovedPackage should be removed now!"));
+		UE_LOG(LogCoreRedirects, Error, TEXT("FCoreRedirect Test Failed: /Game/NotRemovedPackage should be removed now!"));
 	}
 
 	RemoveKnownMissing(ECoreRedirectFlags::Type_Package, FCoreRedirectObjectName(TEXT("/Game/NotRemovedPackage")), ECoreRedirectFlags::Option_MissingLoad);
@@ -927,7 +942,7 @@ bool FCoreRedirects::RunTests()
 	if (!IsKnownMissing(ECoreRedirectFlags::Type_Package, FCoreRedirectObjectName(TEXT("/Game/NotRemovedPackage"))))
 	{
 		bSuccess = false;
-		UE_LOG(LogLinker, Error, TEXT("FCoreRedirect Test Failed: RemoveKnownMissing of /Game/NotRemovedPackage but with bIsMissingLoad=true should not have removed the redirect!"));
+		UE_LOG(LogCoreRedirects, Error, TEXT("FCoreRedirect Test Failed: RemoveKnownMissing of /Game/NotRemovedPackage but with bIsMissingLoad=true should not have removed the redirect!"));
 	}
 
 	RemoveKnownMissing(ECoreRedirectFlags::Type_Package, FCoreRedirectObjectName(TEXT("/Game/NotRemovedPackage")), ECoreRedirectFlags::None);
@@ -935,7 +950,7 @@ bool FCoreRedirects::RunTests()
 	if (IsKnownMissing(ECoreRedirectFlags::Type_Package, FCoreRedirectObjectName(TEXT("/Game/NotRemovedPackage"))))
 	{
 		bSuccess = false;
-		UE_LOG(LogLinker, Error, TEXT("FCoreRedirect Test Failed: /Game/NotRemovedPackage should no longer be removed!"));
+		UE_LOG(LogCoreRedirects, Error, TEXT("FCoreRedirect Test Failed: /Game/NotRemovedPackage should no longer be removed!"));
 	}
 
 	// Restore old state
@@ -1024,7 +1039,7 @@ bool FCoreRedirects::ReadRedirectsFromIni(const FString& IniName)
 
 						if (!Buffer)
 						{
-							UE_LOG(LogLinker, Error, TEXT("ReadRedirectsFromIni failed to parse ValueChanges for Redirect %s!"), *ValueString);
+							UE_LOG(LogCoreRedirects, Error, TEXT("ReadRedirectsFromIni(%s) failed to parse ValueChanges for redirect %s!"), *IniName, *ValueString);
 
 							// Remove added redirect
 							Redirect = nullptr;
@@ -1034,16 +1049,20 @@ bool FCoreRedirects::ReadRedirectsFromIni(const FString& IniName)
 				}
 				else
 				{
-					UE_LOG(LogLinker, Error, TEXT("ReadRedirectsFromIni failed to parse type for Redirect %s!"), *ValueString);
+					UE_LOG(LogCoreRedirects, Error, TEXT("ReadRedirectsFromIni(%s) failed to parse type for redirect %s!"), *IniName, *ValueString);
 				}
 			}
 
 			return AddRedirectList(NewRedirects, IniName);
 		}
+		else
+		{
+			UE_LOG(LogCoreRedirects, Verbose, TEXT("ReadRedirectsFromIni(%s) did not find any redirects"), *IniName);
+		}
 	}
 	else
 	{
-		UE_LOG(LogLinker, Warning, TEXT(" **** CORE REDIRECTS UNABLE TO INITIALIZE! **** "));
+		UE_LOG(LogCoreRedirects, Warning, TEXT(" **** CORE REDIRECTS UNABLE TO INITIALIZE! **** "));
 	}
 	return false;
 }
@@ -1052,31 +1071,39 @@ bool FCoreRedirects::AddRedirectList(TArrayView<const FCoreRedirect> Redirects,
 {
 	Initialize();
 
+	UE_LOG(LogCoreRedirects, Verbose, TEXT("AddRedirect(%s) adding %d redirects"), *SourceString, Redirects.Num());
+
+	if (bInDebugMode && bValidatedOnce)
+	{
+		// Validate on apply because we finished our initial validation pass
+		ValidateRedirectList(Redirects, SourceString);
+	}
+
 	bool bAddedAny = false;
 	for (const FCoreRedirect& NewRedirect : Redirects)
 	{
 		if (!NewRedirect.OldName.IsValid() || !NewRedirect.NewName.IsValid())
 		{
-			UE_LOG(LogLinker, Error, TEXT("AddRedirectList(%s) failed to add redirector from %s to %s with empty name!"), *SourceString, *NewRedirect.OldName.ToString(), *NewRedirect.NewName.ToString());
+			UE_LOG(LogCoreRedirects, Error, TEXT("AddRedirect(%s) failed to add redirect from %s to %s with empty name!"), *SourceString, *NewRedirect.OldName.ToString(), *NewRedirect.NewName.ToString());
 			continue;
 		}
 
 		if ((!NewRedirect.OldName.HasValidCharacters(NewRedirect.RedirectFlags) && !FPackageName::IsVersePackage(NewRedirect.OldName.PackageName.ToString()))
 			|| (!NewRedirect.NewName.HasValidCharacters(NewRedirect.RedirectFlags) && !FPackageName::IsVersePackage(NewRedirect.NewName.PackageName.ToString())))
 		{
-			UE_LOG(LogLinker, Error, TEXT("AddRedirectList(%s) failed to add redirector from %s to %s with invalid characters!"), *SourceString, *NewRedirect.OldName.ToString(), *NewRedirect.NewName.ToString());
+			UE_LOG(LogCoreRedirects, Error, TEXT("AddRedirect(%s) failed to add redirect from %s to %s with invalid characters!"), *SourceString, *NewRedirect.OldName.ToString(), *NewRedirect.NewName.ToString());
 			continue;
 		}
 
 		if (NewRedirect.NewName.PackageName != NewRedirect.OldName.PackageName && NewRedirect.OldName.OuterName != NAME_None)
 		{
-			UE_LOG(LogLinker, Error, TEXT("AddRedirectList(%s) failed to add redirector, it's not valid to modify package from %s to %s while specifying outer!"), *SourceString, *NewRedirect.OldName.ToString(), *NewRedirect.NewName.ToString());
+			UE_LOG(LogCoreRedirects, Error, TEXT("AddRedirect(%s) failed to add redirect, cannot modify package from %s to %s while specifying outer!"), *SourceString, *NewRedirect.OldName.ToString(), *NewRedirect.NewName.ToString());
 			continue;
 		}
 
 		if (NewRedirect.IsSubstringMatch())
 		{
-			UE_LOG(LogLinker, Log, TEXT("AddRedirectList(%s) has substring redirect %s, these are very slow and should be resolved as soon as possible!"), *SourceString, *NewRedirect.OldName.ToString());
+			UE_LOG(LogCoreRedirects, Log, TEXT("AddRedirect(%s) has substring redirect %s, these are very slow and should be resolved as soon as possible!"), *SourceString, *NewRedirect.OldName.ToString());
 		}
 		
 		if (AddSingleRedirect(NewRedirect, SourceString))
@@ -1141,22 +1168,22 @@ bool FCoreRedirects::AddSingleRedirect(const FCoreRedirect& NewRedirect, const F
 				if (bBothHaveValueChanges)
 				{
 					// No warning when there are values changes to merge
+					UE_LOG(LogCoreRedirects, Verbose, TEXT("AddRedirect(%s) merging value redirects for %s"), *SourceString, *ExistingRedirect.NewName.ToString());
 				}
 				else if (bOneIsPartialOtherIsFull)
 				{
-					UE_LOG(LogLinker, Warning, TEXT("AddRedirectList(%s) found duplicate redirectors for %s to %s, one a FullPath and the other an ObjectName-only.")
-						TEXT(" This used to be required for StructRedirects but is no longer. Remove the ObjectName-only redirector and keep the FullPath."),
+					UE_LOG(LogCoreRedirects, Warning, TEXT("AddRedirect(%s) found duplicate redirects for %s to %s, one a FullPath and the other ObjectName-only.")
+						TEXT(" This used to be required for StructRedirects but now you should remove the ObjectName-only redirect and keep the FullPath."),
 						*SourceString, *ExistingRedirect.OldName.ToString(), *ExistingRedirect.NewName.ToString());
 				}
 				else
 				{
-					// ENABLE THIS ONCE INIS ARE CONVERTED
-					//UE_LOG(LogLinker, Warning, TEXT("AddRedirectList(%s) found duplicate redirectors for %s to %s!"), *SourceString, *ExistingRedirect.OldName.ToString(), *ExistingRedirect.NewName.ToString());
+					UE_LOG(LogCoreRedirects, Verbose, TEXT("AddRedirect(%s) ignoring duplicate redirects for %s to %s"), *SourceString, *ExistingRedirect.OldName.ToString(), *ExistingRedirect.NewName.ToString());
 				}
 			}
 			else
 			{
-				UE_LOG(LogLinker, Error, TEXT("AddRedirectList(%s) found conflicting redirectors for %s! Old: %s, New: %s"), *SourceString, *ExistingRedirect.OldName.ToString(), *ExistingRedirect.NewName.ToString(), *NewRedirect.NewName.ToString());
+				UE_LOG(LogCoreRedirects, Error, TEXT("AddRedirect(%s) found conflicting redirects for %s! Old: %s, New: %s"), *SourceString, *ExistingRedirect.OldName.ToString(), *ExistingRedirect.NewName.ToString(), *NewRedirect.NewName.ToString());
 			}
 			break;
 		}
@@ -1173,36 +1200,38 @@ bool FCoreRedirects::AddSingleRedirect(const FCoreRedirect& NewRedirect, const F
 
 bool FCoreRedirects::RemoveRedirectList(TArrayView<const FCoreRedirect> Redirects, const FString& SourceString)
 {
+	UE_LOG(LogCoreRedirects, Verbose, TEXT("RemoveRedirect(%s) Removing %d redirects"), *SourceString, Redirects.Num());
+
 	bool bRemovedAny = false;
 	for (const FCoreRedirect& RedirectToRemove : Redirects)
 	{
 		if (!RedirectToRemove.OldName.IsValid() || !RedirectToRemove.NewName.IsValid())
 		{
-			UE_LOG(LogLinker, Error, TEXT("RemoveRedirectList(%s) failed to remove redirector from %s to %s with empty name!"), *SourceString, *RedirectToRemove.OldName.ToString(), *RedirectToRemove.NewName.ToString());
+			UE_LOG(LogCoreRedirects, Error, TEXT("RemoveRedirect(%s) failed to remove redirect from %s to %s with empty name!"), *SourceString, *RedirectToRemove.OldName.ToString(), *RedirectToRemove.NewName.ToString());
 			continue;
 		}
 
 		if (RedirectToRemove.HasValueChanges())
 		{
-			UE_LOG(LogLinker, Error, TEXT("RemoveRedirectList(%s) failed to remove redirector from %s to %s as it contains value changes!"), *SourceString, *RedirectToRemove.OldName.ToString(), *RedirectToRemove.NewName.ToString());
+			UE_LOG(LogCoreRedirects, Error, TEXT("RemoveRedirect(%s) failed to remove redirect from %s to %s as it contains value changes!"), *SourceString, *RedirectToRemove.OldName.ToString(), *RedirectToRemove.NewName.ToString());
 			continue;
 		}
 
 		if (!RedirectToRemove.OldName.HasValidCharacters(RedirectToRemove.RedirectFlags) || !RedirectToRemove.NewName.HasValidCharacters(RedirectToRemove.RedirectFlags))
 		{
-			UE_LOG(LogLinker, Error, TEXT("RemoveRedirectList(%s) failed to remove redirector from %s to %s with invalid characters!"), *SourceString, *RedirectToRemove.OldName.ToString(), *RedirectToRemove.NewName.ToString());
+			UE_LOG(LogCoreRedirects, Error, TEXT("RemoveRedirect(%s) failed to remove redirect from %s to %s with invalid characters!"), *SourceString, *RedirectToRemove.OldName.ToString(), *RedirectToRemove.NewName.ToString());
 			continue;
 		}
 
 		if (RedirectToRemove.NewName.PackageName != RedirectToRemove.OldName.PackageName && RedirectToRemove.OldName.OuterName != NAME_None)
 		{
-			UE_LOG(LogLinker, Error, TEXT("RemoveRedirectList(%s) failed to remove redirector, it's not valid to modify package from %s to %s while specifying outer!"), *SourceString, *RedirectToRemove.OldName.ToString(), *RedirectToRemove.NewName.ToString());
+			UE_LOG(LogCoreRedirects, Error, TEXT("RemoveRedirect(%s) failed to remove redirect, it's not valid to modify package from %s to %s while specifying outer!"), *SourceString, *RedirectToRemove.OldName.ToString(), *RedirectToRemove.NewName.ToString());
 			continue;
 		}
 
 		if (RedirectToRemove.IsSubstringMatch())
 		{
-			UE_LOG(LogLinker, Log, TEXT("RemoveRedirectList(%s) has substring redirect %s, these are very slow and should be resolved as soon as possible!"), *SourceString, *RedirectToRemove.OldName.ToString());
+			UE_LOG(LogCoreRedirects, Log, TEXT("RemoveRedirect(%s) has substring redirect %s, these are very slow and should be resolved as soon as possible!"), *SourceString, *RedirectToRemove.OldName.ToString());
 		}
 
 		bRemovedAny |= RemoveSingleRedirect(RedirectToRemove, SourceString);
@@ -1233,7 +1262,7 @@ bool FCoreRedirects::RemoveSingleRedirect(const FCoreRedirect& RedirectToRemove,
 		{
 			if (ExistingRedirect.NewName != RedirectToRemove.NewName)
 			{
-				// This isn't the redirector we were looking for... move on in case there's another match for our old name
+				// This isn't the redirect we were looking for... move on in case there's another match for our old name
 				continue;
 			}
 
@@ -1246,6 +1275,112 @@ bool FCoreRedirects::RemoveSingleRedirect(const FCoreRedirect& RedirectToRemove,
 	return bRemovedRedirect;
 }
 
+void FCoreRedirects::ValidateRedirectList(TArrayView<const FCoreRedirect> Redirects, const FString& SourceString)
+{
+	for (const FCoreRedirect& Redirect : Redirects)
+	{
+		if (Redirect.NewName.IsValid())
+		{
+			// If the new package is loaded but the new name isn't, this is very likely a bug
+			// If the new package isn't loaded the redirect can't be validated either way, so report it for manual follow up
+			UObject* NewPackage = FindObjectFast<UPackage>(nullptr, Redirect.NewName.PackageName);
+			FString NewPath = Redirect.NewName.ToString();
+			FString OldPath = Redirect.OldName.ToString();
+
+			if (CheckRedirectFlagsMatch(Redirect.RedirectFlags, ECoreRedirectFlags::Type_Class))
+			{
+				if (Redirect.NewName.PackageName.IsNone())
+				{
+					UE_LOG(LogCoreRedirects, Warning, TEXT("ValidateRedirect(%s) has missing package for Class redirect from %s to %s!"), *SourceString, *OldPath, *NewPath);
+				}
+				else
+				{
+					const UClass* FindClass = FindObject<const UClass>(FTopLevelAssetPath(NewPath));
+					if (!FindClass)
+					{
+						if (NewPackage)
+						{
+							UE_LOG(LogCoreRedirects, Error, TEXT("ValidateRedirect(%s) failed to find destination Class for redirect from %s to %s with loaded package!"), *SourceString, *OldPath, *NewPath);
+						}
+						else
+						{
+							UE_LOG(LogCoreRedirects, Log, TEXT("ValidateRedirect(%s) can't validate destination Class for redirect from %s to %s with unloaded package"), *SourceString, *OldPath, *NewPath);
+						}
+					}
+				}
+			}
+
+			if (CheckRedirectFlagsMatch(Redirect.RedirectFlags, ECoreRedirectFlags::Type_Struct))
+			{
+				if (Redirect.NewName.PackageName.IsNone())
+				{
+					UE_LOG(LogCoreRedirects, Warning, TEXT("ValidateRedirect(%s) has missing package for Struct redirect from %s to %s!"), *SourceString, *OldPath, *NewPath);
+				}
+				else
+				{
+					const UScriptStruct* FindStruct = FindObject<const UScriptStruct>(FTopLevelAssetPath(NewPath));
+					if (!FindStruct)
+					{
+						if (NewPackage)
+						{
+							UE_LOG(LogCoreRedirects, Error, TEXT("ValidateRedirect(%s) failed to find destination Struct for redirect from %s to %s with loaded package!"), *SourceString, *OldPath, *NewPath);
+						}
+						else
+						{
+							UE_LOG(LogCoreRedirects, Log, TEXT("ValidateRedirect(%s) can't validate destination Struct for redirect from %s to %s with unloaded package"), *SourceString, *OldPath, *NewPath);
+						}
+					}
+				}
+			}
+
+			if (CheckRedirectFlagsMatch(Redirect.RedirectFlags, ECoreRedirectFlags::Type_Enum))
+			{
+				if (Redirect.NewName.PackageName.IsNone())
+				{
+					// If the name is the same that's fine and this is just for value redirects
+					if (Redirect.NewName != Redirect.OldName)
+					{
+						UE_LOG(LogCoreRedirects, Warning, TEXT("ValidateRedirect(%s) has missing package for Enum redirect from %s to %s!"), *SourceString, *OldPath, *NewPath);
+					}
+				}
+				else
+				{
+					const UEnum* FindEnum = FindObject<const UEnum>(FTopLevelAssetPath(NewPath));
+					if (!FindEnum)
+					{
+						if (NewPackage)
+						{
+							UE_LOG(LogCoreRedirects, Error, TEXT("ValidateRedirect(%s) failed to find destination Enum for redirect from %s to %s with loaded package!"), *SourceString, *OldPath, *NewPath);
+						}
+						else
+						{
+							UE_LOG(LogCoreRedirects, Log, TEXT("ValidateRedirect(%s) can't validate destination Enum for redirect from %s to %s with unloaded package"), *SourceString, *OldPath, *NewPath);
+						}
+					}
+				}
+			}
+		}
+	}
+}
+
+void FCoreRedirects::ValidateAllRedirects()
+{
+	bValidatedOnce = true;
+
+	// Validate all existing redirects
+
+	for (const TPair<ECoreRedirectFlags, FRedirectNameMap>& Pair : RedirectTypeMap)
+	{
+		ECoreRedirectFlags PairFlags = Pair.Key;
+		FString ListName = FString::Printf(TEXT("Type %d"), (int32)PairFlags);
+
+		for (const TPair<FName, TArray<FCoreRedirect> >& ArrayPair : Pair.Value.RedirectMap)
+		{
+			ValidateRedirectList(ArrayPair.Value, ListName);
+		}
+	}
+}
+
 ECoreRedirectFlags FCoreRedirects::GetFlagsForTypeName(FName PackageName, FName TypeName)
 {
 	if (PackageName == GLongCoreUObjectPackageName)
diff --git a/Engine/Source/Runtime/CoreUObject/Private/UObject/LinkerLoad.cpp b/Engine/Source/Runtime/CoreUObject/Private/UObject/LinkerLoad.cpp
--- a/Engine/Source/Runtime/CoreUObject/Private/UObject/LinkerLoad.cpp
+++ b/Engine/Source/Runtime/CoreUObject/Private/UObject/LinkerLoad.cpp
@@ -349,7 +349,7 @@ void FLinkerLoad::CreateActiveRedirectsMap(const FString& GEngineIniName)
 				}
 			}
 
-			FCoreRedirects::AddRedirectList(NewRedirects, GEngineIniName);
+			FCoreRedirects::AddRedirectList(NewRedirects, FString::Printf(TEXT("ActiveClassRedirects %s"), *GEngineIniName));
 		}
 	}
 	else
diff --git a/Engine/Source/Runtime/CoreUObject/Public/UObject/CoreRedirects.h b/Engine/Source/Runtime/CoreUObject/Public/UObject/CoreRedirects.h
--- a/Engine/Source/Runtime/CoreUObject/Public/UObject/CoreRedirects.h
+++ b/Engine/Source/Runtime/CoreUObject/Public/UObject/CoreRedirects.h
@@ -23,6 +23,8 @@ class UClass;
 
 #define WITH_COREREDIRECTS_MULTITHREAD_WARNING !UE_BUILD_SHIPPING && !IS_PROGRAM && !WITH_EDITOR
 
+DECLARE_LOG_CATEGORY_EXTERN(LogCoreRedirects, Log, All);
+
 /** 
  * Flags describing the type and properties of this redirect
  */
@@ -317,6 +319,15 @@ struct COREUOBJECT_API FCoreRedirects
 	/** Returns true if this has ever been initialized */
 	static bool IsInitialized() { return bInitialized; }
 
+	/** Returns true if this is in debug mode that slows loading and adds additional warnings */
+	static bool IsInDebugMode() { return bInDebugMode; }
+
+	/** Validate a named list of redirects */
+	static void ValidateRedirectList(TArrayView<const FCoreRedirect> Redirects, const FString& SourceString);
+
+	/** Validates all known redirects and warn if they seem to point to missing things */
+	static void ValidateAllRedirects();
+
 	/** Gets map from config key -> Flags */
 	static const TMap<FName, ECoreRedirectFlags>& GetConfigKeyMap() { return ConfigKeyMap; }
 
@@ -359,6 +370,12 @@ private:
 	/** Whether this has been initialized at least once */
 	static bool bInitialized;
 
+	/** True if we are in debug mode that does extra validation */
+	static bool bInDebugMode;
+
+	/** True if we have done our initial validation. After initial validation, each change to redirects will validate independently */
+	static bool bValidatedOnce;
+
 #if WITH_COREREDIRECTS_MULTITHREAD_WARNING
 	/** Whether CoreRedirects is now being used multithreaded and therefore does not support writes to RedirectTypeMap keyvalue pairs */
 	static bool bIsInMultithreadedPhase;
diff --git a/Engine/Source/Runtime/Engine/Private/MemberReference.cpp b/Engine/Source/Runtime/Engine/Private/MemberReference.cpp
--- a/Engine/Source/Runtime/Engine/Private/MemberReference.cpp
+++ b/Engine/Source/Runtime/Engine/Private/MemberReference.cpp
@@ -258,7 +258,7 @@ void FMemberReference::InitFieldRedirectMap()
 				}			
 			}
 
-			FCoreRedirects::AddRedirectList(NewRedirects, GEngineIni);
+			FCoreRedirects::AddRedirectList(NewRedirects, TEXT("InitFieldRedirectMap"));
 			bFieldRedirectMapInitialized = true;
 		}
 	}
