Improve VA logging when the VA system is disabled but code tries to push/pull anyway.

#rnx

- When the VA is disabled we mount the FNullVirtualizationSystem instead.
-- Added logging when this system is used to push or pull payloads. The logging can be disabled by changing the define 'UE_LOG_ON_NULLSYSTEM_USE' to 0
-- The new logging gives errors, but if the define UE_USING_NULLSYSTEM_IS_ERROR is set to 0 then we will log warnings instead.
- EditorBulkData now checks before pulling payloads and if the current VA system is disabled then we print a specific error for this case.

- Overall this should make it easier to identify that the problem is that VA is disabled if we only have access to a small snippet of the log surrounding the crash

diff --git a/Engine/Source/Runtime/Core/Private/Virtualization/VirtualizationSystem.cpp b/Engine/Source/Runtime/Core/Private/Virtualization/VirtualizationSystem.cpp
--- a/Engine/Source/Runtime/Core/Private/Virtualization/VirtualizationSystem.cpp
+++ b/Engine/Source/Runtime/Core/Private/Virtualization/VirtualizationSystem.cpp
@@ -19,6 +19,20 @@
 #ifndef UE_VIRTUALIZATION_SYSTEM_LAZY_INIT
 	#define UE_VIRTUALIZATION_SYSTEM_LAZY_INIT 0
 #endif //UE_VIRTUALIZATION_SYSTEM_LAZY_INIT
+
+// When enabled we will log if FNullVirtualizationSystem tries to push or pull payloads
+#define UE_LOG_ON_NULLSYSTEM_USE 1
+
+// When enabled messages logged when FNullVirtualizationSystem is in use will use the 'Error'
+// severity. When disabled the messages will use 'Warning' severity.
+#define UE_USING_NULLSYSTEM_IS_ERROR 1
+
+#if UE_USING_NULLSYSTEM_IS_ERROR
+	#define UE_NULLSYSTEM_SEVERITY Error
+#else
+	#define UE_NULLSYSTEM_SEVERITY Warning
+#endif // UE_USING_NULLSYSTEM_IS_ERROR
+
 namespace UE::Virtualization
 {
 
@@ -60,16 +74,28 @@ public:
 
 	virtual bool PushData(const FIoHash& Id, const FCompressedBuffer& Payload, EStorageType StorageType, const FString& Context) override
 	{
+#if UE_LOG_ON_NULLSYSTEM_USE
+		UE_LOG(LogVirtualization, UE_NULLSYSTEM_SEVERITY, TEXT("Cannot push payload '%s' as the virtualization system is disabled"), *LexToString(Id));
+#endif //UE_LOG_ON_NULLSYSTEM_USE
+		
 		return false;
 	}
 
 	virtual bool PushData(TArrayView<FPushRequest> Requests, EStorageType StorageType) override
 	{
+#if UE_LOG_ON_NULLSYSTEM_USE
+		UE_LOG(LogVirtualization, UE_NULLSYSTEM_SEVERITY, TEXT("Cannot push payloads as the virtualization system is disabled"));
+#endif //UE_LOG_ON_NULLSYSTEM_USE
+
 		return false;
 	}
 
 	virtual FCompressedBuffer PullData(const FIoHash& Id) override
 	{
+#if UE_LOG_ON_NULLSYSTEM_USE
+		UE_LOG(LogVirtualization, UE_NULLSYSTEM_SEVERITY, TEXT("Cannot pull payload '%s' as the virtualization system is disabled"), *LexToString(Id));
+#endif //UE_LOG_ON_NULLSYSTEM_USE
+
 		return FCompressedBuffer();
 	}
 
diff --git a/Engine/Source/Runtime/CoreUObject/Private/Serialization/EditorBulkData.cpp b/Engine/Source/Runtime/CoreUObject/Private/Serialization/EditorBulkData.cpp
--- a/Engine/Source/Runtime/CoreUObject/Private/Serialization/EditorBulkData.cpp
+++ b/Engine/Source/Runtime/CoreUObject/Private/Serialization/EditorBulkData.cpp
@@ -1458,17 +1458,25 @@ FCompressedBuffer FEditorBulkData::PullData() const
 {
 	TRACE_CPUPROFILER_EVENT_SCOPE(FEditorBulkData::PullData);
 
-	FCompressedBuffer PulledPayload = UE::Virtualization::IVirtualizationSystem::Get().PullData(PayloadContentId);
-
-	if (PulledPayload)
+	UE::Virtualization::IVirtualizationSystem& System = UE::Virtualization::IVirtualizationSystem::Get();
+	if (System.IsEnabled())
 	{
-		checkf(	PayloadSize == PulledPayload.GetRawSize(),
-				TEXT("Mismatch between serialized length (%" INT64_FMT ") and virtualized data length (%" UINT64_FMT ")"),
-				PayloadSize,
-				PulledPayload.GetRawSize());
+		FCompressedBuffer PulledPayload = UE::Virtualization::IVirtualizationSystem::Get().PullData(PayloadContentId);
+
+		checkf(!PulledPayload || PayloadSize == PulledPayload.GetRawSize(),
+			TEXT("Mismatch between serialized length (%" INT64_FMT ") and virtualized data length (%" UINT64_FMT ")"),
+			PayloadSize,
+			PulledPayload.GetRawSize());
+
+		UE_CLOG(PulledPayload.IsNull(), LogSerialization, Error, TEXT("Failed to pull payload '%s'"), *LexToString(PayloadContentId));
+
+		return PulledPayload;
 	}
-	
-	return PulledPayload;
+	else
+	{
+		UE_LOG(LogSerialization, Error, TEXT("Cannot pull payload '%s' as the virtualization system is disabled"), *LexToString(PayloadContentId));
+		return FCompressedBuffer();
+	}	
 }
 
 bool FEditorBulkData::CanUnloadData() const
@@ -1664,7 +1672,7 @@ void FEditorBulkData::SerializeToPackageTrailer(FLinkerSave& LinkerSave, FCompre
 
 				this->OffsetInFile = PayloadOffset;
 				this->Flags = (this->Flags & TransientFlags) | (UpdatedFlags & ~TransientFlags);
-	
+
 				// If the payload is valid we might want to fix up the package path or virtualization flags now
 				// that the package has saved.
 				if (!this->PayloadContentId.IsZero())
@@ -1703,7 +1711,11 @@ void FEditorBulkData::SerializeToPackageTrailer(FLinkerSave& LinkerSave, FCompre
 		}
 	};
 
-	UE::Virtualization::EPayloadFilterReason PayloadFilter = UE::Virtualization::IVirtualizationSystem::Get().FilterPayload(Owner);
+	UE::Virtualization::EPayloadFilterReason PayloadFilter = UE::Virtualization::EPayloadFilterReason::None;
+	if (UE::Virtualization::IVirtualizationSystem::IsInitialized())
+	{
+		PayloadFilter = UE::Virtualization::IVirtualizationSystem::Get().FilterPayload(Owner);
+	}
 
 #if UE_ENABLE_VIRTUALIZATION_TOGGLE
 	if (bSkipVirtualization)
@@ -1781,7 +1793,6 @@ FCompressedBuffer FEditorBulkData::GetDataInternal() const
 		
 		checkf(Payload.IsNull(), TEXT("Pulling data somehow assigned it to the bulk data object!")); //Make sure that we did not assign the buffer internally
 
-		UE_CLOG(CompressedPayload.IsNull(), LogSerialization, Error, TEXT("Failed to pull payload '%s'"), *LexToString(PayloadContentId));
 		UE_CLOG(!IsDataValid(*this, CompressedPayload), LogSerialization, UE_CORRUPTED_DATA_SEVERITY, TEXT("Virtualized payload '%s' is corrupt! Check the backend storage."), *LexToString(PayloadContentId));
 		
 		TRACE_COUNTER_ADD(EditorBulkData_PayloadDataPulled, (int64)CompressedPayload.GetCompressedSize());
