Reverting to Default an entire static array no longer causes a check to fail in FArrayProperty::ImportInnerTextItem().

This check was meant to check if the user had accidentally passed in the top-level FArrayProperty rather than its inner property, but in the case of static arrays the inner property is the top-level property, but with an ArrayDim > 1, which makes the check fail because the property has a valid offset from the start of its parent struct/object.

We now check that the FScriptArrayHelper is set (signaling that we're passing an FArrayProperty) before running the check.

diff --git a/Engine/Source/Runtime/CoreUObject/Private/UObject/PropertyArray.cpp b/Engine/Source/Runtime/CoreUObject/Private/UObject/PropertyArray.cpp
--- a/Engine/Source/Runtime/CoreUObject/Private/UObject/PropertyArray.cpp
+++ b/Engine/Source/Runtime/CoreUObject/Private/UObject/PropertyArray.cpp
@@ -649,7 +649,7 @@ const TCHAR* FArrayProperty::ImportTextInnerItem( const TCHAR* Buffer, const FPr
 		{
 			uint8* Address = ArrayHelper ? ArrayHelper->GetRawPtr(Index) : ((uint8*)Data + Inner->ElementSize * Index);
 			// Parse the item
-			check(Inner->GetOffset_ForInternal() == 0);
+			checkf(ArrayHelper == nullptr || Inner->GetOffset_ForInternal() == 0, TEXT("Expected the Inner property of the FArrayProperty."));
 			Buffer = Inner->ImportText_Direct(Buffer, Address, Parent, PortFlags | PPF_Delimited, ErrorText);
 
 			if(!Buffer)
