Add load checks in the linker that will prevent illegal reference to resolve on load
Current check are for package that are not externally referenceable across plugin boundaries
This won't prevent the load of the package pointed to by the illegal reference but the reference itself won't resolve during loading

diff --git a/Engine/Source/Runtime/CoreUObject/Private/UObject/LinkerLoad.cpp b/Engine/Source/Runtime/CoreUObject/Private/UObject/LinkerLoad.cpp
--- a/Engine/Source/Runtime/CoreUObject/Private/UObject/LinkerLoad.cpp
+++ b/Engine/Source/Runtime/CoreUObject/Private/UObject/LinkerLoad.cpp
@@ -2494,7 +2494,12 @@ UObject* FLinkerLoad::FindExistingImport(int32 ImportIndex)
 	{
 		// if the import outer is null then we have a package, resolve it, potentially remapping it
 		FName ObjectName = InstancingContextRemap(Import.ObjectName);
-		return StaticFindObjectFast(UPackage::StaticClass(), nullptr, ObjectName, /*bExactClass*/true);
+		UPackage* Package = static_cast<UPackage*>(StaticFindObjectFast(UPackage::StaticClass(), nullptr, ObjectName, /*bExactClass*/true));
+		if (!IsPackageReferenceAllowed(Package))
+		{
+			return nullptr;
+		}
+		return Package;
 	}
 	// if our outer is an import, recurse to find it
 	else if (Import.OuterIndex.IsImport())
@@ -3145,6 +3150,12 @@ bool FLinkerLoad::VerifyImportInner(const int32 ImportIndex, FString& WarningSuf
 			return nullptr;
 		}
 
+		if (!IsPackageReferenceAllowed(Package))
+		{
+			UE_LOG(LogLinker, Warning, TEXT("VerifyImport: illegal reference to private package for import object '%s'"), *GetImportFullName(ImportIndex));
+			return nullptr;
+		}
+
 		// while gathering dependencies, there is no need to verify all of the imports for the entire package
 		if (bIsGatheringDependencies)
 		{
@@ -4432,6 +4443,20 @@ bool FLinkerLoad::WillTextureBeLoaded( UClass* Class, int32 ExportIndex )
 	}
 }
 
+bool FLinkerLoad::IsPackageReferenceAllowed(UPackage* InPackage)
+{
+	if (InPackage && !InPackage->IsExternallyReferenceable())
+	{
+		FName MountPointName = FPackageName::GetPackageMountPoint(LinkerRoot->GetName());
+		FName ImportMountPointName = FPackageName::GetPackageMountPoint(InPackage->GetName());
+		if (MountPointName != ImportMountPointName)
+		{
+			return false;
+		}
+	}
+	return true;
+}
+
 UObject* FLinkerLoad::CreateExport( int32 Index )
 {
 	FScopedCreateExportCounter ScopedCounter( this, Index );
@@ -5083,7 +5108,7 @@ UObject* FLinkerLoad::CreateImport( int32 Index )
 		if (!GIsEditor && !IsRunningCommandlet())
 		{
 			// Try to find existing version in memory first.
-			if( UPackage* ClassPackage = FindObjectFast<UPackage>( NULL, Import.ClassPackage, false ) )
+			if( UPackage* ClassPackage = FindObjectFast<UPackage>( nullptr, Import.ClassPackage, false ) )
 			{
 				if( UClass*	FindClass = FindObjectFast<UClass>( ClassPackage, Import.ClassName, false ) ) // 
 				{
@@ -5094,19 +5119,23 @@ UObject* FLinkerLoad::CreateImport( int32 Index )
 					Preload( FindClass );
 
 					FindClass->GetDefaultObject(); // build the CDO if it isn't already built
-					UObject*	FindObject		= NULL;
+					UObject*	FindObject		= nullptr;
 	
 					// Import is a toplevel package.
 					if( Import.OuterIndex.IsNull() )
 					{
 						FName ObjectName = InstancingContextRemap(Import.ObjectName);
-						FindObject = CreatePackage(*ObjectName.ToString());
+						UPackage* Pkg = CreatePackage(*ObjectName.ToString());
+						if (IsPackageReferenceAllowed(Pkg))
+						{
+							FindObject = Pkg;
+						}
 					}
 					// Import is regular import/ export.
 					else
 					{
 						// Find the imports' outer.
-						UObject* FindOuter = NULL;
+						UObject* FindOuter = nullptr;
 						// Import.
 						if( Import.OuterIndex.IsImport() )
 						{
@@ -5120,7 +5149,12 @@ UObject* FLinkerLoad::CreateImport( int32 Index )
 							else if( OuterImport.OuterIndex.IsNull() )
 							{
 								FName ObjectName = InstancingContextRemap(OuterImport.ObjectName);
-								FindOuter = CreatePackage( *ObjectName.ToString() );
+								UPackage* Pkg = CreatePackage(*ObjectName.ToString());
+								if (IsPackageReferenceAllowed(Pkg))
+								{
+									FindOuter = Pkg;
+								}
+
 							}
 							// Outer is regular import/ export, use IndexToObject to potentially recursively load/ find it.
 							else
diff --git a/Engine/Source/Runtime/CoreUObject/Public/UObject/LinkerLoad.h b/Engine/Source/Runtime/CoreUObject/Public/UObject/LinkerLoad.h
--- a/Engine/Source/Runtime/CoreUObject/Public/UObject/LinkerLoad.h
+++ b/Engine/Source/Runtime/CoreUObject/Public/UObject/LinkerLoad.h
@@ -784,6 +784,15 @@ private:
 
 #endif // WITH_EDITOR
 
+
+	/**
+	 * Validate if the current linker may reference a certain package
+	 * Checks if the reference to a certain package from this linker is legal if that package isn't externally referenceable (i.e private outiside its plugin)
+	 * @param InPackage the package to validate
+	 * @returns true if the package can be referenced from this linker
+	 */
+	bool IsPackageReferenceAllowed(UPackage* InPackage);
+
 	UObject* CreateExport( int32 Index );
 
 	/**
diff --git a/Engine/Source/Runtime/CoreUObject/Public/UObject/ObjectMacros.h b/Engine/Source/Runtime/CoreUObject/Public/UObject/ObjectMacros.h
--- a/Engine/Source/Runtime/CoreUObject/Public/UObject/ObjectMacros.h
+++ b/Engine/Source/Runtime/CoreUObject/Public/UObject/ObjectMacros.h
@@ -125,7 +125,7 @@ enum EPackageFlags
 	PKG_UncookedOnly				= 0x00000100,	///< Loaded only in uncooked builds (i.e. runtime in editor)
 	PKG_Cooked						= 0x00000200,	///< Package is cooked
 	PKG_ContainsNoAsset				= 0x00000400,	///< Package doesn't contain any asset object (although asset tags can be present)
-	PKG_NotExternallyReferenceable  = 0x00000800,   ///< (Not Implemented) Objects in this package cannot be referenced in a different plugin or mount point (i.e /Game -> /Engine)
+	PKG_NotExternallyReferenceable  = 0x00000800,   ///< Objects in this package cannot be referenced in a different plugin or mount point (i.e /Game -> /Engine)
 //	PKG_Unused						= 0x00001000,
 	PKG_UnversionedProperties		= 0x00002000,   ///< Uses unversioned property serialization instead of versioned tagged property serialization
 	PKG_ContainsMapData				= 0x00004000,   ///< Contains map data (UObjects only referenced by a single ULevel) but is stored in a different package
