[UE]
* Added TrackedActivity scopes for I/O. This tracking is enabled in non-shipping builds and make it possible to see I/O in windows console. Code is hidden behind define UE_ENABLE_TRACKED_IO

diff --git a/Engine/Source/Runtime/AssetRegistry/Private/AssetDataGatherer.cpp b/Engine/Source/Runtime/AssetRegistry/Private/AssetDataGatherer.cpp
--- a/Engine/Source/Runtime/AssetRegistry/Private/AssetDataGatherer.cpp
+++ b/Engine/Source/Runtime/AssetRegistry/Private/AssetDataGatherer.cpp
@@ -26,6 +26,7 @@
 #include "Misc/PathViews.h"
 #include "Misc/ScopeLock.h"
 #include "Misc/ScopeExit.h"
+#include "Misc/TrackedActivity.h"
 #include "PackageReader.h"
 #include "Serialization/Archive.h"
 #include "Tasks/Task.h"
@@ -3212,6 +3213,7 @@ void FAssetDataGatherer::TickInternal(bool& bOutIsTickInterrupt)
 				ReadContext.bCanceled = true;
 				return;
 			}
+			UE_SCOPED_IO_ACTIVITY(*WriteToString<512>(TEXT("Loading Asset"), ReadContext.PackageName.ToString()));
 			ReadContext.bResult = ReadAssetFile(ReadContext.AssetFileData.LongPackageName, ReadContext.AssetFileData.LocalAbsPath, ReadContext.AssetDataFromFile, ReadContext.DependencyData, ReadContext.CookedPackageNamesWithoutAssetData, ReadContext.bCanAttemptAssetRetry);
 		},
 		EParallelForFlags::Unbalanced | EParallelForFlags::BackgroundPriority
diff --git a/Engine/Source/Runtime/AssetRegistry/Private/AssetRegistry.cpp b/Engine/Source/Runtime/AssetRegistry/Private/AssetRegistry.cpp
--- a/Engine/Source/Runtime/AssetRegistry/Private/AssetRegistry.cpp
+++ b/Engine/Source/Runtime/AssetRegistry/Private/AssetRegistry.cpp
@@ -25,6 +25,7 @@
 #include "Misc/RedirectCollector.h"
 #include "Misc/ScopeExit.h"
 #include "Misc/ScopeRWLock.h"
+#include "Misc/TrackedActivity.h"
 #include "PackageReader.h"
 #include "Serialization/ArrayReader.h"
 #include "Templates/UnrealTemplate.h"
@@ -3218,6 +3219,8 @@ void UAssetRegistryImpl::ScanFilesSynchronous(const TArray<FString>& InFilePaths
 void UAssetRegistryImpl::ScanPathsSynchronousInternal(const TArray<FString>& InDirs, const TArray<FString>& InFiles,
 	bool bInForceRescan, bool bInIgnoreDenyListScanFilters)
 {
+	UE_SCOPED_IO_ACTIVITY(*WriteToString<256>("Scan Paths"));
+
 	TRACE_CPUPROFILER_EVENT_SCOPE(UAssetRegistryImpl::ScanPathsSynchronousInternal);
 	UE_TRACK_REFERENCING_OPNAME_SCOPED(PackageAccessTrackingOps::NAME_ResetContext);
 	const double SearchStartTime = FPlatformTime::Seconds();
diff --git a/Engine/Source/Runtime/Core/Private/Windows/WindowsPlatformProcess.cpp b/Engine/Source/Runtime/Core/Private/Windows/WindowsPlatformProcess.cpp
--- a/Engine/Source/Runtime/Core/Private/Windows/WindowsPlatformProcess.cpp
+++ b/Engine/Source/Runtime/Core/Private/Windows/WindowsPlatformProcess.cpp
@@ -16,6 +16,7 @@
 #include "Misc/SingleThreadEvent.h"
 #include "Misc/CommandLine.h"
 #include "Misc/Paths.h"
+#include "Misc/TrackedActivity.h"
 #include "Internationalization/Internationalization.h"
 #include "CoreGlobals.h"
 #include "Stats/Stats.h"
@@ -1980,6 +1981,8 @@ static void LogImportDiagnostics(const FString& FileName, const TArray<FString>&
 
 void *FWindowsPlatformProcess::LoadLibraryWithSearchPaths(const FString& FileName, const TArray<FString>& SearchPaths)
 {
+	UE_SCOPED_IO_ACTIVITY(*WriteToString<256>("Loading Dll ", FileName));
+
 	// Make sure the initial module exists. If we can't find it from the path we're given, it's probably a system dll.
 	FString FullFileName = FileName;
 	if (FPaths::FileExists(FullFileName))
diff --git a/Engine/Source/Runtime/Core/Public/Misc/TrackedActivity.h b/Engine/Source/Runtime/Core/Public/Misc/TrackedActivity.h
--- a/Engine/Source/Runtime/Core/Public/Misc/TrackedActivity.h
+++ b/Engine/Source/Runtime/Core/Public/Misc/TrackedActivity.h
@@ -127,3 +127,17 @@ public:
 private:
 	FTrackedActivity& Activity;
 };
+
+
+/**
+* Enabled tracking of IO
+*/
+#ifndef UE_ENABLE_TRACKED_IO
+	#if UE_BUILD_SHIPPING
+		#define UE_ENABLE_TRACKED_IO 0
+		#define UE_SCOPED_IO_ACTIVITY(...)
+	#else
+		#define UE_ENABLE_TRACKED_IO 1
+		#define UE_SCOPED_IO_ACTIVITY(...) FTrackedActivityScope ANONYMOUS_VARIABLE(IOActivity_)(FTrackedActivity::GetIOActivity(), __VA_ARGS__);
+	#endif
+#endif
\ No newline at end of file
diff --git a/Engine/Source/Runtime/CoreUObject/Private/Serialization/AsyncLoading.cpp b/Engine/Source/Runtime/CoreUObject/Private/Serialization/AsyncLoading.cpp
--- a/Engine/Source/Runtime/CoreUObject/Private/Serialization/AsyncLoading.cpp
+++ b/Engine/Source/Runtime/CoreUObject/Private/Serialization/AsyncLoading.cpp
@@ -20,6 +20,7 @@
 #include "Misc/MessageDialog.h"
 #include "Misc/PackageName.h"
 #include "Misc/ScopedSlowTask.h"
+#include "Misc/TrackedActivity.h"
 #include "UObject/PackageFileSummary.h"
 #include "UObject/Linker.h"
 #include "UObject/CoreRedirects.h"
@@ -495,7 +496,35 @@ static FName StaticGetNativeClassName(UClass* InClass)
 /** Returns true if we're inside a FGCScopeLock */
 bool IsGarbageCollectionLocked();
 
+#if UE_ENABLE_TRACKED_IO
+struct FActiveAsyncLoadContext* ActiveAsyncLoadContext;
 
+struct FActiveAsyncLoadContext
+{
+	uint32 TotalCompleted = 0;
+
+	FActiveAsyncLoadContext(const TCHAR* Name)
+	{
+		if (*Name == '/')
+			++Name;
+		FTrackedActivity::GetIOActivity().Push(Name);
+		FTrackedActivity::GetIOActivity().Push(TEXT(""), true);
+		ActiveAsyncLoadContext = this;
+	}
+
+	~FActiveAsyncLoadContext()
+	{
+		ActiveAsyncLoadContext = nullptr;
+		FTrackedActivity::GetIOActivity().Pop();
+		FTrackedActivity::GetIOActivity().Pop();
+	}
+
+	void Update(uint32 Queued)
+	{
+		FTrackedActivity::GetIOActivity().Update(*TStringBuilder<256>().Appendf(TEXT(" (%u/%u)"), TotalCompleted, TotalCompleted + Queued));
+	}
+};
+#endif // UE_ENABLE_TRACKED_IO
 
 /**
  * Updates FUObjectThreadContext with the current package when processing it.
@@ -3495,6 +3524,7 @@ void FAsyncPackage::EventDrivenSerializeExport(int32 LocalExportIndex)
 			}
 			else
 			{
+				UE_SCOPED_IO_ACTIVITY(*Object->GetName());
 				Object->Serialize(*Linker);
 			}
 		}
@@ -4249,6 +4279,10 @@ void FAsyncLoadingThread::AddToLoadedPackages(FAsyncPackage* Package)
 #endif
 	if (!LoadedPackages.Contains(Package))
 	{
+		#if UE_ENABLE_TRACKED_IO
+		if (ActiveAsyncLoadContext)
+			++ActiveAsyncLoadContext->TotalCompleted;
+		#endif
 		LoadedPackages.Add(Package);
 		LoadedPackagesNameLookup.Add(Package->GetPackageName(), Package);
 	}
@@ -4300,6 +4334,11 @@ EAsyncPackageState::Type FAsyncLoadingThread::ProcessAsyncLoading(int32& OutPack
 				FCoreDelegates::OnAsyncLoadingFlushUpdate.Broadcast();
 			}
 
+			#if UE_ENABLE_TRACKED_IO
+			if (ActiveAsyncLoadContext)
+				ActiveAsyncLoadContext->Update(AsyncPackages.Num());
+			#endif
+
 			bool bDidSomething = false;
 			{
 				bDidSomething = GetPrecacheHandler().ProcessIncoming();
@@ -4508,6 +4547,11 @@ EAsyncPackageState::Type FAsyncLoadingThread::ProcessAsyncLoading(int32& OutPack
 				CreateAsyncPackagesFromQueue(bUseTimeLimit, bUseFullTimeLimit, RemainingTimeLimit);
 			}
 
+			#if UE_ENABLE_TRACKED_IO
+			if (ActiveAsyncLoadContext)
+				ActiveAsyncLoadContext->Update(AsyncPackages.Num());
+			#endif
+
 			if (bNeedsHeartbeatTick)
 			{
 				// Update heartbeat after each package has been processed
@@ -7467,6 +7511,10 @@ EAsyncPackageState::Type FAsyncLoadingThread::ProcessLoadingUntilComplete(TFunct
 		TimeLimit = 60.0 * 60.0;
 	}
 
+	#if UE_ENABLE_TRACKED_IO
+	FActiveAsyncLoadContext Context(*QueuedPackages[0]->Name.ToString());
+	#endif
+
 	while (IsAsyncLoadingPackages() && TimeLimit > 0 && !CompletionPredicate())
 	{
 		double TickStartTime = FPlatformTime::Seconds();
diff --git a/Engine/Source/Runtime/CoreUObject/Private/UObject/UObjectGlobals.cpp b/Engine/Source/Runtime/CoreUObject/Private/UObject/UObjectGlobals.cpp
--- a/Engine/Source/Runtime/CoreUObject/Private/UObject/UObjectGlobals.cpp
+++ b/Engine/Source/Runtime/CoreUObject/Private/UObject/UObjectGlobals.cpp
@@ -44,6 +44,7 @@
 #include "UObject/LinkerLoad.h"
 #include "Blueprint/BlueprintSupport.h"
 #include "Misc/SecureHash.h"
+#include "Misc/TrackedActivity.h"
 #include "UObject/ConstructorHelpers.h"
 #include "UObject/UObjectThreadContext.h"
 #include "UObject/LinkerManager.h"
@@ -1494,6 +1495,8 @@ UPackage* LoadPackageInternal(UPackage* InOuter, const FPackagePath& PackagePath
 	// Set up a load context
 	TRefCountPtr<FUObjectSerializeContext> LoadContext = FUObjectThreadContext::Get().GetSerializeContext();
 
+	UE_SCOPED_IO_ACTIVITY(*WriteToString<512>(TEXT("Sync %s"), PackagePath.GetDebugName()));
+
 	// Try to load.
 	BeginLoad(LoadContext, *PackagePath.GetDebugName());
 
