Improve event scopes in editorbulkdata to show Registering and Unregistering from the BulkDataRegistry

#rnx

diff --git a/Engine/Source/Runtime/CoreUObject/Private/Serialization/EditorBulkData.cpp b/Engine/Source/Runtime/CoreUObject/Private/Serialization/EditorBulkData.cpp
--- a/Engine/Source/Runtime/CoreUObject/Private/Serialization/EditorBulkData.cpp
+++ b/Engine/Source/Runtime/CoreUObject/Private/Serialization/EditorBulkData.cpp
@@ -498,6 +498,8 @@ static FGuid CreateUniqueGuid(const FGuid& NonUniqueGuid, const UObject* Owner,
 void FEditorBulkData::Register(UObject* Owner, const TCHAR* ErrorLogCaller, bool bAllowUpdateId)
 {
 #if WITH_EDITOR
+	TRACE_CPUPROFILER_EVENT_SCOPE(FEditorBulkData::Register);
+
 	if (BulkDataId.IsValid() && PayloadSize > 0 && !EnumHasAnyFlags(Flags, EFlags::IsTornOff))
 	{
 		UPackage* OwnerPackage = Owner ? Owner->GetPackage() : nullptr;
@@ -548,6 +550,8 @@ void FEditorBulkData::OnExitMemory()
 void FEditorBulkData::Unregister()
 {
 #if WITH_EDITOR
+	TRACE_CPUPROFILER_EVENT_SCOPE(FEditorBulkData::Unregister);
+
 	if (EnumHasAnyFlags(Flags, EFlags::HasRegistered))
 	{
 		check(!EnumHasAnyFlags(Flags, EFlags::IsTornOff));
