Fixed an editor  crash when running from cooked packages stored in pakfiles

#rnx

- More fallout from the BD/BD2 merge. In this case the bulkdata object should be trying to read its payload from a .ubulk file, but ends up using the attached archive by mistake which means we end up reading from the .uexp file instead and either assert due to the offset being invalid or read back garbage and crash later.
- Fixed by bringing back the old logic (we should not use the attached archive if ::IsInExternalResource() is true OR if ::IsInSeparateFile() is true, checking just ::IsInExternalResource was not enough)

diff --git a/Engine/Source/Runtime/CoreUObject/Private/Serialization/BulkData.cpp b/Engine/Source/Runtime/CoreUObject/Private/Serialization/BulkData.cpp
--- a/Engine/Source/Runtime/CoreUObject/Private/Serialization/BulkData.cpp
+++ b/Engine/Source/Runtime/CoreUObject/Private/Serialization/BulkData.cpp
@@ -1944,7 +1944,7 @@ bool FBulkData::TryLoadDataIntoMemory(FIoBuffer Dest)
 	// To make calls to those BulkDatas threadsafe, we need to not use AttachedAr when called multithreaded.
 	// Also don't use the attached archive when loading from the editor domain and the bulk data is referenced
 	// from the original .uasset file
-	if (IsInGameThread() && IsInExternalResource() == false)
+	if (IsInGameThread() && !IsInSeparateFile() && !IsInExternalResource())
 	{
 		if (FArchive* Ar = AttachedAr)
 		{
