[CoreUObject] change to Blueprint loading to only resolve imports that are needed. prevents loading of unnecessary data

diff --git a/Engine/Source/Runtime/CoreUObject/Private/Blueprint/BlueprintSupport.cpp b/Engine/Source/Runtime/CoreUObject/Private/Blueprint/BlueprintSupport.cpp
--- a/Engine/Source/Runtime/CoreUObject/Private/Blueprint/BlueprintSupport.cpp
+++ b/Engine/Source/Runtime/CoreUObject/Private/Blueprint/BlueprintSupport.cpp
@@ -991,6 +991,12 @@ public:
 		}
 	}
 
+	static bool IsEmpty()
+	{
+		FScopeLock UnresolvedStructsLock(&UnresolvedStructsCritical);
+		return UnresolvedStructs.IsEmpty();
+	}
+
 private:
 	/** The struct that is currently being "resolved" */
 	UStruct* TrackedStruct;
@@ -1750,6 +1756,13 @@ void FLinkerLoad::PRIVATE_ForceLoadAllDependencies(UPackage* Package)
 
 void FLinkerLoad::ResolveAllImports()
 {
+	// early out of there are no unresolved structs
+	// saves calls to CreateImport when there is no work to do
+	if (FUnresolvedStructTracker::IsEmpty())
+	{
+		return;
+	}
+
 	for (int32 ImportIndex = 0; ImportIndex < ImportMap.Num() && IsBlueprintFinalizationPending(); ++ImportIndex)
 	{
 		// first, make sure every import object is available... just because 
@@ -1846,11 +1859,8 @@ void FLinkerLoad::FinalizeBlueprint(UClass* LoadClass)
 	// have to)... we do however need it here in FinalizeBlueprint(), because
 	// we need it ran for any super-classes before we regen
 
-	if (!IsImportLazyLoadEnabled())
-	{
-		// @TODO: OBJPTR: Need to find other options for solving this issue of placeholder classes during blueprint compile without forcing all imports to resolve always
-		ResolveAllImports();
-	}
+	// @TODO: OBJPTR: Need to find other options for solving this issue of placeholder classes during blueprint compile without forcing all imports to resolve always
+	ResolveAllImports();
 
 	// Now that imports have been resolved we optionally flush the compilation
 	// queue. This is only done for level blueprints, which will have instances
