Refactor instancing context to make a clear split between package name and object path mappings

diff --git a/Engine/Source/Editor/UnrealEd/Private/EditorServer.cpp b/Engine/Source/Editor/UnrealEd/Private/EditorServer.cpp
--- a/Engine/Source/Editor/UnrealEd/Private/EditorServer.cpp
+++ b/Engine/Source/Editor/UnrealEd/Private/EditorServer.cpp
@@ -2688,7 +2688,10 @@ bool UEditorEngine::Map_Load(const TCHAR* Str, FOutputDevice& Ar)
 					// SoftObjectPaths: Specific case for new maps (/Temp/Untitled) where we need to remap the AssetPath and not just the Package name because the World gets renamed (See UWorld::PostLoad)
 					const FString ShortPackageName = FPackageName::GetShortName(LongTempFname);
 					const FString ShortWorldPackageName = FPackageName::GetShortName(WorldPackage);
-					WorldPackageInstancingContext.AddMapping(*WriteToString<256>(LongTempFname, TEXT("."), ShortPackageName), *WriteToString<256>(WorldPackage->GetName(), TEXT("."), ShortWorldPackageName ));
+					WorldPackageInstancingContext.AddPathMapping(
+						FSoftObjectPath(*WriteToString<256>(LongTempFname, TEXT("."), ShortPackageName)),
+						FSoftObjectPath(*WriteToString<256>(WorldPackage->GetName(), TEXT("."), ShortWorldPackageName))
+					);
 
 					WorldPackage = LoadPackage( WorldPackage, *LongTempFname, LoadFlags, nullptr /* InReaderOverride */, &WorldPackageInstancingContext);
 					WorldPackage->SetPackageFlags(PKG_NewlyCreated);
diff --git a/Engine/Source/Runtime/CoreUObject/Private/Serialization/AsyncLoading.cpp b/Engine/Source/Runtime/CoreUObject/Private/Serialization/AsyncLoading.cpp
--- a/Engine/Source/Runtime/CoreUObject/Private/Serialization/AsyncLoading.cpp
+++ b/Engine/Source/Runtime/CoreUObject/Private/Serialization/AsyncLoading.cpp
@@ -2064,7 +2064,7 @@ FObjectImport* FAsyncPackage::FindExistingImport(int32 LocalImportIndex)
 		Import->bImportSearchedFor = true;
 		if (Import->OuterIndex.IsNull())
 		{
-			Import->XObject = StaticFindObjectFast(UPackage::StaticClass(), nullptr, Linker->GetInstancingContext().Remap(Import->ObjectName), true);
+			Import->XObject = StaticFindObjectFast(UPackage::StaticClass(), nullptr, Linker->GetInstancingContext().RemapPackage(Import->ObjectName), true);
 			check(!Import->XObject || CastChecked<UPackage>(Import->XObject));
 		}
 		else if (Import->OuterIndex.IsImport())
@@ -2194,7 +2194,7 @@ EAsyncPackageState::Type FAsyncPackage::LoadImports_Event()
 		const bool bCompiledInNotDynamic = IsNativeCodePackage(ExistingPackage);
 		// Our import package name is the import name
 		const FName ImportPackageToLoad = !Import->HasPackageName() ? Import->ObjectName : Import->GetPackageName();
-		const FName ImportPackageFName = Linker->GetInstancingContext().Remap(ImportPackageToLoad);
+		const FName ImportPackageFName = Linker->GetInstancingContext().RemapPackage(ImportPackageToLoad);
 		check(!PendingPackage || !bCompiledInNotDynamic); // we should never have a pending package for something that is compiled in
 		if (!PendingPackage && !bCompiledInNotDynamic)
 		{
@@ -2444,7 +2444,7 @@ EAsyncPackageState::Type FAsyncPackage::SetupImports_Event()
 			}
 			FObjectImport& OuterMostImport = Linker->Imp(OuterMostIndex);
 			check(OuterMostImport.OuterIndex.IsNull() || OuterMostImport.HasPackageName());
-			FName ImportPackageName = Linker->GetInstancingContext().Remap(!OuterMostImport.HasPackageName() ? OuterMostImport.ObjectName : OuterMostImport.GetPackageName());
+			FName ImportPackageName = Linker->GetInstancingContext().RemapPackage(!OuterMostImport.HasPackageName() ? OuterMostImport.ObjectName : OuterMostImport.GetPackageName());
 			UPackage* ImportPackage = OuterMostImport.XObject ? OuterMostImport.XObject->GetPackage() : nullptr;
 			if (!ImportPackage)
 			{
@@ -2525,7 +2525,7 @@ EAsyncPackageState::Type FAsyncPackage::SetupImports_Event()
 							OuterName = ImportLinker->LinkerRoot->GetFName();
 						}
 						check(OuterName != NAME_None);
-						check(OuterName == Linker->GetInstancingContext().Remap(Linker->ImpExp(Import.OuterIndex).ObjectName));
+						check(OuterName == Linker->GetInstancingContext().RemapPackage(Linker->ImpExp(Import.OuterIndex).ObjectName));
 					}
 
 					Import.bImportFailed = LocalExportIndex.IsNull();
@@ -6507,7 +6507,7 @@ EAsyncPackageState::Type FAsyncPackage::LoadImports()
 		// Our import package name is the import name itself if import is a package, else the specified package name
 		const FLinkerInstancingContext& InstancingContext = Linker->GetInstancingContext();
 		FName ImportToLoad = !Import->HasPackageName() ? Import->ObjectName : Import->GetPackageName();
-		FName ImportPackageFName = InstancingContext.Remap(ImportToLoad);
+		FName ImportPackageFName = InstancingContext.RemapPackage(ImportToLoad);
 
 		// Don't try to import a package that is in an import table that we know is an invalid entry
 		if (FLinkerLoad::IsKnownMissingPackage(ImportPackageFName))
diff --git a/Engine/Source/Runtime/CoreUObject/Private/UObject/LinkerLoad.cpp b/Engine/Source/Runtime/CoreUObject/Private/UObject/LinkerLoad.cpp
--- a/Engine/Source/Runtime/CoreUObject/Private/UObject/LinkerLoad.cpp
+++ b/Engine/Source/Runtime/CoreUObject/Private/UObject/LinkerLoad.cpp
@@ -942,7 +942,7 @@ FLinkerLoad::FLinkerLoad(UPackage* InParent, const FPackagePath& InPackagePath,
 	FName PackageNameToLoad = GetPackagePath().GetPackageFName();
 	if (LinkerRoot->GetFName() != PackageNameToLoad)
 	{
-		InstancingContext.AddMapping(PackageNameToLoad, LinkerRoot->GetFName());
+		InstancingContext.AddPackageMapping(PackageNameToLoad, LinkerRoot->GetFName());
 	}
 }
 
@@ -1934,7 +1934,7 @@ FLinkerLoad::ELinkerStatus FLinkerLoad::PopulateInstancingContext()
 			// add remapping for all the packages that should be instantiated along with this one
 			for (const FName& InstancingName : InstancingPackageName)
 			{
-				FName& InstancedName = InstancingContext.Mapping.FindOrAdd(InstancingName);
+				FName& InstancedName = InstancingContext.PackageMapping.FindOrAdd(InstancingName);
 				// if there's isn't already a remapping for that package, create one
 				if (InstancedName.IsNone())
 				{
@@ -2493,7 +2493,7 @@ UObject* FLinkerLoad::FindExistingImport(int32 ImportIndex)
 	if (Import.OuterIndex.IsNull())
 	{
 		// if the import outer is null then we have a package, resolve it, potentially remapping it
-		FName ObjectName = InstancingContextRemap(Import.ObjectName);
+		FName ObjectName = InstancingContext.RemapPackage(Import.ObjectName);
 		UPackage* Package = static_cast<UPackage*>(StaticFindObjectFast(UPackage::StaticClass(), nullptr, ObjectName, /*bExactClass*/true));
 		if (!IsPackageReferenceAllowed(Package))
 		{
@@ -3063,7 +3063,7 @@ bool FLinkerLoad::VerifyImportInner(const int32 ImportIndex, FString& WarningSuf
 
 		// Resolve the package name for the import, potentially remapping it, if instancing
 		FName PackageToLoad = !Import.HasPackageName() ? Import.ObjectName : Import.GetPackageName();
-		FName PackageToLoadInto = InstancingContextRemap(PackageToLoad);
+		FName PackageToLoadInto = InstancingContext.RemapPackage(PackageToLoad);
 #if WITH_EDITOR
 		if (SlowTask)
 		{
@@ -3492,7 +3492,7 @@ bool FLinkerLoad::VerifyImportInner(const int32 ImportIndex, FString& WarningSuf
 
 		if (IsCoreUObjectPackage(Import.ClassPackage) && Import.ClassName == NAME_Package && !TmpPkg->GetOuter())
 		{
-			if (InstancingContextRemap(Import.ObjectName) == TmpPkg->GetFName())
+			if (InstancingContext.RemapPackage(Import.ObjectName) == TmpPkg->GetFName())
 			{
 				// except if we are looking for _the_ package...in which case we are looking for TmpPkg, so we are done
 				Import.XObject = TmpPkg;
@@ -5124,7 +5124,7 @@ UObject* FLinkerLoad::CreateImport( int32 Index )
 					// Import is a toplevel package.
 					if( Import.OuterIndex.IsNull() )
 					{
-						FName ObjectName = InstancingContextRemap(Import.ObjectName);
+						FName ObjectName = InstancingContext.RemapPackage(Import.ObjectName);
 						UPackage* Pkg = CreatePackage(*ObjectName.ToString());
 						if (IsPackageReferenceAllowed(Pkg))
 						{
@@ -5148,7 +5148,7 @@ UObject* FLinkerLoad::CreateImport( int32 Index )
 							// Outer is toplevel package, create/ find it.
 							else if( OuterImport.OuterIndex.IsNull() )
 							{
-								FName ObjectName = InstancingContextRemap(OuterImport.ObjectName);
+								FName ObjectName = InstancingContext.RemapPackage(OuterImport.ObjectName);
 								UPackage* Pkg = CreatePackage(*ObjectName.ToString());
 								if (IsPackageReferenceAllowed(Pkg))
 								{
@@ -6220,30 +6220,25 @@ bool FLinkerLoad::IsSoftObjectRemappingEnabled() const
 	return IsContextInstanced() && InstancingContext.GetSoftObjectPathRemappingEnabled();
 }
 
-FName FLinkerLoad::InstancingContextRemap(FName ObjectName) const
-{
-	return InstancingContext.Remap(ObjectName);
-}
-
 void FLinkerLoad::FixupSoftObjectPathForInstancedPackage(FSoftObjectPath& InOutSoftObjectPath)
 {
 	if (IsSoftObjectRemappingEnabled())
 	{
 		// Try remapping AssetPathName before remapping LongPackageName
-		if (FName RemappedAssetPath = InstancingContextRemap(InOutSoftObjectPath.GetAssetPathName()); RemappedAssetPath != InOutSoftObjectPath.GetAssetPathName())
+		if (FSoftObjectPath RemappedAssetPath = InstancingContext.RemapPath(InOutSoftObjectPath); RemappedAssetPath != InOutSoftObjectPath)
 		{
-			InOutSoftObjectPath.SetAssetPathName(RemappedAssetPath);
+			InOutSoftObjectPath = RemappedAssetPath;
 		}
 		else
 		{
 			FName LongPackageName = InOutSoftObjectPath.GetLongPackageFName();
-			if (FName RemappedPackage = InstancingContextRemap(LongPackageName); RemappedPackage != LongPackageName)
+			if (FName RemappedPackage = InstancingContext.RemapPackage(LongPackageName); RemappedPackage != LongPackageName)
 			{
 				FNameBuilder TmpSoftObjectPathBuilder;
 				InOutSoftObjectPath.ToString(TmpSoftObjectPathBuilder);
 
 				TmpSoftObjectPathBuilder.ReplaceAt(0, LongPackageName.GetStringLength(), RemappedPackage.ToString());
-				InOutSoftObjectPath.SetPath(TmpSoftObjectPathBuilder.ToView());
+				InOutSoftObjectPath.SetPath(FName(TmpSoftObjectPathBuilder));
 			}
 		}
 	}
diff --git a/Engine/Source/Runtime/CoreUObject/Public/UObject/LinkerInstancingContext.h b/Engine/Source/Runtime/CoreUObject/Public/UObject/LinkerInstancingContext.h
--- a/Engine/Source/Runtime/CoreUObject/Public/UObject/LinkerInstancingContext.h
+++ b/Engine/Source/Runtime/CoreUObject/Public/UObject/LinkerInstancingContext.h
@@ -2,6 +2,8 @@
 
 #pragma once
 
+#include "UObject/SoftObjectPath.h"
+
 /**
  * Helper class to remap package imports during loading.
  * This is usually when objects in a package are outer-ed to object in another package or vice versa.
@@ -12,9 +14,6 @@ class FLinkerInstancingContext
 {
 public:
 	FLinkerInstancingContext() = default;
-	explicit FLinkerInstancingContext(TMap<FName, FName> InInstanceMapping)
-		: Mapping(MoveTemp(InInstanceMapping))
-	{}
 	explicit FLinkerInstancingContext(TSet<FName> InTags)
 		: Tags(MoveTemp(InTags))
 	{
@@ -26,27 +25,49 @@ public:
 
 	bool IsInstanced() const
 	{
-		return Mapping.Num() > 0;
+		return PackageMapping.Num() > 0 || PathMapping.Num() > 0;
 	}
 
-	/** Remap the object name from the import table to its instanced counterpart, otherwise return the name unmodified. */
-	FName Remap(const FName& ObjectName) const
+	/** Remap the package name from the import table to its instanced counterpart, otherwise return the name unmodified. */
+	FName RemapPackage(const FName& PackageName) const
 	{
-		if (const FName* RemappedName = Mapping.Find(ObjectName))
+		if (const FName* RemappedName = PackageMapping.Find(PackageName))
 		{
 			return *RemappedName;
 		}
-		return ObjectName;
+		return PackageName;
+	}
+
+	/**
+	 * Remap the top level asset part of the path name to its instanced counterpart, otherwise return the name unmodified. 
+	 * i.e. remaps /Path/To/Package.AssetName:Inner to /NewPath/To/NewPackage.NewAssetName:Inner 
+	 */
+	FSoftObjectPath RemapPath(const FSoftObjectPath& Path) const
+	{
+		if (const FName* Remapped = PathMapping.Find(Path.GetAssetPathName()))
+		{
+			FSoftObjectPath NewPath = Path;
+			NewPath.SetAssetPathName(*Remapped);
+			return NewPath;
+		}
+		return Path;
 	}
 
-	void AddMapping(FName Original, FName Instanced)
+	/** Add a mapping from a package name to a new package name. There should be no separators (. or :) in these strings. */
+	void AddPackageMapping(FName Original, FName Instanced)
 	{
-		Mapping.Add(Original, Instanced);
+		PackageMapping.Add(Original, Instanced);
 	}
 
-	void AppendMapping(const TMap<FName, FName>& NewMapping)
+	/** Add a mapping from a top level asset path (/Path/To/Package.AssetName) to another. */
+	void AddPathMapping(FSoftObjectPath Original, FSoftObjectPath Instanced)
 	{
-		Mapping.Append(NewMapping);
+		ensureAlwaysMsgf(Original.GetSubPathString().IsEmpty(), 
+			TEXT("Linker instance remap paths should be top-level assets only: %s->"), *Original.ToString());
+		ensureAlwaysMsgf(Instanced.GetSubPathString().IsEmpty(), 
+			TEXT("Linker instance remap paths should be top-level assets only: ->%s"), *Instanced.ToString());
+	
+		PathMapping.Emplace(Original.GetAssetPathName(), Instanced.GetAssetPathName());
 	}
 
 	void AddTag(FName NewTag)
@@ -80,8 +101,10 @@ public:
 private:
 	friend class FLinkerLoad;
 
-	/** Map of original object name to their instance counterpart. */
-	TMap<FName, FName> Mapping;
+	/** Map of original package name to their instance counterpart. */
+	TMap<FName, FName> PackageMapping;
+	/** Map of original top level asset path to their instance counterpart. This should be FTopLevelAssetPath when FSoftObjectPath is refactored to contain that. */
+	TMap<FName, FName> PathMapping;
 	/** Tags can be used to determine some loading behavior. */
 	TSet<FName> Tags;
 	/** Remap soft object paths */
diff --git a/Engine/Source/Runtime/CoreUObject/Public/UObject/LinkerLoad.h b/Engine/Source/Runtime/CoreUObject/Public/UObject/LinkerLoad.h
--- a/Engine/Source/Runtime/CoreUObject/Public/UObject/LinkerLoad.h
+++ b/Engine/Source/Runtime/CoreUObject/Public/UObject/LinkerLoad.h
@@ -228,9 +228,6 @@ private:
 	// Helper function to query if we should do any SoftObjectPath fixup
 	bool IsSoftObjectRemappingEnabled() const;
 
-	// Helper function to access the InstancingContext, 
-	FName InstancingContextRemap(FName ObjectName) const;
-
 	/** Remaps SoftObjectPaths using InstancingContextRemap. */
 	void FixupSoftObjectPathForInstancedPackage(FSoftObjectPath& InOutSoftObjectPath);
 
diff --git a/Engine/Source/Runtime/CoreUObject/Public/UObject/UObjectGlobals.h b/Engine/Source/Runtime/CoreUObject/Public/UObject/UObjectGlobals.h
--- a/Engine/Source/Runtime/CoreUObject/Public/UObject/UObjectGlobals.h
+++ b/Engine/Source/Runtime/CoreUObject/Public/UObject/UObjectGlobals.h
@@ -37,7 +37,6 @@
 #include "Templates/UniquePtr.h"
 #include "Templates/UnrealTemplate.h"
 #include "Traits/IsCharEncodingCompatibleWith.h"
-#include "UObject/LinkerInstancingContext.h"
 #include "UObject/NameTypes.h"
 #include "UObject/ObjectMacros.h"
 #include "UObject/PrimaryAssetId.h"
diff --git a/Engine/Source/Runtime/Engine/Private/Level.cpp b/Engine/Source/Runtime/Engine/Private/Level.cpp
--- a/Engine/Source/Runtime/Engine/Private/Level.cpp
+++ b/Engine/Source/Runtime/Engine/Private/Level.cpp
@@ -1124,7 +1124,7 @@ void ULevel::PostLoad()
 			FLinkerInstancingContext InstancingContext;
 			if (bInstanced)
 			{
-				InstancingContext.AddMapping(PackageResourceName, LevelPackage->GetFName());
+				InstancingContext.AddPackageMapping(PackageResourceName, LevelPackage->GetFName());
 			}
 
 			TArray<FString> ActorPackageNames = GetOnDiskExternalActorPackages(/*bTryUsingPackageLoadedPath*/ true);
@@ -1136,7 +1136,7 @@ void ULevel::PostLoad()
 					const FString InstancedName = GetExternalActorPackageInstanceName(LevelPackage->GetName(), ActorPackageName);
 					InstancePackageNames.Add(InstancedName);
 
-					InstancingContext.AddMapping(FName(*ActorPackageName), FName(*InstancedName));
+					InstancingContext.AddPackageMapping(FName(*ActorPackageName), FName(*InstancedName));
 				}
 			}
 
diff --git a/Engine/Source/Runtime/Engine/Private/LevelStreaming.cpp b/Engine/Source/Runtime/Engine/Private/LevelStreaming.cpp
--- a/Engine/Source/Runtime/Engine/Private/LevelStreaming.cpp
+++ b/Engine/Source/Runtime/Engine/Private/LevelStreaming.cpp
@@ -1381,7 +1381,7 @@ bool ULevelStreaming::RequestLevel(UWorld* PersistentWorld, bool bAllowLevelLoad
 				for (const FString& ActorPackageName : ActorPackageNames)
 				{
 					const FString InstancedName = ULevel::GetExternalActorPackageInstanceName(DesiredPackageName.ToString(), ActorPackageName);
-					InstancingContext.AddMapping(FName(*ActorPackageName), FName(*InstancedName));
+					InstancingContext.AddPackageMapping(FName(*ActorPackageName), FName(*InstancedName));
 				}
 			}
 #endif
diff --git a/Engine/Source/Runtime/Engine/Private/WorldPartition/WorldPartition.cpp b/Engine/Source/Runtime/Engine/Private/WorldPartition/WorldPartition.cpp
--- a/Engine/Source/Runtime/Engine/Private/WorldPartition/WorldPartition.cpp
+++ b/Engine/Source/Runtime/Engine/Private/WorldPartition/WorldPartition.cpp
@@ -429,10 +429,13 @@ void UWorldPartition::Initialize(UWorld* InWorld, const FTransform& InTransform)
 	
 		if (bIsInstanced)
 		{
-			InstancingContext.AddMapping(PackageName, LevelPackage->GetFName());
+			InstancingContext.AddPackageMapping(PackageName, LevelPackage->GetFName());
 
 			// SoftObjectPaths: Specific case for new maps (/Temp/Untitled) where we need to remap the AssetPath and not just the Package name because the World gets renamed (See UWorld::PostLoad)
-			InstancingContext.AddMapping(*FString::Format(TEXT("{0}.{1}"), { PackageName.ToString(), FPackageName::GetShortName(PackageName) }), *GetWorld()->GetPathName());
+			InstancingContext.AddPathMapping(
+				FSoftObjectPath(*FString::Format(TEXT("{0}.{1}"), {PackageName.ToString(), FPackageName::GetShortName(PackageName)})),
+				FSoftObjectPath(GetWorld())
+				);
 		}
 
 		{
@@ -450,7 +453,7 @@ void UWorldPartition::Initialize(UWorld* InWorld, const FTransform& InTransform)
 					const FString LongActorPackageName = ActorDescIterator->GetActorPackage().ToString();
 					const FString InstancedName = ULevel::GetExternalActorPackageInstanceName(LevelPackage->GetName(), LongActorPackageName);
 
-					InstancingContext.AddMapping(*LongActorPackageName, *InstancedName);
+					InstancingContext.AddPackageMapping(*LongActorPackageName, *InstancedName);
 
 					ActorDescIterator->TransformInstance(SourceWorldPath, RemappedWorldPath, GetInstanceTransform());
 				}
diff --git a/Engine/Source/Runtime/Engine/Private/WorldPartition/WorldPartitionActorDesc.cpp b/Engine/Source/Runtime/Engine/Private/WorldPartition/WorldPartitionActorDesc.cpp
--- a/Engine/Source/Runtime/Engine/Private/WorldPartition/WorldPartitionActorDesc.cpp
+++ b/Engine/Source/Runtime/Engine/Private/WorldPartition/WorldPartitionActorDesc.cpp
@@ -503,7 +503,7 @@ AActor* FWorldPartitionActorDesc::Load() const
 
 		if (InstancingContext)
 		{
-			FName RemappedPackageName = InstancingContext->Remap(ActorPackage);
+			FName RemappedPackageName = InstancingContext->RemapPackage(ActorPackage);
 			check(RemappedPackageName != ActorPath);
 
 			Package = CreatePackage(*RemappedPackageName.ToString());
diff --git a/Engine/Source/Runtime/Engine/Private/WorldPartition/WorldPartitionLevelHelper.cpp b/Engine/Source/Runtime/Engine/Private/WorldPartition/WorldPartitionLevelHelper.cpp
--- a/Engine/Source/Runtime/Engine/Private/WorldPartition/WorldPartitionLevelHelper.cpp
+++ b/Engine/Source/Runtime/Engine/Private/WorldPartition/WorldPartitionLevelHelper.cpp
@@ -224,7 +224,7 @@ bool FWorldPartitionLevelHelper::RemapActorPath(const FActorContainerID& InConta
 			
 		const FString ContainerPackageString = GetContainerPackage(InContainerID, LongPackageName, INDEX_NONE);	
 		const FString NewSubPathString = FWorldPartitionLevelHelper::AddActorContainerIDToSubPathString(InContainerID, SoftObjectPath.GetSubPathString());
-			
+
 		FNameBuilder NewAssetPathBuilder;
 		SoftObjectPath.GetAssetPathName().ToString(NewAssetPathBuilder);
 		NewAssetPathBuilder.ReplaceAt(0, LongPackageName.Len(), ContainerPackageString);
@@ -332,17 +332,17 @@ bool FWorldPartitionLevelHelper::LoadActors(UWorld* InOwningWorld, ULevel* InDes
 
 			FLinkerInstancingContext& NewContext = LinkerInstancingContexts.Add(PackageObjectMapping.ContainerID);
 			NewContext.AddTag(ULevel::DontLoadExternalObjectsTag);
-			NewContext.AddMapping(PackageObjectMapping.ContainerPackage, ContainerPackageInstanceName);
+			NewContext.AddPackageMapping(PackageObjectMapping.ContainerPackage, ContainerPackageInstanceName);
 			Context = &NewContext;
 		}
 		
-		const FName ContainerPackageInstanceName = Context->Remap(PackageObjectMapping.ContainerPackage);
+		const FName ContainerPackageInstanceName = Context->RemapPackage(PackageObjectMapping.ContainerPackage);
 		if (PackageObjectMapping.ContainerPackage != ContainerPackageInstanceName)
 		{
 			const FString ActorPackageName = FPackageName::ObjectPathToPackageName(PackageObjectMapping.Package.ToString());
 			const FString ActorPackageInstanceName = ULevel::GetExternalActorPackageInstanceName(ContainerPackageInstanceName.ToString(), ActorPackageName);
 
-			Context->AddMapping(FName(*ActorPackageName), FName(*ActorPackageInstanceName));
+			Context->AddPackageMapping(FName(*ActorPackageName), FName(*ActorPackageInstanceName));
 		}
 		ActorPackages.Add(&PackageObjectMapping);
 	}
@@ -464,7 +464,7 @@ bool FWorldPartitionLevelHelper::LoadActors(UWorld* InOwningWorld, ULevel* InDes
 
 		FName PackageToLoad(*FPackageName::ObjectPathToPackageName(PackageObjectMapping->Package.ToString()));
 		const FLinkerInstancingContext& ContainerInstancingContext = LinkerInstancingContexts.FindChecked(PackageObjectMapping->ContainerID);
-		FName PackageName = ContainerInstancingContext.Remap(PackageToLoad);
+		FName PackageName = ContainerInstancingContext.RemapPackage(PackageToLoad);
 
 		if (bInLoadAsync)
 		{
diff --git a/Engine/Source/Runtime/Engine/Private/WorldPartition/WorldPartitionLevelStreamingDynamic.cpp b/Engine/Source/Runtime/Engine/Private/WorldPartition/WorldPartitionLevelStreamingDynamic.cpp
--- a/Engine/Source/Runtime/Engine/Private/WorldPartition/WorldPartitionLevelStreamingDynamic.cpp
+++ b/Engine/Source/Runtime/Engine/Private/WorldPartition/WorldPartitionLevelStreamingDynamic.cpp
@@ -286,7 +286,7 @@ bool UWorldPartitionLevelStreamingDynamic::IssueLoadRequests()
 	const bool bSoftObjectRemappingEnabled = false;
 	FLinkerInstancingContext InstancingContext(bSoftObjectRemappingEnabled);
 	UPackage* RuntimePackage = RuntimeLevel->GetPackage();
-	InstancingContext.AddMapping(OriginalLevelPackageName, RuntimePackage->GetFName());
+	InstancingContext.AddPackageMapping(OriginalLevelPackageName, RuntimePackage->GetFName());
 
 	ChildPackagesToLoad.Reset(ChildPackages.Num());
 
diff --git a/Engine/Source/Runtime/Engine/Public/ExternalPackageHelper.h b/Engine/Source/Runtime/Engine/Public/ExternalPackageHelper.h
--- a/Engine/Source/Runtime/Engine/Public/ExternalPackageHelper.h
+++ b/Engine/Source/Runtime/Engine/Public/ExternalPackageHelper.h
@@ -139,12 +139,12 @@ void FExternalPackageHelper::LoadObjectsFromExternalPackages(UObject* InOuter, T
 	const bool bInstanced = !PackageResourceName.IsNone() && (PackageResourceName != OuterPackage->GetFName());
 	if (bInstanced)
 	{
-		InstancingContext.AddMapping(PackageResourceName, OuterPackage->GetFName());
+		InstancingContext.AddPackageMapping(PackageResourceName, OuterPackage->GetFName());
 
 		for (const FString& ObjectPackageName : ObjectPackageNames)
 		{
 			const FString InstancedName = GetExternalObjectPackageInstanceName(OuterPackage->GetName(), ObjectPackageName);
-			InstancingContext.AddMapping(FName(*ObjectPackageName), FName(*InstancedName));
+			InstancingContext.AddPackageMapping(FName(*ObjectPackageName), FName(*InstancedName));
 
 			// Create instance package
 			UPackage* InstancePackage = CreatePackage(*InstancedName);
diff --git a/Engine/Source/Runtime/Engine/Public/WorldPartition/WorldPartition.h b/Engine/Source/Runtime/Engine/Public/WorldPartition/WorldPartition.h
--- a/Engine/Source/Runtime/Engine/Public/WorldPartition/WorldPartition.h
+++ b/Engine/Source/Runtime/Engine/Public/WorldPartition/WorldPartition.h
@@ -8,6 +8,7 @@
 #include "Misc/Optional.h"
 #include "GameFramework/Actor.h"
 #include "Templates/SubclassOf.h"
+#include "UObject/LinkerInstancingContext.h"
 #include "WorldPartition/WorldPartitionLog.h"
 #include "WorldPartition/WorldPartitionActorDesc.h"
 #include "WorldPartition/WorldPartitionStreamingSource.h"
