[Core] added `FObjectPtrProperty::GetObjectPropertyValueAsPtr` to get the FObjectPtr
changed `FPropertyLocalizationDataGatherer::GatherLocalizationDataFromChildTextProperties` to use `GetObjectPropertyValueAsPtr` to avoid loading object unnecessarily

diff --git a/Engine/Source/Runtime/CoreUObject/Private/Serialization/PropertyLocalizationDataGathering.cpp b/Engine/Source/Runtime/CoreUObject/Private/Serialization/PropertyLocalizationDataGathering.cpp
--- a/Engine/Source/Runtime/CoreUObject/Private/Serialization/PropertyLocalizationDataGathering.cpp
+++ b/Engine/Source/Runtime/CoreUObject/Private/Serialization/PropertyLocalizationDataGathering.cpp
@@ -506,7 +506,21 @@ void FPropertyLocalizationDataGatherer::GatherLocalizationDataFromChildTextPrope
 		// Property is an object property.
 		else if (ObjectProperty && !(GatherTextFlags & EPropertyLocalizationGathererTextFlags::SkipSubObjects))
 		{
-			const UObject* InnerObject = ObjectProperty->GetObjectPropertyValue(ElementValueAddress);
+			// if the object is a FObjectPtr it might not be resolved
+			// if unresolved there is no need to resolve it as IsObjectValidForGather would return false anyways
+			const UObject* InnerObject = nullptr;
+			if (auto ObjectPropertyPtr = CastField<FObjectPtrProperty>(ObjectProperty))
+			{
+				const FObjectPtr& ObjectPtr = ObjectPropertyPtr->GetObjectPropertyValueAsPtr(ElementValueAddress);
+				if (ObjectPtr.IsResolved())
+				{
+					InnerObject = ObjectPtr.Get();
+				}
+			}
+			else
+			{
+				InnerObject = ObjectProperty->GetObjectPropertyValue(ElementValueAddress);
+			} 
 			if (InnerObject && IsObjectValidForGather(InnerObject))
 			{
 				GatherLocalizationDataFromObjectWithCallbacks(InnerObject, FixedChildPropertyGatherTextFlags);
diff --git a/Engine/Source/Runtime/CoreUObject/Private/UObject/PropertyObjectPtr.cpp b/Engine/Source/Runtime/CoreUObject/Private/UObject/PropertyObjectPtr.cpp
--- a/Engine/Source/Runtime/CoreUObject/Private/UObject/PropertyObjectPtr.cpp
+++ b/Engine/Source/Runtime/CoreUObject/Private/UObject/PropertyObjectPtr.cpp
@@ -138,6 +138,11 @@ bool FObjectPtrProperty::Identical(const void* A, const void* B, uint32 PortFlag
 	return false;
 }
 
+FObjectPtr& FObjectPtrProperty::GetObjectPropertyValueAsPtr(const void* PropertyValueAddress) const
+{
+	return (FObjectPtr&)GetPropertyValue(PropertyValueAddress);
+}
+
 UObject* FObjectPtrProperty::GetObjectPropertyValue(const void* PropertyValueAddress) const
 {
 	return ((FObjectPtr&)GetPropertyValue(PropertyValueAddress)).Get();
diff --git a/Engine/Source/Runtime/CoreUObject/Public/UObject/UnrealType.h b/Engine/Source/Runtime/CoreUObject/Public/UObject/UnrealType.h
--- a/Engine/Source/Runtime/CoreUObject/Public/UObject/UnrealType.h
+++ b/Engine/Source/Runtime/CoreUObject/Public/UObject/UnrealType.h
@@ -2982,6 +2982,11 @@ class COREUOBJECT_API FObjectPtrProperty : public FObjectProperty
 	// Helper method for sharing code with FClassPtrProperty even though one doesn't inherit from the other
 	static void StaticSerializeItem(const FObjectPropertyBase* ObjectProperty, FStructuredArchive::FSlot Slot, void* Value, void const* Defaults);
 
+	/// @brief Get the property value as FObjectPtr
+	/// @param PropertyValueAddress address of the property
+	/// @return reference to the FObjectPtr value
+	FObjectPtr& GetObjectPropertyValueAsPtr(const void* PropertyValueAddress) const;
+
 	// FObjectProperty interface
 	virtual UObject* GetObjectPropertyValue(const void* PropertyValueAddress) const override;
 	virtual UObject* GetObjectPropertyValue_InContainer(const void* ContainerAddress, int32 ArrayIndex = 0) const override;
