Added support for diffing objects managed by a transaction annotation

ITransactionObjectAnnotation has two new functions; SupportsAdditionalObjectChanges to state whether the annotation type supports object diffing (default false), and ComputeAdditionalObjectChanges to actually compute the diff against another annotation (which must be of the same type).

This allows transaction annotations that entirely manage other objects (such as FActorTransactionAnnotation which manages the SCS and UCS components on an actor) to produce meaningful change information that can be passed to the PostTransacted function of the managed object (tagged with ETransactionObjectChangeCreatedBy::TransactionAnnotation).

Note: It is not safe to cache the result of a ComputeAdditionalObjectChanges outside of the current transaction, as objects managed by a transaction annotation can have volatile names. This means that the transaction annotation itself will have ComputeAdditionalObjectChanges called on it during undo/redo as well as during transaction snapshot/finalization, and means that the transaction annotation must have its own way of pairing old->new objects over time.

diff --git a/Engine/Source/Editor/UnrealEd/Classes/Editor/Transactor.h b/Engine/Source/Editor/UnrealEd/Classes/Editor/Transactor.h
--- a/Engine/Source/Editor/UnrealEd/Classes/Editor/Transactor.h
+++ b/Engine/Source/Editor/UnrealEd/Classes/Editor/Transactor.h
@@ -46,26 +46,26 @@ protected:
 		{
 			void SetObject(const UObject* InObject)
 			{
-				ObjectInfo.SetObject(InObject);
+				ObjectId.SetObject(InObject);
 				ObjectAnnotation = InObject->FindOrCreateTransactionAnnotation();
 			}
 
 			void Reset()
 			{
 				UE::Transaction::FSerializedObject::Reset();
-				ObjectInfo.Reset();
+				ObjectId.Reset();
 				ObjectAnnotation.Reset();
 			}
 
 			void Swap(FSerializedObject& Other)
 			{
 				UE::Transaction::FSerializedObject::Swap(Other);
-				ObjectInfo.Swap(Other.ObjectInfo);
+				ObjectId.Swap(Other.ObjectId);
 				Exchange(ObjectAnnotation, Other.ObjectAnnotation);
 			}
 
-			/** Information about the object when it was serialized */
-			UE::Transaction::FSerializedObjectInfo ObjectInfo;
+			/** ID of the object when it was serialized */
+			FTransactionObjectId ObjectId;
 
 			/** Annotation data for the object stored externally */
 			TSharedPtr<ITransactionObjectAnnotation> ObjectAnnotation;
@@ -111,6 +111,8 @@ protected:
 		TUniquePtr<UE::Transaction::FDiffableObject> DiffableObject;
 		/** The diffable object data when it was last snapshot (always null once finalized) */
 		TUniquePtr<UE::Transaction::FDiffableObject> DiffableObjectSnapshot;
+		/** Annotation data the last time the object was snapshot (always null once finalized) */
+		TSharedPtr<ITransactionObjectAnnotation> ObjectAnnotationSnapshot;
 		/** The combined list of properties that have been passed any Snapshot call for this object (always empty once finalized) */
 		TArray<const FProperty*> AllPropertiesSnapshot;
 		/** Delta change information between the diffable state of the object when the transaction started, and the diffable state of the object when the transaction ended */
@@ -256,20 +258,18 @@ protected:
 
 	struct FChangedObjectValue
 	{
-		FChangedObjectValue()
-			: Annotation()
-			, RecordIndex(INDEX_NONE)
-		{
-		}
+		FChangedObjectValue() = default;
 
-		FChangedObjectValue(const int32 InRecordIndex, const TSharedPtr<ITransactionObjectAnnotation>& InAnnotation)
+		FChangedObjectValue(const int32 InRecordIndex, const TSharedPtr<ITransactionObjectAnnotation>& InAnnotation, const TSharedPtr<ITransactionObjectAnnotation>& InAnnotationSnapshot = nullptr)
 			: Annotation(InAnnotation)
+			, AnnotationSnapshot(InAnnotationSnapshot)
 			, RecordIndex(InRecordIndex)
 		{
 		}
 
 		TSharedPtr<ITransactionObjectAnnotation> Annotation;
-		int32 RecordIndex;
+		TSharedPtr<ITransactionObjectAnnotation> AnnotationSnapshot;
+		int32 RecordIndex = INDEX_NONE;
 	};
 
 	/** Objects that will be changed directly by the transaction, empty when not transacting */
diff --git a/Engine/Source/Editor/UnrealEd/Private/EditorTransaction.cpp b/Engine/Source/Editor/UnrealEd/Private/EditorTransaction.cpp
--- a/Engine/Source/Editor/UnrealEd/Private/EditorTransaction.cpp
+++ b/Engine/Source/Editor/UnrealEd/Private/EditorTransaction.cpp
@@ -315,6 +315,7 @@ void FTransaction::FObjectRecord::Finalize( FTransaction* Owner, TSharedPtr<ITra
 		// Clear out any diff data now as we won't be getting any more diff requests once finalized
 		DiffableObject.Reset();
 		DiffableObjectSnapshot.Reset();
+		ObjectAnnotationSnapshot.Reset();
 		AllPropertiesSnapshot.Reset();
 	}
 }
@@ -388,18 +389,23 @@ void FTransaction::FObjectRecord::Snapshot( FTransaction* Owner, TArrayView<cons
 		}
 		DiffableObjectSnapshot->Swap(CurrentDiffableObject);
 
-		TSharedPtr<ITransactionObjectAnnotation> ChangedObjectTransactionAnnotation = CurrentObject->FindOrCreateTransactionAnnotation();
+		TSharedPtr<ITransactionObjectAnnotation> InitialObjectTransactionAnnotation = ObjectAnnotationSnapshot ? ObjectAnnotationSnapshot : SerializedObject.ObjectAnnotation;
+		ObjectAnnotationSnapshot = CurrentObject->FindOrCreateTransactionAnnotation();
 
 		// Notify any listeners of this change
-		if (SnapshotDeltaChange.HasChanged() || ChangedObjectTransactionAnnotation.IsValid())
+		if (SnapshotDeltaChange.HasChanged() || ObjectAnnotationSnapshot.IsValid())
 		{
-			CurrentObject->PostTransacted(FTransactionObjectEvent(Owner->GetId(), Owner->GetOperationId(), ETransactionObjectEventType::Snapshot, SnapshotDeltaChange, ChangedObjectTransactionAnnotation
-				, InitialDiffableObject.ObjectInfo.ObjectPackageName
-				, InitialDiffableObject.ObjectInfo.ObjectName
-				, InitialDiffableObject.ObjectInfo.ObjectPathName
-				, InitialDiffableObject.ObjectInfo.ObjectOuterPathName
-				, InitialDiffableObject.ObjectInfo.ObjectExternalPackageName
-				, InitialDiffableObject.ObjectInfo.ObjectClassPathName));
+			TMap<UObject*, FTransactionObjectChange> AdditionalObjectChanges;
+			if (ObjectAnnotationSnapshot)
+			{
+				ObjectAnnotationSnapshot->ComputeAdditionalObjectChanges(InitialObjectTransactionAnnotation.Get(), AdditionalObjectChanges);
+			}
+
+			CurrentObject->PostTransacted(FTransactionObjectEvent(Owner->GetId(), Owner->GetOperationId(), ETransactionObjectEventType::Snapshot, ETransactionObjectChangeCreatedBy::TransactionRecord, FTransactionObjectChange{ InitialDiffableObject.ObjectInfo, SnapshotDeltaChange }, ObjectAnnotationSnapshot));
+			for (const TTuple<UObject*, FTransactionObjectChange>& AdditionalObjectChangePair : AdditionalObjectChanges)
+			{
+				AdditionalObjectChangePair.Key->PostTransacted(FTransactionObjectEvent(Owner->GetId(), Owner->GetOperationId(), ETransactionObjectEventType::Snapshot, ETransactionObjectChangeCreatedBy::TransactionAnnotation, AdditionalObjectChangePair.Value, nullptr));
+			}
 		}
 	}
 }
@@ -793,13 +799,18 @@ void FTransaction::Apply()
 		if (DeltaChange.HasChanged() || ChangedObjectTransactionAnnotation.IsValid() || ChangedObjectRecord.bHasSerializedObjectChanges)
 		{
 			const FObjectRecord::FSerializedObject& InitialSerializedObject = ChangedObjectRecord.SerializedObject;
-			ChangedObject->PostTransacted(FTransactionObjectEvent(Id, OperationId, ETransactionObjectEventType::UndoRedo, DeltaChange, ChangedObjectTransactionAnnotation
-				, InitialSerializedObject.ObjectInfo.ObjectPackageName
-				, InitialSerializedObject.ObjectInfo.ObjectName
-				, InitialSerializedObject.ObjectInfo.ObjectPathName
-				, InitialSerializedObject.ObjectInfo.ObjectOuterPathName
-				, InitialSerializedObject.ObjectInfo.ObjectExternalPackageName
-				, InitialSerializedObject.ObjectInfo.ObjectClassPathName));
+
+			TMap<UObject*, FTransactionObjectChange> AdditionalObjectChanges;
+			if (ChangedObjectTransactionAnnotation)
+			{
+				ChangedObjectTransactionAnnotation->ComputeAdditionalObjectChanges(InitialSerializedObject.ObjectAnnotation.Get(), AdditionalObjectChanges);
+			}
+
+			ChangedObject->PostTransacted(FTransactionObjectEvent(Id, OperationId, ETransactionObjectEventType::UndoRedo, ETransactionObjectChangeCreatedBy::TransactionRecord, FTransactionObjectChange{ InitialSerializedObject.ObjectId, DeltaChange }, ChangedObjectTransactionAnnotation));
+			for (const TTuple<UObject*, FTransactionObjectChange>& AdditionalObjectChangePair : AdditionalObjectChanges)
+			{
+				AdditionalObjectChangePair.Key->PostTransacted(FTransactionObjectEvent(Id, OperationId, ETransactionObjectEventType::UndoRedo, ETransactionObjectChangeCreatedBy::TransactionAnnotation, AdditionalObjectChangePair.Value, nullptr));
+			}
 		}
 	}
 
@@ -827,9 +838,10 @@ void FTransaction::Finalize()
 {
 	for (int32 i = 0; i < Records.Num(); ++i)
 	{
-		TSharedPtr<ITransactionObjectAnnotation> FinalizedObjectAnnotation;
-
 		FObjectRecord& ObjectRecord = Records[i];
+
+		TSharedPtr<ITransactionObjectAnnotation> FinalizedObjectAnnotation;
+		TSharedPtr<ITransactionObjectAnnotation> SnapshotObjectAnnotation = ObjectRecord.ObjectAnnotationSnapshot;
 		ObjectRecord.Finalize(this, FinalizedObjectAnnotation);
 
 		UObject* Object = ObjectRecord.Object.Get();
@@ -837,7 +849,7 @@ void FTransaction::Finalize()
 		{
 			if (!ChangedObjects.Contains(Object))
 			{
-				ChangedObjects.Add(Object, FChangedObjectValue(i, FinalizedObjectAnnotation));
+				ChangedObjects.Add(Object, FChangedObjectValue(i, FinalizedObjectAnnotation, SnapshotObjectAnnotation));
 			}
 		}
 	}
@@ -851,22 +863,30 @@ void FTransaction::Finalize()
 
 	for (auto ChangedObjectIt : ChangedObjects)
 	{
-		TSharedPtr<ITransactionObjectAnnotation> ChangedObjectTransactionAnnotation = ChangedObjectIt.Value.Annotation;
-
 		const FObjectRecord& ChangedObjectRecord = Records[ChangedObjectIt.Value.RecordIndex];
+		TSharedPtr<ITransactionObjectAnnotation> ChangedObjectTransactionAnnotation = ChangedObjectIt.Value.Annotation;		
 		const FTransactionObjectDeltaChange& DeltaChange = ChangedObjectRecord.DeltaChange;
 		if (DeltaChange.HasChanged() || ChangedObjectTransactionAnnotation.IsValid() || ChangedObjectRecord.bHasSerializedObjectChanges)
 		{
 			UObject* ChangedObject = ChangedObjectIt.Key;
 			
 			const FObjectRecord::FSerializedObject& InitialSerializedObject = ChangedObjectRecord.SerializedObject;
-			ChangedObject->PostTransacted(FTransactionObjectEvent(Id, OperationId, ETransactionObjectEventType::Finalized, DeltaChange, ChangedObjectTransactionAnnotation
-				, InitialSerializedObject.ObjectInfo.ObjectPackageName
-				, InitialSerializedObject.ObjectInfo.ObjectName
-				, InitialSerializedObject.ObjectInfo.ObjectPathName
-				, InitialSerializedObject.ObjectInfo.ObjectOuterPathName
-				, InitialSerializedObject.ObjectInfo.ObjectExternalPackageName
-				, InitialSerializedObject.ObjectInfo.ObjectClassPathName));
+
+			TMap<UObject*, FTransactionObjectChange> AdditionalObjectChanges;
+			if (ChangedObjectTransactionAnnotation)
+			{
+				ChangedObjectTransactionAnnotation->ComputeAdditionalObjectChanges(InitialSerializedObject.ObjectAnnotation.Get(), AdditionalObjectChanges);
+			}
+			if (ChangedObjectTransactionAnnotation && ChangedObjectIt.Value.AnnotationSnapshot)
+			{
+				ChangedObjectTransactionAnnotation->ComputeAdditionalObjectChanges(ChangedObjectIt.Value.AnnotationSnapshot.Get(), AdditionalObjectChanges);
+			}
+
+			ChangedObject->PostTransacted(FTransactionObjectEvent(Id, OperationId, ETransactionObjectEventType::Finalized, ETransactionObjectChangeCreatedBy::TransactionRecord, FTransactionObjectChange{ InitialSerializedObject.ObjectId, DeltaChange }, ChangedObjectTransactionAnnotation));
+			for (const TTuple<UObject*, FTransactionObjectChange>& AdditionalObjectChangePair : AdditionalObjectChanges)
+			{
+				AdditionalObjectChangePair.Key->PostTransacted(FTransactionObjectEvent(Id, OperationId, ETransactionObjectEventType::Finalized, ETransactionObjectChangeCreatedBy::TransactionAnnotation, AdditionalObjectChangePair.Value, nullptr));
+			}
 		}
 	}
 	ChangedObjects.Reset();
@@ -917,13 +937,10 @@ FTransactionDiff FTransaction::GenerateDiff() const
 				{
 					// Since this transaction is not currently in an undo operation, generate a valid Guid.
 					FGuid Guid = FGuid::NewGuid();
-					TransactionDiff.DiffMap.Emplace(FName(*TransactedObject->GetPathName()), MakeShared<FTransactionObjectEvent>(this->GetId(), Guid, ETransactionObjectEventType::Finalized, ObjectRecord.DeltaChange, ObjectRecord.SerializedObject.ObjectAnnotation
-						, ObjectRecord.SerializedObject.ObjectInfo.ObjectPackageName
-						, ObjectRecord.SerializedObject.ObjectInfo.ObjectName
-						, ObjectRecord.SerializedObject.ObjectInfo.ObjectPathName
-						, ObjectRecord.SerializedObject.ObjectInfo.ObjectOuterPathName
-						, ObjectRecord.SerializedObject.ObjectInfo.ObjectExternalPackageName
-						, ObjectRecord.SerializedObject.ObjectInfo.ObjectClassPathName));
+					TransactionDiff.DiffMap.Emplace(FName(*TransactedObject->GetPathName()), MakeShared<FTransactionObjectEvent>(
+						this->GetId(), Guid, ETransactionObjectEventType::Finalized, ETransactionObjectChangeCreatedBy::TransactionRecord, 
+						FTransactionObjectChange{ ObjectRecord.SerializedObject.ObjectId, ObjectRecord.DeltaChange }, ObjectRecord.SerializedObject.ObjectAnnotation)
+					);
 				}
 			}
 		}
diff --git a/Engine/Source/Runtime/CoreUObject/Private/Misc/TransactionObjectEvent.cpp b/Engine/Source/Runtime/CoreUObject/Private/Misc/TransactionObjectEvent.cpp
new file mode 100644
--- /dev/null
+++ b/Engine/Source/Runtime/CoreUObject/Private/Misc/TransactionObjectEvent.cpp
@@ -0,0 +1,23 @@
+// Copyright Epic Games, Inc. All Rights Reserved.
+
+#include "Misc/TransactionObjectEvent.h"
+#include "UObject/Object.h"
+#include "UObject/Package.h"
+#include "UObject/Class.h"
+
+void FTransactionObjectId::SetObject(const UObject* Object)
+{
+	auto GetPathNameHelper = [](const UObject* ObjectToGet) -> FName
+	{
+		FNameBuilder ObjectPathNameBuilder;
+		ObjectToGet->GetPathName(nullptr, ObjectPathNameBuilder);
+		return FName(ObjectPathNameBuilder);
+	};
+
+	ObjectPackageName = Object->GetPackage()->GetFName();
+	ObjectName = Object->GetFName();
+	ObjectPathName = GetPathNameHelper(Object);
+	ObjectOuterPathName = Object->GetOuter() ? GetPathNameHelper(Object->GetOuter()) : FName();
+	ObjectExternalPackageName = Object->GetExternalPackage() ? Object->GetExternalPackage()->GetFName() : FName();
+	ObjectClassPathName = GetPathNameHelper(Object->GetClass());
+}
diff --git a/Engine/Source/Runtime/CoreUObject/Public/Misc/ITransactionObjectAnnotation.h b/Engine/Source/Runtime/CoreUObject/Public/Misc/ITransactionObjectAnnotation.h
--- a/Engine/Source/Runtime/CoreUObject/Public/Misc/ITransactionObjectAnnotation.h
+++ b/Engine/Source/Runtime/CoreUObject/Public/Misc/ITransactionObjectAnnotation.h
@@ -4,6 +4,9 @@
 
 #include "CoreTypes.h"
 
+class UObject;
+struct FTransactionObjectChange;
+
 /**
  * Interface for transaction object annotations.
  *
@@ -18,4 +21,7 @@ public:
 	virtual ~ITransactionObjectAnnotation() = default;
 	virtual void AddReferencedObjects(class FReferenceCollector& Collector) = 0;
 	virtual void Serialize(class FArchive& Ar) = 0;
+
+	virtual bool SupportsAdditionalObjectChanges() const { return false; }
+	virtual void ComputeAdditionalObjectChanges(const ITransactionObjectAnnotation* OriginalAnnotation, TMap<UObject*, FTransactionObjectChange>& OutAdditionalObjectChanges) {}
 };
diff --git a/Engine/Source/Runtime/CoreUObject/Public/Misc/TransactionObjectEvent.h b/Engine/Source/Runtime/CoreUObject/Public/Misc/TransactionObjectEvent.h
--- a/Engine/Source/Runtime/CoreUObject/Public/Misc/TransactionObjectEvent.h
+++ b/Engine/Source/Runtime/CoreUObject/Public/Misc/TransactionObjectEvent.h
@@ -54,6 +54,107 @@ struct FTransactionObjectDeltaChange
 	TArray<FName> ChangedProperties;
 };
 
+/**
+ * ID for an object that was transacted
+ */
+struct FTransactionObjectId
+{
+public:
+	FTransactionObjectId() = default;
+
+	explicit FTransactionObjectId(const UObject* Object)
+	{
+		SetObject(Object);
+	}
+
+	FTransactionObjectId(const FName InObjectPackageName, const FName InObjectName, const FName InObjectPathName, const FName InObjectOuterPathName, const FName InObjectExternalPackageName, const FName InObjectClassPathName)
+		: ObjectPackageName(InObjectPackageName)
+		, ObjectName(InObjectName)
+		, ObjectPathName(InObjectPathName)
+		, ObjectOuterPathName(InObjectOuterPathName)
+		, ObjectExternalPackageName(InObjectExternalPackageName)
+		, ObjectClassPathName(InObjectClassPathName)
+	{
+	}
+
+	COREUOBJECT_API void SetObject(const UObject* Object);
+
+	void Reset()
+	{
+		ObjectPackageName = FName();
+		ObjectName = FName();
+		ObjectPathName = FName();
+		ObjectOuterPathName = FName();
+		ObjectExternalPackageName = FName();
+		ObjectClassPathName = FName();
+	}
+
+	void Swap(FTransactionObjectId& Other)
+	{
+		Exchange(ObjectPackageName, Other.ObjectPackageName);
+		Exchange(ObjectName, Other.ObjectName);
+		Exchange(ObjectPathName, Other.ObjectPathName);
+		Exchange(ObjectOuterPathName, Other.ObjectOuterPathName);
+		Exchange(ObjectExternalPackageName, Other.ObjectExternalPackageName);
+		Exchange(ObjectClassPathName, Other.ObjectClassPathName);
+	}
+
+	friend bool operator==(const FTransactionObjectId& LHS, const FTransactionObjectId& RHS)
+	{
+		// The other fields are subsets of this data
+		return LHS.ObjectPathName == RHS.ObjectPathName
+			&& LHS.ObjectExternalPackageName == RHS.ObjectExternalPackageName
+			&& LHS.ObjectClassPathName == RHS.ObjectClassPathName;
+	}
+
+	friend bool operator!=(const FTransactionObjectId& LHS, const FTransactionObjectId& RHS)
+	{
+		return !(LHS == RHS);
+	}
+
+	friend uint32 GetTypeHash(const FTransactionObjectId& Id)
+	{
+		return GetTypeHash(Id.ObjectPathName);
+	}
+
+	/** The package name of the object, can be dictated either by outer chain or external package */
+	FName ObjectPackageName;
+
+	/** The name of the object */
+	FName ObjectName;
+
+	/** The path name of the object */
+	FName ObjectPathName;
+
+	/** The outer path name of the object */
+	FName ObjectOuterPathName;
+
+	/** The external package name of the object, if any */
+	FName ObjectExternalPackageName;
+
+	/** The path name of the object's class. */
+	FName ObjectClassPathName;
+};
+
+/** Change information for an object that was transacted */
+struct FTransactionObjectChange
+{
+	/** Original ID of the object that changed */
+	FTransactionObjectId OriginalId;
+
+	/** Information about how the object changed */
+	FTransactionObjectDeltaChange DeltaChange;
+};
+
+/** Different things that can create an object change */
+enum class ETransactionObjectChangeCreatedBy : uint8
+{
+	/** This change was created from an object that was explicitly tracked within a transaction record */
+	TransactionRecord,
+	/** This change was created from an object that was implicitly tracked within a transaction annotation */
+	TransactionAnnotation,
+};
+
 /** Different kinds of actions that can trigger a transaction object event */
 enum class ETransactionObjectEventType : uint8
 {
@@ -77,24 +178,26 @@ class FTransactionObjectEvent
 public:
 	FTransactionObjectEvent() = default;
 
-	FTransactionObjectEvent(const FGuid& InTransactionId, const FGuid& InOperationId, const ETransactionObjectEventType InEventType, const FTransactionObjectDeltaChange& InDeltaChange, const TSharedPtr<ITransactionObjectAnnotation>& InAnnotation
-		, const FName InOriginalObjectPackageName, const FName InOriginalObjectName, const FName InOriginalObjectPathName, const FName InOriginalObjectOuterPathName, const FName InOriginalObjectExternalPackageName, const FName InOriginalObjectClassPathName)
+	FTransactionObjectEvent(const FGuid& InTransactionId, const FGuid& InOperationId, const ETransactionObjectEventType InEventType, const ETransactionObjectChangeCreatedBy InObjectChangeCreatedBy, const FTransactionObjectChange& InObjectChange, const TSharedPtr<ITransactionObjectAnnotation>& InAnnotation)
 		: TransactionId(InTransactionId)
 		, OperationId(InOperationId)
 		, EventType(InEventType)
-		, DeltaChange(InDeltaChange)
+		, ObjectChangeCreatedBy(InObjectChangeCreatedBy)
+		, ObjectChange(InObjectChange)
 		, Annotation(InAnnotation)
-		, OriginalObjectPackageName(InOriginalObjectPackageName)
-		, OriginalObjectName(InOriginalObjectName)
-		, OriginalObjectPathName(InOriginalObjectPathName)
-		, OriginalObjectOuterPathName(InOriginalObjectOuterPathName)
-		, OriginalObjectExternalPackageName(InOriginalObjectExternalPackageName)
-		, OriginalObjectClassPathName(InOriginalObjectClassPathName)
 	{
 		check(TransactionId.IsValid());
 		check(OperationId.IsValid());
 	}
 
+	// TODO: Deprecate this
+	FTransactionObjectEvent(const FGuid& InTransactionId, const FGuid& InOperationId, const ETransactionObjectEventType InEventType, const FTransactionObjectDeltaChange& InDeltaChange, const TSharedPtr<ITransactionObjectAnnotation>& InAnnotation
+		, const FName InOriginalObjectPackageName, const FName InOriginalObjectName, const FName InOriginalObjectPathName, const FName InOriginalObjectOuterPathName, const FName InOriginalObjectExternalPackageName, const FName InOriginalObjectClassPathName)
+		: FTransactionObjectEvent(InTransactionId, InOperationId, InEventType, ETransactionObjectChangeCreatedBy::TransactionRecord, 
+			FTransactionObjectChange{ FTransactionObjectId(InOriginalObjectPackageName, InOriginalObjectName, InOriginalObjectPathName, InOriginalObjectOuterPathName, InOriginalObjectExternalPackageName, InOriginalObjectClassPathName), InDeltaChange }, InAnnotation)
+	{
+	}
+
 	/** The unique identifier of the transaction this event belongs to */
 	const FGuid& GetTransactionId() const
 	{
@@ -113,82 +216,93 @@ public:
 		return EventType;
 	}
 
+	/** What kind of thing created this object change? */
+	ETransactionObjectChangeCreatedBy GetObjectChangeCreatedBy() const
+	{
+		return ObjectChangeCreatedBy;
+	}
+
 	/** Was the pending kill state of this object changed? (implies non-property changes) */
 	bool HasPendingKillChange() const
 	{
-		return DeltaChange.bHasPendingKillChange;
+		return ObjectChange.DeltaChange.bHasPendingKillChange;
 	}
 
 	/** Was the name of this object changed? (implies non-property changes) */
 	bool HasNameChange() const
 	{
-		return DeltaChange.bHasNameChange;
+		return ObjectChange.DeltaChange.bHasNameChange;
+	}
+
+	/** Get the original ID of this object */
+	const FTransactionObjectId& GetOriginalObjectId() const
+	{
+		return ObjectChange.OriginalId;
 	}
 
 	/** Get the original package name of this object */
 	FName GetOriginalObjectPackageName() const
 	{
-		return OriginalObjectPackageName;
+		return ObjectChange.OriginalId.ObjectPackageName;
 	}
 
 	/** Get the original name of this object */
 	FName GetOriginalObjectName() const
 	{
-		return OriginalObjectName;
+		return ObjectChange.OriginalId.ObjectName;
 	}
 
 	/** Get the original path name of this object */
 	FName GetOriginalObjectPathName() const
 	{
-		return OriginalObjectPathName;
+		return ObjectChange.OriginalId.ObjectPathName;
 	}
 
 	FName GetOriginalObjectClassPathName() const
 	{
-		return OriginalObjectClassPathName;
+		return ObjectChange.OriginalId.ObjectClassPathName;
 	}
 
 	/** Was the outer of this object changed? (implies non-property changes) */
 	bool HasOuterChange() const
 	{
-		return DeltaChange.bHasOuterChange;
+		return ObjectChange.DeltaChange.bHasOuterChange;
 	}
 
 	/** Has the package assigned to this object changed? (implies non-property changes) */
 	bool HasExternalPackageChange() const
 	{
-		return DeltaChange.bHasExternalPackageChange;
+		return ObjectChange.DeltaChange.bHasExternalPackageChange;
 	}
 
 	/** Get the original outer path name of this object */
 	FName GetOriginalObjectOuterPathName() const
 	{
-		return OriginalObjectOuterPathName;
+		return ObjectChange.OriginalId.ObjectOuterPathName;
 	}
 
 	/** Get the original package name of this object */
 	FName GetOriginalObjectExternalPackageName() const
 	{
-		return OriginalObjectExternalPackageName;
+		return ObjectChange.OriginalId.ObjectExternalPackageName;
 	}
 
-
 	/** Were any non-property changes made to the object? */
 	bool HasNonPropertyChanges(const bool InSerializationOnly = false) const
 	{
-		return (!InSerializationOnly && (DeltaChange.bHasNameChange || DeltaChange.bHasOuterChange || DeltaChange.bHasExternalPackageChange || DeltaChange.bHasPendingKillChange)) || DeltaChange.bHasNonPropertyChanges;
+		return (!InSerializationOnly && (ObjectChange.DeltaChange.bHasNameChange || ObjectChange.DeltaChange.bHasOuterChange || ObjectChange.DeltaChange.bHasExternalPackageChange || ObjectChange.DeltaChange.bHasPendingKillChange)) || ObjectChange.DeltaChange.bHasNonPropertyChanges;
 	}
 
 	/** Were any property changes made to the object? */
 	bool HasPropertyChanges() const
 	{
-		return DeltaChange.ChangedProperties.Num() > 0;
+		return ObjectChange.DeltaChange.ChangedProperties.Num() > 0;
 	}
 
-	/** Get the list of changed properties. Each entry is actually a chain of property names (root -> leaf) separated by a dot, eg) "ObjProp.StructProp". */
+	/** Get the list of changed properties. */
 	const TArray<FName>& GetChangedProperties() const
 	{
-		return DeltaChange.ChangedProperties;
+		return ObjectChange.DeltaChange.ChangedProperties;
 	}
 
 	/** Get the annotation object associated with the object being transacted (if any). */
@@ -205,21 +319,16 @@ public:
 			EventType = InOther.EventType;
 		}
 
-		DeltaChange.Merge(InOther.DeltaChange);
+		ObjectChange.DeltaChange.Merge(InOther.ObjectChange.DeltaChange);
 	}
 
 private:
 	FGuid TransactionId;
 	FGuid OperationId;
-	ETransactionObjectEventType EventType;
-	FTransactionObjectDeltaChange DeltaChange;
+	ETransactionObjectEventType EventType = ETransactionObjectEventType::Finalized;
+	ETransactionObjectChangeCreatedBy ObjectChangeCreatedBy = ETransactionObjectChangeCreatedBy::TransactionRecord;
+	FTransactionObjectChange ObjectChange;
 	TSharedPtr<ITransactionObjectAnnotation> Annotation;
-	FName OriginalObjectPackageName;
-	FName OriginalObjectName;
-	FName OriginalObjectPathName;
-	FName OriginalObjectOuterPathName;
-	FName OriginalObjectExternalPackageName;
-	FName OriginalObjectClassPathName;
 };
 
 /**
diff --git a/Engine/Source/Runtime/Engine/Private/TransactionCommon.cpp b/Engine/Source/Runtime/Engine/Private/TransactionCommon.cpp
--- a/Engine/Source/Runtime/Engine/Private/TransactionCommon.cpp
+++ b/Engine/Source/Runtime/Engine/Private/TransactionCommon.cpp
@@ -1,7 +1,6 @@
 // Copyright Epic Games, Inc. All Rights Reserved.
 
 #include "TransactionCommon.h"
-#include "Misc/TransactionObjectEvent.h"
 #include "UObject/UnrealType.h"
 #include "Components/ActorComponent.h"
 
diff --git a/Engine/Source/Runtime/Engine/Public/TransactionCommon.h b/Engine/Source/Runtime/Engine/Public/TransactionCommon.h
--- a/Engine/Source/Runtime/Engine/Public/TransactionCommon.h
+++ b/Engine/Source/Runtime/Engine/Public/TransactionCommon.h
@@ -8,9 +8,9 @@
 #include "Serialization/ArchiveUObject.h"
 #include "Serialization/ArchiveSerializedPropertyChain.h"
 #include "Templates/TypeHash.h"
+#include "Misc/TransactionObjectEvent.h"
 
 class FReferenceCollector;
-struct FTransactionObjectDeltaChange;
 
 namespace UE::Transaction
 {
@@ -186,60 +186,34 @@ private:
 	TArray64<uint8> Data;
 };
 
-struct FSerializedObjectInfo
+struct FSerializedObjectInfo : public FTransactionObjectId
 {
 public:
+	FSerializedObjectInfo() = default;
+
+	explicit FSerializedObjectInfo(const UObject* InObject)
+	{
+		SetObject(InObject);
+	}
+
 	void SetObject(const UObject* InObject)
 	{
-		ObjectPackageName = InObject->GetPackage()->GetFName();
-		ObjectName = InObject->GetFName();
-		ObjectPathName = *InObject->GetPathName();
-		ObjectOuterPathName = InObject->GetOuter() ? FName(*InObject->GetOuter()->GetPathName()) : FName();
-		ObjectExternalPackageName = InObject->GetExternalPackage() ? InObject->GetExternalPackage()->GetFName() : FName();
-		ObjectClassPathName = FName(*InObject->GetClass()->GetPathName());
+		FTransactionObjectId::SetObject(InObject);
 		bIsPendingKill = !IsValid(InObject);
 	}
 
 	void Reset()
 	{
-		ObjectPackageName = FName();
-		ObjectName = FName();
-		ObjectPathName = FName();
-		ObjectOuterPathName = FName();
-		ObjectExternalPackageName = FName();
-		ObjectClassPathName = FName();
+		FTransactionObjectId::Reset();
 		bIsPendingKill = false;
 	}
 
 	void Swap(FSerializedObjectInfo& Other)
 	{
-		Exchange(ObjectPackageName, Other.ObjectPackageName);
-		Exchange(ObjectName, Other.ObjectName);
-		Exchange(ObjectPathName, Other.ObjectPathName);
-		Exchange(ObjectOuterPathName, Other.ObjectOuterPathName);
-		Exchange(ObjectExternalPackageName, Other.ObjectExternalPackageName);
-		Exchange(ObjectClassPathName, Other.ObjectClassPathName);
+		FTransactionObjectId::Swap(Other);
 		Exchange(bIsPendingKill, Other.bIsPendingKill);
 	}
 
-	/** The package name of the object when it was serialized, can be dictated either by outer chain or external package */
-	FName ObjectPackageName;
-
-	/** The name of the object when it was serialized */
-	FName ObjectName;
-
-	/** The path name of the object when it was serialized */
-	FName ObjectPathName;
-
-	/** The outer path name of the object when it was serialized */
-	FName ObjectOuterPathName;
-
-	/** The external package name of the object when it was serialized, if any */
-	FName ObjectExternalPackageName;
-
-	/** The path name of the object's class. */
-	FName ObjectClassPathName;
-
 	/** The pending kill state of the object when it was serialized */
 	bool bIsPendingKill = false;
 };
