Fix the generation of duplicate bulkdata identifiers when creating new levels from templates containing bulkdata payloads.

#rnx

- When loading a map from a template we can load light and shadowmaps that contain bulkdata. When loading these they will retain the supposedly unique identifier used in the template. However when we save the level it will become a new package so we need those identifiers to be unique.
- Add a new option to FLinkerInstancingContext that signals that editor bulkdata objects being loaded should generate new identifiers

diff --git a/Engine/Source/Editor/UnrealEd/Private/EditorServer.cpp b/Engine/Source/Editor/UnrealEd/Private/EditorServer.cpp
--- a/Engine/Source/Editor/UnrealEd/Private/EditorServer.cpp
+++ b/Engine/Source/Editor/UnrealEd/Private/EditorServer.cpp
@@ -2693,6 +2693,8 @@ bool UEditorEngine::Map_Load(const TCHAR* Str, FOutputDevice& Ar)
 						FSoftObjectPath(*WriteToString<256>(WorldPackage->GetName(), TEXT("."), ShortWorldPackageName))
 					);
 
+					WorldPackageInstancingContext.SetRegenerateUniqueBulkDataGuids(true);
+
 					WorldPackage = LoadPackage( WorldPackage, *LongTempFname, LoadFlags, nullptr /* InReaderOverride */, &WorldPackageInstancingContext);
 					WorldPackage->SetPackageFlags(PKG_NewlyCreated);
 					UWorld::WorldTypePreLoadMap.Remove(WorldPackageFName);
diff --git a/Engine/Source/Runtime/CoreUObject/Private/Serialization/EditorBulkData.cpp b/Engine/Source/Runtime/CoreUObject/Private/Serialization/EditorBulkData.cpp
--- a/Engine/Source/Runtime/CoreUObject/Private/Serialization/EditorBulkData.cpp
+++ b/Engine/Source/Runtime/CoreUObject/Private/Serialization/EditorBulkData.cpp
@@ -961,8 +961,12 @@ void FEditorBulkData::Serialize(FArchive& Ar, UObject* Owner, bool bAllowRegiste
 			if (BulkDataId.IsValid())
 			{
 				FLinkerLoad* LinkerLoad = Cast<FLinkerLoad>(Ar.GetLinker());
-				if (Ar.HasAllPortFlags(PPF_Duplicate) ||
-					(LinkerLoad && LinkerLoad->GetInstancingContext().IsInstanced()))
+
+				if (LinkerLoad && LinkerLoad->GetInstancingContext().ShouldRegenerateUniqueBulkDataGuids())
+				{
+					BulkDataId = FGuid::NewGuid();
+				}
+				else if (Ar.HasAllPortFlags(PPF_Duplicate) || (LinkerLoad && LinkerLoad->GetInstancingContext().IsInstanced()))
 				{
 					// When duplicating BulkDatas we need to create a new BulkDataId to respect the uniqueness contract
 					BulkDataId = CreateUniqueGuid(BulkDataId, Owner, TEXT("PPF_Duplicate serialization"));
diff --git a/Engine/Source/Runtime/CoreUObject/Public/UObject/LinkerInstancingContext.h b/Engine/Source/Runtime/CoreUObject/Public/UObject/LinkerInstancingContext.h
--- a/Engine/Source/Runtime/CoreUObject/Public/UObject/LinkerInstancingContext.h
+++ b/Engine/Source/Runtime/CoreUObject/Public/UObject/LinkerInstancingContext.h
@@ -90,7 +90,20 @@ public:
 		bSoftObjectPathRemappingEnabled = bInSoftObjectPathRemappingEnabled;
 	}
 
-	bool GetSoftObjectPathRemappingEnabled() const { return bSoftObjectPathRemappingEnabled; }
+	bool GetSoftObjectPathRemappingEnabled() const 
+	{ 
+		return bSoftObjectPathRemappingEnabled; 
+	}
+
+	void SetRegenerateUniqueBulkDataGuids(bool bFlag)
+	{
+		bRegenerateUniqueBulkDataGuids = bFlag;
+	}
+
+	bool ShouldRegenerateUniqueBulkDataGuids() const
+	{
+		return bRegenerateUniqueBulkDataGuids;
+	}
 
 	/** Return the instanced package name for a given instanced outer package and an object package name */
 	static FString GetInstancedPackageName(const FString& InOuterPackageName, const FString& InPackageName)
@@ -109,4 +122,6 @@ private:
 	TSet<FName> Tags;
 	/** Remap soft object paths */
 	bool bSoftObjectPathRemappingEnabled = true;
+	/** When true we will generate new unique identifiers for editor bulkdata objects */
+	bool bRegenerateUniqueBulkDataGuids = false;
 };
