The internal list of all enum values will now store enum value names with the package name the enum was defined in to allow identically named enums in different packages.

+ Replaced a few use cases of UEnum::LookupEnumName with FindObject where LookupEnumName was used to find a known UEnum object

diff --git a/Engine/Source/Editor/BlueprintGraph/Private/K2Node_InputTouch.cpp b/Engine/Source/Editor/BlueprintGraph/Private/K2Node_InputTouch.cpp
--- a/Engine/Source/Editor/BlueprintGraph/Private/K2Node_InputTouch.cpp
+++ b/Engine/Source/Editor/BlueprintGraph/Private/K2Node_InputTouch.cpp
@@ -55,11 +55,11 @@ void UK2Node_InputTouch::PostLoad()
 
 UEnum* UK2Node_InputTouch::GetTouchIndexEnum()
 {
-	static UEnum* TouchIndexEnum = NULL;
-	if(NULL == TouchIndexEnum)
+	static UEnum* TouchIndexEnum = nullptr;
+	if (nullptr == TouchIndexEnum)
 	{
-		FName TouchIndexEnumName(TEXT("ETouchIndex::Touch1"));
-		UEnum::LookupEnumName(TouchIndexEnumName, &TouchIndexEnum);
+		FTopLevelAssetPath TouchIndexEnumPath(TEXT("/Script/InputCore"), TEXT("ETouchIndex"));
+		TouchIndexEnum = FindObject<UEnum>(TouchIndexEnumPath);
 		check(TouchIndexEnum);
 	}
 	return TouchIndexEnum;
diff --git a/Engine/Source/Editor/StaticMeshEditor/Private/StaticMeshEditorTools.cpp b/Engine/Source/Editor/StaticMeshEditor/Private/StaticMeshEditorTools.cpp
--- a/Engine/Source/Editor/StaticMeshEditor/Private/StaticMeshEditorTools.cpp
+++ b/Engine/Source/Editor/StaticMeshEditor/Private/StaticMeshEditorTools.cpp
@@ -424,11 +424,11 @@ int32 SConvexDecomposition::GetVertsPerHullCount() const
 
 static UEnum& GetFeatureImportanceEnum()
 {
-	static FName FeatureImportanceName(TEXT("EMeshFeatureImportance::Off"));
-	static UEnum* FeatureImportanceEnum = NULL;
-	if (FeatureImportanceEnum == NULL)
+	static UEnum* FeatureImportanceEnum = nullptr;
+	if (FeatureImportanceEnum == nullptr)
 	{
-		UEnum::LookupEnumName(FeatureImportanceName, &FeatureImportanceEnum);
+		FTopLevelAssetPath MeshFeatureImportanceEnumPath(TEXT("/Script/Engine"), TEXT("EMeshFeatureImportance"));
+		FeatureImportanceEnum = FindObject<UEnum>(MeshFeatureImportanceEnumPath);
 		check(FeatureImportanceEnum);
 	}
 	return *FeatureImportanceEnum;
@@ -436,11 +436,11 @@ static UEnum& GetFeatureImportanceEnum()
 
 static UEnum& GetTerminationCriterionEunum()
 {
-	static FName Name(TEXT("EStaticMeshReductionTerimationCriterion::Triangles"));
-	static UEnum* EnumPtr = NULL;
-	if (EnumPtr == NULL)
+	static UEnum* EnumPtr = nullptr;
+	if (EnumPtr == nullptr)
 	{
-		UEnum::LookupEnumName(Name, &EnumPtr);
+		FTopLevelAssetPath StaticMeshReductionTerimationCriterionEnumPath(TEXT("/Script/Engine"), TEXT("EStaticMeshReductionTerimationCriterion"));
+		EnumPtr = FindObject<UEnum>(StaticMeshReductionTerimationCriterionEnumPath);
 		check(EnumPtr);
 	}
 	return *EnumPtr;
diff --git a/Engine/Source/Runtime/CoreUObject/Private/UObject/DeferredRegistry.cpp b/Engine/Source/Runtime/CoreUObject/Private/UObject/DeferredRegistry.cpp
--- a/Engine/Source/Runtime/CoreUObject/Private/UObject/DeferredRegistry.cpp
+++ b/Engine/Source/Runtime/CoreUObject/Private/UObject/DeferredRegistry.cpp
@@ -16,31 +16,22 @@ void ReloadProcessObject(UScriptStruct* ScriptStruct, const TCHAR* RenamePrefix)
 }
 
 void ReloadProcessObject(UEnum* Enum, const TCHAR* RenamePrefix)
-{
-	// Make sure the old struct is not used by anything
+{	
+	// Before renaming the enum (moving the enum to transient package) remove its enum value names from the enum master list as the list is populated with enum package name
+	Enum->RemoveNamesFromPrimaryList();
+	// Make sure the old enum is not used by anything
 	Enum->ClearFlags(RF_Standalone | RF_Public);
 	Enum->RemoveFromRoot();
 	const FName OldRename = MakeUniqueObjectName(GetTransientPackage(), Enum->GetClass(), *FString::Printf(TEXT("%s_%s"), RenamePrefix, *Enum->GetName()));
 	Enum->Rename(*OldRename.ToString(), GetTransientPackage());
-
-	Enum->RemoveNamesFromPrimaryList();
 }
 
 void ReloadProcessObject(UClass* Class, const TCHAR* RenamePrefix)
 {
 	FString NameWithoutPrefix = UObjectBase::RemoveClassPrefix(*Class->GetName());
 
-	// Rename the old class and move it to transient package
-	Class->RemoveFromRoot();
-	Class->ClearFlags(RF_Standalone | RF_Public);
-	Class->GetDefaultObject()->RemoveFromRoot();
-	Class->GetDefaultObject()->ClearFlags(RF_Standalone | RF_Public);
-	const FName OldClassRename = MakeUniqueObjectName(GetTransientPackage(), Class->GetClass(), *FString::Printf(TEXT("%s_%s"), RenamePrefix, *NameWithoutPrefix));
-	Class->Rename(*OldClassRename.ToString(), GetTransientPackage());
-	Class->SetFlags(RF_Transient);
-	Class->AddToRoot();
-
-	// Make sure enums de-register their names BEFORE we create the new class, otherwise there will be name conflicts
+	// Make sure enums de-register their names BEFORE we create the new class, otherwise there will be name conflicts.
+	// This needs to be done before renaming the class (moving it to transient package) as enum names are stored with the package name
 	TArray<UObject*> ClassSubobjects;
 	GetObjectsWithOuter(Class, ClassSubobjects);
 	for (auto ClassSubobject : ClassSubobjects)
@@ -51,6 +42,16 @@ void ReloadProcessObject(UClass* Class, const TCHAR* RenamePrefix)
 		}
 	}
 
+	// Rename the old class and move it to transient package
+	Class->RemoveFromRoot();
+	Class->ClearFlags(RF_Standalone | RF_Public);
+	Class->GetDefaultObject()->RemoveFromRoot();
+	Class->GetDefaultObject()->ClearFlags(RF_Standalone | RF_Public);
+	const FName OldClassRename = MakeUniqueObjectName(GetTransientPackage(), Class->GetClass(), *FString::Printf(TEXT("%s_%s"), RenamePrefix, *NameWithoutPrefix));
+	Class->Rename(*OldClassRename.ToString(), GetTransientPackage());
+	Class->SetFlags(RF_Transient);
+	Class->AddToRoot();
+
 	// Reset singletons for any child fields
 #if WITH_LIVE_CODING
 	for (UField* Field = Class->Children; Field != nullptr; Field = Field->Next)
diff --git a/Engine/Source/Runtime/CoreUObject/Private/UObject/Enum.cpp b/Engine/Source/Runtime/CoreUObject/Private/UObject/Enum.cpp
--- a/Engine/Source/Runtime/CoreUObject/Private/UObject/Enum.cpp
+++ b/Engine/Source/Runtime/CoreUObject/Private/UObject/Enum.cpp
@@ -21,7 +21,7 @@ DEFINE_LOG_CATEGORY_STATIC(LogEnum, Log, All);
 -----------------------------------------------------------------------------*/
 
 FRWLock UEnum::AllEnumNamesLock;
-TMap<FName, UEnum*> UEnum::AllEnumNames;
+TMap<FName, TMap<FName, UEnum*> > UEnum::AllEnumNames;
 
 UEnum::UEnum(const FObjectInitializer& ObjectInitializer)
 	: UField(ObjectInitializer)
@@ -262,12 +262,14 @@ bool UEnum::IsValidEnumName(FName InName) const
 void UEnum::AddNamesToMasterList()
 {
 	FWriteScopeLock ScopeLock(AllEnumNamesLock);
+	EnumPackage = GetPackage()->GetFName();
+	TMap<FName, UEnum*>& PackageEnumValues = AllEnumNames.FindOrAdd(EnumPackage);
 	for (TPair<FName, int64> Kvp : Names)
 	{
-		UEnum* Enum = AllEnumNames.FindRef(Kvp.Key);
+		UEnum* Enum = PackageEnumValues.FindRef(Kvp.Key);
 		if (Enum == nullptr || Enum->HasAnyFlags(RF_NewerVersionExists))
 		{
-			AllEnumNames.Add(Kvp.Key, this);
+			PackageEnumValues.Add(Kvp.Key, this);
 		}
 		else if (Enum != this && Enum->GetOutermost() != GetTransientPackage())
 		{
@@ -278,12 +280,14 @@ void UEnum::AddNamesToMasterList()
 
 void UEnum::AddNamesToPrimaryList()
 {
+	EnumPackage = GetPackage()->GetFName();
+	TMap<FName, UEnum*>& PackageEnumValues = AllEnumNames.FindOrAdd(EnumPackage);
 	for (TPair<FName, int64> Kvp : Names)
 	{
-		UEnum* Enum = AllEnumNames.FindRef(Kvp.Key);
+		UEnum* Enum = PackageEnumValues.FindRef(Kvp.Key);
 		if (Enum == nullptr || Enum->HasAnyFlags(RF_NewerVersionExists))
 		{
-			AllEnumNames.Add(Kvp.Key, this);
+			PackageEnumValues.Add(Kvp.Key, this);
 		}
 		else if (Enum != this && Enum->GetOutermost() != GetTransientPackage())
 		{
@@ -295,28 +299,167 @@ void UEnum::AddNamesToPrimaryList()
 void UEnum::RemoveNamesFromMasterList()
 {
 	FWriteScopeLock ScopeLock(AllEnumNamesLock);
-	for (TPair<FName, int64> Kvp : Names)
+	TMap<FName, UEnum*>* PackageEnumValues = AllEnumNames.Find(EnumPackage);
+	if (PackageEnumValues != nullptr)
 	{
-		UEnum* Enum = AllEnumNames.FindRef(Kvp.Key);
-		if (Enum == this)
+		for (TPair<FName, int64> Kvp : Names)
+		{
+			UEnum* Enum = PackageEnumValues->FindRef(Kvp.Key);
+			if (Enum == this)
+			{
+				PackageEnumValues->Remove(Kvp.Key);
+			}
+		}
+		if (PackageEnumValues->Num() == 0)
 		{
-			AllEnumNames.Remove(Kvp.Key);
+			AllEnumNames.Remove(EnumPackage);
 		}
 	}
 }
 
 void UEnum::RemoveNamesFromPrimaryList()
 {
-	for (TPair<FName, int64> Kvp : Names)
+	TMap<FName, UEnum*>* PackageEnumValues = AllEnumNames.Find(EnumPackage);
+	if (PackageEnumValues != nullptr)
 	{
-		UEnum* Enum = AllEnumNames.FindRef(Kvp.Key);
-		if (Enum == this)
+		for (TPair<FName, int64> Kvp : Names)
+		{
+			UEnum* Enum = PackageEnumValues->FindRef(Kvp.Key);
+			if (Enum == this)
+			{
+				PackageEnumValues->Remove(Kvp.Key);
+			}
+		}
+		if (PackageEnumValues->Num() == 0)
 		{
-			AllEnumNames.Remove(Kvp.Key);
+			AllEnumNames.Remove(EnumPackage);
 		}
 	}
 }
 
+UEnum* UEnum::LookupAllEnumNamesWithOptions(FName PackageName, EFindFirstObjectOptions Options, TFunctionRef<bool(FName)> CompareNameFunction)
+{
+	UEnum* TheEnum = nullptr;
+	bool bLastFoundEnumWasNative = false;
+
+	auto LookupAllEnumNamesInPackageWithOptions = [Options, &CompareNameFunction, &bLastFoundEnumWasNative](UEnum*& OutEnum, FName CurrentPackageName, const TMap<FName, UEnum*>& EnumNameToUEnumMap)
+	{		
+		bool bIsNativePackageOrDontCare = !(Options & EFindFirstObjectOptions::NativeFirst) || FPackageName::IsScriptPackage(FNameBuilder(CurrentPackageName));
+
+		for (const TPair<FName, UEnum*>& ValueNameToEnumPair : EnumNameToUEnumMap)
+		{
+			if (CompareNameFunction(ValueNameToEnumPair.Key))
+			{
+				if (OutEnum)
+				{
+					if (!!(Options & EFindFirstObjectOptions::EnsureIfAmbiguous))
+					{
+						ensureAlwaysMsgf(false, TEXT("Ambiguous results in LookupAllEnumNamesWithOptions: first enum found \"%s\" but it could also be: \"%s\""),
+							*OutEnum->GetPathName(), *ValueNameToEnumPair.Value->GetPathName());
+					}
+
+					if (bIsNativePackageOrDontCare && !bLastFoundEnumWasNative)
+					{
+						OutEnum = ValueNameToEnumPair.Value;
+						bLastFoundEnumWasNative = bIsNativePackageOrDontCare;
+					}
+				}
+				else
+				{
+					OutEnum = ValueNameToEnumPair.Value;
+					bLastFoundEnumWasNative = bIsNativePackageOrDontCare;
+
+					// If we don't want to check if the search is ambiguous and we don't care if the enum is native or not (or we do care and it's native), abort now
+					if (!(Options & EFindFirstObjectOptions::EnsureIfAmbiguous) && bIsNativePackageOrDontCare)
+					{
+						break;
+					}
+				}
+			}
+		}
+	};
+
+	if (!PackageName.IsNone())
+	{
+		// Fast path, we only need to look in the specified package
+		const TMap<FName, UEnum*>* EnumNameToUEnumMap = AllEnumNames.Find(PackageName);
+		if (EnumNameToUEnumMap)
+		{
+			LookupAllEnumNamesInPackageWithOptions(TheEnum, PackageName, *EnumNameToUEnumMap);
+		}
+	}
+	else
+	{
+		// Slow path, look through all existing packages
+		for (const TPair<FName, TMap<FName, UEnum*>>& PackageToValuesPair : AllEnumNames)
+		{
+			LookupAllEnumNamesInPackageWithOptions(TheEnum, PackageToValuesPair.Key, PackageToValuesPair.Value);
+
+			if (TheEnum && bLastFoundEnumWasNative && !(Options & EFindFirstObjectOptions::EnsureIfAmbiguous))
+			{
+				break;
+			}
+		}
+	}
+
+	return TheEnum;
+}
+
+int64 UEnum::LookupEnumName(FName PackageName, FName TestName, EFindFirstObjectOptions Options /*= EFindFirstObjectOptions::None*/, UEnum** OutFoundEnum /*= nullptr*/)
+{
+	FReadScopeLock ScopeLock(AllEnumNamesLock);
+	UEnum* TheEnum = nullptr;
+	if (!PackageName.IsNone() && Options == EFindFirstObjectOptions::None)
+	{
+		TMap<FName, UEnum*>* PackageEnumValues = AllEnumNames.Find(PackageName);
+		if (PackageEnumValues)
+		{
+			TheEnum = PackageEnumValues->FindRef(TestName);
+		}
+	}
+	else
+	{
+		if (Options == EFindFirstObjectOptions::None)
+		{
+			// If no options are specified we can just look the value name up with a TMap lookup for each of the packages until we find a match
+			for (const TPair<FName, TMap<FName, UEnum*>>& PackageToValuesPair : AllEnumNames)
+			{
+				TheEnum = PackageToValuesPair.Value.FindRef(TestName);
+				if (TheEnum)
+				{
+					break;
+				}
+			}
+		}
+		else
+		{
+			TheEnum = LookupAllEnumNamesWithOptions(PackageName, Options, [TestName](FName EnumValueName) { return TestName == EnumValueName; });
+		}
+	}
+
+	if (OutFoundEnum != nullptr)
+	{
+		*OutFoundEnum = TheEnum;
+	}
+	return (TheEnum != nullptr) ? TheEnum->GetValueByName(TestName) : INDEX_NONE;
+}
+
+int64 UEnum::LookupEnumNameSlow(FName PackageName, const TCHAR* InTestShortName, EFindFirstObjectOptions Options /*= EFindFirstObjectOptions::None*/, UEnum** OutFoundEnum /*= nullptr*/)
+{
+	FReadScopeLock ScopeLock(AllEnumNamesLock);
+	FName TestName(InTestShortName);
+	FString TestShortName(FString(TEXT("::")) + InTestShortName);
+	UEnum* TheEnum = LookupAllEnumNamesWithOptions(PackageName, Options, [TestName, &TestShortName](FName EnumValueName) { return TestName == EnumValueName || EnumValueName.ToString().Contains(TestShortName); });
+
+	if (OutFoundEnum != nullptr)
+	{
+		*OutFoundEnum = TheEnum;
+	}
+	int64 Result = (TheEnum != nullptr) ? TheEnum->GetValueByName(InTestShortName) : INDEX_NONE;
+
+	return Result;
+}
+
 FString UEnum::GenerateEnumPrefix() const
 {
 	FString Prefix;
@@ -638,7 +781,7 @@ bool UEnum::SetEnums(TArray<TPair<FName, int64>>& InNames, UEnum::ECppForm InCpp
 		if (!ContainsExistingMax())
 		{
 			FName MaxEnumItem = *GenerateFullEnumName(*(GenerateEnumPrefix() + TEXT("_MAX")));
-			if (LookupEnumName(MaxEnumItem) != INDEX_NONE)
+			if (LookupEnumName(GetOutermost()->GetFName(), MaxEnumItem) != INDEX_NONE)
 			{
 				// the MAX identifier is already being used by another enum
 				return false;
@@ -806,7 +949,7 @@ int64 UEnum::ParseEnum(const TCHAR*& Str)
 	if (FParse::AlnumToken(ParsedStr, Token))
 	{
 		FName TheName = FName(*Token, FNAME_Find);
-		int64 Result = LookupEnumName(TheName);
+		int64 Result = LookupEnumName(FName(), TheName);
 		if (Result != INDEX_NONE)
 		{
 			Str = ParsedStr;
diff --git a/Engine/Source/Runtime/CoreUObject/Private/UObject/Property.cpp b/Engine/Source/Runtime/CoreUObject/Private/UObject/Property.cpp
--- a/Engine/Source/Runtime/CoreUObject/Private/UObject/Property.cpp
+++ b/Engine/Source/Runtime/CoreUObject/Private/UObject/Property.cpp
@@ -1476,7 +1476,7 @@ static const int32 ReadArrayIndex(UStruct* ObjectStruct, const TCHAR*& Str, FOut
 					if (IndexTokenName != NAME_None)
 					{
 						// Search for the enum in question.
-						Index = UEnum::LookupEnumName(IndexTokenName);
+						Index = UEnum::LookupEnumName(FName(), IndexTokenName, EFindFirstObjectOptions::NativeFirst /* Only native enums can be used as array indices */);
 						if (Index == INDEX_NONE)
 						{
 							Index = 0;
diff --git a/Engine/Source/Runtime/CoreUObject/Public/UObject/Class.h b/Engine/Source/Runtime/CoreUObject/Public/UObject/Class.h
--- a/Engine/Source/Runtime/CoreUObject/Public/UObject/Class.h
+++ b/Engine/Source/Runtime/CoreUObject/Public/UObject/Class.h
@@ -2053,44 +2053,42 @@ public:
 	 */
 	virtual FString GenerateFullEnumName(const TCHAR* InEnumName) const;
 
+	/**
+	 * Searches the list of all enum value names for the specified name
+	 * @param PackageName Package where the enum was defined (/Script/ModuleName or uasset package name in case of script defined enums). If an empty (None) name is specified the function will look through all enums to find a match
+	 * @param TestName Fully qualified enum value name (EEnumName::ValueName)
+	 * @Param Options Optional options that may limit search results or define ambiguous search result behavior
+	 * @param OutFoundEnum Optional address of a variable where the resulting UEnum object should be stored
+	 * @return The value the specified name represents if found, otherwise INDEX_NONE
+	 */
+	static int64 LookupEnumName(FName PackageName, FName TestName, EFindFirstObjectOptions Options = EFindFirstObjectOptions::None, UEnum** OutFoundEnum = nullptr);
+
 	/** searches the list of all enum value names for the specified name
 	 * @return the value the specified name represents if found, otherwise INDEX_NONE
 	 */
+	UE_DEPRECATED(5.2, "LookupEnumName that takes only enum name is deprecated. Please use the version of this function that also takes enum package name.")
 	static int64 LookupEnumName(FName TestName, UEnum** FoundEnum = nullptr)
 	{
-		FReadScopeLock ScopeLock(AllEnumNamesLock);
-		UEnum* TheEnum = AllEnumNames.FindRef(TestName);
-		if (FoundEnum != nullptr)
-		{
-			*FoundEnum = TheEnum;
-		}
-		return (TheEnum != nullptr) ? TheEnum->GetValueByName(TestName) : INDEX_NONE;
+		return LookupEnumName(FName(), TestName, EFindFirstObjectOptions::None, FoundEnum);
 	}
 
+	/** 
+	 * Searches the list of all enum value names for the specified name
+	 * @param PackageName Package where the enum was defined (/Script/ModuleName or uasset package name in case of script defined enums). If an empty (None) name is specified the function will look through all native enums to find a match
+	 * @param InTestShortName Fully qualified or short enum value name (EEnumName::ValueName or ValueName)
+	 * @Param Options Optional Options that may limit search results or define ambiguous search result behavior
+	 * @param OutFoundEnum Optional address of a variable where the resulting UEnum object should be stored
+	 * @return The value the specified name represents if found, otherwise INDEX_NONE
+	 */
+	static int64 LookupEnumNameSlow(FName PackageName, const TCHAR* InTestShortName, EFindFirstObjectOptions Options = EFindFirstObjectOptions::None, UEnum** OutFoundEnum = nullptr);
+
 	/** searches the list of all enum value names for the specified name
 	 * @return the value the specified name represents if found, otherwise INDEX_NONE
 	 */
+	UE_DEPRECATED(5.2, "LookupEnumNameSlow that takes only enum name is deprecated. Please use the version of this function that also takes enum package name.")
 	static int64 LookupEnumNameSlow(const TCHAR* InTestShortName, UEnum** FoundEnum = nullptr)
 	{
-		int64 Result = LookupEnumName(InTestShortName, FoundEnum);
-		if (Result == INDEX_NONE)
-		{
-			FString TestShortName = FString(TEXT("::")) + InTestShortName;
-			UEnum* TheEnum = nullptr;
-			for (TMap<FName, UEnum*>::TIterator It(AllEnumNames); It; ++It)
-			{
-				if (It.Key().ToString().Contains(TestShortName) )
-				{
-					TheEnum = It.Value();
-				}
-			}
-			if (FoundEnum != nullptr)
-			{
-				*FoundEnum = TheEnum;
-			}
-			Result = (TheEnum != nullptr) ? TheEnum->GetValueByName(InTestShortName) : INDEX_NONE;
-		}
-		return Result;
+		return LookupEnumNameSlow(FName(), InTestShortName, EFindFirstObjectOptions::None, FoundEnum);
 	}
 
 	/** parses the passed in string for a name, then searches for that name in any Enum (in any package)
@@ -2335,11 +2333,14 @@ protected:
 	/** pointer to function used to look up the enum's display name. Currently only assigned for UEnums generated for nativized blueprints */
 	FEnumDisplayNameFn EnumDisplayNameFn;
 
+	/** Package name this enum was in when its names were being added to the master list */
+	FName EnumPackage;
+
 	/** lock to be taken when accessing AllEnumNames */
 	static FRWLock AllEnumNamesLock;
 
 	/** global list of all value names used by all enums in memory, used for property text import */
-	static TMap<FName, UEnum*> AllEnumNames;
+	static TMap<FName, TMap<FName, UEnum*> > AllEnumNames;
 
 	/** adds the Names in this enum to the master AllEnumNames list */
 	UE_DEPRECATED(5.1, "AddNamesToMasterList is deprecated, please use AddNamesToPrimaryList instead.")
@@ -2349,6 +2350,16 @@ protected:
 	void AddNamesToPrimaryList();
 
 private:
+
+	/**
+	 * Searches the list of all enum value names and returns a UEnum object that contains an enum value name matching the provided comparison function criteria
+	 * @param PackageName Package where the enum was defined (/Script/ModuleName or uasset package name in case of script defined enums). If an empty (None) name is specified the function will look through all native enums to find a match
+	 * @Param Options Options that may limit search results or define ambiguous search result behavior
+	 * @param CompareNameFunction Functions used to compare the existing enum value names with the one that should be found
+	 * @return UEnum object if a matching enum value name was found otherwise nullptr
+	 */
+	static UEnum* LookupAllEnumNamesWithOptions(FName PackageName, EFindFirstObjectOptions Options, TFunctionRef<bool(FName)> CompareNameFunction);
+
 	FORCEINLINE static FString GetValueAsString_Internal( const TCHAR* EnumPath, const int64 EnumeratorValue)
 	{
 		UEnum* EnumClass = FindObject<UEnum>( nullptr, EnumPath );
diff --git a/Engine/Source/Runtime/Engine/Private/Texture.cpp b/Engine/Source/Runtime/Engine/Private/Texture.cpp
--- a/Engine/Source/Runtime/Engine/Private/Texture.cpp
+++ b/Engine/Source/Runtime/Engine/Private/Texture.cpp
@@ -1227,12 +1227,12 @@ void UTexture::SetDeterministicLightingGuid()
 UEnum* UTexture::GetPixelFormatEnum()
 {
 	// Lookup the pixel format enum so that the pixel format can be serialized by name.
-	static FName PixelFormatUnknownName(TEXT("PF_Unknown"));
-	static UEnum* PixelFormatEnum = NULL;
-	if (PixelFormatEnum == NULL)
+	static UEnum* PixelFormatEnum = nullptr;
+	if (PixelFormatEnum == nullptr)
 	{
+		FTopLevelAssetPath PixelFormatEnumPath(TEXT("/Script/CoreUObject"), TEXT("EPixelFormat"));
 		check(IsInGameThread());
-		UEnum::LookupEnumName(PixelFormatUnknownName, &PixelFormatEnum);
+		PixelFormatEnum = FindObject<UEnum>(PixelFormatEnumPath);
 		check(PixelFormatEnum);
 	}
 	return PixelFormatEnum;
diff --git a/Engine/Source/Runtime/Engine/Private/Timeline.cpp b/Engine/Source/Runtime/Engine/Private/Timeline.cpp
--- a/Engine/Source/Runtime/Engine/Private/Timeline.cpp
+++ b/Engine/Source/Runtime/Engine/Private/Timeline.cpp
@@ -30,11 +30,11 @@ DECLARE_CYCLE_STAT(TEXT("TimelineComp Tick"), STAT_TimelineCompTick, STATGROUP_D
 
 UEnum* FTimeline::GetTimelineDirectionEnum()
 {
-	static UEnum* TimelineDirectionEnum = NULL;
-	if(NULL == TimelineDirectionEnum)
+	static UEnum* TimelineDirectionEnum = nullptr;
+	if (nullptr == TimelineDirectionEnum)
 	{
-		FName TimelineDirectionEnumName(TEXT("ETimelineDirection::Forward"));
-		UEnum::LookupEnumName(TimelineDirectionEnumName, &TimelineDirectionEnum);
+		FTopLevelAssetPath TimelineDirectionEnumEnumPath(TEXT("/Script/Engine"), TEXT("ETimelineDirection"));
+		TimelineDirectionEnum = FindObject<UEnum>(TimelineDirectionEnumEnumPath);
 		check(TimelineDirectionEnum);
 	}
 	return TimelineDirectionEnum;
diff --git a/Engine/Source/Runtime/Engine/Private/UserDefinedEnum.cpp b/Engine/Source/Runtime/Engine/Private/UserDefinedEnum.cpp
--- a/Engine/Source/Runtime/Engine/Private/UserDefinedEnum.cpp
+++ b/Engine/Source/Runtime/Engine/Private/UserDefinedEnum.cpp
@@ -234,7 +234,7 @@ bool UUserDefinedEnum::SetEnums(TArray<TPair<FName, int64>>& InNames, ECppForm I
 		const FString EnumPrefix = (TryNum == 0) ? BaseEnumPrefix : FString::Printf(TEXT("%s_%d"), *BaseEnumPrefix, TryNum - 1);
 		const FName MaxEnumItem = *GenerateFullEnumName(*(EnumPrefix + TEXT("_MAX")));
 		const int64 MaxEnumItemIndex = GetValueByName(MaxEnumItem);
-		if ((MaxEnumItemIndex == INDEX_NONE) && (LookupEnumName(MaxEnumItem) == INDEX_NONE))
+		if ((MaxEnumItemIndex == INDEX_NONE) && (LookupEnumName(GetPackage()->GetFName(), MaxEnumItem) == INDEX_NONE))
 		{
 			int64 MaxEnumValue = (InNames.Num() == 0)? 0 : GetMaxEnumValue() + 1;
 			Names.Emplace(MaxEnumItem, MaxEnumValue);
