GC timeout can be disabled when address sanitizer is enabled (and avoid timeout being triggered due to ASan build being slower than regular build)

#rnx

diff --git a/Engine/Source/Runtime/CoreUObject/Private/UObject/GarbageCollection.cpp b/Engine/Source/Runtime/CoreUObject/Private/UObject/GarbageCollection.cpp
--- a/Engine/Source/Runtime/CoreUObject/Private/UObject/GarbageCollection.cpp
+++ b/Engine/Source/Runtime/CoreUObject/Private/UObject/GarbageCollection.cpp
@@ -1792,7 +1792,7 @@ void IncrementalPurgeGarbage(bool bUseTimeLimit, double TimeLimit)
 static bool GWarningTimeOutHasBeenDisplayedGC = false;
 
 static bool GEnableTimeoutOnPendingDestroyedObjectGC = true;
-#if UE_BUILD_SHIPPING
+#if UE_BUILD_SHIPPING || USING_ADDRESS_SANITISER
 static FAutoConsoleVariableRef CVarGCEnableTimeoutOnPendingDestroyedObjectInShipping(
 	TEXT("gc.EnableTimeoutOnPendingDestroyedObjectInShipping"),
 	GEnableTimeoutOnPendingDestroyedObjectGC,
