Add a new save flag 'SAVE_RehydratePayloads' which when used will force the package being saved to be rehydrated, so that any virtualized payloads will be downloaded and stored in the resulting package file.

#rnx

diff --git a/Engine/Source/Runtime/CoreUObject/Private/Serialization/EditorBulkData.cpp b/Engine/Source/Runtime/CoreUObject/Private/Serialization/EditorBulkData.cpp
--- a/Engine/Source/Runtime/CoreUObject/Private/Serialization/EditorBulkData.cpp
+++ b/Engine/Source/Runtime/CoreUObject/Private/Serialization/EditorBulkData.cpp
@@ -2080,10 +2080,16 @@ FEditorBulkData::EFlags FEditorBulkData::BuildFlagsForSerialization(FArchive& Ar
 			{
 				EnumRemoveFlags(UpdatedFlags, EFlags::HasPayloadSidecarFile);
 
-				// Remove the virtualization flag if we are rehydrating packages on save unless
-				// referencing the payload data is allowed, in which case we can continue to save
-				// as virtualized.
-				if (LinkerSave != nullptr && !bKeepFileDataByReference && CVarShouldRehydrateOnSave.GetValueOnAnyThread())
+				// Check to see if we need to rehydrate the payload on save, this is true if
+				// a) We are saving to a package (LinkerSave is not null)
+				// b) Either the save package call we made with the  ESaveFlags::SAVE_RehydratePayloads flag set OR
+				// the cvar Serialization.RehydrateOnSave is true
+				// c) We are not saving a reference to a different domain. If we are saving a reference then the goal
+				// is to avoid including the payload in the current target domain if possible, rehydrating for this
+				// domain would prevent this.
+				if ( LinkerSave != nullptr && 
+					(LinkerSave->bRehydratePayloads || CVarShouldRehydrateOnSave.GetValueOnAnyThread()) &&
+					!bKeepFileDataByReference)
 				{
 					EnumRemoveFlags(UpdatedFlags, EFlags::IsVirtualized);
 				}
diff --git a/Engine/Source/Runtime/CoreUObject/Private/UObject/SavePackage/SaveContext.h b/Engine/Source/Runtime/CoreUObject/Private/UObject/SavePackage/SaveContext.h
--- a/Engine/Source/Runtime/CoreUObject/Private/UObject/SavePackage/SaveContext.h
+++ b/Engine/Source/Runtime/CoreUObject/Private/UObject/SavePackage/SaveContext.h
@@ -618,6 +618,11 @@ public:
 		return bIsFixupStandaloneFlags;
 	}
 
+	bool ShouldRehydratePayloads() const
+	{
+		return (SaveArgs.SaveFlags & ESaveFlags::SAVE_RehydratePayloads) != 0;
+	}
+
 	FUObjectSerializeContext* GetSerializeContext() const
 	{
 		return SerializeContext;
diff --git a/Engine/Source/Runtime/CoreUObject/Private/UObject/SavePackage2.cpp b/Engine/Source/Runtime/CoreUObject/Private/UObject/SavePackage2.cpp
--- a/Engine/Source/Runtime/CoreUObject/Private/UObject/SavePackage2.cpp
+++ b/Engine/Source/Runtime/CoreUObject/Private/UObject/SavePackage2.cpp
@@ -893,6 +893,8 @@ ESavePackageResult CreateLinker(FSaveContext& SaveContext)
 		SaveContext.GetLinker()->bUpdatingLoadedPath = SaveContext.IsUpdatingLoadedPath();
 		SaveContext.GetLinker()->bProceduralSave = SaveContext.IsProceduralSave();
 
+		SaveContext.GetLinker()->bRehydratePayloads = SaveContext.ShouldRehydratePayloads();
+		
 		if (UE::FPackageTrailer::IsEnabled())
 		{
 			// The package trailer is not supported for text based assets yet
diff --git a/Engine/Source/Runtime/CoreUObject/Public/UObject/LinkerSave.h b/Engine/Source/Runtime/CoreUObject/Public/UObject/LinkerSave.h
--- a/Engine/Source/Runtime/CoreUObject/Public/UObject/LinkerSave.h
+++ b/Engine/Source/Runtime/CoreUObject/Public/UObject/LinkerSave.h
@@ -123,6 +123,9 @@ public:
 	 * This is used e.g. to decide whether to update the in-memory file offsets for BulkData.
 	 */
 	bool bUpdatingLoadedPath = false;
+
+	/** When set to true, payloads that are currently virtualized should be downloaded and stored locally with the package */
+	bool bRehydratePayloads = false;
 	
 	struct FSidecarStorageInfo
 	{
diff --git a/Engine/Source/Runtime/CoreUObject/Public/UObject/ObjectMacros.h b/Engine/Source/Runtime/CoreUObject/Public/UObject/ObjectMacros.h
--- a/Engine/Source/Runtime/CoreUObject/Public/UObject/ObjectMacros.h
+++ b/Engine/Source/Runtime/CoreUObject/Public/UObject/ObjectMacros.h
@@ -103,6 +103,7 @@ enum ESaveFlags
 													///This is only reasonable to use with full cooked builds for distribution, or with a domain that selects the payload matching a given version externally.
 	SAVE_Optional					= 0x00008000,	///< Indicate that we to save optional exports. This flag is only valid while cooking. Optional exports are filtered if not specified during cooking. 
 	SAVE_AllowTimeout				= 0x00010000,   ///< If present, if SavePackage is blocked by an asynchronous operation, it will quickly fail with ESavePackageResult::Timeout, otherwise it will wait for a while and then fail with ESavePackageResult::Error.
+	SAVE_RehydratePayloads			= 0x00020000,   ///< Any virtualized payloads in the package should be pulled and stored locally in the package file during the save
 	SAVE_Unversioned = SAVE_Unversioned_Native | SAVE_Unversioned_Properties,
 };
 
