Bulk Data - inline reloading ensures

diff --git a/Engine/Source/Runtime/CoreUObject/Private/Serialization/BulkData.cpp b/Engine/Source/Runtime/CoreUObject/Private/Serialization/BulkData.cpp
--- a/Engine/Source/Runtime/CoreUObject/Private/Serialization/BulkData.cpp
+++ b/Engine/Source/Runtime/CoreUObject/Private/Serialization/BulkData.cpp
@@ -23,10 +23,6 @@
 // If set to 0 then we will pretend that optional data does not exist, useful for testing.
 #define ALLOW_OPTIONAL_DATA 1
 
-// When set to 1 attempting to reload inline data will fail with the old loader in the same way that it fails in
-// the new loader to keep the results consistent.
-#define UE_KEEP_INLINE_RELOADING_CONSISTENT 0
-
 //////////////////////////////////////////////////////////////////////////////
 
 const FIoFilenameHash FALLBACK_IO_FILENAME_HASH = INVALID_IO_FILENAME_HASH - 1;
@@ -123,31 +119,6 @@ bool StartAsyncLoad(
 /** Flush pending async load. */
 void FlushAsyncLoad(FBulkData* BulkData);
 
-/** 
-* Returns true if we should not trigger an ensure if we detect an inline bulkdata reload request that will not work with the 
-* IoStore system. This can be done by setting [Core.System]IgnoreInlineBulkDataReloadEnsures to true in the config file.
-* It is NOT recommended that you do this, but is provded in case of unforseen use cases.
-*/
-bool ShouldIgnoreInlineDataReloadEnsures()
-{
-	static struct FIgnoreInlineDataReloadEnsures
-	{
-		bool bEnabled = false;
-
-		FIgnoreInlineDataReloadEnsures()
-		{
-			FConfigFile PlatformEngineIni;
-			FConfigCacheIni::LoadLocalIniFile(PlatformEngineIni, TEXT("Engine"), true, ANSI_TO_TCHAR(FPlatformProperties::IniPlatformName()));
-
-			PlatformEngineIni.GetBool(TEXT("Core.System"), TEXT("IgnoreInlineBulkDataReloadEnsures"), bEnabled);
-
-			UE_LOG(LogSerialization, Display, TEXT("IgnoreInlineDataReloadEnsures: '%s'"), bEnabled ? TEXT("true") : TEXT("false"));
-		}
-	} IgnoreInlineDataReloadEnsures;
-
-	return IgnoreInlineDataReloadEnsures.bEnabled;
-}
-
 //////////////////////////////////////////////////////////////////////////////
 
 static FArchive& SerializeAsInt32(FArchive& Ar, int64& Value)
@@ -780,23 +751,16 @@ bool FBulkData::IsStoredCompressedOnDisk() const
 
 bool FBulkData::CanLoadFromDisk() const
 {
-	if (BulkChunkId.IsValid())
-	{
-		return true;
-	}
-
 #if WITH_EDITOR
-	return AttachedAr != nullptr;
+	return AttachedAr != nullptr || BulkChunkId.IsValid();
 #else
-
 #if UE_KEEP_INLINE_RELOADING_CONSISTENT
 	if (IsInlined())
 	{
 		return false;
 	}
 #endif //UE_KEEP_INLINE_RELOADING_CONSISTENT
-
-	return false;
+	return BulkChunkId.IsValid();
 #endif // WITH_EDITOR
 }
 
diff --git a/Engine/Source/Runtime/CoreUObject/Private/Serialization/BulkDataStreaming.cpp b/Engine/Source/Runtime/CoreUObject/Private/Serialization/BulkDataStreaming.cpp
--- a/Engine/Source/Runtime/CoreUObject/Private/Serialization/BulkDataStreaming.cpp
+++ b/Engine/Source/Runtime/CoreUObject/Private/Serialization/BulkDataStreaming.cpp
@@ -95,6 +95,39 @@ namespace UE::BulkData::Private
 
 //////////////////////////////////////////////////////////////////////////////
 
+/**
+ * Triggers an ensure when trying to stream inline bulk data. This is to prevent
+ * inconsistencies between the Zen loader loading from I/O store (utoc/ucas) and the legacy EDL loader using
+ * the file system (.pak). The Zen loader currenlty does not allow reloading of inline bulk data.
+ * The trigger can be ignored by setting [Core.System]IgnoreInlineBulkDataReloadEnsures to true in the config file.
+ * Ignoring this ensure will most likely break loading of bulk data when packaging with Zen/IO store (.utoc/.ucas).
+ */
+void EnsureCanStreamBulkData(const FBulkMetaData& BulkMeta)
+{
+#if !WITH_EDITOR && !UE_KEEP_INLINE_RELOADING_CONSISTENT	
+	static struct FIgnoreInlineDataReloadEnsures
+	{
+		bool bEnabled = false;
+
+		FIgnoreInlineDataReloadEnsures()
+		{
+			FConfigFile PlatformEngineIni;
+			FConfigCacheIni::LoadLocalIniFile(PlatformEngineIni, TEXT("Engine"), true, ANSI_TO_TCHAR(FPlatformProperties::IniPlatformName()));
+
+			PlatformEngineIni.GetBool(TEXT("Core.System"), TEXT("IgnoreInlineBulkDataReloadEnsures"), bEnabled);
+
+			UE_LOG(LogSerialization, Display, TEXT("IgnoreInlineDataReloadEnsures: '%s'"), bEnabled ? TEXT("true") : TEXT("false"));
+		}
+	} IgnoreInlineDataReloadEnsures;
+
+	const bool bIsInlined = BulkMeta.HasAnyFlags(BULKDATA_PayloadAtEndOfFile) == false;
+	// Note that we only want the ensure to trigger if we have a valid offset (the bulkdata references data from disk)
+	ensureMsgf(!bIsInlined || BulkMeta.GetOffset() == INDEX_NONE || IgnoreInlineDataReloadEnsures.bEnabled,
+				TEXT("Attempting to stream inline BulkData! This operation is not supported by the IoDispatcher and so will eventually stop working."
+				" The calling code should be fixed to retain the inline data in memory and re-use it rather than discard it and then try to reload from disk!"));
+#endif 
+}
+
 FIoChunkId CreateBulkDataIoChunkId(const FBulkMetaData& BulkMeta, const FPackageId& PackageId)
 {
 	if (PackageId.IsValid() == false)
@@ -490,6 +523,8 @@ bool OpenReadBulkData(
 		return false;
 	}
 
+	EnsureCanStreamBulkData(BulkMeta);
+
 	if (FPackageId PackageId = BulkChunkId.GetPackageId(); PackageId.IsValid())
 	{
 		const FIoChunkId ChunkId = CreateBulkDataIoChunkId(BulkMeta, PackageId);
@@ -547,6 +582,8 @@ TUniquePtr<IAsyncReadFileHandle> OpenAsyncReadBulkData(const FBulkMetaData& Bulk
 		return TUniquePtr<IAsyncReadFileHandle>();
 	}
 
+	EnsureCanStreamBulkData(BulkMeta);
+
 	if (FPackageId PackageId = BulkChunkId.GetPackageId(); PackageId.IsValid())
 	{
 		return MakeUnique<FChunkReadFileHandle>(CreateBulkDataIoChunkId(BulkMeta, PackageId));
@@ -585,7 +622,9 @@ TUniquePtr<IBulkDataIORequest> CreateStreamingRequest(
 	{
 		return TUniquePtr<IBulkDataIORequest>();
 	}
-	
+
+	EnsureCanStreamBulkData(BulkMeta);
+
 	if (FPackageId PackageId = BulkChunkId.GetPackageId(); PackageId.IsValid())
 	{
 		FIoBuffer Buffer = UserSuppliedMemory ? FIoBuffer(FIoBuffer::Wrap, UserSuppliedMemory, Size) : FIoBuffer(Size);
diff --git a/Engine/Source/Runtime/CoreUObject/Public/Serialization/BulkData.h b/Engine/Source/Runtime/CoreUObject/Public/Serialization/BulkData.h
--- a/Engine/Source/Runtime/CoreUObject/Public/Serialization/BulkData.h
+++ b/Engine/Source/Runtime/CoreUObject/Public/Serialization/BulkData.h
@@ -51,6 +51,10 @@ namespace UE { namespace Serialization { class FEditorBulkData; } }
 #define USE_BULKDATA_STREAMING_TOKEN UE_DEPRECATED_MACRO(5.0, "USE_BULKDATA_STREAMING_TOKEN now always evaluates to 0 and will be removed") 0
 #define STREAMINGTOKEN_PARAM(param) UE_DEPRECATED_MACRO(5.0, "STREAMINGTOKEN_PARAM now always evaluates to a NOP")
 
+// When set to 1 attempting to reload inline data will fail with the old loader in the same way that it fails in
+// the new loader to keep the results consistent.
+#define UE_KEEP_INLINE_RELOADING_CONSISTENT 0
+
 class IMappedFileHandle;
 class IMappedFileRegion;
 
