AsyncLoading2: Lower the new NotifyRegistrationEvent ensure (added in 21696962) to a log message

#rnx

diff --git a/Engine/Source/Runtime/CoreUObject/Private/Serialization/AsyncLoading2.cpp b/Engine/Source/Runtime/CoreUObject/Private/Serialization/AsyncLoading2.cpp
--- a/Engine/Source/Runtime/CoreUObject/Private/Serialization/AsyncLoading2.cpp
+++ b/Engine/Source/Runtime/CoreUObject/Private/Serialization/AsyncLoading2.cpp
@@ -3405,7 +3405,7 @@ void FGlobalImportStore::RegistrationComplete()
 			if (!ScriptObjects.Contains(PackageGlobalImportIndex))
 			{
 				ScriptObjects.Add(PackageGlobalImportIndex, Package);
-				ensureMsgf(false, TEXT("Script package %s (0x%016llX) is missing a NotifyRegistrationEvent from the initial load phase."),
+				UE_LOG(LogStreaming, Display, TEXT("Script package %s (0x%016llX) is missing a NotifyRegistrationEvent from the initial load phase."),
 					*Package->GetFullName(),
 					PackageGlobalImportIndex.Value());
 			}
