[Engine runtime]
* Fixed formatting of text inside scoped io activity

diff --git a/Engine/Source/Runtime/CoreUObject/Private/UObject/UObjectGlobals.cpp b/Engine/Source/Runtime/CoreUObject/Private/UObject/UObjectGlobals.cpp
--- a/Engine/Source/Runtime/CoreUObject/Private/UObject/UObjectGlobals.cpp
+++ b/Engine/Source/Runtime/CoreUObject/Private/UObject/UObjectGlobals.cpp
@@ -1495,7 +1495,7 @@ UPackage* LoadPackageInternal(UPackage* InOuter, const FPackagePath& PackagePath
 	// Set up a load context
 	TRefCountPtr<FUObjectSerializeContext> LoadContext = FUObjectThreadContext::Get().GetSerializeContext();
 
-	UE_SCOPED_IO_ACTIVITY(*WriteToString<512>(TEXT("Sync %s"), PackagePath.GetDebugName()));
+	UE_SCOPED_IO_ACTIVITY(*WriteToString<512>(TEXT("Sync "), PackagePath.GetDebugName()));
 
 	// Try to load.
 	BeginLoad(LoadContext, *PackagePath.GetDebugName());
