[Runtime Core]
* AsyncLoading - Fixed out-of-bounds bug where QueuedPackages[0] was accessed when array was empty

Fixes a incremental editor test

diff --git a/Engine/Source/Runtime/CoreUObject/Private/Serialization/AsyncLoading.cpp b/Engine/Source/Runtime/CoreUObject/Private/Serialization/AsyncLoading.cpp
--- a/Engine/Source/Runtime/CoreUObject/Private/Serialization/AsyncLoading.cpp
+++ b/Engine/Source/Runtime/CoreUObject/Private/Serialization/AsyncLoading.cpp
@@ -504,11 +504,18 @@ struct FActiveAsyncLoadContext
 {
 	uint32 TotalCompleted = 0;
 
-	FActiveAsyncLoadContext(const TCHAR* Name)
+	FActiveAsyncLoadContext(const TArray<FAsyncPackageDesc*>& Queue)
 	{
-		if (*Name == '/')
-			++Name;
-		FTrackedActivity::GetIOActivity().Push(Name);
+		TStringBuilder<256> Str;
+		if (Queue.Num() > 0)
+		{
+			Queue[0]->Name.AppendString(Str);
+			if (Str.GetData()[0] == '/')
+			{
+				Str.RemoveAt(0, 1);
+			}
+		}
+		FTrackedActivity::GetIOActivity().Push(*Str);
 		FTrackedActivity::GetIOActivity().Push(TEXT(""), true);
 		ActiveAsyncLoadContext = this;
 	}
@@ -7526,7 +7533,7 @@ EAsyncPackageState::Type FAsyncLoadingThread::ProcessLoadingUntilComplete(TFunct
 	}
 
 	#if UE_ENABLE_TRACKED_IO
-	FActiveAsyncLoadContext Context(*QueuedPackages[0]->Name.ToString());
+	FActiveAsyncLoadContext Context(QueuedPackages);
 	#endif
 
 	while (IsAsyncLoadingPackages() && TimeLimit > 0 && !CompletionPredicate())
