ZenLoader: When sync loading a package only execute nodes required for that package to finish loading

diff --git a/Engine/Source/Runtime/CoreUObject/Private/Serialization/AsyncLoading2.cpp b/Engine/Source/Runtime/CoreUObject/Private/Serialization/AsyncLoading2.cpp
--- a/Engine/Source/Runtime/CoreUObject/Private/Serialization/AsyncLoading2.cpp
+++ b/Engine/Source/Runtime/CoreUObject/Private/Serialization/AsyncLoading2.cpp
@@ -50,7 +50,6 @@
 #include "Serialization/UnversionedPropertySerialization.h"
 #include "Serialization/Zenaphore.h"
 #include "UObject/GCObject.h"
-#include "Experimental/Containers/FAAArrayQueue.h"
 #include "UObject/ObjectRedirector.h"
 #include "Serialization/BulkData.h"
 #include "Serialization/LargeMemoryReader.h"
@@ -377,6 +376,8 @@ struct FExportObject
 	bool bWasFoundInMemory = false;
 };
 
+class FAsyncLoadingSyncLoadContext;
+
 struct FPackageRequest
 {
 	int32 RequestId = -1;
@@ -1584,48 +1585,50 @@ enum class EAsyncPackageLoadingState2 : uint8
 class FEventLoadGraphAllocator;
 struct FAsyncLoadEventSpec;
 struct FAsyncLoadingThreadState2;
+class FAsyncLoadEventQueue2;
+
+enum class EEventLoadNodeExecutionResult : uint8
+{
+	Timeout,
+	Complete,
+};
 
 /** [EDL] Event Load Node */
 class FEventLoadNode2
 {
-	enum class ENodeState : uint8
-	{
-		Waiting = 0,
-		Executing,
-		Timeout,
-		Completed
-	};
 public:
 	FEventLoadNode2(const FAsyncLoadEventSpec* InSpec, FAsyncPackage2* InPackage, int32 InImportOrExportIndex, int32 InBarrierCount);
 	void DependsOn(FEventLoadNode2* Other);
 	void AddBarrier();
 	void AddBarrier(int32 Count);
 	void ReleaseBarrier(FAsyncLoadingThreadState2* ThreadState = nullptr);
-	void Execute(FAsyncLoadingThreadState2& ThreadState);
+	EEventLoadNodeExecutionResult Execute(FAsyncLoadingThreadState2& ThreadState);
 
-	int32 GetBarrierCount()
+	inline bool IsDone() const
 	{
-		return BarrierCount.Load();
+		return bIsDone.load(std::memory_order_acquire) == 1;
 	}
 
-	inline bool IsDone()
+	inline void MarkAsDone()
 	{
-		return ENodeState::Completed == static_cast<ENodeState>(NodeState.Load());
+		bIsDone.store(1, std::memory_order_release);
 	}
 
-	inline bool IsExecuting() const
+	FAsyncPackage2* GetPackage() const
 	{
-		return ENodeState::Executing == static_cast<ENodeState>(NodeState.Load());
-	}
-
-	inline void SetState(ENodeState InNodeState)
-	{
-		NodeState.Store(static_cast<uint8>(InNodeState));
+		return Package;
 	}
 
 	int32 ReferencerRequestId() const;
 
 private:
+	friend class FAsyncLoadEventQueue2;
+
+	FEventLoadNode2()
+	{
+
+	}
+
 	void ProcessDependencies(FAsyncLoadingThreadState2& ThreadState);
 	void Fire(FAsyncLoadingThreadState2* ThreadState = nullptr);
 
@@ -1634,14 +1637,12 @@ private:
 		FEventLoadNode2* SingleDependent;
 		FEventLoadNode2** MultipleDependents;
 	};
+	std::atomic<FEventLoadNode2*> Next { nullptr };
 	uint32 DependenciesCount = 0;
 	uint32 DependenciesCapacity = 0;
-	TAtomic<int32> BarrierCount { 0 };
-	TAtomic<uint8> DependencyWriterCount { 0 };
-	TAtomic<uint8> NodeState { static_cast<uint8>(ENodeState::Waiting) };
-#if UE_BUILD_DEBUG || UE_BUILD_DEVELOPMENT
-	TAtomic<uint8> bFired { 0 };
-#endif
+	std::atomic<int32> BarrierCount { 0 };
+	std::atomic<uint8> DependencyWriterCount { 0 };
+	std::atomic<uint8> bIsDone { 0 };
 
 	const FAsyncLoadEventSpec* Spec = nullptr;
 	FAsyncPackage2* Package = nullptr;
@@ -1673,6 +1674,10 @@ public:
 	TAtomic<int64> TotalAllocated { 0 };
 };
 
+/** Queue based on DepletableMpscQueue.h 
+	Producers will add nodes to the linked list defined by Sentinel->Tail
+	Nodes will be executed from the linked list defined by LocalHead->LocalTail
+	All the nodes from the producer queue will be moved to the local queue when needed */
 class FAsyncLoadEventQueue2
 {
 public:
@@ -1686,31 +1691,77 @@ public:
 
 	bool PopAndExecute(FAsyncLoadingThreadState2& ThreadState);
 	void Push(FEventLoadNode2* Node);
+	bool ExecuteSyncLoadEvents(FAsyncLoadingThreadState2& ThreadState, FAsyncLoadingSyncLoadContext& SyncLoadContext);
 
 private:
+	void Deplete();
+
 	FZenaphore* Zenaphore = nullptr;
-	FAAArrayQueue<FEventLoadNode2> Queue;
+	FEventLoadNode2 Sentinel; // `Sentinel.Next` is the head of the queue
+	std::atomic<FEventLoadNode2*> Tail{ &Sentinel };
+	FEventLoadNode2* LocalHead = nullptr;
+	FEventLoadNode2* LocalTail = nullptr;
+	FEventLoadNode2* TimedOutEventNode = nullptr;
+	int32 ExecuteSyncLoadEventsCallCounter = 0;
+	bool bIsInPopAndExecute = false;
+
 };
 
 struct FAsyncLoadEventSpec
 {
-	typedef EAsyncPackageState::Type(*FAsyncLoadEventFunc)(FAsyncLoadingThreadState2&, FAsyncPackage2*, int32);
+	typedef EEventLoadNodeExecutionResult(*FAsyncLoadEventFunc)(FAsyncLoadingThreadState2&, FAsyncPackage2*, int32);
 	FAsyncLoadEventFunc Func = nullptr;
 	FAsyncLoadEventQueue2* EventQueue = nullptr;
 	bool bExecuteImmediately = false;
 };
 
+class FAsyncLoadingSyncLoadContext
+{
+public:
+	FAsyncLoadingSyncLoadContext(int32 InRequestId)
+		: RequestId(InRequestId)
+	{
+
+	}
+
+	~FAsyncLoadingSyncLoadContext()
+	{
+		check(Packages.IsEmpty());
+	}
+
+	void AddRef()
+	{
+		++RefCount;
+	}
+
+	void ReleaseRef()
+	{
+		int32 NewRefCount = --RefCount;
+		check(NewRefCount >= 0);
+		if (NewRefCount == 0)
+		{
+			delete this;
+		}
+	}
+
+	int32 RequestId;
+	FAsyncPackage2* RequestedPackageDebug = nullptr;
+	FAsyncPackage2* RequestingPackageDebug = nullptr;
+	FCriticalSection CriticalSection;
+	TSet<FAsyncPackage2*> Packages;
+	bool bHasFoundRequestedPackage = false;
+
+private:
+	std::atomic<int32> RefCount = 1;
+};
+
 struct FAsyncLoadingThreadState2
-	: public FTlsAutoCleanup
 {
-	static FAsyncLoadingThreadState2* Create(FAsyncLoadEventGraphAllocator& GraphAllocator, FIoDispatcher& IoDispatcher)
+	static void Set(FAsyncLoadingThreadState2* State)
 	{
 		check(TlsSlot != 0);
 		check(!FPlatformTLS::GetTlsValue(TlsSlot));
-		FAsyncLoadingThreadState2* State = new FAsyncLoadingThreadState2(GraphAllocator, IoDispatcher);
-		State->Register();
 		FPlatformTLS::SetTlsValue(TlsSlot, State);
-		return State;
 	}
 
 	static FAsyncLoadingThreadState2* Get()
@@ -1784,7 +1835,12 @@ struct FAsyncLoadingThreadState2
 	FAsyncLoadEventGraphAllocator& GraphAllocator;
 	TArray<TTuple<FEventLoadNode2**, uint32>> DeferredFreeArcs;
 	TArray<FEventLoadNode2*> NodesToFire;
-	FEventLoadNode2* CurrentEventNode = nullptr;
+	FEventLoadNode2* CurrentlyExecutingEventNode = nullptr;
+	TArray<FAsyncLoadingSyncLoadContext*> SyncLoadContextStack;
+	TArray<FAsyncPackage2*> PackagesExcludedFromStateDependencyChecks;
+	TSpscQueue<FAsyncLoadingSyncLoadContext*> SyncLoadContextsCreatedOnGameThread;
+	bool bIsAsyncLoadingThread = false;
+	bool bCanAccessAsyncLoadingThreadData = true;
 	bool bShouldFireNodes = true;
 	bool bUseTimeLimit = false;
 	double TimeLimit = 0.0;
@@ -1868,9 +1924,6 @@ struct FAsyncPackage2
 
 	void ClearImportedPackages();
 
-	/** Marks a specific request as complete */
-	void MarkRequestIDsAsComplete();
-
 	/**
 	 * @return Time load begun. This is NOT the time the load was requested in the case of other pending requests.
 	 */
@@ -1952,6 +2005,8 @@ private:
 	FPackageImportStore ImportStore;
 	/** Package which is going to have its exports and imports loaded */
 	UPackage* LinkerRoot = nullptr;
+	// The sync load context associated with this package
+	FAsyncLoadingSyncLoadContext* SyncLoadContext = nullptr;
 	/** Time load begun. This is NOT the time the load was requested in the case of pending requests.	*/
 	double						LoadStartTime = 0.0;
 	TAtomic<int32> RefCount{ 0 };
@@ -2086,13 +2141,13 @@ public:
 
 	/** [EDL] Begin Event driven loader specific stuff */
 
-	static EAsyncPackageState::Type Event_ProcessExportBundle(FAsyncLoadingThreadState2& ThreadState, FAsyncPackage2* Package, int32 ExportBundleIndex);
-	static EAsyncPackageState::Type Event_ProcessPackageSummary(FAsyncLoadingThreadState2& ThreadState, FAsyncPackage2* Package, int32);
-	static EAsyncPackageState::Type Event_SetupDependencies(FAsyncLoadingThreadState2& ThreadState, FAsyncPackage2* Package, int32);
-	static EAsyncPackageState::Type Event_ExportsDone(FAsyncLoadingThreadState2& ThreadState, FAsyncPackage2* Package, int32);
-	static EAsyncPackageState::Type Event_PostLoadExportBundle(FAsyncLoadingThreadState2& ThreadState, FAsyncPackage2* Package, int32 ExportBundleIndex);
-	static EAsyncPackageState::Type Event_DeferredPostLoadExportBundle(FAsyncLoadingThreadState2& ThreadState, FAsyncPackage2* Package, int32 ExportBundleIndex);
-	
+	static EEventLoadNodeExecutionResult Event_ProcessExportBundle(FAsyncLoadingThreadState2& ThreadState, FAsyncPackage2* Package, int32 ExportBundleIndex);
+	static EEventLoadNodeExecutionResult Event_ProcessPackageSummary(FAsyncLoadingThreadState2& ThreadState, FAsyncPackage2* Package, int32);
+	static EEventLoadNodeExecutionResult Event_SetupDependencies(FAsyncLoadingThreadState2& ThreadState, FAsyncPackage2* Package, int32);
+	static EEventLoadNodeExecutionResult Event_ExportsDone(FAsyncLoadingThreadState2& ThreadState, FAsyncPackage2* Package, int32);
+	static EEventLoadNodeExecutionResult Event_PostLoadExportBundle(FAsyncLoadingThreadState2& ThreadState, FAsyncPackage2* Package, int32 ExportBundleIndex);
+	static EEventLoadNodeExecutionResult Event_DeferredPostLoadExportBundle(FAsyncLoadingThreadState2& ThreadState, FAsyncPackage2* Package, int32 ExportBundleIndex);
+
 	void EventDrivenCreateExport(const FAsyncPackageHeaderData& Header, const TArrayView<FExportObject>& Exports, int32 LocalExportIndex);
 	bool EventDrivenSerializeExport(const FAsyncPackageHeaderData& Header, const TArrayView<FExportObject>& Exports, int32 LocalExportIndex, FExportArchive& Ar);
 
@@ -2120,12 +2175,12 @@ private:
 	void CreateExportBundleNodes(const FAsyncLoadEventSpec* EventSpecs);
 	void SetupSerializedArcs(const FAsyncPackageHeaderData& Header, const TArrayView<FEventLoadNode2>& ExportBundleNodes, const TArrayView<FAsyncPackage2*>& ImportedAsyncPackages);
 	void SetupScriptDependencies();
-	bool HaveAllDependenciesReachedStateDebug(FAsyncPackage2* Package, TSet<FAsyncPackage2*>& VisitedPackages, EAsyncPackageLoadingState2 WaitForPackageState);
-	bool HaveAllDependenciesReachedState(FAllDependenciesState FAsyncPackage2::* StateMemberPtr, EAsyncPackageLoadingState2 WaitForPackageState, uint32 CurrentTick);
-	void UpdateDependenciesStateRecursive(FAllDependenciesState FAsyncPackage2::* StateMemberPtr, EAsyncPackageLoadingState2 WaitForPackageState, uint32 CurrentTick, FAsyncPackage2* Root);
-	void WaitForAllDependenciesToReachState(FAllDependenciesState FAsyncPackage2::* StateMemberPtr, EAsyncPackageLoadingState2 WaitForPackageState, uint32& CurrentTickVariable, TFunctionRef<void(FAsyncPackage2*)> OnStateReached);
-	void ConditionalBeginPostLoad();
-	void ConditionalFinishLoading();
+	bool HaveAllDependenciesReachedStateDebug(FAsyncLoadingThreadState2& ThreadState, FAsyncPackage2* Package, TSet<FAsyncPackage2*>& VisitedPackages, EAsyncPackageLoadingState2 WaitForPackageState);
+	bool HaveAllDependenciesReachedState(FAsyncLoadingThreadState2& ThreadState, FAllDependenciesState FAsyncPackage2::* StateMemberPtr, EAsyncPackageLoadingState2 WaitForPackageState, uint32 CurrentTick);
+	void UpdateDependenciesStateRecursive(FAsyncLoadingThreadState2& ThreadState, FAllDependenciesState FAsyncPackage2::* StateMemberPtr, EAsyncPackageLoadingState2 WaitForPackageState, uint32 CurrentTick, FAsyncPackage2* Root);
+	void WaitForAllDependenciesToReachState(FAsyncLoadingThreadState2& ThreadState, FAllDependenciesState FAsyncPackage2::* StateMemberPtr, EAsyncPackageLoadingState2 WaitForPackageState, uint32& CurrentTickVariable, TFunctionRef<void(FAsyncPackage2*)> OnStateReached);
+	void ConditionalBeginPostLoad(FAsyncLoadingThreadState2& ThreadState);
+	void ConditionalFinishLoading(FAsyncLoadingThreadState2& ThreadState);
 
 	/**
 	 * Begin async loading process. Simulates parts of BeginLoad.
@@ -2176,7 +2231,7 @@ struct FScopedAsyncPackageEvent2
 	FAsyncPackage2* Package;
 	/** Outer scope package */
 	FAsyncPackage2* PreviousPackage;
-#if WITH_IOSTORE_IN_EDITOR
+#if WITH_EDITOR
 	IAsyncPackageLoader* PreviousAsyncPackageLoader;
 #endif
 
@@ -2202,10 +2257,6 @@ private:
 	/** [ASYNC/GAME THREAD] true if the async thread is actually started. We don't start it until after we boot because the boot process on the game thread can create objects that are also being created by the loader */
 	bool bThreadStarted = false;
 
-#if ALT2_VERIFY_RECURSIVE_LOADS
-	int32 LoadRecursionLevel = 0;
-#endif
-
 #if !UE_BUILD_SHIPPING
 	FPlatformFileOpenLog* FileOpenLogWrapper = nullptr;
 #endif
@@ -2219,7 +2270,7 @@ private:
 	TArray<FAsyncPackage2*> LoadedPackagesToProcess;
 	/** [GAME THREAD] Game thread CompletedPackages list */
 	TArray<FAsyncPackage2*> CompletedPackages;
-#if WITH_IOSTORE_IN_EDITOR
+#if WITH_EDITOR
 	/** [GAME THREAD] Game thread LoadedAssets list */
 	TSet<FWeakObjectPtr> LoadedAssets;
 #endif
@@ -2246,6 +2297,7 @@ private:
 	TSet<int32> PendingRequests;
 	/** Synchronization object for PendingRequests list */
 	FCriticalSection PendingRequestsCritical;
+	TMap<int32, FAsyncPackage2*> RequestIdToPackageMap; // Only accessed from the async loading thread
 
 	/** [ASYNC/GAME THREAD] Number of package load requests in the async loading queue */
 	TAtomic<int32> QueuedPackagesCounter { 0 };
@@ -2283,6 +2335,9 @@ private:
 	FCriticalSection UnreachableObjectsCritical;
 	FUnreachableObjects UnreachableObjects;
 
+	TUniquePtr<FAsyncLoadingThreadState2> GameThreadState;
+	TUniquePtr<FAsyncLoadingThreadState2> AsyncLoadingThreadState;
+
 	uint32 ConditionalBeginPostLoadTick = 0;
 	uint32 ConditionalFinishLoadingTick = 0;
 
@@ -2302,7 +2357,6 @@ public:
 	FAsyncLoadEventGraphAllocator GraphAllocator;
 	FAsyncLoadEventQueue2 EventQueue;
 	FAsyncLoadEventQueue2 MainThreadEventQueue;
-	TArray<FAsyncLoadEventQueue2*> AltEventQueues;
 	TArray<FAsyncLoadEventSpec> EventSpecs;
 
 	/** True if multithreaded async loading is currently being used. */
@@ -2396,7 +2450,7 @@ public:
 	* @param OutPackagesProcessed Number of packages processed in this call.
 	* @return The current state of async loading
 	*/
-	EAsyncPackageState::Type ProcessAsyncLoadingFromGameThread(FAsyncLoadingThreadState2& ThreadState, int32& OutPackagesProcessed);
+	EAsyncPackageState::Type ProcessAsyncLoadingFromGameThread(FAsyncLoadingThreadState2& ThreadState, bool& bDidSomething);
 
 	/**
 	* [GAME THREAD] Ticks game thread side of async loading.
@@ -2407,7 +2461,7 @@ public:
 	* @param FlushTree Package dependency tree to be flushed
 	* @return The current state of async loading
 	*/
-	EAsyncPackageState::Type TickAsyncLoadingFromGameThread(FAsyncLoadingThreadState2& ThreadState, bool bUseTimeLimit, bool bUseFullTimeLimit, float TimeLimit, int32 FlushRequestID = INDEX_NONE);
+	EAsyncPackageState::Type TickAsyncLoadingFromGameThread(FAsyncLoadingThreadState2& ThreadState, bool bUseTimeLimit, bool bUseFullTimeLimit, float TimeLimit, int32 FlushRequestID, bool& bDidSomething);
 
 	/**
 	* [ASYNC THREAD] Main thread loop
@@ -2526,10 +2580,11 @@ private:
 
 	void RemoveUnreachableObjects(FUnreachableObjects& ObjectsToRemove);
 
-	void ProcessPendingCDOs()
+	bool ProcessPendingCDOs()
 	{
 		TRACE_CPUPROFILER_EVENT_SCOPE(ProcessPendingCDOs);
 
+		bool bDidSomething = false;
 		UClass* Class = nullptr;
 		int32 MaxRequestId = -1;
 		for (TMap<UClass*, TArray<FEventLoadNode2*>>::TIterator It = PendingCDOs.CreateIterator(); It; ++It)
@@ -2585,6 +2640,8 @@ private:
 			{
 				Node->ReleaseBarrier();
 			}
+
+			bDidSomething = true;
 		}
 		else
 		{
@@ -2599,6 +2656,8 @@ private:
 					*CurrentClass->GetFullName(), Nodes.Num());
 			}
 		}
+
+		return bDidSomething;
 	}
 
 	/**
@@ -2612,6 +2671,10 @@ private:
 	*/
 	EAsyncPackageState::Type ProcessLoadedPackagesFromGameThread(FAsyncLoadingThreadState2& ThreadState, bool& bDidSomething, int32 FlushRequestID = INDEX_NONE);
 
+	void IncludePackageInSyncLoadContext(FAsyncLoadingSyncLoadContext& SyncContext, FAsyncPackage2* Package);
+	void DetachPackageFromSyncLoadContext(FAsyncPackage2* Package);
+	FAsyncLoadingSyncLoadContext* UpdateSyncLoadContext(FAsyncLoadingThreadState2& ThreadState);
+
 	bool CreateAsyncPackagesFromQueue(FAsyncLoadingThreadState2& ThreadState);
 
 	FAsyncPackage2* CreateAsyncPackage(const FAsyncPackageDesc2& Desc)
@@ -2732,6 +2795,10 @@ private:
 	{
 		TRACE_CPUPROFILER_EVENT_SCOPE(DeleteAsyncPackage);
 		UE_ASYNC_PACKAGE_DEBUG(Package->Desc);
+		for (int32 RequestId : Package->RequestIDs)
+		{
+			RequestIdToPackageMap.Remove(RequestId);
+		}
 		delete Package;
 		--PackagesWithRemainingWorkCounter;
 	}
@@ -2875,6 +2942,48 @@ FAsyncPackage2* FAsyncLoadingThread2::FindOrInsertPackage(FAsyncPackageDesc2& De
 	return Package;
 }
 
+void FAsyncLoadingThread2::IncludePackageInSyncLoadContext(FAsyncLoadingSyncLoadContext& SyncLoadContext, FAsyncPackage2* Package)
+{
+	if (Package->SyncLoadContext)
+	{
+		if (Package->SyncLoadContext != &SyncLoadContext)
+		{
+			DetachPackageFromSyncLoadContext(Package);
+		}
+		else
+		{
+			return;
+		}
+	}
+	// Assumes that SyncLoadContext.CriticalSection is already locked by the caller
+	Package->SyncLoadContext = &SyncLoadContext;
+	SyncLoadContext.Packages.Add(Package);
+	for (FAsyncPackage2* ImportedPackage : Package->Data.ImportedAsyncPackages)
+	{
+		if (ImportedPackage && ImportedPackage->SyncLoadContext != &SyncLoadContext)
+		{
+			IncludePackageInSyncLoadContext(SyncLoadContext, ImportedPackage);
+		}
+	}
+}
+
+void FAsyncLoadingThread2::DetachPackageFromSyncLoadContext(FAsyncPackage2* Package)
+{
+	FAsyncLoadingSyncLoadContext* SyncLoadContext = Package->SyncLoadContext;
+	check(SyncLoadContext);
+	FScopeLock Lock(&SyncLoadContext->CriticalSection);
+	SyncLoadContext->Packages.Remove(Package);
+	if (SyncLoadContext->RequestedPackageDebug == Package)
+	{
+		SyncLoadContext->RequestedPackageDebug = nullptr;
+	}
+	if (SyncLoadContext->RequestingPackageDebug == Package)
+	{
+		SyncLoadContext->RequestingPackageDebug = nullptr;
+	}
+	Package->SyncLoadContext = nullptr;
+}
+
 bool FAsyncLoadingThread2::CreateAsyncPackagesFromQueue(FAsyncLoadingThreadState2& ThreadState)
 {
 	TRACE_CPUPROFILER_EVENT_SCOPE(CreateAsyncPackagesFromQueue);
@@ -3010,6 +3119,8 @@ bool FAsyncLoadingThread2::CreateAsyncPackagesFromQueue(FAsyncLoadingThreadState
 						TEXT("Package is alreay being loaded."));
 					--PackagesWithRemainingWorkCounter;
 				}
+
+				RequestIdToPackageMap.Add(PackageDesc.RequestID, Package);
 			}
 		}
 		
@@ -3041,12 +3152,10 @@ void FEventLoadNode2::DependsOn(FEventLoadNode2* Other)
 	TRACE_CPUPROFILER_EVENT_SCOPE(DependsOn);
 #if UE_BUILD_DEBUG || UE_BUILD_DEVELOPMENT
 	check(!IsDone());
-	check(!bFired);
 #endif
 	uint8 Expected = 0;
-	while (!Other->DependencyWriterCount.CompareExchange(Expected, 1))
+	while (!Other->DependencyWriterCount.compare_exchange_weak(Expected, 1))
 	{
-		TRACE_CPUPROFILER_EVENT_SCOPE(DependsOnContested);
 		check(Expected == 1);
 		Expected = 0;
 	}
@@ -3084,14 +3193,13 @@ void FEventLoadNode2::DependsOn(FEventLoadNode2* Other)
 			Other->MultipleDependents[Other->DependenciesCount++] = this;
 		}
 	}
-	Other->DependencyWriterCount.Store(0);
+	Other->DependencyWriterCount.store(0, std::memory_order_release);
 }
 
 void FEventLoadNode2::AddBarrier()
 {
 #if UE_BUILD_DEBUG || UE_BUILD_DEVELOPMENT
 	check(!IsDone());
-	check(!bFired);
 #endif
 	++BarrierCount;
 }
@@ -3100,7 +3208,6 @@ void FEventLoadNode2::AddBarrier(int32 Count)
 {
 #if UE_BUILD_DEBUG || UE_BUILD_DEVELOPMENT
 	check(!IsDone());
-	check(!bFired);
 #endif
 	BarrierCount += Count;
 }
@@ -3117,14 +3224,10 @@ void FEventLoadNode2::ReleaseBarrier(FAsyncLoadingThreadState2* ThreadState)
 void FEventLoadNode2::Fire(FAsyncLoadingThreadState2* ThreadState)
 {
 	//TRACE_CPUPROFILER_EVENT_SCOPE(Fire);
-
-#if UE_BUILD_DEBUG || UE_BUILD_DEVELOPMENT
-	bFired.Store(1);
-#endif
-
-	if (Spec->bExecuteImmediately && ThreadState && !ThreadState->CurrentEventNode)
+	if (Spec->bExecuteImmediately && ThreadState)
 	{
-		Execute(*ThreadState);
+		EEventLoadNodeExecutionResult Result = Execute(*ThreadState);
+		check(Result == EEventLoadNodeExecutionResult::Complete);
 	}
 	else
 	{
@@ -3132,47 +3235,33 @@ void FEventLoadNode2::Fire(FAsyncLoadingThreadState2* ThreadState)
 	}
 }
 
-void FEventLoadNode2::Execute(FAsyncLoadingThreadState2& ThreadState)
+EEventLoadNodeExecutionResult FEventLoadNode2::Execute(FAsyncLoadingThreadState2& ThreadState)
 {
 	//TRACE_CPUPROFILER_EVENT_SCOPE(ExecuteEvent);
-	check(BarrierCount.Load() == 0);
-	check(WITH_IOSTORE_IN_EDITOR || !ThreadState.CurrentEventNode || ThreadState.CurrentEventNode == this);
-
-#if WITH_IOSTORE_IN_EDITOR
-	// Allow recursive execution of event nodes in editor builds
-	FEventLoadNode2* PrevNode = ThreadState.CurrentEventNode != this ? ThreadState.CurrentEventNode : nullptr;
-	SetState(ENodeState::Executing);
-#endif
-	ThreadState.CurrentEventNode = this;
-	EAsyncPackageState::Type State = Spec->Func(ThreadState, Package, ImportOrExportIndex);
-	if (State == EAsyncPackageState::Complete)
+	check(BarrierCount.load(std::memory_order_relaxed) == 0);
+	EEventLoadNodeExecutionResult Result;
 	{
-		SetState(ENodeState::Completed);
-		ThreadState.CurrentEventNode = nullptr;
-		ProcessDependencies(ThreadState);
-#if WITH_IOSTORE_IN_EDITOR
-		ThreadState.CurrentEventNode = PrevNode;
-#endif
+		TGuardValue<FEventLoadNode2*> GuardCurrentlyExecutingEventNode(ThreadState.CurrentlyExecutingEventNode, this);
+		Result = Spec->Func(ThreadState, Package, ImportOrExportIndex);
 	}
-#if WITH_IOSTORE_IN_EDITOR
-	else
+	if (Result == EEventLoadNodeExecutionResult::Complete)
 	{
-		check(PrevNode == nullptr);
-		SetState(ENodeState::Timeout);
+		MarkAsDone();
+		ProcessDependencies(ThreadState);
 	}
-#endif
+	return Result;
 }
 
 void FEventLoadNode2::ProcessDependencies(FAsyncLoadingThreadState2& ThreadState)
 {
 	//TRACE_CPUPROFILER_EVENT_SCOPE(ProcessDependencies);
-	if (DependencyWriterCount.Load() != 0)
+	if (DependencyWriterCount.load(std::memory_order_acquire) != 0)
 	{
 		TRACE_CPUPROFILER_EVENT_SCOPE(ConcurrentWriter);
 		do
 		{
 			FPlatformProcess::Sleep(0);
-		} while (DependencyWriterCount.Load() != 0);
+		} while (DependencyWriterCount.load(std::memory_order_acquire) != 0);
 	}
 
 	if (DependenciesCount == 1)
@@ -3224,38 +3313,172 @@ FAsyncLoadEventQueue2::~FAsyncLoadEventQueue2()
 
 void FAsyncLoadEventQueue2::Push(FEventLoadNode2* Node)
 {
-	LLM_SCOPE_BYNAME(TEXT("AsyncLoadEventQueue2"));
-	Queue.enqueue(Node);
-
+	// switch `Tail` to the new node and only then link the old tail to the new one. The list is not fully linked between these ops,
+	// this is explicitly handled by the consumer by waiting for the link
+	FEventLoadNode2* Prev = Tail.exchange(Node, std::memory_order_release); // `release` to make sure the new node is fully constructed before
+	check(Prev->Next.load(std::memory_order_relaxed) == nullptr); // `Tail` is assigned before its Next
+	Prev->Next.store(Node, std::memory_order_relaxed);
 	if (Zenaphore)
 	{
-		Zenaphore->NotifyOne();
+		bool bWasEmpty = Prev == &Sentinel;
+		if (bWasEmpty)
+		{
+			Zenaphore->NotifyOne();
+		}
+	}
+}
+
+void FAsyncLoadEventQueue2::Deplete()
+{
+	FEventLoadNode2* First = Sentinel.Next.load(std::memory_order_relaxed);
+	if (First == nullptr)
+	{
+		return; // empty
 	}
+
+	// reset the head so the next consumption can detect that the queue is empty
+	// `Sentinel.Next` is not touched by producers right now because it's already not null
+	Sentinel.Next.store(nullptr, std::memory_order_relaxed);
+
+	// reset the queue to the empty state. this redirects producers to start from `Sentinel` again.
+	// take note of the tail on resetting it because the list can be still not fully linked and so `Node.Next == nullptr` can't be 
+	// used to detect the end of the list
+	FEventLoadNode2* Last = Tail.exchange(&Sentinel, std::memory_order_acquire); // `acquire` to sync with producers' tail modifications,
+	// `Sentinel.Next = nullptr` must happen before modifying `Tail` but it's important only for the single consumer so doesn't need extra sync
+	check(Last->Next.load(std::memory_order_relaxed) == nullptr); // `Tail` is assigned before its Next
+	// the previously queued items are detached from the instance (as a linked list, though potentially not fully linked yet)
+
+	check(Last != &Sentinel); // can't be empty because of `First != nullptr` above
+
+	if (LocalHead)
+	{
+		check(LocalTail);
+		check(LocalTail->Next.load(std::memory_order_relaxed) == nullptr);
+		LocalTail->Next.store(First, std::memory_order_relaxed);
+	}
+	else
+	{
+		LocalHead = First;
+	}
+	LocalTail = Last;
 }
 
 bool FAsyncLoadEventQueue2::PopAndExecute(FAsyncLoadingThreadState2& ThreadState)
 {
-	if (ThreadState.CurrentEventNode
-#if WITH_IOSTORE_IN_EDITOR
-		&& !ThreadState.CurrentEventNode->IsExecuting()
-#endif
-		)
+	check(!bIsInPopAndExecute);
+	TGuardValue<bool> GuardIsInPopAndExecute(bIsInPopAndExecute, true);
+	if (TimedOutEventNode)
 	{
-		check(!ThreadState.CurrentEventNode->IsDone());
-		ThreadState.CurrentEventNode->Execute(ThreadState);
+		EEventLoadNodeExecutionResult Result = TimedOutEventNode->Execute(ThreadState);
+		if (Result != EEventLoadNodeExecutionResult::Timeout)
+		{
+			TimedOutEventNode = nullptr;
+		}
 		return true;
 	}
 
-	if (FEventLoadNode2* Node = Queue.dequeue())
+	if (!LocalHead)
 	{
-		//TRACE_CPUPROFILER_EVENT_SCOPE(Execute);
-		Node->Execute(ThreadState);
+		Deplete();
+	}
+	if (LocalHead)
+	{
+		FEventLoadNode2* Next = nullptr;
+		if (LocalHead != LocalTail)
+		{
+			// producers can still be updating `Next`, wait until the link to the next element is established
+			do
+			{
+				Next = LocalHead->Next.load(std::memory_order_relaxed);
+			} while (Next == nullptr);
+		}
+		FEventLoadNode2* Node = LocalHead;
+		LocalHead = Next;
+		if (!LocalHead)
+		{
+			LocalTail = nullptr;
+		}
+
+		EEventLoadNodeExecutionResult Result = Node->Execute(ThreadState);
+		if (Result == EEventLoadNodeExecutionResult::Timeout)
+		{
+			TimedOutEventNode = Node;
+		}
 		return true;
 	}
-	else
+	return false;
+}
+
+bool FAsyncLoadEventQueue2::ExecuteSyncLoadEvents(FAsyncLoadingThreadState2& ThreadState, FAsyncLoadingSyncLoadContext& SyncLoadContext)
+{
+	int32 ThisCallCounter = ++ExecuteSyncLoadEventsCallCounter;
+
+	auto ShouldExecuteNode = [&SyncLoadContext](FEventLoadNode2* Node) -> bool
 	{
-		return false;
+		FScopeLock Lock(&SyncLoadContext.CriticalSection);
+		return SyncLoadContext.Packages.Contains(Node->Package);
+	};
+
+	bool bDidSomething = false;
+	if (TimedOutEventNode && ShouldExecuteNode(TimedOutEventNode))
+	{
+		EEventLoadNodeExecutionResult Result = TimedOutEventNode->Execute(ThreadState);
+		check(Result == EEventLoadNodeExecutionResult::Complete); // we can't timeout during a sync load operation
+		TimedOutEventNode = nullptr;
+		bDidSomething = true;
+	}
+
+	Deplete();
+
+	FEventLoadNode2* Node = LocalHead;
+	FEventLoadNode2* Prev = nullptr;
+	while (Node)
+	{
+		FEventLoadNode2* Next = nullptr;
+		if (Node != LocalTail)
+		{
+			// producers can still be updating `Next`, wait until the link to the next element is established
+			do
+			{
+				Next = Node->Next.load(std::memory_order_relaxed);
+			} while (Next == nullptr);
+		}
+
+		if (ShouldExecuteNode(Node))
+		{
+			if (Prev)
+			{
+				Prev->Next.store(Next, std::memory_order_relaxed);
+			}
+			if (Node == LocalHead)
+			{
+				LocalHead = Next;
+				if (!LocalHead)
+				{
+					LocalTail = nullptr;
+				}
+			}
+			if (Node == LocalTail)
+			{
+				LocalTail = Prev;
+			}
+
+			EEventLoadNodeExecutionResult Result = Node->Execute(ThreadState);
+			check(Result == EEventLoadNodeExecutionResult::Complete); // we can't timeout during a sync load operation
+			if (ExecuteSyncLoadEventsCallCounter != ThisCallCounter)
+			{
+				// ExecuteSyncLoadEvents was called recursively and our view of the list might have been compromised, start over
+				return true;
+			}
+			bDidSomething = true;
+		}
+		else
+		{
+			Prev = Node;
+		}
+		Node = Next;
 	}
+	return bDidSomething;
 }
 
 FScopedAsyncPackageEvent2::FScopedAsyncPackageEvent2(FAsyncPackage2* InPackage)
@@ -3623,6 +3846,18 @@ void FAsyncPackage2::ImportPackagesRecursive(FIoBatch& IoBatch, FPackageStore& P
 	}
 #endif
 
+	if (SyncLoadContext)
+	{
+		FScopeLock Lock(&SyncLoadContext->CriticalSection);
+		for (FAsyncPackage2* ImportedPackage : Data.ImportedAsyncPackages)
+		{
+			if (ImportedPackage)
+			{
+				AsyncLoadingThread.IncludePackageInSyncLoadContext(*SyncLoadContext, ImportedPackage);
+			}
+		}
+	}
+
 	UE_ASYNC_PACKAGE_LOG_VERBOSE(VeryVerbose, Desc, TEXT("ImportPackages: ImportsDone"),
 		TEXT("All imported packages are now being loaded."));
 
@@ -3756,7 +3991,7 @@ static void ReadAsyncPackageHeader(FAsyncPackageSerializationState& Serializatio
 	SerializationState.AllExportDataPtr = PackageHeaderDataPtr + PackageSummary->HeaderSize;
 }
 
-EAsyncPackageState::Type FAsyncPackage2::Event_ProcessPackageSummary(FAsyncLoadingThreadState2& ThreadState, FAsyncPackage2* Package, int32)
+EEventLoadNodeExecutionResult FAsyncPackage2::Event_ProcessPackageSummary(FAsyncLoadingThreadState2& ThreadState, FAsyncPackage2* Package, int32)
 {
 	TRACE_CPUPROFILER_EVENT_SCOPE(Event_ProcessPackageSummary);
 	UE_ASYNC_PACKAGE_DEBUG(Package->Desc);
@@ -3857,10 +4092,10 @@ EAsyncPackageState::Type FAsyncPackage2::Event_ProcessPackageSummary(FAsyncLoadi
 
 	Package->AsyncPackageLoadingState = EAsyncPackageLoadingState2::SetupDependencies;
 	Package->GetPackageNode(Package_SetupDependencies).ReleaseBarrier();
-	return EAsyncPackageState::Complete;
+	return EEventLoadNodeExecutionResult::Complete;
 }
 
-EAsyncPackageState::Type FAsyncPackage2::Event_SetupDependencies(FAsyncLoadingThreadState2& ThreadState, FAsyncPackage2* Package, int32)
+EEventLoadNodeExecutionResult FAsyncPackage2::Event_SetupDependencies(FAsyncLoadingThreadState2& ThreadState, FAsyncPackage2* Package, int32)
 {
 	TRACE_CPUPROFILER_EVENT_SCOPE(Event_SetupDependencies);
 	check(Package->AsyncPackageLoadingState == EAsyncPackageLoadingState2::SetupDependencies);
@@ -3886,10 +4121,10 @@ EAsyncPackageState::Type FAsyncPackage2::Event_SetupDependencies(FAsyncLoadingTh
 	{
 		Package->GetExportBundleNode(ExportBundle_Process, ExportBundleIndex).ReleaseBarrier();
 	}
-	return EAsyncPackageState::Complete;
+	return EEventLoadNodeExecutionResult::Complete;
 }
 
-EAsyncPackageState::Type FAsyncPackage2::Event_ProcessExportBundle(FAsyncLoadingThreadState2& ThreadState, FAsyncPackage2* Package, int32 InExportBundleIndex)
+EEventLoadNodeExecutionResult FAsyncPackage2::Event_ProcessExportBundle(FAsyncLoadingThreadState2& ThreadState, FAsyncPackage2* Package, int32 InExportBundleIndex)
 {
 	TRACE_CPUPROFILER_EVENT_SCOPE(Event_ProcessExportBundle);
 	UE_ASYNC_PACKAGE_DEBUG(Package->Desc);
@@ -4000,7 +4235,7 @@ EAsyncPackageState::Type FAsyncPackage2::Event_ProcessExportBundle(FAsyncLoading
 			const FExportBundleEntry& BundleEntry = HeaderData->ExportBundleEntries[ExportBundle->FirstEntryIndex + Package->ExportBundleEntryIndex];
 			if (ThreadState.IsTimeLimitExceeded(TEXT("Event_ProcessExportBundle")))
 			{
-				return EAsyncPackageState::TimeOut;
+				return EEventLoadNodeExecutionResult::Timeout;
 			}
 			const FExportMapEntry& ExportMapEntry = HeaderData->ExportMap[BundleEntry.LocalExportIndex];
 			FExportObject& Export = Exports[BundleEntry.LocalExportIndex];
@@ -4076,7 +4311,7 @@ EAsyncPackageState::Type FAsyncPackage2::Event_ProcessExportBundle(FAsyncLoading
 		}
 	}
 
-	return EAsyncPackageState::Complete;
+	return EEventLoadNodeExecutionResult::Complete;
 }
 
 UObject* FAsyncPackage2::EventDrivenIndexToObject(const FAsyncPackageHeaderData& Header, const TArrayView<const FExportObject>& Exports, FPackageObjectIndex Index, bool bCheckSerialized)
@@ -4419,7 +4654,7 @@ bool FAsyncPackage2::EventDrivenSerializeExport(const FAsyncPackageHeaderData& H
 	return true;
 }
 
-EAsyncPackageState::Type FAsyncPackage2::Event_ExportsDone(FAsyncLoadingThreadState2& ThreadState, FAsyncPackage2* Package, int32)
+EEventLoadNodeExecutionResult FAsyncPackage2::Event_ExportsDone(FAsyncLoadingThreadState2& ThreadState, FAsyncPackage2* Package, int32)
 {
 	TRACE_CPUPROFILER_EVENT_SCOPE(Event_ExportsDone);
 	UE_ASYNC_PACKAGE_DEBUG(Package->Desc);
@@ -4438,17 +4673,20 @@ EAsyncPackageState::Type FAsyncPackage2::Event_ExportsDone(FAsyncLoadingThreadSt
 		FCoreDelegates::ReleasePreloadedPackageShaderMaps.ExecuteIfBound(Package->Data.ShaderMapHashes);
 	}
 
-	Package->ConditionalBeginPostLoad();
-	return EAsyncPackageState::Complete;
+	Package->ConditionalBeginPostLoad(ThreadState);
+	return EEventLoadNodeExecutionResult::Complete;
 }
 
-bool FAsyncPackage2::HaveAllDependenciesReachedStateDebug(FAsyncPackage2* Package, TSet<FAsyncPackage2*>& VisitedPackages, EAsyncPackageLoadingState2 WaitForPackageState)
+bool FAsyncPackage2::HaveAllDependenciesReachedStateDebug(FAsyncLoadingThreadState2& ThreadState, FAsyncPackage2* Package, TSet<FAsyncPackage2*>& VisitedPackages, EAsyncPackageLoadingState2 WaitForPackageState)
 {
 	for (FAsyncPackage2* ImportedPackage : Package->Data.ImportedAsyncPackages)
 	{
 		if (ImportedPackage)
 		{
-
+			if (ThreadState.PackagesExcludedFromStateDependencyChecks.Contains(ImportedPackage))
+			{
+				continue;
+			}
 			if (VisitedPackages.Contains(ImportedPackage))
 			{
 				continue;
@@ -4460,7 +4698,7 @@ bool FAsyncPackage2::HaveAllDependenciesReachedStateDebug(FAsyncPackage2* Packag
 				return false;
 			}
 
-			if (!HaveAllDependenciesReachedStateDebug(ImportedPackage, VisitedPackages, WaitForPackageState))
+			if (!HaveAllDependenciesReachedStateDebug(ThreadState, ImportedPackage, VisitedPackages, WaitForPackageState))
 			{
 				return false;
 			}
@@ -4469,7 +4707,7 @@ bool FAsyncPackage2::HaveAllDependenciesReachedStateDebug(FAsyncPackage2* Packag
 	return true;
 }
 
-bool FAsyncPackage2::HaveAllDependenciesReachedState(FAllDependenciesState FAsyncPackage2::* StateMemberPtr, EAsyncPackageLoadingState2 WaitForPackageState, uint32 CurrentTick)
+bool FAsyncPackage2::HaveAllDependenciesReachedState(FAsyncLoadingThreadState2& ThreadState, FAllDependenciesState FAsyncPackage2::* StateMemberPtr, EAsyncPackageLoadingState2 WaitForPackageState, uint32 CurrentTick)
 {
 	FAllDependenciesState& ThisState = this->*StateMemberPtr;
 	if (ThisState.bAllDone)
@@ -4481,12 +4719,12 @@ bool FAsyncPackage2::HaveAllDependenciesReachedState(FAllDependenciesState FAsyn
 		return false;
 	}
 	ThisState.UpdateTick(CurrentTick);
-	UpdateDependenciesStateRecursive(StateMemberPtr, WaitForPackageState, CurrentTick, this);
+	UpdateDependenciesStateRecursive(ThreadState, StateMemberPtr, WaitForPackageState, CurrentTick, this);
 	check(ThisState.bAllDone || (ThisState.WaitingForPackage && ThisState.WaitingForPackage->AsyncPackageLoadingState <= WaitForPackageState));
 	return ThisState.bAllDone;
 }
 
-void FAsyncPackage2::UpdateDependenciesStateRecursive(FAllDependenciesState FAsyncPackage2::* StateMemberPtr, EAsyncPackageLoadingState2 WaitForPackageState, uint32 CurrentTick, FAsyncPackage2* Root)
+void FAsyncPackage2::UpdateDependenciesStateRecursive(FAsyncLoadingThreadState2& ThreadState, FAllDependenciesState FAsyncPackage2::* StateMemberPtr, EAsyncPackageLoadingState2 WaitForPackageState, uint32 CurrentTick, FAsyncPackage2* Root)
 {
 	FAllDependenciesState& ThisState = this->*StateMemberPtr;
 
@@ -4498,24 +4736,31 @@ void FAsyncPackage2::UpdateDependenciesStateRecursive(FAllDependenciesState FAsy
 
 	if (FAsyncPackage2* WaitingForPackage = ThisState.WaitingForPackage)
 	{
-		FAllDependenciesState& WaitingForPackageState = WaitingForPackage->*StateMemberPtr;
-		if (WaitingForPackage->AsyncPackageLoadingState < WaitForPackageState)
+		if (ThreadState.PackagesExcludedFromStateDependencyChecks.Contains(WaitingForPackage))
 		{
-			ThisState.bAnyNotDone = true;
-			return;
+			ThisState.WaitingForPackage = nullptr;
 		}
-		else if (!WaitingForPackageState.bAllDone)
+		else
 		{
-			WaitingForPackageState.UpdateTick(CurrentTick);
-			if (!WaitingForPackageState.bVisitedMark)
-			{
-				WaitingForPackage->UpdateDependenciesStateRecursive(StateMemberPtr, WaitForPackageState, CurrentTick, Root);
-			}
-			if (WaitingForPackageState.bAnyNotDone)
+			FAllDependenciesState& WaitingForPackageState = WaitingForPackage->*StateMemberPtr;
+			if (WaitingForPackage->AsyncPackageLoadingState < WaitForPackageState)
 			{
 				ThisState.bAnyNotDone = true;
 				return;
 			}
+			else if (!WaitingForPackageState.bAllDone)
+			{
+				WaitingForPackageState.UpdateTick(CurrentTick);
+				if (!WaitingForPackageState.bVisitedMark)
+				{
+					WaitingForPackage->UpdateDependenciesStateRecursive(ThreadState, StateMemberPtr, WaitForPackageState, CurrentTick, Root);
+				}
+				if (WaitingForPackageState.bAnyNotDone)
+				{
+					ThisState.bAnyNotDone = true;
+					return;
+				}
+			}
 		}
 	}
 
@@ -4528,6 +4773,11 @@ void FAsyncPackage2::UpdateDependenciesStateRecursive(FAllDependenciesState FAsy
 			continue;
 		}
 
+		if (ThreadState.PackagesExcludedFromStateDependencyChecks.Contains(ImportedPackage))
+		{
+			continue;
+		}
+
 		FAllDependenciesState& ImportedPackageState = ImportedPackage->*StateMemberPtr;
 
 		if (ImportedPackageState.bAllDone)
@@ -4543,7 +4793,7 @@ void FAsyncPackage2::UpdateDependenciesStateRecursive(FAllDependenciesState FAsy
 		}
 		else if (!ImportedPackageState.bVisitedMark)
 		{
-			ImportedPackage->UpdateDependenciesStateRecursive(StateMemberPtr, WaitForPackageState, CurrentTick, Root);
+			ImportedPackage->UpdateDependenciesStateRecursive(ThreadState, StateMemberPtr, WaitForPackageState, CurrentTick, Root);
 		}
 
 		if (ImportedPackageState.bAnyNotDone)
@@ -4575,9 +4825,9 @@ void FAsyncPackage2::UpdateDependenciesStateRecursive(FAllDependenciesState FAsy
 	}
 }
 
-void FAsyncPackage2::WaitForAllDependenciesToReachState(FAllDependenciesState FAsyncPackage2::* StateMemberPtr, EAsyncPackageLoadingState2 WaitForPackageState, uint32& CurrentTickVariable, TFunctionRef<void(FAsyncPackage2*)> OnStateReached)
+void FAsyncPackage2::WaitForAllDependenciesToReachState(FAsyncLoadingThreadState2& ThreadState, FAllDependenciesState FAsyncPackage2::* StateMemberPtr, EAsyncPackageLoadingState2 WaitForPackageState, uint32& CurrentTickVariable, TFunctionRef<void(FAsyncPackage2*)> OnStateReached)
 {
-	if (HaveAllDependenciesReachedState(StateMemberPtr, WaitForPackageState, CurrentTickVariable++))
+	if (HaveAllDependenciesReachedState(ThreadState, StateMemberPtr, WaitForPackageState, CurrentTickVariable++))
 	{
 		FAsyncPackage2* FirstPackageReadyToProceed = this;
 
@@ -4594,14 +4844,14 @@ void FAsyncPackage2::WaitForAllDependenciesToReachState(FAllDependenciesState FA
 
 #if DO_CHECK
 			TSet<FAsyncPackage2*> VisitedPackages;
-			check(HaveAllDependenciesReachedStateDebug(this, VisitedPackages, WaitForPackageState));
+			check(HaveAllDependenciesReachedStateDebug(ThreadState, this, VisitedPackages, WaitForPackageState));
 #endif
 
 			while (FAsyncPackage2* WaitingPackage = PackageReadyToProceedState.PackagesWaitingForThisHead)
 			{
 				FAllDependenciesState& WaitingPackageState = WaitingPackage->*StateMemberPtr;
 				check(WaitingPackageState.WaitingForPackage == PackageReadyToProceed);
-				if (WaitingPackage->HaveAllDependenciesReachedState(StateMemberPtr, WaitForPackageState, CurrentTickVariable++))
+				if (WaitingPackage->HaveAllDependenciesReachedState(ThreadState, StateMemberPtr, WaitForPackageState, CurrentTickVariable++))
 				{
 					FAllDependenciesState::RemoveFromWaitList(StateMemberPtr, PackageReadyToProceed, WaitingPackage);
 					WaitingPackageState.NextLink = FirstPackageReadyToProceed;
@@ -4616,10 +4866,10 @@ void FAsyncPackage2::WaitForAllDependenciesToReachState(FAllDependenciesState FA
 	}
 }
 
-void FAsyncPackage2::ConditionalBeginPostLoad()
+void FAsyncPackage2::ConditionalBeginPostLoad(FAsyncLoadingThreadState2& ThreadState)
 {
 	TRACE_CPUPROFILER_EVENT_SCOPE(ConditionalBeginPostLoad);
-	WaitForAllDependenciesToReachState(&FAsyncPackage2::AllDependenciesSerializedState, EAsyncPackageLoadingState2::ExportsDone, AsyncLoadingThread.ConditionalBeginPostLoadTick,
+	WaitForAllDependenciesToReachState(ThreadState, &FAsyncPackage2::AllDependenciesSerializedState, EAsyncPackageLoadingState2::ExportsDone, AsyncLoadingThread.ConditionalBeginPostLoadTick,
 		[](FAsyncPackage2* Package)
 		{
 			for (int32 ExportBundleIndex = 0; ExportBundleIndex < Package->Data.ExportBundleCount; ++ExportBundleIndex)
@@ -4629,17 +4879,17 @@ void FAsyncPackage2::ConditionalBeginPostLoad()
 		});
 }
 
-void FAsyncPackage2::ConditionalFinishLoading()
+void FAsyncPackage2::ConditionalFinishLoading(FAsyncLoadingThreadState2& ThreadState)
 {
 	TRACE_CPUPROFILER_EVENT_SCOPE(ConditionalFinishLoading);
-	WaitForAllDependenciesToReachState(&FAsyncPackage2::AllDependenciesFullyLoadedState, EAsyncPackageLoadingState2::DeferredPostLoadDone, AsyncLoadingThread.ConditionalFinishLoadingTick,
+	WaitForAllDependenciesToReachState(ThreadState, &FAsyncPackage2::AllDependenciesFullyLoadedState, EAsyncPackageLoadingState2::DeferredPostLoadDone, AsyncLoadingThread.ConditionalFinishLoadingTick,
 		[](FAsyncPackage2* Package)
 		{
 			Package->AsyncLoadingThread.LoadedPackagesToProcess.Add(Package);
 		});
 }
 
-EAsyncPackageState::Type FAsyncPackage2::Event_PostLoadExportBundle(FAsyncLoadingThreadState2& ThreadState, FAsyncPackage2* Package, int32 InExportBundleIndex)
+EEventLoadNodeExecutionResult FAsyncPackage2::Event_PostLoadExportBundle(FAsyncLoadingThreadState2& ThreadState, FAsyncPackage2* Package, int32 InExportBundleIndex)
 {
 	TRACE_CPUPROFILER_EVENT_SCOPE(Event_PostLoad);
 	UE_ASYNC_PACKAGE_DEBUG(Package->Desc);
@@ -4671,7 +4921,7 @@ EAsyncPackageState::Type FAsyncPackage2::Event_PostLoadExportBundle(FAsyncLoadin
 	
 	check(InExportBundleIndex < Package->Data.ExportBundleCount);
 
-	EAsyncPackageState::Type LoadingState = EAsyncPackageState::Complete;
+	EEventLoadNodeExecutionResult LoadingState = EEventLoadNodeExecutionResult::Complete;
 
 	if (!Package->bLoadHasFailed)
 	{
@@ -4714,7 +4964,7 @@ EAsyncPackageState::Type FAsyncPackage2::Event_PostLoadExportBundle(FAsyncLoadin
 			const FExportBundleEntry& BundleEntry = HeaderData->ExportBundleEntries[ExportBundle->FirstEntryIndex + Package->ExportBundleEntryIndex];
 			if (ThreadState.IsTimeLimitExceeded(TEXT("Event_PostLoadExportBundle")))
 			{
-				LoadingState = EAsyncPackageState::TimeOut;
+				LoadingState = EEventLoadNodeExecutionResult::Timeout;
 				break;
 			}
 			
@@ -4755,7 +5005,7 @@ EAsyncPackageState::Type FAsyncPackage2::Event_PostLoadExportBundle(FAsyncLoadin
 		Package->EndAsyncLoad();
 	}
 	
-	if (LoadingState == EAsyncPackageState::TimeOut)
+	if (LoadingState == EEventLoadNodeExecutionResult::Timeout)
 	{
 		return LoadingState;
 	}
@@ -4782,10 +5032,10 @@ EAsyncPackageState::Type FAsyncPackage2::Event_PostLoadExportBundle(FAsyncLoadin
 		}
 	}
 
-	return EAsyncPackageState::Complete;
+	return EEventLoadNodeExecutionResult::Complete;
 }
 
-EAsyncPackageState::Type FAsyncPackage2::Event_DeferredPostLoadExportBundle(FAsyncLoadingThreadState2& ThreadState, FAsyncPackage2* Package, int32 InExportBundleIndex)
+EEventLoadNodeExecutionResult FAsyncPackage2::Event_DeferredPostLoadExportBundle(FAsyncLoadingThreadState2& ThreadState, FAsyncPackage2* Package, int32 InExportBundleIndex)
 {
 	SCOPE_CYCLE_COUNTER(STAT_FAsyncPackage_PostLoadObjectsGameThread);
 	TRACE_CPUPROFILER_EVENT_SCOPE(Event_DeferredPostLoad);
@@ -4795,7 +5045,7 @@ EAsyncPackageState::Type FAsyncPackage2::Event_DeferredPostLoadExportBundle(FAsy
 	FAsyncPackageScope2 PackageScope(Package);
 
 	check(InExportBundleIndex < Package->Data.ExportBundleCount);
-	EAsyncPackageState::Type LoadingState = EAsyncPackageState::Complete;
+	EEventLoadNodeExecutionResult LoadingState = EEventLoadNodeExecutionResult::Complete;
 
 	if (Package->bLoadHasFailed)
 	{
@@ -4835,7 +5085,7 @@ EAsyncPackageState::Type FAsyncPackage2::Event_DeferredPostLoadExportBundle(FAsy
 			const FExportBundleEntry& BundleEntry = HeaderData->ExportBundleEntries[ExportBundle->FirstEntryIndex + Package->ExportBundleEntryIndex];
 			if (ThreadState.IsTimeLimitExceeded(TEXT("Event_DeferredPostLoadExportBundle")))
 			{
-				LoadingState = EAsyncPackageState::TimeOut;
+				LoadingState = EEventLoadNodeExecutionResult::Timeout;
 				break;
 			}
 
@@ -4868,7 +5118,7 @@ EAsyncPackageState::Type FAsyncPackage2::Event_DeferredPostLoadExportBundle(FAsy
 		}
 	}
 
-	if (LoadingState == EAsyncPackageState::TimeOut)
+	if (LoadingState == EEventLoadNodeExecutionResult::Timeout)
 	{
 		return LoadingState;
 	}
@@ -4880,10 +5130,10 @@ EAsyncPackageState::Type FAsyncPackage2::Event_DeferredPostLoadExportBundle(FAsy
 		Package->ProcessedExportBundlesCount = 0;
 		check(Package->AsyncPackageLoadingState == EAsyncPackageLoadingState2::DeferredPostLoad);
 		Package->AsyncPackageLoadingState = EAsyncPackageLoadingState2::DeferredPostLoadDone;
-		Package->ConditionalFinishLoading();
+		Package->ConditionalFinishLoading(ThreadState);
 	}
 
-	return EAsyncPackageState::Complete;
+	return EEventLoadNodeExecutionResult::Complete;
 }
 
 FEventLoadNode2& FAsyncPackage2::GetPackageNode(EEventLoadNode2 Phase)
@@ -4899,36 +5149,50 @@ FEventLoadNode2& FAsyncPackage2::GetExportBundleNode(EEventLoadNode2 Phase, uint
 	return Data.ExportBundleNodes[ExportBundleNodeIndex];
 }
 
-#if ALT2_VERIFY_RECURSIVE_LOADS 
-struct FScopedLoadRecursionVerifier
+FAsyncLoadingSyncLoadContext* FAsyncLoadingThread2::UpdateSyncLoadContext(FAsyncLoadingThreadState2& ThreadState)
 {
-	int32& Level;
-	FScopedLoadRecursionVerifier(int32& InLevel) : Level(InLevel)
+	if (ThreadState.bIsAsyncLoadingThread)
 	{
-		UE_CLOG(Level > 0, LogStreaming, Error, TEXT("Entering recursive load level: %d"), Level);
-		++Level;
-		check(Level == 1);
+		FAsyncLoadingSyncLoadContext* CreatedOnMainThread;
+		while (ThreadState.SyncLoadContextsCreatedOnGameThread.Dequeue(CreatedOnMainThread))
+		{
+			ThreadState.SyncLoadContextStack.Push(CreatedOnMainThread);
+		}
 	}
-	~FScopedLoadRecursionVerifier()
+	if (ThreadState.SyncLoadContextStack.IsEmpty())
 	{
-		--Level;
-		UE_CLOG(Level > 0, LogStreaming, Error, TEXT("Leaving recursive load level: %d"), Level);
-		check(Level == 0);
+		return nullptr;
 	}
-};
-#endif
+	FAsyncLoadingSyncLoadContext* SyncLoadContext = ThreadState.SyncLoadContextStack.Top();
+	if (ThreadState.bIsAsyncLoadingThread && !ContainsRequestID(SyncLoadContext->RequestId))
+	{
+		SyncLoadContext->ReleaseRef();
+		ThreadState.SyncLoadContextStack.Pop();
+		if (ThreadState.SyncLoadContextStack.IsEmpty())
+		{
+			return nullptr;
+		}
+		SyncLoadContext = ThreadState.SyncLoadContextStack.Top();
+	}
+	if (ThreadState.bCanAccessAsyncLoadingThreadData && !SyncLoadContext->bHasFoundRequestedPackage)
+	{
+		if (FAsyncPackage2* RequestedPackage = RequestIdToPackageMap.FindRef(SyncLoadContext->RequestId))
+		{
+			FScopeLock Lock(&SyncLoadContext->CriticalSection);
+			SyncLoadContext->bHasFoundRequestedPackage = true;
+			SyncLoadContext->RequestedPackageDebug = RequestedPackage;
+			IncludePackageInSyncLoadContext(*SyncLoadContext, RequestedPackage);
+		}
+	}
+	return SyncLoadContext;
+}
 
-EAsyncPackageState::Type FAsyncLoadingThread2::ProcessAsyncLoadingFromGameThread(FAsyncLoadingThreadState2& ThreadState, int32& OutPackagesProcessed)
+EAsyncPackageState::Type FAsyncLoadingThread2::ProcessAsyncLoadingFromGameThread(FAsyncLoadingThreadState2& ThreadState, bool& bDidSomething)
 {
 	SCOPED_LOADTIMER(AsyncLoadingTime);
 
 	check(IsInGameThread());
 
-	OutPackagesProcessed = 0;
-
-#if ALT2_VERIFY_RECURSIVE_LOADS 
-	FScopedLoadRecursionVerifier LoadRecursionVerifier(this->LoadRecursionLevel);
-#endif
 	FAsyncLoadingTickScope2 InAsyncLoadingTick(*this);
 	uint32 LoopIterations = 0;
 
@@ -4959,7 +5223,7 @@ EAsyncPackageState::Type FAsyncLoadingThread2::ProcessAsyncLoadingFromGameThread
 			{
 				if (CreateAsyncPackagesFromQueue(ThreadState))
 				{
-					OutPackagesProcessed++;
+					bDidSomething = true;
 					break;
 				}
 				else
@@ -4968,18 +5232,17 @@ EAsyncPackageState::Type FAsyncLoadingThread2::ProcessAsyncLoadingFromGameThread
 				}
 			}
 
-			bool bPopped = false;
-			for (FAsyncLoadEventQueue2* Queue : AltEventQueues)
+			if (FAsyncLoadingSyncLoadContext* SyncLoadContext = UpdateSyncLoadContext(ThreadState))
 			{
-				if (Queue->PopAndExecute(ThreadState))
+				if (EventQueue.ExecuteSyncLoadEvents(ThreadState, *SyncLoadContext))
 				{
-					bPopped = true;
+					bDidSomething = true;
 					break;
 				}
 			}
-			if (bPopped)
+			else if (EventQueue.PopAndExecute(ThreadState))
 			{
-				OutPackagesProcessed++;
+				bDidSomething = true;
 				break;
 			}
 
@@ -4993,7 +5256,7 @@ EAsyncPackageState::Type FAsyncLoadingThread2::ProcessAsyncLoadingFromGameThread
 				EAsyncPackageState::Type Result = Package->ProcessExternalReads(FAsyncPackage2::ExternalReadAction_Wait);
 				check(Result == EAsyncPackageState::Complete);
 
-				OutPackagesProcessed++;
+				bDidSomething = true;
 				break;
 			}
 
@@ -5004,7 +5267,7 @@ EAsyncPackageState::Type FAsyncLoadingThread2::ProcessAsyncLoadingFromGameThread
 				FAsyncPackage2* Package = nullptr;
 				DeferredDeletePackages.Dequeue(Package);
 				DeleteAsyncPackage(Package);
-				OutPackagesProcessed++;
+				bDidSomething = true;
 				break;
 			}
 
@@ -5045,11 +5308,18 @@ EAsyncPackageState::Type FAsyncLoadingThread2::ProcessLoadedPackagesFromGameThre
 		}
 
 		bool bLocalDidSomething = false;
-		bLocalDidSomething |= MainThreadEventQueue.PopAndExecute(ThreadState);
+		if (FAsyncLoadingSyncLoadContext* SyncLoadContext = UpdateSyncLoadContext(ThreadState))
+		{
+			bLocalDidSomething |= MainThreadEventQueue.ExecuteSyncLoadEvents(ThreadState, *SyncLoadContext);
+		}
+		else
+		{
+			bLocalDidSomething |= MainThreadEventQueue.PopAndExecute(ThreadState);
+		}
 
 		bLocalDidSomething |= LoadedPackagesToProcess.Num() > 0;
 		TArray<FAsyncPackage2*, TInlineAllocator<4>> PackagesReadyForCallback;
-#if WITH_IOSTORE_IN_EDITOR
+#if WITH_EDITOR
 		TSet<UPackage*> CompletedUPackages;
 #endif
 		for (int32 PackageIndex = 0; PackageIndex < LoadedPackagesToProcess.Num(); ++PackageIndex)
@@ -5060,6 +5330,11 @@ EAsyncPackageState::Type FAsyncLoadingThread2::ProcessLoadedPackagesFromGameThre
 			check(Package->AsyncPackageLoadingState >= EAsyncPackageLoadingState2::Finalize &&
 				  Package->AsyncPackageLoadingState <= EAsyncPackageLoadingState2::CreateClusters);
 
+			if (Package->SyncLoadContext)
+			{
+				DetachPackageFromSyncLoadContext(Package);
+			}
+
 			if (Package->AsyncPackageLoadingState == EAsyncPackageLoadingState2::Finalize)
 			{
 				TArray<UObject*> CDODefaultSubobjects;
@@ -5148,7 +5423,7 @@ EAsyncPackageState::Type FAsyncLoadingThread2::ProcessLoadedPackagesFromGameThre
 				AsyncPackageLookup.Remove(Package->Desc.UPackageId);
 				if (!Package->bLoadHasFailed)
 				{
-#if WITH_IOSTORE_IN_EDITOR
+#if WITH_EDITOR
 					// In the editor we need to find any assets and packages and add them to list for later callback
 					Package->GetLoadedAssetsAndPackages(LoadedAssets, CompletedUPackages);
 #endif
@@ -5182,12 +5457,12 @@ EAsyncPackageState::Type FAsyncLoadingThread2::ProcessLoadedPackagesFromGameThre
 
 			check(!CompletedPackages.Contains(Package));
 			CompletedPackages.Add(Package);
-			Package->MarkRequestIDsAsComplete();
+			RemovePendingRequests(Package->RequestIDs);
 
 			UE_ASYNC_PACKAGE_LOG(Verbose, Package->Desc, TEXT("GameThread: LoadCompleted"),
 				TEXT("All loading of package is done, and the async package and load request will be deleted."));
 		}
-#if WITH_IOSTORE_IN_EDITOR
+#if WITH_EDITOR
 		// Call the global delegate for package endloads and set the bHasBeenLoaded flag that is used to
 		// check which packages have reached this state
 		for (UPackage* CompletedUPackage : CompletedUPackages)
@@ -5271,7 +5546,7 @@ EAsyncPackageState::Type FAsyncLoadingThread2::ProcessLoadedPackagesFromGameThre
 	return Result;
 }
 
-EAsyncPackageState::Type FAsyncLoadingThread2::TickAsyncLoadingFromGameThread(FAsyncLoadingThreadState2& ThreadState, bool bUseTimeLimit, bool bUseFullTimeLimit, float TimeLimit, int32 FlushRequestID)
+EAsyncPackageState::Type FAsyncLoadingThread2::TickAsyncLoadingFromGameThread(FAsyncLoadingThreadState2& ThreadState, bool bUseTimeLimit, bool bUseFullTimeLimit, float TimeLimit, int32 FlushRequestID, bool& bDidSomething)
 {
 	SCOPE_CYCLE_COUNTER(STAT_FAsyncPackage_TickAsyncLoadingGameThread);
 	//TRACE_INT_VALUE(QueuedPackagesCounter, QueuedPackagesCounter);
@@ -5293,7 +5568,6 @@ EAsyncPackageState::Type FAsyncLoadingThread2::TickAsyncLoadingFromGameThread(FA
 		const bool bIsMultithreaded = FAsyncLoadingThread2::IsMultithreaded();
 		double TickStartTime = FPlatformTime::Seconds();
 
-		bool bDidSomething = false;
 		{
 			Result = ProcessLoadedPackagesFromGameThread(ThreadState, bDidSomething, FlushRequestID);
 			double TimeLimitUsedForProcessLoaded = FPlatformTime::Seconds() - TickStartTime;
@@ -5309,7 +5583,7 @@ EAsyncPackageState::Type FAsyncLoadingThread2::TickAsyncLoadingFromGameThread(FA
 		{
 			if (!bDidSomething && PendingCDOs.Num() > 0)
 			{
-				ProcessPendingCDOs();
+				bDidSomething = ProcessPendingCDOs();
 			}
 
 			// Flush deferred messages
@@ -5333,16 +5607,12 @@ FAsyncLoadingThread2::FAsyncLoadingThread2(FIoDispatcher& InIoDispatcher, IAsync
 	, PackageStore(FPackageStore::Get())
 	, GlobalImportStore(LoadedPackageStore)
 {
-#if !WITH_IOSTORE_IN_EDITOR
+#if !WITH_EDITOR
 	IsEventDrivenLoaderEnabled(); // make sure the one time init inside runs
 	GEventDrivenLoaderEnabled = true;
 #endif
 
-	AltEventQueues.Add(&EventQueue);
-	for (FAsyncLoadEventQueue2* Queue : AltEventQueues)
-	{
-		Queue->SetZenaphore(&AltZenaphore);
-	}
+	EventQueue.SetZenaphore(&AltZenaphore);
 
 	EventSpecs.AddDefaulted(EEventLoadNode2::Package_NumPhases + EEventLoadNode2::ExportBundle_NumPhases);
 	EventSpecs[EEventLoadNode2::Package_ProcessSummary] = { &FAsyncPackage2::Event_ProcessPackageSummary, &EventQueue, false };
@@ -5361,8 +5631,9 @@ FAsyncLoadingThread2::FAsyncLoadingThread2(FIoDispatcher& InIoDispatcher, IAsync
 	FCoreUObjectInternalDelegates::GetOnLeakedPackageRenameDelegate().AddRaw(this, &FAsyncLoadingThread2::OnLeakedPackageRename);
 
 	FAsyncLoadingThreadState2::TlsSlot = FPlatformTLS::AllocTlsSlot();
-	FAsyncLoadingThreadState2::Create(GraphAllocator, IoDispatcher);
-
+	GameThreadState = MakeUnique<FAsyncLoadingThreadState2>(GraphAllocator, IoDispatcher);
+	FAsyncLoadingThreadState2::Set(GameThreadState.Get());
+	
 	UE_LOG(LogStreaming, Display, TEXT("AsyncLoading2 - Created: Event Driven Loader: %s, Async Loading Thread: %s, Async Post Load: %s"),
 		GEventDrivenLoaderEnabled ? TEXT("true") : TEXT("false"),
 		FAsyncLoadingThreadSettings::Get().bAsyncLoadingThreadEnabled ? TEXT("true") : TEXT("false"),
@@ -5406,6 +5677,10 @@ void FAsyncLoadingThread2::StartThread()
 
 	if (FAsyncLoadingThreadSettings::Get().bAsyncLoadingThreadEnabled && !Thread)
 	{
+		AsyncLoadingThreadState = MakeUnique<FAsyncLoadingThreadState2>(GraphAllocator, IoDispatcher);
+		AsyncLoadingThreadState->bIsAsyncLoadingThread = true;
+		AsyncLoadingThreadState->bCanAccessAsyncLoadingThreadData = true;
+		GameThreadState->bCanAccessAsyncLoadingThreadData = false;
 		UE_LOG(LogStreaming, Log, TEXT("Starting Async Loading Thread."));
 		bThreadStarted = true;
 		FPlatformMisc::MemoryBarrier();
@@ -5430,7 +5705,7 @@ uint32 FAsyncLoadingThread2::Run()
 
 	AsyncLoadingThreadID = FPlatformTLS::GetCurrentThreadId();
 
-	FAsyncLoadingThreadState2::Create(GraphAllocator, IoDispatcher);
+	FAsyncLoadingThreadState2::Set(AsyncLoadingThreadState.Get());
 
 	TRACE_LOADTIME_START_ASYNC_LOADING();
 
@@ -5528,20 +5803,25 @@ uint32 FAsyncLoadingThread2::Run()
 						do 
 						{
 							bPopped = false;
-							for (FAsyncLoadEventQueue2* Queue : AltEventQueues)
+							if (FAsyncLoadingSyncLoadContext* SyncLoadContext = UpdateSyncLoadContext(ThreadState))
 							{
-								if (Queue->PopAndExecute(ThreadState))
+								if (EventQueue.ExecuteSyncLoadEvents(ThreadState, *SyncLoadContext))
 								{
 									bPopped = true;
 									bDidSomething = true;
 								}
-								if (bSuspendRequested.Load(EMemoryOrder::Relaxed) || IsGarbageCollectionWaiting())
-								{
-									bShouldSuspend = true;
-									bDidSomething = true;
-									bPopped = false;
-									break;
-								}
+							}
+							else if (EventQueue.PopAndExecute(ThreadState))
+							{
+								bPopped = true;
+								bDidSomething = true;
+							}
+							if (bSuspendRequested.Load(EMemoryOrder::Relaxed) || IsGarbageCollectionWaiting())
+							{
+								bShouldSuspend = true;
+								bDidSomething = true;
+								bPopped = false;
+								break;
 							}
 						} while (bPopped);
 
@@ -5611,7 +5891,6 @@ EAsyncPackageState::Type FAsyncLoadingThread2::TickAsyncThreadFromGameThread(FAs
 	check(IsInGameThread());
 	EAsyncPackageState::Type Result = EAsyncPackageState::Complete;
 	
-	int32 ProcessedRequests = 0;
 	if (AsyncThreadReady.GetValue())
 	{
 		if (ThreadState.IsTimeLimitExceeded(TEXT("TickAsyncThreadFromGameThread")))
@@ -5620,8 +5899,7 @@ EAsyncPackageState::Type FAsyncLoadingThread2::TickAsyncThreadFromGameThread(FAs
 		}
 		else
 		{
-			Result = ProcessAsyncLoadingFromGameThread(ThreadState, ProcessedRequests);
-			bDidSomething = bDidSomething || ProcessedRequests > 0;
+			Result = ProcessAsyncLoadingFromGameThread(ThreadState, bDidSomething);
 		}
 	}
 
@@ -6009,9 +6287,6 @@ FAsyncPackage2::~FAsyncPackage2()
 	checkf(RefCount == 0, TEXT("RefCount is not 0 when deleting package %s"),
 		*Desc.PackagePathToLoad.GetPackageFName().ToString());
 
-	checkf(RequestIDs.Num() == 0, TEXT("MarkRequestIDsAsComplete() has not been called for package %s"),
-		*Desc.PackagePathToLoad.GetPackageFName().ToString());
-	
 	checkf(ConstructedObjects.Num() == 0, TEXT("ClearConstructedObjects() has not been called for package %s"),
 		*Desc.PackagePathToLoad.GetPackageFName().ToString());
 
@@ -6090,12 +6365,6 @@ void FAsyncPackage2::AddRequestID(int32 Id)
 	}
 }
 
-void FAsyncPackage2::MarkRequestIDsAsComplete()
-{
-	AsyncLoadingThread.RemovePendingRequests(RequestIDs);
-	RequestIDs.Reset();
-}
-
 /**
  * @return Time load begun. This is NOT the time the load was requested in the case of other pending requests.
  */
@@ -6104,7 +6373,7 @@ double FAsyncPackage2::GetLoadStartTime() const
 	return LoadStartTime;
 }
 
-#if WITH_IOSTORE_IN_EDITOR
+#if WITH_EDITOR
 void FAsyncPackage2::GetLoadedAssetsAndPackages(TSet<FWeakObjectPtr>& AssetList, TSet<UPackage*>& PackageList)
 {
 	for (UObject* Object : ConstructedObjects)
@@ -6461,7 +6730,8 @@ EAsyncPackageState::Type FAsyncLoadingThread2::ProcessLoadingFromGameThread(FAsy
 	CSV_CUSTOM_STAT(FileIO, QueuedPackagesQueueDepth, GetNumQueuedPackages(), ECsvCustomStatOp::Set);
 	CSV_CUSTOM_STAT(FileIO, ExistingQueuedPackagesQueueDepth, GetNumAsyncPackages(), ECsvCustomStatOp::Set);
 	
-	TickAsyncLoadingFromGameThread(ThreadState, bUseTimeLimit, bUseFullTimeLimit, TimeLimit);
+	bool bDidSomething = false;
+	TickAsyncLoadingFromGameThread(ThreadState, bUseTimeLimit, bUseFullTimeLimit, TimeLimit, INDEX_NONE, bDidSomething);
 	return IsAsyncLoadingPackages() ? EAsyncPackageState::TimeOut : EAsyncPackageState::Complete;
 }
 
@@ -6484,17 +6754,49 @@ void FAsyncLoadingThread2::FlushLoading(int32 RequestId)
 		double StartTime = FPlatformTime::Seconds();
 		double LogFlushTime = StartTime;
 
+		FAsyncPackage2* CurrentlyExecutingPackage = nullptr;
+		if (GameThreadState->CurrentlyExecutingEventNode)
+		{
+			UE_CLOG(RequestId == INDEX_NONE, LogStreaming, Fatal, TEXT("Flushing async loading while creating, serializing or postloading an object is not permitted"));
+			CurrentlyExecutingPackage = GameThreadState->CurrentlyExecutingEventNode->GetPackage();
+			GameThreadState->PackagesExcludedFromStateDependencyChecks.Push(CurrentlyExecutingPackage);
+		}
+
+		FAsyncLoadingSyncLoadContext* SyncLoadContext = nullptr;
+		if (RequestId != INDEX_NONE)
+		{
+			SyncLoadContext = new FAsyncLoadingSyncLoadContext(RequestId);
+			SyncLoadContext->RequestingPackageDebug = CurrentlyExecutingPackage;
+			GameThreadState->SyncLoadContextStack.Push(SyncLoadContext);
+			if (AsyncLoadingThreadState)
+			{
+				SyncLoadContext->AddRef();
+				AsyncLoadingThreadState->SyncLoadContextsCreatedOnGameThread.Enqueue(SyncLoadContext);
+				AltZenaphore.NotifyOne();
+			}
+		}
+
 		// Flush async loaders by not using a time limit. Needed for e.g. garbage collection.
 		{
-			FAsyncLoadingThreadState2& ThreadState = *FAsyncLoadingThreadState2::Get();
 			while (IsAsyncLoadingPackages())
 			{
-				EAsyncPackageState::Type Result = TickAsyncLoadingFromGameThread(ThreadState, false, false, 0, RequestId);
+				bool bDidSomething = false;
+				EAsyncPackageState::Type Result = TickAsyncLoadingFromGameThread(*GameThreadState, false, false, 0, RequestId, bDidSomething);
 				if (RequestId != INDEX_NONE && !ContainsRequestID(RequestId))
 				{
 					break;
 				}
 
+#if 0
+				if (!bDidSomething)
+				{
+					if (!AsyncLoadingThreadState && !PendingIoRequestsCounter)
+					{
+						UE_LOG(LogStreaming, Fatal, TEXT("Loading is stuck, flush will never finish"));
+					}
+				}
+#endif
+
 				if (IsMultithreaded())
 				{
 					// Update the heartbeat and sleep. If we're not multithreading, the heartbeat is updated after each package has been processed
@@ -6514,6 +6816,19 @@ void FAsyncLoadingThread2::FlushLoading(int32 RequestId)
 			}
 		}
 
+		if (SyncLoadContext)
+		{
+			check(GameThreadState->SyncLoadContextStack.Top() == SyncLoadContext);
+			SyncLoadContext->ReleaseRef();
+			GameThreadState->SyncLoadContextStack.Pop();
+		}
+
+		if (CurrentlyExecutingPackage)
+		{
+			check(GameThreadState->PackagesExcludedFromStateDependencyChecks.Top() == CurrentlyExecutingPackage);
+			GameThreadState->PackagesExcludedFromStateDependencyChecks.Pop();
+		}
+
 		check(RequestId != INDEX_NONE || !IsAsyncLoadingPackages());
 	}
 }
@@ -6534,7 +6849,8 @@ EAsyncPackageState::Type FAsyncLoadingThread2::ProcessLoadingUntilCompleteFromGa
 	bool bUseTimeLimit = TimeLimit > 0.0f;
 	double TimeLoadingPackage = 0.0f;
 
-	while (IsAsyncLoadingPackages() && (!bUseTimeLimit || TimeLimit > 0.0f) && !CompletionPredicate())
+	bool bLoadingComplete = !IsAsyncLoadingPackages() || CompletionPredicate();
+	while (!bLoadingComplete && (!bUseTimeLimit || TimeLimit > 0.0f))
 	{
 		double TickStartTime = FPlatformTime::Seconds();
 		if (ProcessLoadingFromGameThread(ThreadState, bUseTimeLimit, bUseTimeLimit, TimeLimit) == EAsyncPackageState::Complete)
@@ -6556,9 +6872,11 @@ EAsyncPackageState::Type FAsyncLoadingThread2::ProcessLoadingUntilCompleteFromGa
 		double TimeDelta = (FPlatformTime::Seconds() - TickStartTime);
 		TimeLimit -= TimeDelta;
 		TimeLoadingPackage += TimeDelta;
+
+		bLoadingComplete = !IsAsyncLoadingPackages() || CompletionPredicate();
 	}
 
-	return TimeLimit <= 0 ? EAsyncPackageState::TimeOut : EAsyncPackageState::Complete;
+	return bLoadingComplete ? EAsyncPackageState::Complete : EAsyncPackageState::TimeOut;
 }
 
 IAsyncPackageLoader* MakeAsyncPackageLoader2(FIoDispatcher& InIoDispatcher, IAsyncPackageLoader* InUncookedPackageLoader)
diff --git a/Engine/Source/Runtime/CoreUObject/Private/UObject/UObjectGlobals.cpp b/Engine/Source/Runtime/CoreUObject/Private/UObject/UObjectGlobals.cpp
--- a/Engine/Source/Runtime/CoreUObject/Private/UObject/UObjectGlobals.cpp
+++ b/Engine/Source/Runtime/CoreUObject/Private/UObject/UObjectGlobals.cpp
@@ -59,9 +59,7 @@
 #include "HAL/LowLevelMemStats.h"
 #include "Misc/CoreDelegates.h"
 #include "ProfilingDebugging/CsvProfiler.h"
-#if WITH_IOSTORE_IN_EDITOR
 #include "IO/IoDispatcher.h"
-#endif
 #include "ProfilingDebugging/LoadTimeTracker.h"
 #include "Misc/PackageAccessTracking.h"
 #include "Misc/PackageAccessTracking.h"
@@ -1459,7 +1457,10 @@ UPackage* LoadPackageInternal(UPackage* InOuter, const FPackagePath& PackagePath
 		}
 
 		GSyncLoadUsingAsyncLoaderCount++;
-		int32 RequestID = LoadPackageAsync(PackagePath, PackageName);
+		constexpr EPackageFlags PackageFlags = PKG_None;
+		constexpr int32 PIEInstanceID = INDEX_NONE;
+		constexpr int32 Priority = INT32_MAX;
+		int32 RequestID = LoadPackageAsync(PackagePath, PackageName, FLoadPackageAsyncDelegate(), PackageFlags, PIEInstanceID, Priority, InstancingContext);
 
 		if (RequestID != INDEX_NONE)
 		{
