Fixed minor issue with LiveCoding/HotReload where the first character of the old class name would be removed.

diff --git a/Engine/Source/Runtime/CoreUObject/Private/UObject/DeferredRegistry.cpp b/Engine/Source/Runtime/CoreUObject/Private/UObject/DeferredRegistry.cpp
--- a/Engine/Source/Runtime/CoreUObject/Private/UObject/DeferredRegistry.cpp
+++ b/Engine/Source/Runtime/CoreUObject/Private/UObject/DeferredRegistry.cpp
@@ -28,8 +28,6 @@ void ReloadProcessObject(UEnum* Enum, const TCHAR* RenamePrefix)
 
 void ReloadProcessObject(UClass* Class, const TCHAR* RenamePrefix)
 {
-	FString NameWithoutPrefix = UObjectBase::RemoveClassPrefix(*Class->GetName());
-
 	// Make sure enums de-register their names BEFORE we create the new class, otherwise there will be name conflicts.
 	// This needs to be done before renaming the class (moving it to transient package) as enum names are stored with the package name
 	TArray<UObject*> ClassSubobjects;
@@ -47,7 +45,7 @@ void ReloadProcessObject(UClass* Class, const TCHAR* RenamePrefix)
 	Class->ClearFlags(RF_Standalone | RF_Public);
 	Class->GetDefaultObject()->RemoveFromRoot();
 	Class->GetDefaultObject()->ClearFlags(RF_Standalone | RF_Public);
-	const FName OldClassRename = MakeUniqueObjectName(GetTransientPackage(), Class->GetClass(), *FString::Printf(TEXT("%s_%s"), RenamePrefix, *NameWithoutPrefix));
+	const FName OldClassRename = MakeUniqueObjectName(GetTransientPackage(), Class->GetClass(), *FString::Printf(TEXT("%s_%s"), RenamePrefix, *Class->GetName()));
 	Class->Rename(*OldClassRename.ToString(), GetTransientPackage());
 	Class->SetFlags(RF_Transient);
 	Class->AddToRoot();
