[Backout] - CL21736945. Probable cause of crash in autotest
Original CL Desc
-----------------------------------------------------------------
ZenLoader: When sync loading a package only execute nodes required for that package to finish loading

diff --git a/Engine/Source/Runtime/CoreUObject/Private/Serialization/AsyncLoading2.cpp b/Engine/Source/Runtime/CoreUObject/Private/Serialization/AsyncLoading2.cpp
--- a/Engine/Source/Runtime/CoreUObject/Private/Serialization/AsyncLoading2.cpp
+++ b/Engine/Source/Runtime/CoreUObject/Private/Serialization/AsyncLoading2.cpp
@@ -50,6 +50,7 @@
 #include "Serialization/UnversionedPropertySerialization.h"
 #include "Serialization/Zenaphore.h"
 #include "UObject/GCObject.h"
+#include "Experimental/Containers/FAAArrayQueue.h"
 #include "UObject/ObjectRedirector.h"
 #include "Serialization/BulkData.h"
 #include "Serialization/LargeMemoryReader.h"
@@ -376,8 +377,6 @@ struct FExportObject
 	bool bWasFoundInMemory = false;
 };
 
-class FAsyncLoadingSyncLoadContext;
-
 struct FPackageRequest
 {
 	int32 RequestId = -1;
@@ -1585,50 +1584,48 @@ enum class EAsyncPackageLoadingState2 : uint8
 class FEventLoadGraphAllocator;
 struct FAsyncLoadEventSpec;
 struct FAsyncLoadingThreadState2;
-class FAsyncLoadEventQueue2;
-
-enum class EEventLoadNodeExecutionResult : uint8
-{
-	Timeout,
-	Complete,
-};
 
 /** [EDL] Event Load Node */
 class FEventLoadNode2
 {
+	enum class ENodeState : uint8
+	{
+		Waiting = 0,
+		Executing,
+		Timeout,
+		Completed
+	};
 public:
 	FEventLoadNode2(const FAsyncLoadEventSpec* InSpec, FAsyncPackage2* InPackage, int32 InImportOrExportIndex, int32 InBarrierCount);
 	void DependsOn(FEventLoadNode2* Other);
 	void AddBarrier();
 	void AddBarrier(int32 Count);
 	void ReleaseBarrier(FAsyncLoadingThreadState2* ThreadState = nullptr);
-	EEventLoadNodeExecutionResult Execute(FAsyncLoadingThreadState2& ThreadState);
+	void Execute(FAsyncLoadingThreadState2& ThreadState);
 
-	inline bool IsDone() const
+	int32 GetBarrierCount()
 	{
-		return bIsDone.load(std::memory_order_acquire) == 1;
+		return BarrierCount.Load();
 	}
 
-	inline void MarkAsDone()
+	inline bool IsDone()
 	{
-		bIsDone.store(1, std::memory_order_release);
+		return ENodeState::Completed == static_cast<ENodeState>(NodeState.Load());
 	}
 
-	FAsyncPackage2* GetPackage() const
+	inline bool IsExecuting() const
 	{
-		return Package;
+		return ENodeState::Executing == static_cast<ENodeState>(NodeState.Load());
 	}
 
-	int32 ReferencerRequestId() const;
-
-private:
-	friend class FAsyncLoadEventQueue2;
-
-	FEventLoadNode2()
+	inline void SetState(ENodeState InNodeState)
 	{
-
+		NodeState.Store(static_cast<uint8>(InNodeState));
 	}
 
+	int32 ReferencerRequestId() const;
+
+private:
 	void ProcessDependencies(FAsyncLoadingThreadState2& ThreadState);
 	void Fire(FAsyncLoadingThreadState2* ThreadState = nullptr);
 
@@ -1637,12 +1634,14 @@ private:
 		FEventLoadNode2* SingleDependent;
 		FEventLoadNode2** MultipleDependents;
 	};
-	std::atomic<FEventLoadNode2*> Next { nullptr };
 	uint32 DependenciesCount = 0;
 	uint32 DependenciesCapacity = 0;
-	std::atomic<int32> BarrierCount { 0 };
-	std::atomic<uint8> DependencyWriterCount { 0 };
-	std::atomic<uint8> bIsDone { 0 };
+	TAtomic<int32> BarrierCount { 0 };
+	TAtomic<uint8> DependencyWriterCount { 0 };
+	TAtomic<uint8> NodeState { static_cast<uint8>(ENodeState::Waiting) };
+#if UE_BUILD_DEBUG || UE_BUILD_DEVELOPMENT
+	TAtomic<uint8> bFired { 0 };
+#endif
 
 	const FAsyncLoadEventSpec* Spec = nullptr;
 	FAsyncPackage2* Package = nullptr;
@@ -1674,10 +1673,6 @@ public:
 	TAtomic<int64> TotalAllocated { 0 };
 };
 
-/** Queue based on DepletableMpscQueue.h 
-	Producers will add nodes to the linked list defined by Sentinel->Tail
-	Nodes will be executed from the linked list defined by LocalHead->LocalTail
-	All the nodes from the producer queue will be moved to the local queue when needed */
 class FAsyncLoadEventQueue2
 {
 public:
@@ -1691,77 +1686,31 @@ public:
 
 	bool PopAndExecute(FAsyncLoadingThreadState2& ThreadState);
 	void Push(FEventLoadNode2* Node);
-	bool ExecuteSyncLoadEvents(FAsyncLoadingThreadState2& ThreadState, FAsyncLoadingSyncLoadContext& SyncLoadContext);
 
 private:
-	void Deplete();
-
 	FZenaphore* Zenaphore = nullptr;
-	FEventLoadNode2 Sentinel; // `Sentinel.Next` is the head of the queue
-	std::atomic<FEventLoadNode2*> Tail{ &Sentinel };
-	FEventLoadNode2* LocalHead = nullptr;
-	FEventLoadNode2* LocalTail = nullptr;
-	FEventLoadNode2* TimedOutEventNode = nullptr;
-	int32 ExecuteSyncLoadEventsCallCounter = 0;
-	bool bIsInPopAndExecute = false;
-
+	FAAArrayQueue<FEventLoadNode2> Queue;
 };
 
 struct FAsyncLoadEventSpec
 {
-	typedef EEventLoadNodeExecutionResult(*FAsyncLoadEventFunc)(FAsyncLoadingThreadState2&, FAsyncPackage2*, int32);
+	typedef EAsyncPackageState::Type(*FAsyncLoadEventFunc)(FAsyncLoadingThreadState2&, FAsyncPackage2*, int32);
 	FAsyncLoadEventFunc Func = nullptr;
 	FAsyncLoadEventQueue2* EventQueue = nullptr;
 	bool bExecuteImmediately = false;
 };
 
-class FAsyncLoadingSyncLoadContext
-{
-public:
-	FAsyncLoadingSyncLoadContext(int32 InRequestId)
-		: RequestId(InRequestId)
-	{
-
-	}
-
-	~FAsyncLoadingSyncLoadContext()
-	{
-		check(Packages.IsEmpty());
-	}
-
-	void AddRef()
-	{
-		++RefCount;
-	}
-
-	void ReleaseRef()
-	{
-		int32 NewRefCount = --RefCount;
-		check(NewRefCount >= 0);
-		if (NewRefCount == 0)
-		{
-			delete this;
-		}
-	}
-
-	int32 RequestId;
-	FAsyncPackage2* RequestedPackageDebug = nullptr;
-	FAsyncPackage2* RequestingPackageDebug = nullptr;
-	FCriticalSection CriticalSection;
-	TSet<FAsyncPackage2*> Packages;
-	bool bHasFoundRequestedPackage = false;
-
-private:
-	std::atomic<int32> RefCount = 1;
-};
-
 struct FAsyncLoadingThreadState2
+	: public FTlsAutoCleanup
 {
-	static void Set(FAsyncLoadingThreadState2* State)
+	static FAsyncLoadingThreadState2* Create(FAsyncLoadEventGraphAllocator& GraphAllocator, FIoDispatcher& IoDispatcher)
 	{
 		check(TlsSlot != 0);
 		check(!FPlatformTLS::GetTlsValue(TlsSlot));
+		FAsyncLoadingThreadState2* State = new FAsyncLoadingThreadState2(GraphAllocator, IoDispatcher);
+		State->Register();
 		FPlatformTLS::SetTlsValue(TlsSlot, State);
+		return State;
 	}
 
 	static FAsyncLoadingThreadState2* Get()
@@ -1835,12 +1784,7 @@ struct FAsyncLoadingThreadState2
 	FAsyncLoadEventGraphAllocator& GraphAllocator;
 	TArray<TTuple<FEventLoadNode2**, uint32>> DeferredFreeArcs;
 	TArray<FEventLoadNode2*> NodesToFire;
-	FEventLoadNode2* CurrentlyExecutingEventNode = nullptr;
-	TArray<FAsyncLoadingSyncLoadContext*> SyncLoadContextStack;
-	TArray<FAsyncPackage2*> PackagesExcludedFromStateDependencyChecks;
-	TSpscQueue<FAsyncLoadingSyncLoadContext*> SyncLoadContextsCreatedOnGameThread;
-	bool bIsAsyncLoadingThread = false;
-	bool bCanAccessAsyncLoadingThreadData = true;
+	FEventLoadNode2* CurrentEventNode = nullptr;
 	bool bShouldFireNodes = true;
 	bool bUseTimeLimit = false;
 	double TimeLimit = 0.0;
@@ -1924,6 +1868,9 @@ struct FAsyncPackage2
 
 	void ClearImportedPackages();
 
+	/** Marks a specific request as complete */
+	void MarkRequestIDsAsComplete();
+
 	/**
 	 * @return Time load begun. This is NOT the time the load was requested in the case of other pending requests.
 	 */
@@ -2005,8 +1952,6 @@ private:
 	FPackageImportStore ImportStore;
 	/** Package which is going to have its exports and imports loaded */
 	UPackage* LinkerRoot = nullptr;
-	// The sync load context associated with this package
-	FAsyncLoadingSyncLoadContext* SyncLoadContext = nullptr;
 	/** Time load begun. This is NOT the time the load was requested in the case of pending requests.	*/
 	double						LoadStartTime = 0.0;
 	TAtomic<int32> RefCount{ 0 };
@@ -2141,13 +2086,13 @@ public:
 
 	/** [EDL] Begin Event driven loader specific stuff */
 
-	static EEventLoadNodeExecutionResult Event_ProcessExportBundle(FAsyncLoadingThreadState2& ThreadState, FAsyncPackage2* Package, int32 ExportBundleIndex);
-	static EEventLoadNodeExecutionResult Event_ProcessPackageSummary(FAsyncLoadingThreadState2& ThreadState, FAsyncPackage2* Package, int32);
-	static EEventLoadNodeExecutionResult Event_SetupDependencies(FAsyncLoadingThreadState2& ThreadState, FAsyncPackage2* Package, int32);
-	static EEventLoadNodeExecutionResult Event_ExportsDone(FAsyncLoadingThreadState2& ThreadState, FAsyncPackage2* Package, int32);
-	static EEventLoadNodeExecutionResult Event_PostLoadExportBundle(FAsyncLoadingThreadState2& ThreadState, FAsyncPackage2* Package, int32 ExportBundleIndex);
-	static EEventLoadNodeExecutionResult Event_DeferredPostLoadExportBundle(FAsyncLoadingThreadState2& ThreadState, FAsyncPackage2* Package, int32 ExportBundleIndex);
-
+	static EAsyncPackageState::Type Event_ProcessExportBundle(FAsyncLoadingThreadState2& ThreadState, FAsyncPackage2* Package, int32 ExportBundleIndex);
+	static EAsyncPackageState::Type Event_ProcessPackageSummary(FAsyncLoadingThreadState2& ThreadState, FAsyncPackage2* Package, int32);
+	static EAsyncPackageState::Type Event_SetupDependencies(FAsyncLoadingThreadState2& ThreadState, FAsyncPackage2* Package, int32);
+	static EAsyncPackageState::Type Event_ExportsDone(FAsyncLoadingThreadState2& ThreadState, FAsyncPackage2* Package, int32);
+	static EAsyncPackageState::Type Event_PostLoadExportBundle(FAsyncLoadingThreadState2& ThreadState, FAsyncPackage2* Package, int32 ExportBundleIndex);
+	static EAsyncPackageState::Type Event_DeferredPostLoadExportBundle(FAsyncLoadingThreadState2& ThreadState, FAsyncPackage2* Package, int32 ExportBundleIndex);
+	
 	void EventDrivenCreateExport(const FAsyncPackageHeaderData& Header, const TArrayView<FExportObject>& Exports, int32 LocalExportIndex);
 	bool EventDrivenSerializeExport(const FAsyncPackageHeaderData& Header, const TArrayView<FExportObject>& Exports, int32 LocalExportIndex, FExportArchive& Ar);
 
@@ -2175,12 +2120,12 @@ private:
 	void CreateExportBundleNodes(const FAsyncLoadEventSpec* EventSpecs);
 	void SetupSerializedArcs(const FAsyncPackageHeaderData& Header, const TArrayView<FEventLoadNode2>& ExportBundleNodes, const TArrayView<FAsyncPackage2*>& ImportedAsyncPackages);
 	void SetupScriptDependencies();
-	bool HaveAllDependenciesReachedStateDebug(FAsyncLoadingThreadState2& ThreadState, FAsyncPackage2* Package, TSet<FAsyncPackage2*>& VisitedPackages, EAsyncPackageLoadingState2 WaitForPackageState);
-	bool HaveAllDependenciesReachedState(FAsyncLoadingThreadState2& ThreadState, FAllDependenciesState FAsyncPackage2::* StateMemberPtr, EAsyncPackageLoadingState2 WaitForPackageState, uint32 CurrentTick);
-	void UpdateDependenciesStateRecursive(FAsyncLoadingThreadState2& ThreadState, FAllDependenciesState FAsyncPackage2::* StateMemberPtr, EAsyncPackageLoadingState2 WaitForPackageState, uint32 CurrentTick, FAsyncPackage2* Root);
-	void WaitForAllDependenciesToReachState(FAsyncLoadingThreadState2& ThreadState, FAllDependenciesState FAsyncPackage2::* StateMemberPtr, EAsyncPackageLoadingState2 WaitForPackageState, uint32& CurrentTickVariable, TFunctionRef<void(FAsyncPackage2*)> OnStateReached);
-	void ConditionalBeginPostLoad(FAsyncLoadingThreadState2& ThreadState);
-	void ConditionalFinishLoading(FAsyncLoadingThreadState2& ThreadState);
+	bool HaveAllDependenciesReachedStateDebug(FAsyncPackage2* Package, TSet<FAsyncPackage2*>& VisitedPackages, EAsyncPackageLoadingState2 WaitForPackageState);
+	bool HaveAllDependenciesReachedState(FAllDependenciesState FAsyncPackage2::* StateMemberPtr, EAsyncPackageLoadingState2 WaitForPackageState, uint32 CurrentTick);
+	void UpdateDependenciesStateRecursive(FAllDependenciesState FAsyncPackage2::* StateMemberPtr, EAsyncPackageLoadingState2 WaitForPackageState, uint32 CurrentTick, FAsyncPackage2* Root);
+	void WaitForAllDependenciesToReachState(FAllDependenciesState FAsyncPackage2::* StateMemberPtr, EAsyncPackageLoadingState2 WaitForPackageState, uint32& CurrentTickVariable, TFunctionRef<void(FAsyncPackage2*)> OnStateReached);
+	void ConditionalBeginPostLoad();
+	void ConditionalFinishLoading();
 
 	/**
 	 * Begin async loading process. Simulates parts of BeginLoad.
@@ -2231,7 +2176,7 @@ struct FScopedAsyncPackageEvent2
 	FAsyncPackage2* Package;
 	/** Outer scope package */
 	FAsyncPackage2* PreviousPackage;
-#if WITH_EDITOR
+#if WITH_IOSTORE_IN_EDITOR
 	IAsyncPackageLoader* PreviousAsyncPackageLoader;
 #endif
 
@@ -2257,6 +2202,10 @@ private:
 	/** [ASYNC/GAME THREAD] true if the async thread is actually started. We don't start it until after we boot because the boot process on the game thread can create objects that are also being created by the loader */
 	bool bThreadStarted = false;
 
+#if ALT2_VERIFY_RECURSIVE_LOADS
+	int32 LoadRecursionLevel = 0;
+#endif
+
 #if !UE_BUILD_SHIPPING
 	FPlatformFileOpenLog* FileOpenLogWrapper = nullptr;
 #endif
@@ -2270,7 +2219,7 @@ private:
 	TArray<FAsyncPackage2*> LoadedPackagesToProcess;
 	/** [GAME THREAD] Game thread CompletedPackages list */
 	TArray<FAsyncPackage2*> CompletedPackages;
-#if WITH_EDITOR
+#if WITH_IOSTORE_IN_EDITOR
 	/** [GAME THREAD] Game thread LoadedAssets list */
 	TSet<FWeakObjectPtr> LoadedAssets;
 #endif
@@ -2297,7 +2246,6 @@ private:
 	TSet<int32> PendingRequests;
 	/** Synchronization object for PendingRequests list */
 	FCriticalSection PendingRequestsCritical;
-	TMap<int32, FAsyncPackage2*> RequestIdToPackageMap; // Only accessed from the async loading thread
 
 	/** [ASYNC/GAME THREAD] Number of package load requests in the async loading queue */
 	TAtomic<int32> QueuedPackagesCounter { 0 };
@@ -2335,9 +2283,6 @@ private:
 	FCriticalSection UnreachableObjectsCritical;
 	FUnreachableObjects UnreachableObjects;
 
-	TUniquePtr<FAsyncLoadingThreadState2> GameThreadState;
-	TUniquePtr<FAsyncLoadingThreadState2> AsyncLoadingThreadState;
-
 	uint32 ConditionalBeginPostLoadTick = 0;
 	uint32 ConditionalFinishLoadingTick = 0;
 
@@ -2357,6 +2302,7 @@ public:
 	FAsyncLoadEventGraphAllocator GraphAllocator;
 	FAsyncLoadEventQueue2 EventQueue;
 	FAsyncLoadEventQueue2 MainThreadEventQueue;
+	TArray<FAsyncLoadEventQueue2*> AltEventQueues;
 	TArray<FAsyncLoadEventSpec> EventSpecs;
 
 	/** True if multithreaded async loading is currently being used. */
@@ -2450,7 +2396,7 @@ public:
 	* @param OutPackagesProcessed Number of packages processed in this call.
 	* @return The current state of async loading
 	*/
-	EAsyncPackageState::Type ProcessAsyncLoadingFromGameThread(FAsyncLoadingThreadState2& ThreadState, bool& bDidSomething);
+	EAsyncPackageState::Type ProcessAsyncLoadingFromGameThread(FAsyncLoadingThreadState2& ThreadState, int32& OutPackagesProcessed);
 
 	/**
 	* [GAME THREAD] Ticks game thread side of async loading.
@@ -2461,7 +2407,7 @@ public:
 	* @param FlushTree Package dependency tree to be flushed
 	* @return The current state of async loading
 	*/
-	EAsyncPackageState::Type TickAsyncLoadingFromGameThread(FAsyncLoadingThreadState2& ThreadState, bool bUseTimeLimit, bool bUseFullTimeLimit, float TimeLimit, int32 FlushRequestID, bool& bDidSomething);
+	EAsyncPackageState::Type TickAsyncLoadingFromGameThread(FAsyncLoadingThreadState2& ThreadState, bool bUseTimeLimit, bool bUseFullTimeLimit, float TimeLimit, int32 FlushRequestID = INDEX_NONE);
 
 	/**
 	* [ASYNC THREAD] Main thread loop
@@ -2580,11 +2526,10 @@ private:
 
 	void RemoveUnreachableObjects(FUnreachableObjects& ObjectsToRemove);
 
-	bool ProcessPendingCDOs()
+	void ProcessPendingCDOs()
 	{
 		TRACE_CPUPROFILER_EVENT_SCOPE(ProcessPendingCDOs);
 
-		bool bDidSomething = false;
 		UClass* Class = nullptr;
 		int32 MaxRequestId = -1;
 		for (TMap<UClass*, TArray<FEventLoadNode2*>>::TIterator It = PendingCDOs.CreateIterator(); It; ++It)
@@ -2640,8 +2585,6 @@ private:
 			{
 				Node->ReleaseBarrier();
 			}
-
-			bDidSomething = true;
 		}
 		else
 		{
@@ -2656,8 +2599,6 @@ private:
 					*CurrentClass->GetFullName(), Nodes.Num());
 			}
 		}
-
-		return bDidSomething;
 	}
 
 	/**
@@ -2671,10 +2612,6 @@ private:
 	*/
 	EAsyncPackageState::Type ProcessLoadedPackagesFromGameThread(FAsyncLoadingThreadState2& ThreadState, bool& bDidSomething, int32 FlushRequestID = INDEX_NONE);
 
-	void IncludePackageInSyncLoadContext(FAsyncLoadingSyncLoadContext& SyncContext, FAsyncPackage2* Package);
-	void DetachPackageFromSyncLoadContext(FAsyncPackage2* Package);
-	FAsyncLoadingSyncLoadContext* UpdateSyncLoadContext(FAsyncLoadingThreadState2& ThreadState);
-
 	bool CreateAsyncPackagesFromQueue(FAsyncLoadingThreadState2& ThreadState);
 
 	FAsyncPackage2* CreateAsyncPackage(const FAsyncPackageDesc2& Desc)
@@ -2795,10 +2732,6 @@ private:
 	{
 		TRACE_CPUPROFILER_EVENT_SCOPE(DeleteAsyncPackage);
 		UE_ASYNC_PACKAGE_DEBUG(Package->Desc);
-		for (int32 RequestId : Package->RequestIDs)
-		{
-			RequestIdToPackageMap.Remove(RequestId);
-		}
 		delete Package;
 		--PackagesWithRemainingWorkCounter;
 	}
@@ -2942,48 +2875,6 @@ FAsyncPackage2* FAsyncLoadingThread2::FindOrInsertPackage(FAsyncPackageDesc2& De
 	return Package;
 }
 
-void FAsyncLoadingThread2::IncludePackageInSyncLoadContext(FAsyncLoadingSyncLoadContext& SyncLoadContext, FAsyncPackage2* Package)
-{
-	if (Package->SyncLoadContext)
-	{
-		if (Package->SyncLoadContext != &SyncLoadContext)
-		{
-			DetachPackageFromSyncLoadContext(Package);
-		}
-		else
-		{
-			return;
-		}
-	}
-	// Assumes that SyncLoadContext.CriticalSection is already locked by the caller
-	Package->SyncLoadContext = &SyncLoadContext;
-	SyncLoadContext.Packages.Add(Package);
-	for (FAsyncPackage2* ImportedPackage : Package->Data.ImportedAsyncPackages)
-	{
-		if (ImportedPackage && ImportedPackage->SyncLoadContext != &SyncLoadContext)
-		{
-			IncludePackageInSyncLoadContext(SyncLoadContext, ImportedPackage);
-		}
-	}
-}
-
-void FAsyncLoadingThread2::DetachPackageFromSyncLoadContext(FAsyncPackage2* Package)
-{
-	FAsyncLoadingSyncLoadContext* SyncLoadContext = Package->SyncLoadContext;
-	check(SyncLoadContext);
-	FScopeLock Lock(&SyncLoadContext->CriticalSection);
-	SyncLoadContext->Packages.Remove(Package);
-	if (SyncLoadContext->RequestedPackageDebug == Package)
-	{
-		SyncLoadContext->RequestedPackageDebug = nullptr;
-	}
-	if (SyncLoadContext->RequestingPackageDebug == Package)
-	{
-		SyncLoadContext->RequestingPackageDebug = nullptr;
-	}
-	Package->SyncLoadContext = nullptr;
-}
-
 bool FAsyncLoadingThread2::CreateAsyncPackagesFromQueue(FAsyncLoadingThreadState2& ThreadState)
 {
 	TRACE_CPUPROFILER_EVENT_SCOPE(CreateAsyncPackagesFromQueue);
@@ -3119,8 +3010,6 @@ bool FAsyncLoadingThread2::CreateAsyncPackagesFromQueue(FAsyncLoadingThreadState
 						TEXT("Package is alreay being loaded."));
 					--PackagesWithRemainingWorkCounter;
 				}
-
-				RequestIdToPackageMap.Add(PackageDesc.RequestID, Package);
 			}
 		}
 		
@@ -3152,10 +3041,12 @@ void FEventLoadNode2::DependsOn(FEventLoadNode2* Other)
 	TRACE_CPUPROFILER_EVENT_SCOPE(DependsOn);
 #if UE_BUILD_DEBUG || UE_BUILD_DEVELOPMENT
 	check(!IsDone());
+	check(!bFired);
 #endif
 	uint8 Expected = 0;
-	while (!Other->DependencyWriterCount.compare_exchange_weak(Expected, 1))
+	while (!Other->DependencyWriterCount.CompareExchange(Expected, 1))
 	{
+		TRACE_CPUPROFILER_EVENT_SCOPE(DependsOnContested);
 		check(Expected == 1);
 		Expected = 0;
 	}
@@ -3193,13 +3084,14 @@ void FEventLoadNode2::DependsOn(FEventLoadNode2* Other)
 			Other->MultipleDependents[Other->DependenciesCount++] = this;
 		}
 	}
-	Other->DependencyWriterCount.store(0, std::memory_order_release);
+	Other->DependencyWriterCount.Store(0);
 }
 
 void FEventLoadNode2::AddBarrier()
 {
 #if UE_BUILD_DEBUG || UE_BUILD_DEVELOPMENT
 	check(!IsDone());
+	check(!bFired);
 #endif
 	++BarrierCount;
 }
@@ -3208,6 +3100,7 @@ void FEventLoadNode2::AddBarrier(int32 Count)
 {
 #if UE_BUILD_DEBUG || UE_BUILD_DEVELOPMENT
 	check(!IsDone());
+	check(!bFired);
 #endif
 	BarrierCount += Count;
 }
@@ -3224,10 +3117,14 @@ void FEventLoadNode2::ReleaseBarrier(FAsyncLoadingThreadState2* ThreadState)
 void FEventLoadNode2::Fire(FAsyncLoadingThreadState2* ThreadState)
 {
 	//TRACE_CPUPROFILER_EVENT_SCOPE(Fire);
-	if (Spec->bExecuteImmediately && ThreadState)
+
+#if UE_BUILD_DEBUG || UE_BUILD_DEVELOPMENT
+	bFired.Store(1);
+#endif
+
+	if (Spec->bExecuteImmediately && ThreadState && !ThreadState->CurrentEventNode)
 	{
-		EEventLoadNodeExecutionResult Result = Execute(*ThreadState);
-		check(Result == EEventLoadNodeExecutionResult::Complete);
+		Execute(*ThreadState);
 	}
 	else
 	{
@@ -3235,33 +3132,47 @@ void FEventLoadNode2::Fire(FAsyncLoadingThreadState2* ThreadState)
 	}
 }
 
-EEventLoadNodeExecutionResult FEventLoadNode2::Execute(FAsyncLoadingThreadState2& ThreadState)
+void FEventLoadNode2::Execute(FAsyncLoadingThreadState2& ThreadState)
 {
 	//TRACE_CPUPROFILER_EVENT_SCOPE(ExecuteEvent);
-	check(BarrierCount.load(std::memory_order_relaxed) == 0);
-	EEventLoadNodeExecutionResult Result;
+	check(BarrierCount.Load() == 0);
+	check(WITH_IOSTORE_IN_EDITOR || !ThreadState.CurrentEventNode || ThreadState.CurrentEventNode == this);
+
+#if WITH_IOSTORE_IN_EDITOR
+	// Allow recursive execution of event nodes in editor builds
+	FEventLoadNode2* PrevNode = ThreadState.CurrentEventNode != this ? ThreadState.CurrentEventNode : nullptr;
+	SetState(ENodeState::Executing);
+#endif
+	ThreadState.CurrentEventNode = this;
+	EAsyncPackageState::Type State = Spec->Func(ThreadState, Package, ImportOrExportIndex);
+	if (State == EAsyncPackageState::Complete)
 	{
-		TGuardValue<FEventLoadNode2*> GuardCurrentlyExecutingEventNode(ThreadState.CurrentlyExecutingEventNode, this);
-		Result = Spec->Func(ThreadState, Package, ImportOrExportIndex);
+		SetState(ENodeState::Completed);
+		ThreadState.CurrentEventNode = nullptr;
+		ProcessDependencies(ThreadState);
+#if WITH_IOSTORE_IN_EDITOR
+		ThreadState.CurrentEventNode = PrevNode;
+#endif
 	}
-	if (Result == EEventLoadNodeExecutionResult::Complete)
+#if WITH_IOSTORE_IN_EDITOR
+	else
 	{
-		MarkAsDone();
-		ProcessDependencies(ThreadState);
+		check(PrevNode == nullptr);
+		SetState(ENodeState::Timeout);
 	}
-	return Result;
+#endif
 }
 
 void FEventLoadNode2::ProcessDependencies(FAsyncLoadingThreadState2& ThreadState)
 {
 	//TRACE_CPUPROFILER_EVENT_SCOPE(ProcessDependencies);
-	if (DependencyWriterCount.load(std::memory_order_acquire) != 0)
+	if (DependencyWriterCount.Load() != 0)
 	{
 		TRACE_CPUPROFILER_EVENT_SCOPE(ConcurrentWriter);
 		do
 		{
 			FPlatformProcess::Sleep(0);
-		} while (DependencyWriterCount.load(std::memory_order_acquire) != 0);
+		} while (DependencyWriterCount.Load() != 0);
 	}
 
 	if (DependenciesCount == 1)
@@ -3313,172 +3224,38 @@ FAsyncLoadEventQueue2::~FAsyncLoadEventQueue2()
 
 void FAsyncLoadEventQueue2::Push(FEventLoadNode2* Node)
 {
-	// switch `Tail` to the new node and only then link the old tail to the new one. The list is not fully linked between these ops,
-	// this is explicitly handled by the consumer by waiting for the link
-	FEventLoadNode2* Prev = Tail.exchange(Node, std::memory_order_release); // `release` to make sure the new node is fully constructed before
-	check(Prev->Next.load(std::memory_order_relaxed) == nullptr); // `Tail` is assigned before its Next
-	Prev->Next.store(Node, std::memory_order_relaxed);
-	if (Zenaphore)
-	{
-		bool bWasEmpty = Prev == &Sentinel;
-		if (bWasEmpty)
-		{
-			Zenaphore->NotifyOne();
-		}
-	}
-}
-
-void FAsyncLoadEventQueue2::Deplete()
-{
-	FEventLoadNode2* First = Sentinel.Next.load(std::memory_order_relaxed);
-	if (First == nullptr)
-	{
-		return; // empty
-	}
-
-	// reset the head so the next consumption can detect that the queue is empty
-	// `Sentinel.Next` is not touched by producers right now because it's already not null
-	Sentinel.Next.store(nullptr, std::memory_order_relaxed);
-
-	// reset the queue to the empty state. this redirects producers to start from `Sentinel` again.
-	// take note of the tail on resetting it because the list can be still not fully linked and so `Node.Next == nullptr` can't be 
-	// used to detect the end of the list
-	FEventLoadNode2* Last = Tail.exchange(&Sentinel, std::memory_order_acquire); // `acquire` to sync with producers' tail modifications,
-	// `Sentinel.Next = nullptr` must happen before modifying `Tail` but it's important only for the single consumer so doesn't need extra sync
-	check(Last->Next.load(std::memory_order_relaxed) == nullptr); // `Tail` is assigned before its Next
-	// the previously queued items are detached from the instance (as a linked list, though potentially not fully linked yet)
-
-	check(Last != &Sentinel); // can't be empty because of `First != nullptr` above
+	LLM_SCOPE_BYNAME(TEXT("AsyncLoadEventQueue2"));
+	Queue.enqueue(Node);
 
-	if (LocalHead)
-	{
-		check(LocalTail);
-		check(LocalTail->Next.load(std::memory_order_relaxed) == nullptr);
-		LocalTail->Next.store(First, std::memory_order_relaxed);
-	}
-	else
+	if (Zenaphore)
 	{
-		LocalHead = First;
+		Zenaphore->NotifyOne();
 	}
-	LocalTail = Last;
 }
 
 bool FAsyncLoadEventQueue2::PopAndExecute(FAsyncLoadingThreadState2& ThreadState)
 {
-	check(!bIsInPopAndExecute);
-	TGuardValue<bool> GuardIsInPopAndExecute(bIsInPopAndExecute, true);
-	if (TimedOutEventNode)
+	if (ThreadState.CurrentEventNode
+#if WITH_IOSTORE_IN_EDITOR
+		&& !ThreadState.CurrentEventNode->IsExecuting()
+#endif
+		)
 	{
-		EEventLoadNodeExecutionResult Result = TimedOutEventNode->Execute(ThreadState);
-		if (Result != EEventLoadNodeExecutionResult::Timeout)
-		{
-			TimedOutEventNode = nullptr;
-		}
+		check(!ThreadState.CurrentEventNode->IsDone());
+		ThreadState.CurrentEventNode->Execute(ThreadState);
 		return true;
 	}
 
-	if (!LocalHead)
+	if (FEventLoadNode2* Node = Queue.dequeue())
 	{
-		Deplete();
-	}
-	if (LocalHead)
-	{
-		FEventLoadNode2* Next = nullptr;
-		if (LocalHead != LocalTail)
-		{
-			// producers can still be updating `Next`, wait until the link to the next element is established
-			do
-			{
-				Next = LocalHead->Next.load(std::memory_order_relaxed);
-			} while (Next == nullptr);
-		}
-		FEventLoadNode2* Node = LocalHead;
-		LocalHead = Next;
-		if (!LocalHead)
-		{
-			LocalTail = nullptr;
-		}
-
-		EEventLoadNodeExecutionResult Result = Node->Execute(ThreadState);
-		if (Result == EEventLoadNodeExecutionResult::Timeout)
-		{
-			TimedOutEventNode = Node;
-		}
+		//TRACE_CPUPROFILER_EVENT_SCOPE(Execute);
+		Node->Execute(ThreadState);
 		return true;
 	}
-	return false;
-}
-
-bool FAsyncLoadEventQueue2::ExecuteSyncLoadEvents(FAsyncLoadingThreadState2& ThreadState, FAsyncLoadingSyncLoadContext& SyncLoadContext)
-{
-	int32 ThisCallCounter = ++ExecuteSyncLoadEventsCallCounter;
-
-	auto ShouldExecuteNode = [&SyncLoadContext](FEventLoadNode2* Node) -> bool
-	{
-		FScopeLock Lock(&SyncLoadContext.CriticalSection);
-		return SyncLoadContext.Packages.Contains(Node->Package);
-	};
-
-	bool bDidSomething = false;
-	if (TimedOutEventNode && ShouldExecuteNode(TimedOutEventNode))
-	{
-		EEventLoadNodeExecutionResult Result = TimedOutEventNode->Execute(ThreadState);
-		check(Result == EEventLoadNodeExecutionResult::Complete); // we can't timeout during a sync load operation
-		TimedOutEventNode = nullptr;
-		bDidSomething = true;
-	}
-
-	Deplete();
-
-	FEventLoadNode2* Node = LocalHead;
-	FEventLoadNode2* Prev = nullptr;
-	while (Node)
+	else
 	{
-		FEventLoadNode2* Next = nullptr;
-		if (Node != LocalTail)
-		{
-			// producers can still be updating `Next`, wait until the link to the next element is established
-			do
-			{
-				Next = Node->Next.load(std::memory_order_relaxed);
-			} while (Next == nullptr);
-		}
-
-		if (ShouldExecuteNode(Node))
-		{
-			if (Prev)
-			{
-				Prev->Next.store(Next, std::memory_order_relaxed);
-			}
-			if (Node == LocalHead)
-			{
-				LocalHead = Next;
-				if (!LocalHead)
-				{
-					LocalTail = nullptr;
-				}
-			}
-			if (Node == LocalTail)
-			{
-				LocalTail = Prev;
-			}
-
-			EEventLoadNodeExecutionResult Result = Node->Execute(ThreadState);
-			check(Result == EEventLoadNodeExecutionResult::Complete); // we can't timeout during a sync load operation
-			if (ExecuteSyncLoadEventsCallCounter != ThisCallCounter)
-			{
-				// ExecuteSyncLoadEvents was called recursively and our view of the list might have been compromised, start over
-				return true;
-			}
-			bDidSomething = true;
-		}
-		else
-		{
-			Prev = Node;
-		}
-		Node = Next;
+		return false;
 	}
-	return bDidSomething;
 }
 
 FScopedAsyncPackageEvent2::FScopedAsyncPackageEvent2(FAsyncPackage2* InPackage)
@@ -3846,18 +3623,6 @@ void FAsyncPackage2::ImportPackagesRecursive(FIoBatch& IoBatch, FPackageStore& P
 	}
 #endif
 
-	if (SyncLoadContext)
-	{
-		FScopeLock Lock(&SyncLoadContext->CriticalSection);
-		for (FAsyncPackage2* ImportedPackage : Data.ImportedAsyncPackages)
-		{
-			if (ImportedPackage)
-			{
-				AsyncLoadingThread.IncludePackageInSyncLoadContext(*SyncLoadContext, ImportedPackage);
-			}
-		}
-	}
-
 	UE_ASYNC_PACKAGE_LOG_VERBOSE(VeryVerbose, Desc, TEXT("ImportPackages: ImportsDone"),
 		TEXT("All imported packages are now being loaded."));
 
@@ -3991,7 +3756,7 @@ static void ReadAsyncPackageHeader(FAsyncPackageSerializationState& Serializatio
 	SerializationState.AllExportDataPtr = PackageHeaderDataPtr + PackageSummary->HeaderSize;
 }
 
-EEventLoadNodeExecutionResult FAsyncPackage2::Event_ProcessPackageSummary(FAsyncLoadingThreadState2& ThreadState, FAsyncPackage2* Package, int32)
+EAsyncPackageState::Type FAsyncPackage2::Event_ProcessPackageSummary(FAsyncLoadingThreadState2& ThreadState, FAsyncPackage2* Package, int32)
 {
 	TRACE_CPUPROFILER_EVENT_SCOPE(Event_ProcessPackageSummary);
 	UE_ASYNC_PACKAGE_DEBUG(Package->Desc);
@@ -4092,10 +3857,10 @@ EEventLoadNodeExecutionResult FAsyncPackage2::Event_ProcessPackageSummary(FAsync
 
 	Package->AsyncPackageLoadingState = EAsyncPackageLoadingState2::SetupDependencies;
 	Package->GetPackageNode(Package_SetupDependencies).ReleaseBarrier();
-	return EEventLoadNodeExecutionResult::Complete;
+	return EAsyncPackageState::Complete;
 }
 
-EEventLoadNodeExecutionResult FAsyncPackage2::Event_SetupDependencies(FAsyncLoadingThreadState2& ThreadState, FAsyncPackage2* Package, int32)
+EAsyncPackageState::Type FAsyncPackage2::Event_SetupDependencies(FAsyncLoadingThreadState2& ThreadState, FAsyncPackage2* Package, int32)
 {
 	TRACE_CPUPROFILER_EVENT_SCOPE(Event_SetupDependencies);
 	check(Package->AsyncPackageLoadingState == EAsyncPackageLoadingState2::SetupDependencies);
@@ -4121,10 +3886,10 @@ EEventLoadNodeExecutionResult FAsyncPackage2::Event_SetupDependencies(FAsyncLoad
 	{
 		Package->GetExportBundleNode(ExportBundle_Process, ExportBundleIndex).ReleaseBarrier();
 	}
-	return EEventLoadNodeExecutionResult::Complete;
+	return EAsyncPackageState::Complete;
 }
 
-EEventLoadNodeExecutionResult FAsyncPackage2::Event_ProcessExportBundle(FAsyncLoadingThreadState2& ThreadState, FAsyncPackage2* Package, int32 InExportBundleIndex)
+EAsyncPackageState::Type FAsyncPackage2::Event_ProcessExportBundle(FAsyncLoadingThreadState2& ThreadState, FAsyncPackage2* Package, int32 InExportBundleIndex)
 {
 	TRACE_CPUPROFILER_EVENT_SCOPE(Event_ProcessExportBundle);
 	UE_ASYNC_PACKAGE_DEBUG(Package->Desc);
@@ -4235,7 +4000,7 @@ EEventLoadNodeExecutionResult FAsyncPackage2::Event_ProcessExportBundle(FAsyncLo
 			const FExportBundleEntry& BundleEntry = HeaderData->ExportBundleEntries[ExportBundle->FirstEntryIndex + Package->ExportBundleEntryIndex];
 			if (ThreadState.IsTimeLimitExceeded(TEXT("Event_ProcessExportBundle")))
 			{
-				return EEventLoadNodeExecutionResult::Timeout;
+				return EAsyncPackageState::TimeOut;
 			}
 			const FExportMapEntry& ExportMapEntry = HeaderData->ExportMap[BundleEntry.LocalExportIndex];
 			FExportObject& Export = Exports[BundleEntry.LocalExportIndex];
@@ -4311,7 +4076,7 @@ EEventLoadNodeExecutionResult FAsyncPackage2::Event_ProcessExportBundle(FAsyncLo
 		}
 	}
 
-	return EEventLoadNodeExecutionResult::Complete;
+	return EAsyncPackageState::Complete;
 }
 
 UObject* FAsyncPackage2::EventDrivenIndexToObject(const FAsyncPackageHeaderData& Header, const TArrayView<const FExportObject>& Exports, FPackageObjectIndex Index, bool bCheckSerialized)
@@ -4654,7 +4419,7 @@ bool FAsyncPackage2::EventDrivenSerializeExport(const FAsyncPackageHeaderData& H
 	return true;
 }
 
-EEventLoadNodeExecutionResult FAsyncPackage2::Event_ExportsDone(FAsyncLoadingThreadState2& ThreadState, FAsyncPackage2* Package, int32)
+EAsyncPackageState::Type FAsyncPackage2::Event_ExportsDone(FAsyncLoadingThreadState2& ThreadState, FAsyncPackage2* Package, int32)
 {
 	TRACE_CPUPROFILER_EVENT_SCOPE(Event_ExportsDone);
 	UE_ASYNC_PACKAGE_DEBUG(Package->Desc);
@@ -4673,20 +4438,17 @@ EEventLoadNodeExecutionResult FAsyncPackage2::Event_ExportsDone(FAsyncLoadingThr
 		FCoreDelegates::ReleasePreloadedPackageShaderMaps.ExecuteIfBound(Package->Data.ShaderMapHashes);
 	}
 
-	Package->ConditionalBeginPostLoad(ThreadState);
-	return EEventLoadNodeExecutionResult::Complete;
+	Package->ConditionalBeginPostLoad();
+	return EAsyncPackageState::Complete;
 }
 
-bool FAsyncPackage2::HaveAllDependenciesReachedStateDebug(FAsyncLoadingThreadState2& ThreadState, FAsyncPackage2* Package, TSet<FAsyncPackage2*>& VisitedPackages, EAsyncPackageLoadingState2 WaitForPackageState)
+bool FAsyncPackage2::HaveAllDependenciesReachedStateDebug(FAsyncPackage2* Package, TSet<FAsyncPackage2*>& VisitedPackages, EAsyncPackageLoadingState2 WaitForPackageState)
 {
 	for (FAsyncPackage2* ImportedPackage : Package->Data.ImportedAsyncPackages)
 	{
 		if (ImportedPackage)
 		{
-			if (ThreadState.PackagesExcludedFromStateDependencyChecks.Contains(ImportedPackage))
-			{
-				continue;
-			}
+
 			if (VisitedPackages.Contains(ImportedPackage))
 			{
 				continue;
@@ -4698,7 +4460,7 @@ bool FAsyncPackage2::HaveAllDependenciesReachedStateDebug(FAsyncLoadingThreadSta
 				return false;
 			}
 
-			if (!HaveAllDependenciesReachedStateDebug(ThreadState, ImportedPackage, VisitedPackages, WaitForPackageState))
+			if (!HaveAllDependenciesReachedStateDebug(ImportedPackage, VisitedPackages, WaitForPackageState))
 			{
 				return false;
 			}
@@ -4707,7 +4469,7 @@ bool FAsyncPackage2::HaveAllDependenciesReachedStateDebug(FAsyncLoadingThreadSta
 	return true;
 }
 
-bool FAsyncPackage2::HaveAllDependenciesReachedState(FAsyncLoadingThreadState2& ThreadState, FAllDependenciesState FAsyncPackage2::* StateMemberPtr, EAsyncPackageLoadingState2 WaitForPackageState, uint32 CurrentTick)
+bool FAsyncPackage2::HaveAllDependenciesReachedState(FAllDependenciesState FAsyncPackage2::* StateMemberPtr, EAsyncPackageLoadingState2 WaitForPackageState, uint32 CurrentTick)
 {
 	FAllDependenciesState& ThisState = this->*StateMemberPtr;
 	if (ThisState.bAllDone)
@@ -4719,12 +4481,12 @@ bool FAsyncPackage2::HaveAllDependenciesReachedState(FAsyncLoadingThreadState2&
 		return false;
 	}
 	ThisState.UpdateTick(CurrentTick);
-	UpdateDependenciesStateRecursive(ThreadState, StateMemberPtr, WaitForPackageState, CurrentTick, this);
+	UpdateDependenciesStateRecursive(StateMemberPtr, WaitForPackageState, CurrentTick, this);
 	check(ThisState.bAllDone || (ThisState.WaitingForPackage && ThisState.WaitingForPackage->AsyncPackageLoadingState <= WaitForPackageState));
 	return ThisState.bAllDone;
 }
 
-void FAsyncPackage2::UpdateDependenciesStateRecursive(FAsyncLoadingThreadState2& ThreadState, FAllDependenciesState FAsyncPackage2::* StateMemberPtr, EAsyncPackageLoadingState2 WaitForPackageState, uint32 CurrentTick, FAsyncPackage2* Root)
+void FAsyncPackage2::UpdateDependenciesStateRecursive(FAllDependenciesState FAsyncPackage2::* StateMemberPtr, EAsyncPackageLoadingState2 WaitForPackageState, uint32 CurrentTick, FAsyncPackage2* Root)
 {
 	FAllDependenciesState& ThisState = this->*StateMemberPtr;
 
@@ -4736,30 +4498,23 @@ void FAsyncPackage2::UpdateDependenciesStateRecursive(FAsyncLoadingThreadState2&
 
 	if (FAsyncPackage2* WaitingForPackage = ThisState.WaitingForPackage)
 	{
-		if (ThreadState.PackagesExcludedFromStateDependencyChecks.Contains(WaitingForPackage))
+		FAllDependenciesState& WaitingForPackageState = WaitingForPackage->*StateMemberPtr;
+		if (WaitingForPackage->AsyncPackageLoadingState < WaitForPackageState)
 		{
-			ThisState.WaitingForPackage = nullptr;
+			ThisState.bAnyNotDone = true;
+			return;
 		}
-		else
+		else if (!WaitingForPackageState.bAllDone)
 		{
-			FAllDependenciesState& WaitingForPackageState = WaitingForPackage->*StateMemberPtr;
-			if (WaitingForPackage->AsyncPackageLoadingState < WaitForPackageState)
+			WaitingForPackageState.UpdateTick(CurrentTick);
+			if (!WaitingForPackageState.bVisitedMark)
 			{
-				ThisState.bAnyNotDone = true;
-				return;
+				WaitingForPackage->UpdateDependenciesStateRecursive(StateMemberPtr, WaitForPackageState, CurrentTick, Root);
 			}
-			else if (!WaitingForPackageState.bAllDone)
+			if (WaitingForPackageState.bAnyNotDone)
 			{
-				WaitingForPackageState.UpdateTick(CurrentTick);
-				if (!WaitingForPackageState.bVisitedMark)
-				{
-					WaitingForPackage->UpdateDependenciesStateRecursive(ThreadState, StateMemberPtr, WaitForPackageState, CurrentTick, Root);
-				}
-				if (WaitingForPackageState.bAnyNotDone)
-				{
-					ThisState.bAnyNotDone = true;
-					return;
-				}
+				ThisState.bAnyNotDone = true;
+				return;
 			}
 		}
 	}
@@ -4773,11 +4528,6 @@ void FAsyncPackage2::UpdateDependenciesStateRecursive(FAsyncLoadingThreadState2&
 			continue;
 		}
 
-		if (ThreadState.PackagesExcludedFromStateDependencyChecks.Contains(ImportedPackage))
-		{
-			continue;
-		}
-
 		FAllDependenciesState& ImportedPackageState = ImportedPackage->*StateMemberPtr;
 
 		if (ImportedPackageState.bAllDone)
@@ -4793,7 +4543,7 @@ void FAsyncPackage2::UpdateDependenciesStateRecursive(FAsyncLoadingThreadState2&
 		}
 		else if (!ImportedPackageState.bVisitedMark)
 		{
-			ImportedPackage->UpdateDependenciesStateRecursive(ThreadState, StateMemberPtr, WaitForPackageState, CurrentTick, Root);
+			ImportedPackage->UpdateDependenciesStateRecursive(StateMemberPtr, WaitForPackageState, CurrentTick, Root);
 		}
 
 		if (ImportedPackageState.bAnyNotDone)
@@ -4825,9 +4575,9 @@ void FAsyncPackage2::UpdateDependenciesStateRecursive(FAsyncLoadingThreadState2&
 	}
 }
 
-void FAsyncPackage2::WaitForAllDependenciesToReachState(FAsyncLoadingThreadState2& ThreadState, FAllDependenciesState FAsyncPackage2::* StateMemberPtr, EAsyncPackageLoadingState2 WaitForPackageState, uint32& CurrentTickVariable, TFunctionRef<void(FAsyncPackage2*)> OnStateReached)
+void FAsyncPackage2::WaitForAllDependenciesToReachState(FAllDependenciesState FAsyncPackage2::* StateMemberPtr, EAsyncPackageLoadingState2 WaitForPackageState, uint32& CurrentTickVariable, TFunctionRef<void(FAsyncPackage2*)> OnStateReached)
 {
-	if (HaveAllDependenciesReachedState(ThreadState, StateMemberPtr, WaitForPackageState, CurrentTickVariable++))
+	if (HaveAllDependenciesReachedState(StateMemberPtr, WaitForPackageState, CurrentTickVariable++))
 	{
 		FAsyncPackage2* FirstPackageReadyToProceed = this;
 
@@ -4844,14 +4594,14 @@ void FAsyncPackage2::WaitForAllDependenciesToReachState(FAsyncLoadingThreadState
 
 #if DO_CHECK
 			TSet<FAsyncPackage2*> VisitedPackages;
-			check(HaveAllDependenciesReachedStateDebug(ThreadState, this, VisitedPackages, WaitForPackageState));
+			check(HaveAllDependenciesReachedStateDebug(this, VisitedPackages, WaitForPackageState));
 #endif
 
 			while (FAsyncPackage2* WaitingPackage = PackageReadyToProceedState.PackagesWaitingForThisHead)
 			{
 				FAllDependenciesState& WaitingPackageState = WaitingPackage->*StateMemberPtr;
 				check(WaitingPackageState.WaitingForPackage == PackageReadyToProceed);
-				if (WaitingPackage->HaveAllDependenciesReachedState(ThreadState, StateMemberPtr, WaitForPackageState, CurrentTickVariable++))
+				if (WaitingPackage->HaveAllDependenciesReachedState(StateMemberPtr, WaitForPackageState, CurrentTickVariable++))
 				{
 					FAllDependenciesState::RemoveFromWaitList(StateMemberPtr, PackageReadyToProceed, WaitingPackage);
 					WaitingPackageState.NextLink = FirstPackageReadyToProceed;
@@ -4866,10 +4616,10 @@ void FAsyncPackage2::WaitForAllDependenciesToReachState(FAsyncLoadingThreadState
 	}
 }
 
-void FAsyncPackage2::ConditionalBeginPostLoad(FAsyncLoadingThreadState2& ThreadState)
+void FAsyncPackage2::ConditionalBeginPostLoad()
 {
 	TRACE_CPUPROFILER_EVENT_SCOPE(ConditionalBeginPostLoad);
-	WaitForAllDependenciesToReachState(ThreadState, &FAsyncPackage2::AllDependenciesSerializedState, EAsyncPackageLoadingState2::ExportsDone, AsyncLoadingThread.ConditionalBeginPostLoadTick,
+	WaitForAllDependenciesToReachState(&FAsyncPackage2::AllDependenciesSerializedState, EAsyncPackageLoadingState2::ExportsDone, AsyncLoadingThread.ConditionalBeginPostLoadTick,
 		[](FAsyncPackage2* Package)
 		{
 			for (int32 ExportBundleIndex = 0; ExportBundleIndex < Package->Data.ExportBundleCount; ++ExportBundleIndex)
@@ -4879,17 +4629,17 @@ void FAsyncPackage2::ConditionalBeginPostLoad(FAsyncLoadingThreadState2& ThreadS
 		});
 }
 
-void FAsyncPackage2::ConditionalFinishLoading(FAsyncLoadingThreadState2& ThreadState)
+void FAsyncPackage2::ConditionalFinishLoading()
 {
 	TRACE_CPUPROFILER_EVENT_SCOPE(ConditionalFinishLoading);
-	WaitForAllDependenciesToReachState(ThreadState, &FAsyncPackage2::AllDependenciesFullyLoadedState, EAsyncPackageLoadingState2::DeferredPostLoadDone, AsyncLoadingThread.ConditionalFinishLoadingTick,
+	WaitForAllDependenciesToReachState(&FAsyncPackage2::AllDependenciesFullyLoadedState, EAsyncPackageLoadingState2::DeferredPostLoadDone, AsyncLoadingThread.ConditionalFinishLoadingTick,
 		[](FAsyncPackage2* Package)
 		{
 			Package->AsyncLoadingThread.LoadedPackagesToProcess.Add(Package);
 		});
 }
 
-EEventLoadNodeExecutionResult FAsyncPackage2::Event_PostLoadExportBundle(FAsyncLoadingThreadState2& ThreadState, FAsyncPackage2* Package, int32 InExportBundleIndex)
+EAsyncPackageState::Type FAsyncPackage2::Event_PostLoadExportBundle(FAsyncLoadingThreadState2& ThreadState, FAsyncPackage2* Package, int32 InExportBundleIndex)
 {
 	TRACE_CPUPROFILER_EVENT_SCOPE(Event_PostLoad);
 	UE_ASYNC_PACKAGE_DEBUG(Package->Desc);
@@ -4921,7 +4671,7 @@ EEventLoadNodeExecutionResult FAsyncPackage2::Event_PostLoadExportBundle(FAsyncL
 	
 	check(InExportBundleIndex < Package->Data.ExportBundleCount);
 
-	EEventLoadNodeExecutionResult LoadingState = EEventLoadNodeExecutionResult::Complete;
+	EAsyncPackageState::Type LoadingState = EAsyncPackageState::Complete;
 
 	if (!Package->bLoadHasFailed)
 	{
@@ -4964,7 +4714,7 @@ EEventLoadNodeExecutionResult FAsyncPackage2::Event_PostLoadExportBundle(FAsyncL
 			const FExportBundleEntry& BundleEntry = HeaderData->ExportBundleEntries[ExportBundle->FirstEntryIndex + Package->ExportBundleEntryIndex];
 			if (ThreadState.IsTimeLimitExceeded(TEXT("Event_PostLoadExportBundle")))
 			{
-				LoadingState = EEventLoadNodeExecutionResult::Timeout;
+				LoadingState = EAsyncPackageState::TimeOut;
 				break;
 			}
 			
@@ -5005,7 +4755,7 @@ EEventLoadNodeExecutionResult FAsyncPackage2::Event_PostLoadExportBundle(FAsyncL
 		Package->EndAsyncLoad();
 	}
 	
-	if (LoadingState == EEventLoadNodeExecutionResult::Timeout)
+	if (LoadingState == EAsyncPackageState::TimeOut)
 	{
 		return LoadingState;
 	}
@@ -5032,10 +4782,10 @@ EEventLoadNodeExecutionResult FAsyncPackage2::Event_PostLoadExportBundle(FAsyncL
 		}
 	}
 
-	return EEventLoadNodeExecutionResult::Complete;
+	return EAsyncPackageState::Complete;
 }
 
-EEventLoadNodeExecutionResult FAsyncPackage2::Event_DeferredPostLoadExportBundle(FAsyncLoadingThreadState2& ThreadState, FAsyncPackage2* Package, int32 InExportBundleIndex)
+EAsyncPackageState::Type FAsyncPackage2::Event_DeferredPostLoadExportBundle(FAsyncLoadingThreadState2& ThreadState, FAsyncPackage2* Package, int32 InExportBundleIndex)
 {
 	SCOPE_CYCLE_COUNTER(STAT_FAsyncPackage_PostLoadObjectsGameThread);
 	TRACE_CPUPROFILER_EVENT_SCOPE(Event_DeferredPostLoad);
@@ -5045,7 +4795,7 @@ EEventLoadNodeExecutionResult FAsyncPackage2::Event_DeferredPostLoadExportBundle
 	FAsyncPackageScope2 PackageScope(Package);
 
 	check(InExportBundleIndex < Package->Data.ExportBundleCount);
-	EEventLoadNodeExecutionResult LoadingState = EEventLoadNodeExecutionResult::Complete;
+	EAsyncPackageState::Type LoadingState = EAsyncPackageState::Complete;
 
 	if (Package->bLoadHasFailed)
 	{
@@ -5085,7 +4835,7 @@ EEventLoadNodeExecutionResult FAsyncPackage2::Event_DeferredPostLoadExportBundle
 			const FExportBundleEntry& BundleEntry = HeaderData->ExportBundleEntries[ExportBundle->FirstEntryIndex + Package->ExportBundleEntryIndex];
 			if (ThreadState.IsTimeLimitExceeded(TEXT("Event_DeferredPostLoadExportBundle")))
 			{
-				LoadingState = EEventLoadNodeExecutionResult::Timeout;
+				LoadingState = EAsyncPackageState::TimeOut;
 				break;
 			}
 
@@ -5118,7 +4868,7 @@ EEventLoadNodeExecutionResult FAsyncPackage2::Event_DeferredPostLoadExportBundle
 		}
 	}
 
-	if (LoadingState == EEventLoadNodeExecutionResult::Timeout)
+	if (LoadingState == EAsyncPackageState::TimeOut)
 	{
 		return LoadingState;
 	}
@@ -5130,10 +4880,10 @@ EEventLoadNodeExecutionResult FAsyncPackage2::Event_DeferredPostLoadExportBundle
 		Package->ProcessedExportBundlesCount = 0;
 		check(Package->AsyncPackageLoadingState == EAsyncPackageLoadingState2::DeferredPostLoad);
 		Package->AsyncPackageLoadingState = EAsyncPackageLoadingState2::DeferredPostLoadDone;
-		Package->ConditionalFinishLoading(ThreadState);
+		Package->ConditionalFinishLoading();
 	}
 
-	return EEventLoadNodeExecutionResult::Complete;
+	return EAsyncPackageState::Complete;
 }
 
 FEventLoadNode2& FAsyncPackage2::GetPackageNode(EEventLoadNode2 Phase)
@@ -5149,50 +4899,36 @@ FEventLoadNode2& FAsyncPackage2::GetExportBundleNode(EEventLoadNode2 Phase, uint
 	return Data.ExportBundleNodes[ExportBundleNodeIndex];
 }
 
-FAsyncLoadingSyncLoadContext* FAsyncLoadingThread2::UpdateSyncLoadContext(FAsyncLoadingThreadState2& ThreadState)
+#if ALT2_VERIFY_RECURSIVE_LOADS 
+struct FScopedLoadRecursionVerifier
 {
-	if (ThreadState.bIsAsyncLoadingThread)
-	{
-		FAsyncLoadingSyncLoadContext* CreatedOnMainThread;
-		while (ThreadState.SyncLoadContextsCreatedOnGameThread.Dequeue(CreatedOnMainThread))
-		{
-			ThreadState.SyncLoadContextStack.Push(CreatedOnMainThread);
-		}
-	}
-	if (ThreadState.SyncLoadContextStack.IsEmpty())
-	{
-		return nullptr;
-	}
-	FAsyncLoadingSyncLoadContext* SyncLoadContext = ThreadState.SyncLoadContextStack.Top();
-	if (ThreadState.bIsAsyncLoadingThread && !ContainsRequestID(SyncLoadContext->RequestId))
+	int32& Level;
+	FScopedLoadRecursionVerifier(int32& InLevel) : Level(InLevel)
 	{
-		SyncLoadContext->ReleaseRef();
-		ThreadState.SyncLoadContextStack.Pop();
-		if (ThreadState.SyncLoadContextStack.IsEmpty())
-		{
-			return nullptr;
-		}
-		SyncLoadContext = ThreadState.SyncLoadContextStack.Top();
+		UE_CLOG(Level > 0, LogStreaming, Error, TEXT("Entering recursive load level: %d"), Level);
+		++Level;
+		check(Level == 1);
 	}
-	if (ThreadState.bCanAccessAsyncLoadingThreadData && !SyncLoadContext->bHasFoundRequestedPackage)
+	~FScopedLoadRecursionVerifier()
 	{
-		if (FAsyncPackage2* RequestedPackage = RequestIdToPackageMap.FindRef(SyncLoadContext->RequestId))
-		{
-			FScopeLock Lock(&SyncLoadContext->CriticalSection);
-			SyncLoadContext->bHasFoundRequestedPackage = true;
-			SyncLoadContext->RequestedPackageDebug = RequestedPackage;
-			IncludePackageInSyncLoadContext(*SyncLoadContext, RequestedPackage);
-		}
+		--Level;
+		UE_CLOG(Level > 0, LogStreaming, Error, TEXT("Leaving recursive load level: %d"), Level);
+		check(Level == 0);
 	}
-	return SyncLoadContext;
-}
+};
+#endif
 
-EAsyncPackageState::Type FAsyncLoadingThread2::ProcessAsyncLoadingFromGameThread(FAsyncLoadingThreadState2& ThreadState, bool& bDidSomething)
+EAsyncPackageState::Type FAsyncLoadingThread2::ProcessAsyncLoadingFromGameThread(FAsyncLoadingThreadState2& ThreadState, int32& OutPackagesProcessed)
 {
 	SCOPED_LOADTIMER(AsyncLoadingTime);
 
 	check(IsInGameThread());
 
+	OutPackagesProcessed = 0;
+
+#if ALT2_VERIFY_RECURSIVE_LOADS 
+	FScopedLoadRecursionVerifier LoadRecursionVerifier(this->LoadRecursionLevel);
+#endif
 	FAsyncLoadingTickScope2 InAsyncLoadingTick(*this);
 	uint32 LoopIterations = 0;
 
@@ -5223,7 +4959,7 @@ EAsyncPackageState::Type FAsyncLoadingThread2::ProcessAsyncLoadingFromGameThread
 			{
 				if (CreateAsyncPackagesFromQueue(ThreadState))
 				{
-					bDidSomething = true;
+					OutPackagesProcessed++;
 					break;
 				}
 				else
@@ -5232,17 +4968,18 @@ EAsyncPackageState::Type FAsyncLoadingThread2::ProcessAsyncLoadingFromGameThread
 				}
 			}
 
-			if (FAsyncLoadingSyncLoadContext* SyncLoadContext = UpdateSyncLoadContext(ThreadState))
+			bool bPopped = false;
+			for (FAsyncLoadEventQueue2* Queue : AltEventQueues)
 			{
-				if (EventQueue.ExecuteSyncLoadEvents(ThreadState, *SyncLoadContext))
+				if (Queue->PopAndExecute(ThreadState))
 				{
-					bDidSomething = true;
+					bPopped = true;
 					break;
 				}
 			}
-			else if (EventQueue.PopAndExecute(ThreadState))
+			if (bPopped)
 			{
-				bDidSomething = true;
+				OutPackagesProcessed++;
 				break;
 			}
 
@@ -5256,7 +4993,7 @@ EAsyncPackageState::Type FAsyncLoadingThread2::ProcessAsyncLoadingFromGameThread
 				EAsyncPackageState::Type Result = Package->ProcessExternalReads(FAsyncPackage2::ExternalReadAction_Wait);
 				check(Result == EAsyncPackageState::Complete);
 
-				bDidSomething = true;
+				OutPackagesProcessed++;
 				break;
 			}
 
@@ -5267,7 +5004,7 @@ EAsyncPackageState::Type FAsyncLoadingThread2::ProcessAsyncLoadingFromGameThread
 				FAsyncPackage2* Package = nullptr;
 				DeferredDeletePackages.Dequeue(Package);
 				DeleteAsyncPackage(Package);
-				bDidSomething = true;
+				OutPackagesProcessed++;
 				break;
 			}
 
@@ -5308,18 +5045,11 @@ EAsyncPackageState::Type FAsyncLoadingThread2::ProcessLoadedPackagesFromGameThre
 		}
 
 		bool bLocalDidSomething = false;
-		if (FAsyncLoadingSyncLoadContext* SyncLoadContext = UpdateSyncLoadContext(ThreadState))
-		{
-			bLocalDidSomething |= MainThreadEventQueue.ExecuteSyncLoadEvents(ThreadState, *SyncLoadContext);
-		}
-		else
-		{
-			bLocalDidSomething |= MainThreadEventQueue.PopAndExecute(ThreadState);
-		}
+		bLocalDidSomething |= MainThreadEventQueue.PopAndExecute(ThreadState);
 
 		bLocalDidSomething |= LoadedPackagesToProcess.Num() > 0;
 		TArray<FAsyncPackage2*, TInlineAllocator<4>> PackagesReadyForCallback;
-#if WITH_EDITOR
+#if WITH_IOSTORE_IN_EDITOR
 		TSet<UPackage*> CompletedUPackages;
 #endif
 		for (int32 PackageIndex = 0; PackageIndex < LoadedPackagesToProcess.Num(); ++PackageIndex)
@@ -5330,11 +5060,6 @@ EAsyncPackageState::Type FAsyncLoadingThread2::ProcessLoadedPackagesFromGameThre
 			check(Package->AsyncPackageLoadingState >= EAsyncPackageLoadingState2::Finalize &&
 				  Package->AsyncPackageLoadingState <= EAsyncPackageLoadingState2::CreateClusters);
 
-			if (Package->SyncLoadContext)
-			{
-				DetachPackageFromSyncLoadContext(Package);
-			}
-
 			if (Package->AsyncPackageLoadingState == EAsyncPackageLoadingState2::Finalize)
 			{
 				TArray<UObject*> CDODefaultSubobjects;
@@ -5423,7 +5148,7 @@ EAsyncPackageState::Type FAsyncLoadingThread2::ProcessLoadedPackagesFromGameThre
 				AsyncPackageLookup.Remove(Package->Desc.UPackageId);
 				if (!Package->bLoadHasFailed)
 				{
-#if WITH_EDITOR
+#if WITH_IOSTORE_IN_EDITOR
 					// In the editor we need to find any assets and packages and add them to list for later callback
 					Package->GetLoadedAssetsAndPackages(LoadedAssets, CompletedUPackages);
 #endif
@@ -5457,12 +5182,12 @@ EAsyncPackageState::Type FAsyncLoadingThread2::ProcessLoadedPackagesFromGameThre
 
 			check(!CompletedPackages.Contains(Package));
 			CompletedPackages.Add(Package);
-			RemovePendingRequests(Package->RequestIDs);
+			Package->MarkRequestIDsAsComplete();
 
 			UE_ASYNC_PACKAGE_LOG(Verbose, Package->Desc, TEXT("GameThread: LoadCompleted"),
 				TEXT("All loading of package is done, and the async package and load request will be deleted."));
 		}
-#if WITH_EDITOR
+#if WITH_IOSTORE_IN_EDITOR
 		// Call the global delegate for package endloads and set the bHasBeenLoaded flag that is used to
 		// check which packages have reached this state
 		for (UPackage* CompletedUPackage : CompletedUPackages)
@@ -5546,7 +5271,7 @@ EAsyncPackageState::Type FAsyncLoadingThread2::ProcessLoadedPackagesFromGameThre
 	return Result;
 }
 
-EAsyncPackageState::Type FAsyncLoadingThread2::TickAsyncLoadingFromGameThread(FAsyncLoadingThreadState2& ThreadState, bool bUseTimeLimit, bool bUseFullTimeLimit, float TimeLimit, int32 FlushRequestID, bool& bDidSomething)
+EAsyncPackageState::Type FAsyncLoadingThread2::TickAsyncLoadingFromGameThread(FAsyncLoadingThreadState2& ThreadState, bool bUseTimeLimit, bool bUseFullTimeLimit, float TimeLimit, int32 FlushRequestID)
 {
 	SCOPE_CYCLE_COUNTER(STAT_FAsyncPackage_TickAsyncLoadingGameThread);
 	//TRACE_INT_VALUE(QueuedPackagesCounter, QueuedPackagesCounter);
@@ -5568,6 +5293,7 @@ EAsyncPackageState::Type FAsyncLoadingThread2::TickAsyncLoadingFromGameThread(FA
 		const bool bIsMultithreaded = FAsyncLoadingThread2::IsMultithreaded();
 		double TickStartTime = FPlatformTime::Seconds();
 
+		bool bDidSomething = false;
 		{
 			Result = ProcessLoadedPackagesFromGameThread(ThreadState, bDidSomething, FlushRequestID);
 			double TimeLimitUsedForProcessLoaded = FPlatformTime::Seconds() - TickStartTime;
@@ -5583,7 +5309,7 @@ EAsyncPackageState::Type FAsyncLoadingThread2::TickAsyncLoadingFromGameThread(FA
 		{
 			if (!bDidSomething && PendingCDOs.Num() > 0)
 			{
-				bDidSomething = ProcessPendingCDOs();
+				ProcessPendingCDOs();
 			}
 
 			// Flush deferred messages
@@ -5607,12 +5333,16 @@ FAsyncLoadingThread2::FAsyncLoadingThread2(FIoDispatcher& InIoDispatcher, IAsync
 	, PackageStore(FPackageStore::Get())
 	, GlobalImportStore(LoadedPackageStore)
 {
-#if !WITH_EDITOR
+#if !WITH_IOSTORE_IN_EDITOR
 	IsEventDrivenLoaderEnabled(); // make sure the one time init inside runs
 	GEventDrivenLoaderEnabled = true;
 #endif
 
-	EventQueue.SetZenaphore(&AltZenaphore);
+	AltEventQueues.Add(&EventQueue);
+	for (FAsyncLoadEventQueue2* Queue : AltEventQueues)
+	{
+		Queue->SetZenaphore(&AltZenaphore);
+	}
 
 	EventSpecs.AddDefaulted(EEventLoadNode2::Package_NumPhases + EEventLoadNode2::ExportBundle_NumPhases);
 	EventSpecs[EEventLoadNode2::Package_ProcessSummary] = { &FAsyncPackage2::Event_ProcessPackageSummary, &EventQueue, false };
@@ -5631,9 +5361,8 @@ FAsyncLoadingThread2::FAsyncLoadingThread2(FIoDispatcher& InIoDispatcher, IAsync
 	FCoreUObjectInternalDelegates::GetOnLeakedPackageRenameDelegate().AddRaw(this, &FAsyncLoadingThread2::OnLeakedPackageRename);
 
 	FAsyncLoadingThreadState2::TlsSlot = FPlatformTLS::AllocTlsSlot();
-	GameThreadState = MakeUnique<FAsyncLoadingThreadState2>(GraphAllocator, IoDispatcher);
-	FAsyncLoadingThreadState2::Set(GameThreadState.Get());
-	
+	FAsyncLoadingThreadState2::Create(GraphAllocator, IoDispatcher);
+
 	UE_LOG(LogStreaming, Display, TEXT("AsyncLoading2 - Created: Event Driven Loader: %s, Async Loading Thread: %s, Async Post Load: %s"),
 		GEventDrivenLoaderEnabled ? TEXT("true") : TEXT("false"),
 		FAsyncLoadingThreadSettings::Get().bAsyncLoadingThreadEnabled ? TEXT("true") : TEXT("false"),
@@ -5677,10 +5406,6 @@ void FAsyncLoadingThread2::StartThread()
 
 	if (FAsyncLoadingThreadSettings::Get().bAsyncLoadingThreadEnabled && !Thread)
 	{
-		AsyncLoadingThreadState = MakeUnique<FAsyncLoadingThreadState2>(GraphAllocator, IoDispatcher);
-		AsyncLoadingThreadState->bIsAsyncLoadingThread = true;
-		AsyncLoadingThreadState->bCanAccessAsyncLoadingThreadData = true;
-		GameThreadState->bCanAccessAsyncLoadingThreadData = false;
 		UE_LOG(LogStreaming, Log, TEXT("Starting Async Loading Thread."));
 		bThreadStarted = true;
 		FPlatformMisc::MemoryBarrier();
@@ -5705,7 +5430,7 @@ uint32 FAsyncLoadingThread2::Run()
 
 	AsyncLoadingThreadID = FPlatformTLS::GetCurrentThreadId();
 
-	FAsyncLoadingThreadState2::Set(AsyncLoadingThreadState.Get());
+	FAsyncLoadingThreadState2::Create(GraphAllocator, IoDispatcher);
 
 	TRACE_LOADTIME_START_ASYNC_LOADING();
 
@@ -5803,25 +5528,20 @@ uint32 FAsyncLoadingThread2::Run()
 						do 
 						{
 							bPopped = false;
-							if (FAsyncLoadingSyncLoadContext* SyncLoadContext = UpdateSyncLoadContext(ThreadState))
+							for (FAsyncLoadEventQueue2* Queue : AltEventQueues)
 							{
-								if (EventQueue.ExecuteSyncLoadEvents(ThreadState, *SyncLoadContext))
+								if (Queue->PopAndExecute(ThreadState))
 								{
 									bPopped = true;
 									bDidSomething = true;
 								}
-							}
-							else if (EventQueue.PopAndExecute(ThreadState))
-							{
-								bPopped = true;
-								bDidSomething = true;
-							}
-							if (bSuspendRequested.Load(EMemoryOrder::Relaxed) || IsGarbageCollectionWaiting())
-							{
-								bShouldSuspend = true;
-								bDidSomething = true;
-								bPopped = false;
-								break;
+								if (bSuspendRequested.Load(EMemoryOrder::Relaxed) || IsGarbageCollectionWaiting())
+								{
+									bShouldSuspend = true;
+									bDidSomething = true;
+									bPopped = false;
+									break;
+								}
 							}
 						} while (bPopped);
 
@@ -5891,6 +5611,7 @@ EAsyncPackageState::Type FAsyncLoadingThread2::TickAsyncThreadFromGameThread(FAs
 	check(IsInGameThread());
 	EAsyncPackageState::Type Result = EAsyncPackageState::Complete;
 	
+	int32 ProcessedRequests = 0;
 	if (AsyncThreadReady.GetValue())
 	{
 		if (ThreadState.IsTimeLimitExceeded(TEXT("TickAsyncThreadFromGameThread")))
@@ -5899,7 +5620,8 @@ EAsyncPackageState::Type FAsyncLoadingThread2::TickAsyncThreadFromGameThread(FAs
 		}
 		else
 		{
-			Result = ProcessAsyncLoadingFromGameThread(ThreadState, bDidSomething);
+			Result = ProcessAsyncLoadingFromGameThread(ThreadState, ProcessedRequests);
+			bDidSomething = bDidSomething || ProcessedRequests > 0;
 		}
 	}
 
@@ -6287,6 +6009,9 @@ FAsyncPackage2::~FAsyncPackage2()
 	checkf(RefCount == 0, TEXT("RefCount is not 0 when deleting package %s"),
 		*Desc.PackagePathToLoad.GetPackageFName().ToString());
 
+	checkf(RequestIDs.Num() == 0, TEXT("MarkRequestIDsAsComplete() has not been called for package %s"),
+		*Desc.PackagePathToLoad.GetPackageFName().ToString());
+	
 	checkf(ConstructedObjects.Num() == 0, TEXT("ClearConstructedObjects() has not been called for package %s"),
 		*Desc.PackagePathToLoad.GetPackageFName().ToString());
 
@@ -6365,6 +6090,12 @@ void FAsyncPackage2::AddRequestID(int32 Id)
 	}
 }
 
+void FAsyncPackage2::MarkRequestIDsAsComplete()
+{
+	AsyncLoadingThread.RemovePendingRequests(RequestIDs);
+	RequestIDs.Reset();
+}
+
 /**
  * @return Time load begun. This is NOT the time the load was requested in the case of other pending requests.
  */
@@ -6373,7 +6104,7 @@ double FAsyncPackage2::GetLoadStartTime() const
 	return LoadStartTime;
 }
 
-#if WITH_EDITOR
+#if WITH_IOSTORE_IN_EDITOR
 void FAsyncPackage2::GetLoadedAssetsAndPackages(TSet<FWeakObjectPtr>& AssetList, TSet<UPackage*>& PackageList)
 {
 	for (UObject* Object : ConstructedObjects)
@@ -6730,8 +6461,7 @@ EAsyncPackageState::Type FAsyncLoadingThread2::ProcessLoadingFromGameThread(FAsy
 	CSV_CUSTOM_STAT(FileIO, QueuedPackagesQueueDepth, GetNumQueuedPackages(), ECsvCustomStatOp::Set);
 	CSV_CUSTOM_STAT(FileIO, ExistingQueuedPackagesQueueDepth, GetNumAsyncPackages(), ECsvCustomStatOp::Set);
 	
-	bool bDidSomething = false;
-	TickAsyncLoadingFromGameThread(ThreadState, bUseTimeLimit, bUseFullTimeLimit, TimeLimit, INDEX_NONE, bDidSomething);
+	TickAsyncLoadingFromGameThread(ThreadState, bUseTimeLimit, bUseFullTimeLimit, TimeLimit);
 	return IsAsyncLoadingPackages() ? EAsyncPackageState::TimeOut : EAsyncPackageState::Complete;
 }
 
@@ -6754,49 +6484,17 @@ void FAsyncLoadingThread2::FlushLoading(int32 RequestId)
 		double StartTime = FPlatformTime::Seconds();
 		double LogFlushTime = StartTime;
 
-		FAsyncPackage2* CurrentlyExecutingPackage = nullptr;
-		if (GameThreadState->CurrentlyExecutingEventNode)
-		{
-			UE_CLOG(RequestId == INDEX_NONE, LogStreaming, Fatal, TEXT("Flushing async loading while creating, serializing or postloading an object is not permitted"));
-			CurrentlyExecutingPackage = GameThreadState->CurrentlyExecutingEventNode->GetPackage();
-			GameThreadState->PackagesExcludedFromStateDependencyChecks.Push(CurrentlyExecutingPackage);
-		}
-
-		FAsyncLoadingSyncLoadContext* SyncLoadContext = nullptr;
-		if (RequestId != INDEX_NONE)
-		{
-			SyncLoadContext = new FAsyncLoadingSyncLoadContext(RequestId);
-			SyncLoadContext->RequestingPackageDebug = CurrentlyExecutingPackage;
-			GameThreadState->SyncLoadContextStack.Push(SyncLoadContext);
-			if (AsyncLoadingThreadState)
-			{
-				SyncLoadContext->AddRef();
-				AsyncLoadingThreadState->SyncLoadContextsCreatedOnGameThread.Enqueue(SyncLoadContext);
-				AltZenaphore.NotifyOne();
-			}
-		}
-
 		// Flush async loaders by not using a time limit. Needed for e.g. garbage collection.
 		{
+			FAsyncLoadingThreadState2& ThreadState = *FAsyncLoadingThreadState2::Get();
 			while (IsAsyncLoadingPackages())
 			{
-				bool bDidSomething = false;
-				EAsyncPackageState::Type Result = TickAsyncLoadingFromGameThread(*GameThreadState, false, false, 0, RequestId, bDidSomething);
+				EAsyncPackageState::Type Result = TickAsyncLoadingFromGameThread(ThreadState, false, false, 0, RequestId);
 				if (RequestId != INDEX_NONE && !ContainsRequestID(RequestId))
 				{
 					break;
 				}
 
-#if 0
-				if (!bDidSomething)
-				{
-					if (!AsyncLoadingThreadState && !PendingIoRequestsCounter)
-					{
-						UE_LOG(LogStreaming, Fatal, TEXT("Loading is stuck, flush will never finish"));
-					}
-				}
-#endif
-
 				if (IsMultithreaded())
 				{
 					// Update the heartbeat and sleep. If we're not multithreading, the heartbeat is updated after each package has been processed
@@ -6816,19 +6514,6 @@ void FAsyncLoadingThread2::FlushLoading(int32 RequestId)
 			}
 		}
 
-		if (SyncLoadContext)
-		{
-			check(GameThreadState->SyncLoadContextStack.Top() == SyncLoadContext);
-			SyncLoadContext->ReleaseRef();
-			GameThreadState->SyncLoadContextStack.Pop();
-		}
-
-		if (CurrentlyExecutingPackage)
-		{
-			check(GameThreadState->PackagesExcludedFromStateDependencyChecks.Top() == CurrentlyExecutingPackage);
-			GameThreadState->PackagesExcludedFromStateDependencyChecks.Pop();
-		}
-
 		check(RequestId != INDEX_NONE || !IsAsyncLoadingPackages());
 	}
 }
@@ -6849,8 +6534,7 @@ EAsyncPackageState::Type FAsyncLoadingThread2::ProcessLoadingUntilCompleteFromGa
 	bool bUseTimeLimit = TimeLimit > 0.0f;
 	double TimeLoadingPackage = 0.0f;
 
-	bool bLoadingComplete = !IsAsyncLoadingPackages() || CompletionPredicate();
-	while (!bLoadingComplete && (!bUseTimeLimit || TimeLimit > 0.0f))
+	while (IsAsyncLoadingPackages() && (!bUseTimeLimit || TimeLimit > 0.0f) && !CompletionPredicate())
 	{
 		double TickStartTime = FPlatformTime::Seconds();
 		if (ProcessLoadingFromGameThread(ThreadState, bUseTimeLimit, bUseTimeLimit, TimeLimit) == EAsyncPackageState::Complete)
@@ -6872,11 +6556,9 @@ EAsyncPackageState::Type FAsyncLoadingThread2::ProcessLoadingUntilCompleteFromGa
 		double TimeDelta = (FPlatformTime::Seconds() - TickStartTime);
 		TimeLimit -= TimeDelta;
 		TimeLoadingPackage += TimeDelta;
-
-		bLoadingComplete = !IsAsyncLoadingPackages() || CompletionPredicate();
 	}
 
-	return bLoadingComplete ? EAsyncPackageState::Complete : EAsyncPackageState::TimeOut;
+	return TimeLimit <= 0 ? EAsyncPackageState::TimeOut : EAsyncPackageState::Complete;
 }
 
 IAsyncPackageLoader* MakeAsyncPackageLoader2(FIoDispatcher& InIoDispatcher, IAsyncPackageLoader* InUncookedPackageLoader)
diff --git a/Engine/Source/Runtime/CoreUObject/Private/UObject/UObjectGlobals.cpp b/Engine/Source/Runtime/CoreUObject/Private/UObject/UObjectGlobals.cpp
--- a/Engine/Source/Runtime/CoreUObject/Private/UObject/UObjectGlobals.cpp
+++ b/Engine/Source/Runtime/CoreUObject/Private/UObject/UObjectGlobals.cpp
@@ -59,7 +59,9 @@
 #include "HAL/LowLevelMemStats.h"
 #include "Misc/CoreDelegates.h"
 #include "ProfilingDebugging/CsvProfiler.h"
+#if WITH_IOSTORE_IN_EDITOR
 #include "IO/IoDispatcher.h"
+#endif
 #include "ProfilingDebugging/LoadTimeTracker.h"
 #include "Misc/PackageAccessTracking.h"
 #include "Misc/PackageAccessTracking.h"
@@ -1457,10 +1459,7 @@ UPackage* LoadPackageInternal(UPackage* InOuter, const FPackagePath& PackagePath
 		}
 
 		GSyncLoadUsingAsyncLoaderCount++;
-		constexpr EPackageFlags PackageFlags = PKG_None;
-		constexpr int32 PIEInstanceID = INDEX_NONE;
-		constexpr int32 Priority = INT32_MAX;
-		int32 RequestID = LoadPackageAsync(PackagePath, PackageName, FLoadPackageAsyncDelegate(), PackageFlags, PIEInstanceID, Priority, InstancingContext);
+		int32 RequestID = LoadPackageAsync(PackagePath, PackageName);
 
 		if (RequestID != INDEX_NONE)
 		{
