Bulk Data - fix bug when merging I/O batch read request(s)

- This change fixes an issue when batching bulk data I/O requests where the bulk data is layed out contiguous in the chunk/file. Internally
the max read size was previously capped to the size of the bulk data but when merging read requests we read past the first bulk data
instance, i.e. from first to last. This is safe as long as the range is stored contiguously in the same chunk/file.

diff --git a/Engine/Source/Runtime/CoreUObject/Private/Serialization/BulkDataStreaming.cpp b/Engine/Source/Runtime/CoreUObject/Private/Serialization/BulkDataStreaming.cpp
--- a/Engine/Source/Runtime/CoreUObject/Private/Serialization/BulkDataStreaming.cpp
+++ b/Engine/Source/Runtime/CoreUObject/Private/Serialization/BulkDataStreaming.cpp
@@ -1084,16 +1084,8 @@ public:
 	{
 		const FIoChunkId ChunkId	= CreateBulkDataIoChunkId(BulkMeta, BulkChunkId.GetPackageId());
 		const int32 IoPriority		= ConvertToIoDispatcherPriority(Priority);
-		const uint64 BulkOffset		= uint64(BulkMeta.GetOffset());
-		const uint64 BulkSize		= uint64(BulkMeta.GetSize());
-		const uint64 BytesToRead	= FMath::Min(BulkSize, Options.GetSize());
-
-		FIoRequest IoRequest = IoBatch.ReadWithCallback(
-			ChunkId,
-			FIoReadOptions(BulkOffset + Options.GetOffset(), BytesToRead, Options.GetTargetVa()),
-			IoPriority,
-			MoveTemp(Callback));
 
+		FIoRequest IoRequest = IoBatch.ReadWithCallback(ChunkId, Options, IoPriority, MoveTemp(Callback));
 		FChunkBatchReadRequest* Request = new (Requests) FChunkBatchReadRequest(this, MoveTemp(IoRequest));
 
 		if (OutRequest)
@@ -1240,20 +1232,9 @@ public:
 		FIoReadCallback&& Callback,
 		FBulkDataBatchReadRequest* OutRequest) override
 	{
-		const uint64 BulkOffset		= uint64(BulkMeta.GetOffset());
-		const uint64 BulkSize		= uint64(BulkMeta.GetSize());
-		const uint64 BytesToRead	= FMath::Min(BulkSize, Options.GetSize());
-		
 		FFileSystemBatchReadRequest* Request = new (Requests) FFileSystemBatchReadRequest(this, MoveTemp(Callback));
 
-		RequestParams.Add(FRequestParams
-		{
-			Request,
-			BulkChunkId.GetPackagePath(),
-			FIoReadOptions(BulkOffset + Options.GetOffset(), BytesToRead, Options.GetTargetVa()),
-			BulkMeta,
-			Priority
-		});
+		RequestParams.Add(FRequestParams {Request, BulkChunkId.GetPackagePath(), Options, BulkMeta, Priority});
 
 		if (OutRequest)
 		{
@@ -1546,7 +1527,7 @@ FBulkDataBatchRequest::FBatchBuilder& FBulkDataBatchRequest::FBatchBuilder::Read
 	GetBatch(BulkData).Read(
 		BulkData.BulkMeta,
 		BulkData.BulkChunkId,
-		FIoReadOptions(0, BulkData.GetBulkDataSize()),
+		FIoReadOptions(BulkData.GetBulkDataOffsetInFile(), BulkData.GetBulkDataSize()),
 		Priority,
 		[BulkData = &BulkData](TIoStatusOr<FIoBuffer> Status)
 		{
@@ -1566,61 +1547,34 @@ FBulkDataBatchRequest::FBatchBuilder& FBulkDataBatchRequest::FBatchBuilder::Read
 	return *this;
 }
 
-FBulkDataBatchRequest::FBatchBuilder& FBulkDataBatchRequest::FBatchBuilder::Read(
-	const FBulkData& BulkData,
-	uint64 Offset,
-	uint64 Size,
-	EAsyncIOPriorityAndFlags Priority,
-	FIoBuffer& Dst)
-{
-	const uint64 BytesToRead = FMath::Min(uint64(BulkData.GetBulkDataSize()), Size);
-
-	check(BatchMax == -1 || BatchCount < BatchMax);
-	check(Dst.GetSize() == 0 || Dst.GetSize() == BytesToRead);
-
-	if (Dst.GetSize() == 0)
-	{
-		Dst = FIoBuffer(BytesToRead);
-	}
-
-	GetBatch(BulkData).Read(
-		BulkData.BulkMeta,
-		BulkData.BulkChunkId,
-		FIoReadOptions(Offset, Dst.GetSize(), Dst.GetData()),
-		Priority,
-		FIoReadCallback(),
-		nullptr);
-
-	++BatchCount;
-
-	return *this;
-}
-
 FBulkDataBatchRequest::FBatchBuilder& FBulkDataBatchRequest::FBatchBuilder::Read(
 	const FBulkData& BulkData,
 	uint64 Offset,
 	uint64 Size,
 	EAsyncIOPriorityAndFlags Priority,
 	FIoBuffer& Dst, 
-	FBulkDataBatchReadRequest& OutRequest)
+	FBulkDataBatchReadRequest* OutRequest)
 {
-	const uint64 BytesToRead = FMath::Min(uint64(BulkData.GetBulkDataSize()), Size);
+	check(Size == MAX_uint64 || Size <= uint64(BulkData.GetBulkDataSize()));
+
+	const uint64 ReadOffset = BulkData.GetBulkDataOffsetInFile() + Offset;
+	const uint64 ReadSize	= FMath::Min(uint64(BulkData.GetBulkDataSize()), Size);
 
 	check(BatchMax == -1 || BatchCount < BatchMax);
-	check(Dst.GetSize() == 0 || Dst.GetSize() == BytesToRead);
+	check(Dst.GetSize() == 0 || Dst.GetSize() == ReadSize);
 
 	if (Dst.GetSize() == 0)
 	{
-		Dst = FIoBuffer(BytesToRead);
+		Dst = FIoBuffer(ReadSize);
 	}
 
 	GetBatch(BulkData).Read(
 		BulkData.BulkMeta,
 		BulkData.BulkChunkId,
-		FIoReadOptions(Offset, Dst.GetSize(), Dst.GetData()),
+		FIoReadOptions(ReadOffset, ReadSize, Dst.GetData()),
 		Priority,
 		FIoReadCallback(),
-		&OutRequest);
+		OutRequest);
 	
 	++BatchCount;
 	
@@ -1653,45 +1607,43 @@ FBulkDataRequest::EStatus FBulkDataBatchRequest::FBatchBuilder::Issue()
 FBulkDataBatchRequest::FScatterGatherBuilder::FScatterGatherBuilder(int32 MaxCount)
 	: FBuilder(MaxCount)
 {
+	if (MaxCount > 0)
+	{
+		Requests.Reserve(MaxCount);
+	}
 }
 
 FBulkDataBatchRequest::FScatterGatherBuilder& FBulkDataBatchRequest::FScatterGatherBuilder::Read(const FBulkData& BulkData, uint64 Offset, uint64 Size)
 {
-	if (Requests.Num() > 0 && Offset == 0)
+	check(Size == MAX_uint64 || Size <= uint64(BulkData.GetBulkDataSize()));
+
+	const uint64 ReadOffset = BulkData.GetBulkDataOffsetInFile() + Offset;
+	const uint64 ReadSize	= FMath::Min(uint64(BulkData.GetBulkDataSize()), Size);
+
+	if (Requests.Num() > 0)
 	{
 		FRequest& Last = Requests.Last();
 
-		const bool bContiguous = 
-			Last.Offset + Last.Size == BulkData.GetBulkDataOffsetInFile() &&
+		const bool bContiguous =
+			Last.Offset + Last.Size == ReadOffset &&
 			Last.BulkData->GetBulkDataFlags() == BulkData.GetBulkDataFlags() &&
 			Last.BulkData->BulkChunkId == BulkData.BulkChunkId;
 
 		if (bContiguous)
 		{
-			Last.Size += FMath::Min(uint64(BulkData.GetBulkDataSize()), Size);
+			Last.Size += ReadSize; 
+			return *this;
 		}
-
-		return *this;
 	}
 
-	Requests.Add(FRequest
-	{
-		&BulkData,
-		Offset,
-		FMath::Min(uint64(BulkData.GetBulkDataSize()), Size)
-	});
+	Requests.Add(FRequest {&BulkData, ReadOffset, ReadSize});
 
 	return *this;
 }
 
 FBulkDataRequest::EStatus FBulkDataBatchRequest::FScatterGatherBuilder::Issue(FIoBuffer& Dst, EAsyncIOPriorityAndFlags Priority, FCompletionCallback&& Callback, FBulkDataBatchRequest& OutRequest)
 {
-	if (Requests.IsEmpty())
-	{
-		check(false);
-		OutRequest = FBulkDataBatchRequest(new FHandleBase(EStatus::Error));
-		return EStatus::Error;
-	}
+	check(Requests.IsEmpty() == false);
 
 	uint64 TotalSize = 0;
 	for (const FRequest& Request : Requests)
@@ -1699,6 +1651,8 @@ FBulkDataRequest::EStatus FBulkDataBatchRequest::FScatterGatherBuilder::Issue(FI
 		TotalSize += Request.Size;
 	}
 
+	check(Dst.GetSize() == 0 || Dst.GetSize() == TotalSize);
+
 	if (Dst.GetSize() != TotalSize)
 	{
 		Dst = FIoBuffer(TotalSize);
diff --git a/Engine/Source/Runtime/CoreUObject/Public/Serialization/BulkData.h b/Engine/Source/Runtime/CoreUObject/Public/Serialization/BulkData.h
--- a/Engine/Source/Runtime/CoreUObject/Public/Serialization/BulkData.h
+++ b/Engine/Source/Runtime/CoreUObject/Public/Serialization/BulkData.h
@@ -1477,7 +1477,10 @@ public:
 		 * @param Priority		The I/O priority.
 		 * @param Dst			An empty or preallocated I/O buffer.	
 		 */
-		FBatchBuilder& Read(const FBulkData& BulkData, uint64 Offset, uint64 Size, EAsyncIOPriorityAndFlags Priority, FIoBuffer& Dst);
+		FBatchBuilder& Read(const FBulkData& BulkData, uint64 Offset, uint64 Size, EAsyncIOPriorityAndFlags Priority, FIoBuffer& Dst)
+		{
+			return Read(BulkData, Offset, Size, Priority, Dst, nullptr);
+		}
 		/**
 		 * Read the bulk data from the specified offset and size and copy the result into the destination buffer.
 		 * @param BulkData		The bulk data instance.
@@ -1487,7 +1490,11 @@ public:
 		 * @param Dst			An empty or preallocated I/O buffer.	
 		 * @param OutRequest	A handle to the read request.	
 		 */
-		FBatchBuilder& Read(const FBulkData& BulkData, uint64 Offset, uint64 Size, EAsyncIOPriorityAndFlags Priority, FIoBuffer& Dst, FBulkDataBatchReadRequest& OutRequest);
+		FBatchBuilder& Read(const FBulkData& BulkData, uint64 Offset, uint64 Size, EAsyncIOPriorityAndFlags Priority, FIoBuffer& Dst, FBulkDataBatchReadRequest& OutRequest)
+		{
+			return Read(BulkData, Offset, Size, Priority, Dst, &OutRequest);
+		}
+
 		/**
 		 * Issue the batch.
 		 * @param OutRequest	A handle to the batch request.
@@ -1501,6 +1508,9 @@ public:
 		 * @note Assumes one or more handle(s) has been passed into any of the read operations.
 		 */
 		[[nodiscard]] EStatus Issue();
+
+	private:
+		FBatchBuilder& Read(const FBulkData& BulkData, uint64 Offset, uint64 Size, EAsyncIOPriorityAndFlags Priority, FIoBuffer& Dst, FBulkDataBatchReadRequest* OutRequest);
 	};
 
 	/** Reads one or more bulk data and copies the result into a single I/O buffer. */
@@ -1530,7 +1540,7 @@ public:
 			uint64 Size					= MAX_uint64;
 		};
 
-		TArray<FRequest> Requests;
+		TArray<FRequest, TInlineAllocator<8>> Requests;
 	};
 
 	/** Returns a request builder that dispatches one or more I/O requests in a single batch. */
