ZenLoader: Initial support for loading packages in the editor

diff --git a/Engine/Source/Runtime/CoreUObject/Private/Serialization/AsyncLoading2.cpp b/Engine/Source/Runtime/CoreUObject/Private/Serialization/AsyncLoading2.cpp
--- a/Engine/Source/Runtime/CoreUObject/Private/Serialization/AsyncLoading2.cpp
+++ b/Engine/Source/Runtime/CoreUObject/Private/Serialization/AsyncLoading2.cpp
@@ -182,20 +182,26 @@ void FindAllRuntimeScriptPackages(TArray<UPackage*>& OutPackages)
 	}, /*bIncludeDerivedClasses*/false);
 }
 
-#if WITH_ASYNCLOADING2
-
 #ifndef ALT2_VERIFY_ASYNC_FLAGS
-#define ALT2_VERIFY_ASYNC_FLAGS DO_CHECK && !(WITH_IOSTORE_IN_EDITOR)
-#endif
-
-#ifndef ALT2_VERIFY_RECURSIVE_LOADS
-#define ALT2_VERIFY_RECURSIVE_LOADS !(WITH_IOSTORE_IN_EDITOR) && DO_CHECK
+#define ALT2_VERIFY_ASYNC_FLAGS DO_CHECK && !(WITH_EDITOR)
 #endif
 
 #ifndef ALT2_VERIFY_UNREACHABLE_OBJECTS
 #define ALT2_VERIFY_UNREACHABLE_OBJECTS DO_CHECK
 #endif
 
+#ifndef ALT2_VERIFY_LINKERLOAD_MATCHES_IMPORTSTORE
+#define ALT2_VERIFY_LINKERLOAD_MATCHES_IMPORTSTORE 0
+#endif
+
+#ifndef ALT2_ENABLE_LINKERLOAD_SUPPORT
+#define ALT2_ENABLE_LINKERLOAD_SUPPORT WITH_EDITOR
+#endif
+
+#ifndef ALT2_ENABLE_NEW_ARCHIVE_FOR_LINKERLOAD
+#define ALT2_ENABLE_NEW_ARCHIVE_FOR_LINKERLOAD 0
+#endif
+
 #ifndef ALT2_LOG_VERBOSE
 #define ALT2_LOG_VERBOSE DO_CHECK
 #endif
@@ -382,17 +388,31 @@ struct FPackageRequest
 {
 	int32 RequestId = -1;
 	int32 Priority = -1;
+	EPackageFlags PackageFlags = PKG_None;
+#if WITH_EDITOR
+	int32 PIEInstanceID = INDEX_NONE;
+#endif
+#if WITH_EDITORONLY_DATA
+	FLinkerInstancingContext InstancingContext;
+#endif
 	FName CustomName;
 	FPackagePath PackagePath;
 	TUniquePtr<FLoadPackageAsyncDelegate> PackageLoadedDelegate;
 	FPackageRequest* Next = nullptr;
 
-	static FPackageRequest Create(int32 RequestId, int32 Priority, const FPackagePath& PackagePath, FName CustomName, TUniquePtr<FLoadPackageAsyncDelegate> PackageLoadedDelegate)
+	static FPackageRequest Create(int32 RequestId, EPackageFlags PackageFlags, int32 PIEInstanceID, int32 Priority, const FLinkerInstancingContext* InstancingContext, const FPackagePath& PackagePath, FName CustomName, TUniquePtr<FLoadPackageAsyncDelegate> PackageLoadedDelegate)
 	{
 		return FPackageRequest
 		{
 			RequestId,
 			Priority,
+			PackageFlags,
+#if WITH_EDITOR
+			PIEInstanceID,
+#endif
+#if WITH_EDITORONLY_DATA
+			InstancingContext ? *InstancingContext : FLinkerInstancingContext(),
+#endif
 			CustomName,
 			PackagePath,
 			MoveTemp(PackageLoadedDelegate),
@@ -409,6 +429,16 @@ struct FAsyncPackageDesc2
 	int32 ReferencerRequestId;
 	// Package priority
 	int32 Priority;
+	/** The flags that should be applied to the package */
+	EPackageFlags PackageFlags;
+#if WITH_EDITOR
+	/** PIE instance ID this package belongs to, INDEX_NONE otherwise */
+	int32 PIEInstanceID;
+#endif
+#if WITH_EDITORONLY_DATA
+	/** Instancing context, maps original package to their instanced counterpart, used to remap imports. */
+	FLinkerInstancingContext InstancingContext;
+#endif
 	// The package id of the UPackage being loaded
 	// It will be used as key when tracking active async packages
 	FPackageId UPackageId;
@@ -423,43 +453,70 @@ struct FAsyncPackageDesc2
 	// Packages with a a custom name can't be imported
 	bool bCanBeImported;
 
+#if WITH_EDITORONLY_DATA
+	const FLinkerInstancingContext* GetInstancingContext() const
+	{
+		return &InstancingContext;
+	}
+#else
+	const FLinkerInstancingContext* GetInstancingContext() const
+	{
+		return nullptr;
+	}
+#endif 
+
 	static FAsyncPackageDesc2 FromPackageRequest(
-		int32 RequestID,
-		int32 Priority,
+		FPackageRequest& Request,
 		FName UPackageName,
-		FPackageId PackageIdToLoad,
-		FPackagePath&& PackagePathToLoad,
-		bool bHasCustomName)
+		FPackageId PackageIdToLoad)
 	{
 		return FAsyncPackageDesc2
 		{
-			RequestID,
-			RequestID,
-			Priority,
+			Request.RequestId,
+			Request.RequestId,
+			Request.Priority,
+			Request.PackageFlags,
+#if WITH_EDITOR
+			Request.PIEInstanceID,
+#endif
+#if WITH_EDITORONLY_DATA
+			MoveTemp(Request.InstancingContext),
+#endif
 			FPackageId::FromName(UPackageName),
 			PackageIdToLoad,
 			UPackageName,
-			MoveTemp(PackagePathToLoad),
-			!bHasCustomName
+			MoveTemp(Request.PackagePath),
+#if WITH_EDITOR
+			true
+#else
+			Request.CustomName.IsNone()
+#endif
 		};
 	}
 
 	static FAsyncPackageDesc2 FromPackageImport(
-		int32 ReferencerRequestId,
-		int32 Priority,
+		const FAsyncPackageDesc2& ImportingPackageDesc,
+		FName UPackageName,
 		FPackageId ImportedPackageId,
 		FPackageId PackageIdToLoad,
-		FName UPackageName)
+		FPackagePath&& PackagePathToLoad)
 	{
 		return FAsyncPackageDesc2
 		{
 			INDEX_NONE,
-			ReferencerRequestId,
-			Priority,
+			ImportingPackageDesc.ReferencerRequestId,
+			ImportingPackageDesc.Priority,
+			PKG_None,
+#if WITH_EDITOR
+			INDEX_NONE,
+#endif
+#if WITH_EDITORONLY_DATA
+			ImportingPackageDesc.InstancingContext,
+#endif
 			ImportedPackageId,
 			PackageIdToLoad,
 			UPackageName,
-			FPackagePath(),
+			MoveTemp(PackagePathToLoad),
 			true
 		};
 	}
@@ -880,7 +937,7 @@ public:
 		ensureMsgf(!PackageRef.bHasBeenLoadedDebug || PackageRef.bAreAllPublicExportsLoaded || PackageRef.bIsMissing || PackageRef.bHasFailed,
 			TEXT("LoadedPackageRef from None (0x%llX) to %s (0x%llX) should not have been released when the package is not complete.")
 			TEXT("RefCount=%d, AreAllExportsLoaded=%d, IsMissing=%d, HasFailed=%d, HasBeenLoaded=%d"),
-			FromPackageId.Value(),
+			FromPackageId.ValueForDebugging(),
 			PackageRef.Package ? *PackageRef.Package->GetName() : TEXT("None"),
 			PackageId.Value(),
 			PackageRef.RefCount,
@@ -1153,6 +1210,12 @@ struct FAsyncPackageHeaderData
 	TArrayView<const uint8> ArcsData;
 	// Backed by allocation in FAsyncPackageData
 	TArrayView<FPackageId> ImportedPackageIds;
+	struct FImportPackageNames
+	{
+		FName UPackageName;
+		FName PackageNameToLoad;
+	};
+	TArrayView<FImportPackageNames> ImportedPackageNames;
 	TArrayView<FExportBundleHeader> ExportBundleHeaders;
 	TArrayView<FExportBundleEntry> ExportBundleEntries;
 
@@ -1166,6 +1229,15 @@ struct FAsyncPackageHeaderData
 	}
 };
 
+#if ALT2_ENABLE_LINKERLOAD_SUPPORT
+struct FAsyncPackageLinkerLoadHeaderData
+{
+	TArray<uint64> ImportedPublicExportHashes;
+	TArray<FPackageObjectIndex> ImportMap;
+	TArray<uint64> PublicExportHashes;
+};
+#endif
+
 struct FPackageImportStore
 {
 	FGlobalImportStore& GlobalImportStore;
@@ -1204,6 +1276,12 @@ struct FPackageImportStore
 		else if (GlobalIndex.IsPackageImport())
 		{
 			Object = GlobalImportStore.FindPublicExportObject(FPublicExportKey::FromPackageImport(GlobalIndex, Header.ImportedPackageIds, Header.ImportedPublicExportHashes));
+#if WITH_EDITOR
+			if (UObjectRedirector* Redirector = dynamic_cast<UObjectRedirector*>(Object))
+			{
+				Object = Redirector->DestinationObject;
+			}
+#endif
 		}
 		else
 		{
@@ -1564,8 +1642,6 @@ private:
 enum class EAsyncPackageLoadingState2 : uint8
 {
 	NewPackage,
-	ImportPackages,
-	ImportPackagesDone,
 	WaitingForIo,
 	ProcessPackageSummary,
 	SetupDependencies,
@@ -1898,6 +1974,228 @@ struct FAsyncPackageSerializationState
 	}
 };
 
+#if ALT2_ENABLE_NEW_ARCHIVE_FOR_LINKERLOAD
+class FLinkerLoadArchive2 final
+	: public FArchive
+{
+public:
+	FLinkerLoadArchive2(const FPackagePath& InPackagePath)
+	{
+		FOpenAsyncPackageResult UassetOpenResult = IPackageResourceManager::Get().OpenAsyncReadPackage(InPackagePath, EPackageSegment::Header);
+		UassetFileHandle = UassetOpenResult.Handle.Release();
+		check(UassetFileHandle); // OpenAsyncReadPackage guarantees a non-null return value; the handle will fail to read later if the path does not exist
+		if (UassetOpenResult.Format != EPackageFormat::Binary)
+		{
+			UE_LOG(LogStreaming, Fatal, TEXT("Only binary assets are supported")); // TODO
+			SetError();
+			return;
+		}
+		bNeedsEngineVersionChecks = UassetOpenResult.bNeedsEngineVersionChecks;
+
+		if (FPlatformProperties::RequiresCookedData())
+		{
+			FOpenAsyncPackageResult UexpOpenResult = IPackageResourceManager::Get().OpenAsyncReadPackage(InPackagePath, EPackageSegment::Exports);
+			UexpFileHandle = UexpOpenResult.Handle.Release();
+			check(UexpFileHandle); // OpenAsyncReadPackage guarantees a non-null return value; the handle will fail to read later if the path does not exist
+			if (UexpOpenResult.Format != EPackageFormat::Binary)
+			{
+				UE_LOG(LogStreaming, Fatal, TEXT("Only binary assets are supported")); // TODO
+				SetError();
+				return;
+			}
+		}
+		else
+		{
+			UexpSize = 0;
+		}
+	}
+
+	virtual ~FLinkerLoadArchive2()
+	{
+		WaitForRequests();
+		delete UassetFileHandle;
+		delete UexpFileHandle;
+	}
+
+	bool NeedsEngineVersionChecks() const
+	{
+		return bNeedsEngineVersionChecks;
+	}
+
+	void BeginRead(FEventLoadNode2* InDependentNode)
+	{
+		check(PendingSizeRequests == 0);
+		check(PendingReadRequests == 0);
+		check(!DependentNode);
+		if (UexpFileHandle)
+		{
+			PendingSizeRequests = 2;
+			PendingReadRequests = 2;
+		}
+		else
+		{
+			PendingSizeRequests = 1;
+			PendingReadRequests = 1;
+		}
+		DependentNode = InDependentNode;
+		StartSizeRequests();
+	}
+
+	virtual int64 TotalSize() override
+	{
+		check(bDone);
+		return ActiveFPLB->EndFastPathLoadBuffer - ActiveFPLB->OriginalFastPathLoadBuffer;
+	}
+
+	virtual int64 Tell() override
+	{
+		check(bDone);
+		return ActiveFPLB->StartFastPathLoadBuffer - ActiveFPLB->OriginalFastPathLoadBuffer;
+	}
+
+	virtual void Seek(int64 InPos) override
+	{
+		check(bDone);
+		ActiveFPLB->StartFastPathLoadBuffer = ActiveFPLB->OriginalFastPathLoadBuffer + InPos;
+		check(ActiveFPLB->StartFastPathLoadBuffer <= ActiveFPLB->EndFastPathLoadBuffer);
+	}
+
+	virtual void Serialize(void* Data, int64 Length) override
+	{
+		check(bDone);
+		if (!Length || IsError())
+		{
+			return;
+		}
+		check(ActiveFPLB->StartFastPathLoadBuffer + Length <= ActiveFPLB->EndFastPathLoadBuffer);
+		FMemory::Memcpy(Data, ActiveFPLB->StartFastPathLoadBuffer, Length);
+		ActiveFPLB->StartFastPathLoadBuffer += Length;
+	}
+
+	/*virtual bool Close() override;
+	virtual bool Precache(int64 PrecacheOffset, int64 PrecacheSize) override;
+	virtual void FlushCache() override;*/
+
+private:
+	void StartSizeRequests()
+	{
+		FAsyncFileCallBack UassetSizeRequestCallbackFunction = [this](bool bWasCancelled, IAsyncReadRequest* Request)
+		{
+			if (!bWasCancelled)
+			{
+				UassetSize = Request->GetSizeResults();
+			}
+			if (--PendingSizeRequests == 0)
+			{
+				StartReadRequests();
+			}
+		};
+		UassetSizeRequest = UassetFileHandle->SizeRequest(&UassetSizeRequestCallbackFunction);
+		if (UexpFileHandle)
+		{
+			FAsyncFileCallBack UexpSizeRequestCallbackFunction = [this](bool bWasCancelled, IAsyncReadRequest* Request)
+			{
+				if (!bWasCancelled)
+				{
+					UexpSize = Request->GetSizeResults();
+				}
+				if (--PendingSizeRequests == 0)
+				{
+					StartReadRequests();
+				}
+			};
+			UexpSizeRequest = UexpFileHandle->SizeRequest(&UexpSizeRequestCallbackFunction);
+		}
+	}
+
+	void StartReadRequests()
+	{
+		FAsyncFileCallBack ReadRequestCallbackFunction = [this](bool bWasCancelled, IAsyncReadRequest* Request)
+		{
+			if (bWasCancelled)
+			{
+				bFailed = true;
+			}
+			if (--PendingReadRequests == 0)
+			{
+				FinishedReading();
+			}
+		};
+		if (UassetSize <= 0 || (UexpFileHandle && UexpSize <= 0))
+		{
+			SetError();
+			FinishedReading();
+			return;
+		}
+		IoBuffer = FIoBuffer(UassetSize + UexpSize);
+		UassetReadRequest = UassetFileHandle->ReadRequest(0, UassetSize, AIOP_Normal, &ReadRequestCallbackFunction, IoBuffer.Data());
+		if (UexpFileHandle)
+		{
+			UexpReadRequest = UexpFileHandle->ReadRequest(0, UexpSize, AIOP_Normal, &ReadRequestCallbackFunction, IoBuffer.Data() + UassetSize);
+		}
+	}
+
+	void FinishedReading()
+	{
+		ActiveFPLB->OriginalFastPathLoadBuffer = IoBuffer.Data();
+		ActiveFPLB->StartFastPathLoadBuffer = IoBuffer.Data();
+		ActiveFPLB->EndFastPathLoadBuffer = IoBuffer.Data() + IoBuffer.DataSize();
+		bDone = true;
+		DependentNode->ReleaseBarrier();
+		DependentNode = nullptr;
+	}
+
+	void WaitForRequests()
+	{
+		if (UassetSizeRequest)
+		{
+			UassetSizeRequest->WaitCompletion();
+			delete UassetSizeRequest;
+			UassetSizeRequest = nullptr;
+		}
+		if (UexpSizeRequest)
+		{
+			UexpSizeRequest->WaitCompletion();
+			delete UexpSizeRequest;
+			UexpSizeRequest = nullptr;
+		}
+		if (UassetReadRequest)
+		{
+			UassetReadRequest->WaitCompletion();
+			delete UassetReadRequest;
+			UassetReadRequest = nullptr;
+		}
+		if (UexpReadRequest)
+		{
+			UexpReadRequest->WaitCompletion();
+			delete UexpReadRequest;
+			UexpReadRequest = nullptr;
+		}
+	}
+
+#if (!DEVIRTUALIZE_FLinkerLoad_Serialize)
+	FArchive::FFastPathLoadBuffer InlineFPLB;
+	FArchive::FFastPathLoadBuffer* ActiveFPLB = &InlineFPLB;
+#endif
+	FEventLoadNode2* DependentNode = nullptr;
+	FIoBuffer IoBuffer;
+	int64 Offset = 0;
+	IAsyncReadFileHandle* UassetFileHandle = nullptr;
+	IAsyncReadFileHandle* UexpFileHandle = nullptr;
+	IAsyncReadRequest* UassetSizeRequest = nullptr;
+	IAsyncReadRequest* UexpSizeRequest = nullptr;
+	IAsyncReadRequest* UassetReadRequest = nullptr;
+	IAsyncReadRequest* UexpReadRequest = nullptr;
+	int64 UassetSize = -1;
+	int64 UexpSize = -1;
+	std::atomic<int8> PendingSizeRequests = 0;
+	std::atomic<int8> PendingReadRequests = 0;
+	std::atomic<bool> bDone = false;
+	std::atomic<bool> bFailed = false;
+	bool bNeedsEngineVersionChecks = false;
+};
+#endif // ALT2_ENABLE_NEW_ARCHIVE_FOR_LINKERLOAD
+
 /**
 * Structure containing intermediate data required for async loading of all exports of a package.
 */
@@ -1907,6 +2205,7 @@ struct FAsyncPackage2
 	friend struct FScopedAsyncPackageEvent2;
 	friend struct FAsyncPackageScope2;
 	friend class FAsyncLoadingThread2;
+	friend class FAsyncLoadEventQueue2;
 
 	FAsyncPackage2(const FAsyncPackageDesc2& InDesc,
 		FAsyncLoadingThread2& InAsyncLoadingThread,
@@ -1978,18 +2277,23 @@ struct FAsyncPackage2
 	/** Creates GC clusters from loaded objects */
 	EAsyncPackageState::Type CreateClusters(FAsyncLoadingThreadState2& ThreadState);
 
-	void ImportPackagesRecursive(FIoBatch& IoBatch, FPackageStore& PackageStore);
+	void ImportPackagesRecursive(FAsyncLoadingThreadState2& ThreadState, FIoBatch& IoBatch, FPackageStore& PackageStore);
 	void StartLoading(FIoBatch& IoBatch);
 
-#if WITH_IOSTORE_IN_EDITOR
+#if ALT2_ENABLE_LINKERLOAD_SUPPORT
+	void CreateLinker();
+	void DetachLinker();
+#endif
+
+#if WITH_EDITOR
 	void GetLoadedAssetsAndPackages(TSet<FWeakObjectPtr>& AssetList, TSet<UPackage*>& PackageList);
 #endif
 
 private:
-
-	void ImportPackagesRecursiveInner(FIoBatch& IoBatch, FPackageStore& PackageStore, const TArrayView<const FPackageId>& ImportedPackageIds, int32& ImportedPackageIndex);
+	void ImportPackagesRecursiveInner(FAsyncLoadingThreadState2& ThreadState, FIoBatch& IoBatch, FPackageStore& PackageStore, const TArrayView<const FPackageId>& ImportedPackageIds, const TArrayView<const FAsyncPackageHeaderData::FImportPackageNames>& ImportedPackageNames, int32& ImportedPackageIndex);
 
 	uint8 PackageNodesMemory[EEventLoadNode2::Package_NumPhases * sizeof(FEventLoadNode2)];
+	TArrayView<FEventLoadNode2> PackageNodes;
 	/** Basic information associated with this package */
 	FAsyncPackageDesc2 Desc;
 	FAsyncPackageData Data;
@@ -1998,6 +2302,24 @@ private:
 #if WITH_EDITOR
 	TOptional<FAsyncPackageHeaderData> OptionalSegmentHeaderData;
 	TOptional<FAsyncPackageSerializationState> OptionalSegmentSerializationState;
+#endif
+#if ALT2_ENABLE_LINKERLOAD_SUPPORT
+	struct FLinkerLoadState
+	{
+		FLinkerLoad* Linker = nullptr;
+		int32 ProcessingImportedPackageIndex = 0;
+		int32 CreateImportIndex = 0;
+		int32 CreateExportIndex = 0;
+		int32 SerializeExportIndex = 0;
+		int32 PostLoadExportIndex = 0;
+#if WITH_EDITORONLY_DATA
+		TOptional<int32> MetaDataIndex;
+#endif
+		bool bIsCurrentlyProcessingImports = false;
+
+		FAsyncPackageLinkerLoadHeaderData LinkerLoadHeaderData;
+	};
+	TOptional<FLinkerLoadState> LinkerLoadState;
 #endif
 	/** Cached async loading thread object this package was created by */
 	FAsyncLoadingThread2& AsyncLoadingThread;
@@ -2010,6 +2332,7 @@ private:
 	/** Time load begun. This is NOT the time the load was requested in the case of pending requests.	*/
 	double						LoadStartTime = 0.0;
 	TAtomic<int32> RefCount{ 0 };
+	bool						bHasStartedImportingPackages = false;
 	int32						ProcessedExportBundlesCount = 0;
 	/** Current bundle entry index in the current export bundle */
 	int32						ExportBundleEntryIndex = 0;
@@ -2111,6 +2434,9 @@ private:
 			PackageToRemoveState.WaitingForPackage = nullptr;
 		}
 	};
+#if ALT2_ENABLE_LINKERLOAD_SUPPORT
+	FAllDependenciesState		AllDependenciesSetupState;
+#endif
 	FAllDependenciesState		AllDependenciesSerializedState;
 	FAllDependenciesState		AllDependenciesFullyLoadedState;
 
@@ -2175,13 +2501,22 @@ private:
 	void CreateExportBundleNodes(const FAsyncLoadEventSpec* EventSpecs);
 	void SetupSerializedArcs(const FAsyncPackageHeaderData& Header, const TArrayView<FEventLoadNode2>& ExportBundleNodes, const TArrayView<FAsyncPackage2*>& ImportedAsyncPackages);
 	void SetupScriptDependencies();
-	bool HaveAllDependenciesReachedStateDebug(FAsyncLoadingThreadState2& ThreadState, FAsyncPackage2* Package, TSet<FAsyncPackage2*>& VisitedPackages, EAsyncPackageLoadingState2 WaitForPackageState);
-	bool HaveAllDependenciesReachedState(FAsyncLoadingThreadState2& ThreadState, FAllDependenciesState FAsyncPackage2::* StateMemberPtr, EAsyncPackageLoadingState2 WaitForPackageState, uint32 CurrentTick);
-	void UpdateDependenciesStateRecursive(FAsyncLoadingThreadState2& ThreadState, FAllDependenciesState FAsyncPackage2::* StateMemberPtr, EAsyncPackageLoadingState2 WaitForPackageState, uint32 CurrentTick, FAsyncPackage2* Root);
-	void WaitForAllDependenciesToReachState(FAsyncLoadingThreadState2& ThreadState, FAllDependenciesState FAsyncPackage2::* StateMemberPtr, EAsyncPackageLoadingState2 WaitForPackageState, uint32& CurrentTickVariable, TFunctionRef<void(FAsyncPackage2*)> OnStateReached);
+	bool HaveAllDependenciesReachedStateDebug(FAsyncLoadingThreadState2& ThreadState, FAsyncPackage2* Package, TSet<FAsyncPackage2*>& VisitedPackages, EAsyncPackageLoadingState2 WaitForPackageState, bool bExcludeZenPackages);
+	bool HaveAllDependenciesReachedState(FAsyncLoadingThreadState2& ThreadState, FAllDependenciesState FAsyncPackage2::* StateMemberPtr, EAsyncPackageLoadingState2 WaitForPackageState, uint32 CurrentTick, bool bExcludeZenPackages);
+	void UpdateDependenciesStateRecursive(FAsyncLoadingThreadState2& ThreadState, FAllDependenciesState FAsyncPackage2::* StateMemberPtr, EAsyncPackageLoadingState2 WaitForPackageState, uint32 CurrentTick, FAsyncPackage2* Root, bool bExcludeZenPackages);
+	void WaitForAllDependenciesToReachState(FAsyncLoadingThreadState2& ThreadState, FAllDependenciesState FAsyncPackage2::* StateMemberPtr, EAsyncPackageLoadingState2 WaitForPackageState, uint32& CurrentTickVariable, bool bExcludeZenPackages, TFunctionRef<void(FAsyncPackage2*)> OnStateReached);
 	void ConditionalBeginPostLoad(FAsyncLoadingThreadState2& ThreadState);
 	void ConditionalFinishLoading(FAsyncLoadingThreadState2& ThreadState);
 
+#if ALT2_ENABLE_LINKERLOAD_SUPPORT
+	void ConditionalProcessLinkerLoadPackageImportsAndExports(FAsyncLoadingThreadState2& ThreadState);
+	EEventLoadNodeExecutionResult ProcessLinkerLoadPackageSummary(FAsyncLoadingThreadState2& ThreadState);
+	bool ProcessLinkerLoadPackageImports(FAsyncLoadingThreadState2& ThreadState);
+	bool ProcessLinkerLoadPackageExports(FAsyncLoadingThreadState2& ThreadState);
+	EEventLoadNodeExecutionResult ProcessLinkerLoadPackageImportsAndExports(FAsyncLoadingThreadState2& ThreadState);
+	EEventLoadNodeExecutionResult PostLoadLinkerLoadPackageExports(FAsyncLoadingThreadState2& ThreadState);
+#endif
+
 	/**
 	 * Begin async loading process. Simulates parts of BeginLoad.
 	 *
@@ -2197,7 +2532,7 @@ private:
 	 *
 	 * @return true
 	 */
-	void CreateUPackage(const FZenPackageSummary* PackageSummary, const FZenPackageVersioningInfo* VersioningInfo);
+	void CreateUPackage();
 
 	/**
 	 * Finish up UPackage
@@ -2338,6 +2673,9 @@ private:
 	TUniquePtr<FAsyncLoadingThreadState2> GameThreadState;
 	TUniquePtr<FAsyncLoadingThreadState2> AsyncLoadingThreadState;
 
+#if ALT2_ENABLE_LINKERLOAD_SUPPORT
+	uint32 ConditionalProcessLinkerLoadPackageExportsTick = 0;
+#endif
 	uint32 ConditionalBeginPostLoadTick = 0;
 	uint32 ConditionalFinishLoadingTick = 0;
 
@@ -2575,6 +2913,14 @@ public:
 	}
 
 private:
+#if ALT2_ENABLE_LINKERLOAD_SUPPORT
+	bool ShouldLoadPackageFromFileSystem(FName PackageNameToLoad, FPackagePath& OutPackagePath)
+	{
+		return !PackageNameToLoad.IsNone() &&
+			FPackagePath::TryFromPackageName(PackageNameToLoad, OutPackagePath) &&
+			FPackageName::DoesPackageExistEx(OutPackagePath, FPackageName::EPackageLocationFilter::FileSystem, false /* bMatchCaseOnDisk */, &OutPackagePath) != FPackageName::EPackageLocationFilter::None;
+	}
+#endif
 
 	void OnLeakedPackageRename(UPackage* Package);
 
@@ -2685,7 +3031,7 @@ private:
 		return new FAsyncPackage2(Desc, *this, GraphAllocator, EventSpecs.GetData());
 	}
 
-	void InitializeAsyncPackageFromPackageStore(FAsyncPackage2* AsyncPackage, const FPackageStoreEntry& PackageStoreEntry)
+	void InitializeAsyncPackageFromPackageStore(FAsyncLoadingThreadState2& ThreadState, FIoBatch& IoBatch, FAsyncPackage2* AsyncPackage, const FPackageStoreEntry& PackageStoreEntry, const TArrayView<FAsyncPackageHeaderData::FImportPackageNames>& ImportedPackageNames)
 	{
 		TRACE_CPUPROFILER_EVENT_SCOPE(InitializeAsyncPackageFromPackageStore);
 		UE_ASYNC_PACKAGE_DEBUG(AsyncPackage->Desc);
@@ -2700,6 +3046,7 @@ private:
 		const uint64 ExportBundleEntriesMemSize = Align(sizeof(FExportBundleEntry) * ExportBundleEntriesCount, 8);
 		const int32 ImportedPackagesCount = PackageStoreEntry.ImportedPackageIds.Num();
 		const uint64 ImportedPackageIdsMemSize = Align(sizeof(FPackageId) * ImportedPackagesCount, 8);
+		const uint64 ImportedPackageNamesMemSize = Align(sizeof(FAsyncPackageHeaderData::FImportPackageNames) * ImportedPackageNames.Num(), 8);
 #if WITH_EDITOR
 		const int32 OptionalSegmentExportCount = PackageStoreEntry.OptionalSegmentExportInfo.ExportCount;
 		const int32 OptionalSegmentExportBundleCount = PackageStoreEntry.OptionalSegmentExportInfo.ExportBundleCount;
@@ -2734,6 +3081,7 @@ private:
 			ExportBundleHeadersMemSize +
 			ExportBundleEntriesMemSize +
 			ImportedPackageIdsMemSize +
+			ImportedPackageNamesMemSize +
 			ImportedPackagesMemSize +
 			ExportsMemSize +
 			ExportBundleNodesMemSize +
@@ -2764,6 +3112,12 @@ private:
 		HeaderData.ImportedPackageIds = MakeArrayView(reinterpret_cast<FPackageId*>(DataPtr), ImportedPackagesCount);
 		FMemory::Memcpy((void*)HeaderData.ImportedPackageIds.GetData(), PackageStoreEntry.ImportedPackageIds.GetData(), sizeof(FPackageId) * ImportedPackagesCount);
 		DataPtr += ImportedPackageIdsMemSize;
+		if (ImportedPackageNamesMemSize)
+		{
+			HeaderData.ImportedPackageNames = MakeArrayView(reinterpret_cast<FAsyncPackageHeaderData::FImportPackageNames*>(DataPtr), ImportedPackageNames.Num());
+			FMemory::Memcpy((void*)HeaderData.ImportedPackageNames.GetData(), ImportedPackageNames.GetData(), sizeof(FAsyncPackageHeaderData::FImportPackageNames) * ImportedPackageNames.Num());
+			DataPtr += ImportedPackageNamesMemSize;
+		}
 #if WITH_EDITOR
 		if (OptionalSegmentExportCount)
 		{
@@ -2781,7 +3135,6 @@ private:
 		}
 #endif
 		AsyncPackage->ImportStore.AddImportedPackageReferences(HeaderData.ImportedPackageIds);
-		AsyncPackage->ImportStore.AddPackageReference(AsyncPackage->Desc);
 		check(DataPtr - Data.MemoryBuffer == MemoryBufferSize);
 		AsyncPackage->CreateExportBundleNodes(EventSpecs.GetData());
 
@@ -2790,6 +3143,8 @@ private:
 		{
 			Export = FExportObject();
 		}
+
+		AsyncPackage->ImportPackagesRecursive(ThreadState, IoBatch, PackageStore);
 	}
 
 	void DeleteAsyncPackage(FAsyncPackage2* Package)
@@ -3025,8 +3380,7 @@ bool FAsyncLoadingThread2::CreateAsyncPackagesFromQueue(FAsyncLoadingThreadState
 		FPackageStoreEntry PackageEntry;
 		if (EPackageStoreEntryStatus::Ok == PackageStore.GetPackageStoreEntry(PendingPackage->Desc.PackageIdToLoad, PackageEntry))
 		{
-			InitializeAsyncPackageFromPackageStore(PendingPackage, PackageEntry);
-			PendingPackage->ImportPackagesRecursive(IoBatch, PackageStore);
+			InitializeAsyncPackageFromPackageStore(ThreadState, IoBatch, PendingPackage, PackageEntry, TArrayView<FAsyncPackageHeaderData::FImportPackageNames>());
 			PendingPackage->StartLoading(IoBatch);
 			It.RemoveCurrent();
 		}
@@ -3046,6 +3400,8 @@ bool FAsyncLoadingThread2::CreateAsyncPackagesFromQueue(FAsyncLoadingThreadState
 			++NumDequeued;
 		
 			FPackageRequest& Request = OptionalRequest.GetValue();
+			EPackageStoreEntryStatus PackageStatus = EPackageStoreEntryStatus::Missing;
+			FPackageStoreEntry PackageEntry;
 			FName PackageNameToLoad = Request.PackagePath.GetPackageFName();
 			TCHAR NameBuffer[FName::StringBufferSize];
 			uint32 NameLen = PackageNameToLoad.ToString(NameBuffer);
@@ -3068,10 +3424,35 @@ bool FAsyncLoadingThread2::CreateAsyncPackagesFromQueue(FAsyncLoadingThreadState
 				{
 					PackageIdToLoad = RedirectedToPackageId;
 					Request.PackagePath.Empty(); // We no longer know the path but it will be set again when serializing the package summary
+					PackageNameToLoad = NAME_None;
 					UPackageName = SourcePackageName;
 				}
 			}
 
+			PackageStatus = PackageStore.GetPackageStoreEntry(PackageIdToLoad, PackageEntry);
+			if (PackageStatus == EPackageStoreEntryStatus::Missing)
+			{
+				// While there is an active load request for (InName=/Temp/PackageABC_abc, InPackageToLoadFrom=/Game/PackageABC), then allow these requests too:
+				// (InName=/Temp/PackageA_abc, InPackageToLoadFrom=/Temp/PackageABC_abc) and (InName=/Temp/PackageABC_xyz, InPackageToLoadFrom=/Temp/PackageABC_abc)
+				FAsyncPackage2* Package = GetAsyncPackage(PackageIdToLoad);
+				if (Package)
+				{
+					PackageIdToLoad = Package->Desc.PackageIdToLoad;
+					Request.PackagePath = Package->Desc.PackagePathToLoad;
+					PackageNameToLoad = Request.PackagePath.GetPackageFName();
+					PackageStatus = PackageStore.GetPackageStoreEntry(PackageIdToLoad, PackageEntry);
+				}
+			}
+
+#if ALT2_ENABLE_LINKERLOAD_SUPPORT
+			bool bIsZenPackage = true;
+			if (ShouldLoadPackageFromFileSystem(PackageNameToLoad, Request.PackagePath))
+			{
+				bIsZenPackage = false;
+				PackageStatus = EPackageStoreEntryStatus::Ok;
+			}
+#endif
+
 			// Fixup CustomName to handle any input string that can be converted to a long package name.
 			if (!Request.CustomName.IsNone())
 			{
@@ -3088,20 +3469,7 @@ bool FAsyncLoadingThread2::CreateAsyncPackagesFromQueue(FAsyncLoadingThreadState
 				UPackageName = Request.CustomName;
 			}
 
-			FPackageStoreEntry PackageEntry;
-			EPackageStoreEntryStatus PackageStatus = PackageStore.GetPackageStoreEntry(PackageIdToLoad, PackageEntry);
-			if (PackageStatus == EPackageStoreEntryStatus::Missing)
-			{
-				// While there is an active load request for (InName=/Temp/PackageABC_abc, InPackageToLoadFrom=/Game/PackageABC), then allow these requests too:
-				// (InName=/Temp/PackageA_abc, InPackageToLoadFrom=/Temp/PackageABC_abc) and (InName=/Temp/PackageABC_xyz, InPackageToLoadFrom=/Temp/PackageABC_abc)
-				FAsyncPackage2* Package = GetAsyncPackage(PackageIdToLoad);
-				if (Package)
-				{
-					PackageIdToLoad = Package->Desc.PackageIdToLoad;
-					PackageStatus = PackageStore.GetPackageStoreEntry(PackageIdToLoad, PackageEntry);
-				}
-			}
-			FAsyncPackageDesc2 PackageDesc = FAsyncPackageDesc2::FromPackageRequest(Request.RequestId, Request.Priority, UPackageName, PackageIdToLoad, MoveTemp(Request.PackagePath), !Request.CustomName.IsNone());
+			FAsyncPackageDesc2 PackageDesc = FAsyncPackageDesc2::FromPackageRequest(Request, UPackageName, PackageIdToLoad);
 			if (PackageStatus == EPackageStoreEntryStatus::Missing)
 			{
 				QueueMissingPackage(PackageDesc, MoveTemp(Request.PackageLoadedDelegate));
@@ -3122,19 +3490,24 @@ bool FAsyncLoadingThread2::CreateAsyncPackagesFromQueue(FAsyncLoadingThreadState
 						FileOpenLogWrapper->AddPackageToOpenLog(*PackageDesc.UPackageName.ToString());
 					}
 #endif
-					if (PackageStatus == EPackageStoreEntryStatus::Ok)
+					if (PackageStatus == EPackageStoreEntryStatus::Pending)
 					{
-						InitializeAsyncPackageFromPackageStore(Package, PackageEntry);
-						{
-							TRACE_CPUPROFILER_EVENT_SCOPE(ImportPackages);
-							Package->ImportPackagesRecursive(IoBatch, PackageStore);
-						}
-						Package->StartLoading(IoBatch);
+						PendingPackages.Add(Package);
 					}
 					else
 					{
-						check(PackageStatus == EPackageStoreEntryStatus::Pending);
-						PendingPackages.Add(Package);
+						check(PackageStatus == EPackageStoreEntryStatus::Ok);
+#if ALT2_ENABLE_LINKERLOAD_SUPPORT
+						if (!bIsZenPackage)
+						{
+							Package->LinkerLoadState.Emplace();
+						}
+						else
+#endif
+						{
+							InitializeAsyncPackageFromPackageStore(ThreadState, IoBatch, Package, PackageEntry, TArrayView<FAsyncPackageHeaderData::FImportPackageNames>());
+						}
+						Package->StartLoading(IoBatch);
 					}
 				}
 				else
@@ -3681,26 +4054,74 @@ void FGlobalImportStore::RegistrationComplete()
 	ScriptObjects.Shrink();
 }
 
-void FAsyncPackage2::ImportPackagesRecursiveInner(FIoBatch& IoBatch, FPackageStore& PackageStore, const TArrayView<const FPackageId>& ImportedPackageIds, int32& ImportedPackageIndex)
+void FAsyncPackage2::ImportPackagesRecursiveInner(FAsyncLoadingThreadState2& ThreadState, FIoBatch& IoBatch, FPackageStore& PackageStore, const TArrayView<const FPackageId>& ImportedPackageIds, const TArrayView<const FAsyncPackageHeaderData::FImportPackageNames>& ImportedPackageNames, int32& ImportedPackageIndex)
 {
-	for (const FPackageId& ImportedPackageId : ImportedPackageIds)
+	const int32 ImportedPackageCount = ImportedPackageIds.Num();
+	if (!ImportedPackageCount)
+	{
+		return;
+	}
+	bool bHasImportedPackageNames = !ImportedPackageNames.IsEmpty();
+	check(ImportedPackageNames.Num() == 0 || ImportedPackageNames.Num() == ImportedPackageCount);
+	for (int32 LocalImportedPackageIndex = 0; LocalImportedPackageIndex < ImportedPackageCount; ++LocalImportedPackageIndex)
 	{
-		FName ImportedPackageUPackageName = FName();
+		FPackageId ImportedPackageId = ImportedPackageIds[LocalImportedPackageIndex];
+		const FAsyncPackageHeaderData::FImportPackageNames* NamesForPackage = bHasImportedPackageNames ? &ImportedPackageNames[LocalImportedPackageIndex] : nullptr;
+		EPackageStoreEntryStatus ImportedPackageStatus = EPackageStoreEntryStatus::Missing;
+		FPackageStoreEntry ImportedPackageEntry;
+		FName ImportedPackageUPackageName;
+		FName ImportedPackageNameToLoad;
 		FPackageId ImportedPackageIdToLoad = ImportedPackageId;
+		if (NamesForPackage)
+		{
+			if (NamesForPackage->UPackageName != NamesForPackage->PackageNameToLoad)
+			{
+				ImportedPackageIdToLoad = FPackageId::FromName(NamesForPackage->PackageNameToLoad);
+			}
+			ImportedPackageUPackageName = NamesForPackage->UPackageName;
+			ImportedPackageNameToLoad = NamesForPackage->PackageNameToLoad;
+		}
+
 		{
 			FName SourcePackageName;
 			FPackageId RedirectedToPackageId;
-			if (PackageStore.GetPackageRedirectInfo(ImportedPackageId, SourcePackageName, RedirectedToPackageId))
+			if (PackageStore.GetPackageRedirectInfo(ImportedPackageIdToLoad, SourcePackageName, RedirectedToPackageId))
 			{
-				ImportedPackageUPackageName = SourcePackageName;
+				if (ImportedPackageUPackageName.IsNone())
+				{
+					ImportedPackageUPackageName = SourcePackageName;
+				}
 				ImportedPackageIdToLoad = RedirectedToPackageId;
+				ImportedPackageNameToLoad = NAME_None;
 			}
 		}
-		
-		FLoadedPackageRef& ImportedPackageRef = ImportStore.LoadedPackageStore.FindPackageRefChecked(ImportedPackageId);
-		FPackageStoreEntry ImportedPackageEntry;
-		EPackageStoreEntryStatus ImportedPackageStatus = PackageStore.GetPackageStoreEntry(ImportedPackageIdToLoad, ImportedPackageEntry);
+		ImportedPackageStatus = PackageStore.GetPackageStoreEntry(ImportedPackageIdToLoad, ImportedPackageEntry);
+#if WITH_EDITOR
+		if (ImportedPackageStatus == EPackageStoreEntryStatus::Ok && !ImportedPackageEntry.UncookedPackageName.IsNone())
+		{
+			ImportedPackageNameToLoad = ImportedPackageEntry.UncookedPackageName;
+			if (ImportedPackageUPackageName.IsNone())
+			{
+				ImportedPackageUPackageName = ImportedPackageNameToLoad;
+			}
+		}
+#endif
+
+		FPackagePath ImportedPackagePath;
+#if ALT2_ENABLE_LINKERLOAD_SUPPORT
+		bool bIsZenPackage = !LinkerLoadState.IsSet();
+		bool bIsZenPackageImport = true;
+		if (AsyncLoadingThread.ShouldLoadPackageFromFileSystem(ImportedPackageNameToLoad, ImportedPackagePath))
+		{
+			bIsZenPackageImport = false;
+			ImportedPackageStatus = EPackageStoreEntryStatus::Ok;
+		}
+#else
+		constexpr bool bIsZenPackage = true;
+		constexpr bool bIsZenPackageImport = true;
+#endif
 
+		FLoadedPackageRef& ImportedPackageRef = ImportStore.LoadedPackageStore.FindPackageRefChecked(ImportedPackageId);
 		if (ImportedPackageStatus == EPackageStoreEntryStatus::Missing)
 		{
 			UE_ASYNC_PACKAGE_LOG(Warning, Desc, TEXT("ImportPackages: SkipPackage"),
@@ -3710,7 +4131,7 @@ void FAsyncPackage2::ImportPackagesRecursiveInner(FIoBatch& IoBatch, FPackageSto
 			continue;
 		}
 #if WITH_EDITOR
-		else if (!ImportedPackageEntry.UncookedPackageName.IsNone())
+		else if (AsyncLoadingThread.UncookedPackageLoader && !ImportedPackageEntry.UncookedPackageName.IsNone())
 		{
 			UPackage* UncookedPackage = nullptr;
 			if (!ImportedPackageRef.AreAllPublicExportsLoaded())
@@ -3718,7 +4139,7 @@ void FAsyncPackage2::ImportPackagesRecursiveInner(FIoBatch& IoBatch, FPackageSto
 				UE_ASYNC_PACKAGE_LOG(Verbose, Desc, TEXT("ImportPackages: LoadUncookedImport"), TEXT("Loading imported uncooked package '%s' '0x%llX'"), *ImportedPackageEntry.UncookedPackageName.ToString(), ImportedPackageId.ValueForDebugging());
 				check(IsInGameThread());
 				IoBatch.Issue(); // The batch might already contain requests for packages being imported from the uncooked one we're going to load so make sure that those are started before blocking
-				FPackagePath ImportedPackagePath = FPackagePath::FromPackageNameUnchecked(ImportedPackageEntry.UncookedPackageName);
+				ImportedPackagePath = FPackagePath::FromPackageNameUnchecked(ImportedPackageEntry.UncookedPackageName);
 				ImportedPackagePath.SetHeaderExtension(static_cast<EPackageExtension>(ImportedPackageEntry.UncookedPackageHeaderExtension));
 				int32 ImportRequestId = AsyncLoadingThread.UncookedPackageLoader->LoadPackage(ImportedPackagePath, NAME_None, FLoadPackageAsyncDelegate(), PKG_None, INDEX_NONE, 0, nullptr);
 				AsyncLoadingThread.UncookedPackageLoader->FlushLoading(ImportRequestId);
@@ -3782,7 +4203,7 @@ void FAsyncPackage2::ImportPackagesRecursiveInner(FIoBatch& IoBatch, FPackageSto
 
 		FAsyncPackage2* ImportedPackage = nullptr;
 		bool bInserted = false;
-		FAsyncPackageDesc2 PackageDesc = FAsyncPackageDesc2::FromPackageImport(Desc.ReferencerRequestId, Desc.Priority, ImportedPackageId, ImportedPackageIdToLoad, ImportedPackageUPackageName);
+		FAsyncPackageDesc2 PackageDesc = FAsyncPackageDesc2::FromPackageImport(Desc, ImportedPackageUPackageName, ImportedPackageId, ImportedPackageIdToLoad, MoveTemp(ImportedPackagePath));
 		if (ImportedPackageRef.AreAllPublicExportsLoaded())
 		{
 			ImportedPackage = AsyncLoadingThread.FindAsyncPackage(ImportedPackageId);
@@ -3814,32 +4235,52 @@ void FAsyncPackage2::ImportPackagesRecursiveInner(FIoBatch& IoBatch, FPackageSto
 		}
 		ImportedPackage->AddRef();
 		Data.ImportedAsyncPackages[ImportedPackageIndex++] = ImportedPackage;
-		GetPackageNode(Package_SetupDependencies).DependsOn(&ImportedPackage->GetPackageNode(Package_ProcessSummary));
+
+		if (bIsZenPackage != bIsZenPackageImport)
+		{
+			// When importing a linker load package from a zen package or vice versa we need to wait for all the exports in the imported package to be created
+			// and serialized before we can start processing our own exports
+			GetPackageNode(Package_SetupDependencies).DependsOn(&ImportedPackage->GetPackageNode(Package_ExportsSerialized));
+		}
+		else if (bIsZenPackage)
+		{
+			// When importing a zen package from another zen package we need the summary of the imported package to be ready before we can hook up our export
+			// bundle node dependenices to the imported package
+			GetPackageNode(Package_SetupDependencies).DependsOn(&ImportedPackage->GetPackageNode(Package_ProcessSummary));
+		}
 
 		if (bInserted)
 		{
-			if (ImportedPackageStatus == EPackageStoreEntryStatus::Ok)
+			if (ImportedPackageStatus == EPackageStoreEntryStatus::Pending)
 			{
-				AsyncLoadingThread.InitializeAsyncPackageFromPackageStore(ImportedPackage, ImportedPackageEntry);
-				ImportedPackage->ImportPackagesRecursive(IoBatch, PackageStore);
-				ImportedPackage->StartLoading(IoBatch);
+				AsyncLoadingThread.PendingPackages.Add(ImportedPackage);
 			}
 			else
 			{
-				check(ImportedPackageStatus == EPackageStoreEntryStatus::Pending);
-				AsyncLoadingThread.PendingPackages.Add(ImportedPackage);
-			}
+				check(ImportedPackageStatus == EPackageStoreEntryStatus::Ok);
+#if ALT2_ENABLE_LINKERLOAD_SUPPORT
+				if (!bIsZenPackageImport)
+				{
+					ImportedPackage->LinkerLoadState.Emplace();
+				}
+				else
+#endif
+				{
+					AsyncLoadingThread.InitializeAsyncPackageFromPackageStore(ThreadState, IoBatch, ImportedPackage, ImportedPackageEntry, TArrayView<FAsyncPackageHeaderData::FImportPackageNames>());
+				}
+				ImportedPackage->StartLoading(IoBatch);
+			}
 		}
 	}
 }
 
-void FAsyncPackage2::ImportPackagesRecursive(FIoBatch& IoBatch, FPackageStore& PackageStore)
+void FAsyncPackage2::ImportPackagesRecursive(FAsyncLoadingThreadState2& ThreadState, FIoBatch& IoBatch, FPackageStore& PackageStore)
 {
-	if (AsyncPackageLoadingState >= EAsyncPackageLoadingState2::ImportPackages)
+	if (bHasStartedImportingPackages)
 	{
 		return;
 	}
-	check(AsyncPackageLoadingState == EAsyncPackageLoadingState2::NewPackage);
+	bHasStartedImportingPackages = true;
 
 	int32 ImportedPackageCount = HeaderData.ImportedPackageIds.Num();
 #if WITH_EDITOR
@@ -3850,23 +4291,18 @@ void FAsyncPackage2::ImportPackagesRecursive(FIoBatch& IoBatch, FPackageStore& P
 #endif
 	if (!ImportedPackageCount)
 	{
-		AsyncPackageLoadingState = EAsyncPackageLoadingState2::ImportPackagesDone;
 		return;
 	}
-	else
-	{
-		AsyncPackageLoadingState = EAsyncPackageLoadingState2::ImportPackages;
-	}
 
 	int32 ImportedPackageIndex = 0;
 
 	Data.ImportedAsyncPackages = MakeArrayView(Data.ImportedAsyncPackages.GetData(), ImportedPackageCount);
 	
-	ImportPackagesRecursiveInner(IoBatch, PackageStore, HeaderData.ImportedPackageIds, ImportedPackageIndex);
+	ImportPackagesRecursiveInner(ThreadState, IoBatch, PackageStore, HeaderData.ImportedPackageIds, HeaderData.ImportedPackageNames, ImportedPackageIndex);
 #if WITH_EDITOR
 	if (OptionalSegmentHeaderData.IsSet())
 	{
-		ImportPackagesRecursiveInner(IoBatch, PackageStore, OptionalSegmentHeaderData->ImportedPackageIds, ImportedPackageIndex);
+		ImportPackagesRecursiveInner(ThreadState, IoBatch, PackageStore, OptionalSegmentHeaderData->ImportedPackageIds, TArrayView<FAsyncPackageHeaderData::FImportPackageNames>(), ImportedPackageIndex);
 	}
 #endif
 
@@ -3884,22 +4320,70 @@ void FAsyncPackage2::ImportPackagesRecursive(FIoBatch& IoBatch, FPackageStore& P
 
 	UE_ASYNC_PACKAGE_LOG_VERBOSE(VeryVerbose, Desc, TEXT("ImportPackages: ImportsDone"),
 		TEXT("All imported packages are now being loaded."));
+}
+
+#if ALT2_ENABLE_LINKERLOAD_SUPPORT
+void FAsyncPackage2::CreateLinker()
+{
+	FLinkerLoad* Linker = FLinkerLoad::FindExistingLinkerForPackage(LinkerRoot);
+	if (!Linker)
+	{
+		uint32 LinkerFlags = (LOAD_Async | LOAD_NoVerify | LOAD_SkipLoadImportedPackages);
+		if ((!FApp::IsGame() || GIsEditor) && (Desc.PackageFlags & PKG_PlayInEditor) != 0)
+		{
+			LinkerFlags |= LOAD_PackageForPIE;
+		}
+		const FLinkerInstancingContext* InstancingContext = Desc.GetInstancingContext();
+		FUObjectSerializeContext* LoadContext = GetSerializeContext();
+#if ALT2_ENABLE_NEW_ARCHIVE_FOR_LINKERLOAD
+		Linker = new FLinkerLoad(LinkerRoot, Desc.PackagePathToLoad, LinkerFlags, InstancingContext ? *InstancingContext : FLinkerInstancingContext());
+		Linker->SetSerializeContext(LoadContext);
+		LinkerRoot->SetLinker(Linker);
+		FLinkerLoadArchive2* Loader = new FLinkerLoadArchive2(Desc.PackagePathToLoad);
+		Linker->SetLoader(Loader, Loader->NeedsEngineVersionChecks());
+#else
+		Linker = FLinkerLoad::CreateLinkerAsync(LoadContext, LinkerRoot, Desc.PackagePathToLoad, LinkerFlags, Desc.GetInstancingContext(), TFunction<void()>([]() {}));
+#endif
+	}
+	check(Linker);
+	check(Linker->LinkerRoot == LinkerRoot);
+	LinkerLoadState->Linker = Linker;
+#if ALT2_ENABLE_NEW_ARCHIVE_FOR_LINKERLOAD
+	Linker->ResetStatusInfo();
+	static_cast<FLinkerLoadArchive2*>(Linker->GetLoader())->BeginRead(&GetPackageNode(EEventLoadNode2::Package_ProcessSummary));
+#else
+	GetPackageNode(EEventLoadNode2::Package_ProcessSummary).ReleaseBarrier();
+#endif
+}
 
-	check(AsyncPackageLoadingState == EAsyncPackageLoadingState2::ImportPackages);
-	AsyncPackageLoadingState = EAsyncPackageLoadingState2::ImportPackagesDone;
+void FAsyncPackage2::DetachLinker()
+{
+	if (LinkerLoadState.IsSet() && LinkerLoadState->Linker)
+	{
+		LinkerLoadState->Linker = nullptr;
+	}
 }
+#endif
 
 void FAsyncPackage2::StartLoading(FIoBatch& IoBatch)
 {
 	TRACE_CPUPROFILER_EVENT_SCOPE(StartLoading);
 	TRACE_LOADTIME_BEGIN_LOAD_ASYNC_PACKAGE(this);
-	check(AsyncPackageLoadingState == EAsyncPackageLoadingState2::ImportPackagesDone);
 
 	LoadStartTime = FPlatformTime::Seconds();
 
 	AsyncPackageLoadingState = EAsyncPackageLoadingState2::WaitingForIo;
 
 	FIoReadOptions ReadOptions;
+#if ALT2_ENABLE_LINKERLOAD_SUPPORT
+	if (LinkerLoadState.IsSet())
+	{
+		CreateUPackage();
+		CreateLinker();
+		return;
+	}
+#endif
+
 #if WITH_EDITOR
 	if (OptionalSegmentHeaderData.IsSet())
 	{
@@ -4015,6 +4499,397 @@ static void ReadAsyncPackageHeader(FAsyncPackageSerializationState& Serializatio
 	SerializationState.AllExportDataPtr = PackageHeaderDataPtr + PackageSummary->HeaderSize;
 }
 
+#if ALT2_ENABLE_LINKERLOAD_SUPPORT
+void FAsyncPackage2::ConditionalProcessLinkerLoadPackageImportsAndExports(FAsyncLoadingThreadState2& ThreadState)
+{
+	TRACE_CPUPROFILER_EVENT_SCOPE(ConditionalProcessLinkerLoadPackageExports);
+	constexpr bool bExcludeZenPackages = true;
+	WaitForAllDependenciesToReachState(ThreadState, &FAsyncPackage2::AllDependenciesSetupState, EAsyncPackageLoadingState2::ProcessExportBundles, AsyncLoadingThread.ConditionalProcessLinkerLoadPackageExportsTick, bExcludeZenPackages,
+		[](FAsyncPackage2* Package)
+		{
+			Package->AsyncPackageLoadingState = EAsyncPackageLoadingState2::ProcessExportBundles;
+			Package->GetExportBundleNode(ExportBundle_Process, 0).ReleaseBarrier();
+		});
+}
+
+EEventLoadNodeExecutionResult FAsyncPackage2::ProcessLinkerLoadPackageSummary(FAsyncLoadingThreadState2& ThreadState)
+{
+	TRACE_CPUPROFILER_EVENT_SCOPE(ProcessLinkerLoadPackageSummary);
+
+#if ALT2_ENABLE_NEW_ARCHIVE_FOR_LINKERLOAD
+	if (Package->LinkerLoadState->Linker->GetLoader()->IsError())
+	{
+		Package->bLoadHasFailed = true;
+		return 	return EEventLoadNodeExecutionResult::Complete;
+	}
+	LinkerLoadState->Linker->bUseTimeLimit = false;
+	FLinkerLoad::ELinkerStatus LinkerResult = LinkerLoadState->Linker->ProcessPackageSummary(nullptr);
+#else
+	FLinkerLoad::ELinkerStatus LinkerResult = LinkerLoadState->Linker->Tick(/* RemainingTimeLimit */ 0.0, /* bUseTimeLimit */ false, /* bUseFullTimeLimit */ false, nullptr);
+#endif
+	check(LinkerResult != FLinkerLoad::LINKER_TimedOut); // TODO: Add support for timeouts here
+	if (LinkerResult == FLinkerLoad::LINKER_Failed)
+	{
+		bLoadHasFailed = true;
+		return EEventLoadNodeExecutionResult::Complete;
+	}
+	check(LinkerLoadState->Linker->HasFinishedInitialization());
+
+	LinkerLoadState->LinkerLoadHeaderData.ImportMap.SetNum(LinkerLoadState->Linker->ImportMap.Num());
+	TArray<FAsyncPackageHeaderData::FImportPackageNames, TInlineAllocator<128>> ImportedPackageNames;
+	TArray<FPackageId, TInlineAllocator<128>> ImportedPackageIds;
+	for (int32 ImportIndex = 0; ImportIndex < LinkerLoadState->Linker->ImportMap.Num(); ++ImportIndex)
+	{
+		const FObjectImport& LinkerImport = LinkerLoadState->Linker->ImportMap[ImportIndex];
+		TArray<int32, TInlineAllocator<128>> PathComponents;
+		int32 PathIndex = ImportIndex;
+		PathComponents.Push(PathIndex);
+		while (LinkerLoadState->Linker->ImportMap[PathIndex].OuterIndex.IsImport() && !LinkerLoadState->Linker->ImportMap[PathIndex].HasPackageName())
+		{
+			PathIndex = LinkerLoadState->Linker->ImportMap[PathIndex].OuterIndex.ToImport();
+			PathComponents.Push(PathIndex);
+		}
+		int32 PackageImportIndex = PathComponents.Top();
+		FObjectImport& PackageImport = LinkerLoadState->Linker->ImportMap[PackageImportIndex];
+		FName ImportPackageName;
+		const bool bImportHasPackageName = PackageImport.HasPackageName();
+		if (bImportHasPackageName)
+		{
+			ImportPackageName = PackageImport.GetPackageName();
+		}
+		else
+		{
+			ImportPackageName =  PackageImport.ObjectName;
+		}
+		TCHAR NameStr[FName::StringBufferSize];
+		uint32 NameLen = ImportPackageName.ToString(NameStr);
+		bool bIsScriptImport = FPackageName::IsScriptPackage(FStringView(NameStr, NameLen));
+		if (bIsScriptImport)
+		{
+			check(!bImportHasPackageName);
+			TStringBuilder<256> FullPath;
+			while (!PathComponents.IsEmpty())
+			{
+				NameLen = LinkerLoadState->Linker->ImportMap[PathComponents.Pop(false)].ObjectName.ToString(NameStr);
+				FPathViews::Append(FullPath, FStringView(NameStr, NameLen));
+				LinkerLoadState->LinkerLoadHeaderData.ImportMap[ImportIndex] = FPackageObjectIndex::FromScriptPath(FullPath);
+			}
+		}
+		else
+		{
+			FPackageId ImportedPackageId = FPackageId::FromName(ImportPackageName);
+			int32 ImportedPackageIndex = ImportedPackageIds.AddUnique(ImportedPackageId);
+			if (ImportedPackageIndex == ImportedPackageNames.Num())
+			{
+				ImportedPackageNames.AddDefaulted();
+			}
+			ImportedPackageNames[ImportedPackageIndex].PackageNameToLoad = ImportPackageName;
+			bool bIsPackageImport = ImportIndex == PackageImportIndex;
+			if (!bIsPackageImport || bImportHasPackageName)
+			{
+				if (!bImportHasPackageName)
+				{
+					PathComponents.Pop(false);
+				}
+				TStringBuilder<256> PackageRelativeExportPath;
+				while (!PathComponents.IsEmpty())
+				{
+					NameLen = LinkerLoadState->Linker->ImportMap[PathComponents.Pop(false)].ObjectName.ToString(NameStr);
+					for (uint32 I = 0; I < NameLen; ++I)
+					{
+						NameStr[I] = TChar<TCHAR>::ToLower(NameStr[I]);
+					}
+					PackageRelativeExportPath.AppendChar('/');
+					PackageRelativeExportPath.Append(FStringView(NameStr, NameLen));
+				}
+				FPackageImportReference PackageImportRef(ImportedPackageIndex, LinkerLoadState->LinkerLoadHeaderData.ImportedPublicExportHashes.Num());
+				LinkerLoadState->LinkerLoadHeaderData.ImportMap[ImportIndex] = FPackageObjectIndex::FromPackageImportRef(PackageImportRef);
+				LinkerLoadState->LinkerLoadHeaderData.ImportedPublicExportHashes.Add(CityHash64(reinterpret_cast<const char*>(PackageRelativeExportPath.GetData() + 1), (PackageRelativeExportPath.Len() - 1) * sizeof(TCHAR)));
+			}
+		}
+	}
+
+	const FLinkerInstancingContext& InstancingContext = LinkerLoadState->Linker->GetInstancingContext();
+	for (int32 PackageImportIndex = 0; PackageImportIndex < ImportedPackageNames.Num(); ++PackageImportIndex)
+	{
+		FAsyncPackageHeaderData::FImportPackageNames& ImportedPackage = ImportedPackageNames[PackageImportIndex];
+		ImportedPackage.UPackageName = InstancingContext.Remap(ImportedPackage.PackageNameToLoad);
+		if (ImportedPackage.UPackageName != ImportedPackage.PackageNameToLoad)
+		{
+			ImportedPackageIds[PackageImportIndex] = FPackageId::FromName(ImportedPackage.UPackageName);
+		}
+	}
+
+	LinkerLoadState->LinkerLoadHeaderData.PublicExportHashes.SetNum(LinkerLoadState->Linker->ExportMap.Num());
+	for (int32 ExportIndex = 0; ExportIndex < LinkerLoadState->Linker->ExportMap.Num(); ++ExportIndex)
+	{
+		FObjectExport& ObjectExport = LinkerLoadState->Linker->ExportMap[ExportIndex];
+		//if ((ObjectExport.ObjectFlags & RF_Public) > 0) // We need hashes for all objects, external actors are breaking assumptions about public exports
+		{
+			TArray<int32, TInlineAllocator<128>> FullPath;
+			int32 PathIndex = ExportIndex;
+			FullPath.Push(PathIndex);
+			while (LinkerLoadState->Linker->ExportMap[PathIndex].OuterIndex.IsExport())
+			{
+				PathIndex = LinkerLoadState->Linker->ExportMap[PathIndex].OuterIndex.ToExport();
+				FullPath.Push(PathIndex);
+			}
+			TStringBuilder<256> PackageRelativeExportPath;
+			while (!FullPath.IsEmpty())
+			{
+				TCHAR NameStr[FName::StringBufferSize];
+				uint32 NameLen = LinkerLoadState->Linker->ExportMap[FullPath.Pop(false)].ObjectName.ToString(NameStr);
+				for (uint32 I = 0; I < NameLen; ++I)
+				{
+					NameStr[I] = TChar<TCHAR>::ToLower(NameStr[I]);
+				}
+				PackageRelativeExportPath.AppendChar('/');
+				PackageRelativeExportPath.Append(FStringView(NameStr, NameLen));
+			}
+			uint64 PublicExportHash = CityHash64(reinterpret_cast<const char*>(PackageRelativeExportPath.GetData() + 1), (PackageRelativeExportPath.Len() - 1) * sizeof(TCHAR));
+			LinkerLoadState->LinkerLoadHeaderData.PublicExportHashes[ExportIndex] = PublicExportHash;
+		}
+	}
+
+	FPackageStoreEntry PackageStoreEntry;
+	PackageStoreEntry.ExportInfo.ExportCount = LinkerLoadState->Linker->ExportMap.Num();
+	PackageStoreEntry.ExportInfo.ExportBundleCount = 1;
+	PackageStoreEntry.ImportedPackageIds = ImportedPackageIds;
+
+	FIoBatch IoBatch = AsyncLoadingThread.IoDispatcher.NewBatch();
+	{
+		FPackageStoreReadScope _(AsyncLoadingThread.PackageStore);
+		AsyncLoadingThread.InitializeAsyncPackageFromPackageStore(ThreadState, IoBatch, this, PackageStoreEntry, ImportedPackageNames);
+	}
+	IoBatch.Issue();
+
+	HeaderData.ImportedPublicExportHashes = LinkerLoadState->LinkerLoadHeaderData.ImportedPublicExportHashes;
+	HeaderData.ImportMap = LinkerLoadState->LinkerLoadHeaderData.ImportMap;
+
+	AsyncPackageLoadingState = EAsyncPackageLoadingState2::SetupDependencies;
+	GetPackageNode(Package_SetupDependencies).ReleaseBarrier();
+	return EEventLoadNodeExecutionResult::Complete;
+}
+
+EEventLoadNodeExecutionResult FAsyncPackage2::ProcessLinkerLoadPackageImportsAndExports(FAsyncLoadingThreadState2& ThreadState)
+{
+	ON_SCOPE_EXIT
+	{
+		FUObjectSerializeContext * LoadContext = GetSerializeContext();
+		TArray<UObject*>& ThreadObjLoaded = LoadContext->PRIVATE_GetObjectsLoadedInternalUseOnly();
+		ThreadObjLoaded.Reset();
+	};
+
+	if (!ProcessLinkerLoadPackageImports(ThreadState))
+	{
+		return EEventLoadNodeExecutionResult::Timeout;
+	}
+	if (!ProcessLinkerLoadPackageExports(ThreadState))
+	{
+		return EEventLoadNodeExecutionResult::Timeout;
+	}
+
+	if (ExternalReadDependencies.Num() == 0)
+	{
+		AsyncPackageLoadingState = EAsyncPackageLoadingState2::ProcessExportBundles;
+		GetPackageNode(Package_ExportsSerialized).ReleaseBarrier(&ThreadState);
+	}
+	else
+	{
+		AsyncPackageLoadingState = EAsyncPackageLoadingState2::WaitingForExternalReads;
+		AsyncLoadingThread.ExternalReadQueue.Enqueue(this);
+	}
+	return EEventLoadNodeExecutionResult::Complete;
+}
+
+bool FAsyncPackage2::ProcessLinkerLoadPackageImports(FAsyncLoadingThreadState2& ThreadState)
+{
+	if (LinkerLoadState->bIsCurrentlyProcessingImports)
+	{
+		return true;
+	}
+	TGuardValue<bool> GuardIsCurrentlyProcessingImports(LinkerLoadState->bIsCurrentlyProcessingImports, true);
+	const int32 ImportedPackagesCount = Data.ImportedAsyncPackages.Num();
+	while (LinkerLoadState->ProcessingImportedPackageIndex < ImportedPackagesCount)
+	{
+		FAsyncPackage2* ImportedPackage = Data.ImportedAsyncPackages[LinkerLoadState->ProcessingImportedPackageIndex];
+		if (ImportedPackage && ImportedPackage->LinkerLoadState.IsSet() && ImportedPackage->LinkerLoadState->Linker)
+		{
+			if (!ImportedPackage->ProcessLinkerLoadPackageImports(ThreadState))
+			{
+				return false;
+			}
+			if (!ImportedPackage->ProcessLinkerLoadPackageExports(ThreadState))
+			{
+				return false;
+			}
+		}
+		++LinkerLoadState->ProcessingImportedPackageIndex;
+	}
+
+	const int32 ImportCount = LinkerLoadState->Linker->ImportMap.Num();
+	while (LinkerLoadState->CreateImportIndex < ImportCount)
+	{
+		const int32 ImportIndex = LinkerLoadState->CreateImportIndex++;
+		const FPackageObjectIndex& GlobalImportIndex = HeaderData.ImportMap[ImportIndex];
+		if (!GlobalImportIndex.IsNull())
+		{
+			UObject* FromImportStore = ImportStore.FindOrGetImportObject(HeaderData, GlobalImportIndex);
+#if ALT2_VERIFY_LINKERLOAD_MATCHES_IMPORTSTORE
+			/*if (Desc.UPackageId.ValueForDebugging() == 0xF37353A71BF5C938 && ImportIndex == 0x0000005d)
+			{
+				UE_DEBUG_BREAK();
+			}*/
+			UObject* FromLinker = LinkerLoadState->Linker->CreateImport(ImportIndex);
+			if (FromImportStore != FromLinker)
+			{
+				bool bIsAcceptableDeviation = false;
+				FObjectImport& LinkerImport = LinkerLoadState->Linker->ImportMap[ImportIndex];
+				if (FromLinker)
+				{
+					check(LinkerImport.SourceLinker);
+					check(LinkerImport.SourceIndex >= 0);
+					FObjectExport& SourceExport = LinkerImport.SourceLinker->ExportMap[LinkerImport.SourceIndex];
+					if (!FromImportStore && SourceExport.bExportLoadFailed)
+					{
+						bIsAcceptableDeviation = true; // Exports can be marked as failed after they have already been returned to CreateImport. Is this a bug or a feature?
+					}
+					else if (FromImportStore && FromImportStore->GetName() == FromLinker->GetName() && FromLinker->GetOutermost() == GetTransientPackage())
+					{
+						bIsAcceptableDeviation = true; // Linker are sometimes stuck with objects that have been moved to the transient package. Is this a bug or a feature?
+					}
+				}
+				check(bIsAcceptableDeviation);
+			}
+#endif
+			FObjectImport& LinkerImport = LinkerLoadState->Linker->ImportMap[ImportIndex];
+			if (!LinkerImport.XObject)
+			{
+				LinkerImport.XObject = FromImportStore;
+			}
+		}
+	}
+	return true;
+}
+
+bool FAsyncPackage2::ProcessLinkerLoadPackageExports(FAsyncLoadingThreadState2& ThreadState)
+{
+#if WITH_EDITORONLY_DATA
+	// Create metadata object
+	int32 MetaDataIndex = LinkerLoadState->MetaDataIndex.IsSet() ? LinkerLoadState->MetaDataIndex.GetValue() : -1;
+	if (MetaDataIndex < 0)
+	{
+		MetaDataIndex = LinkerLoadState->Linker->LoadMetaDataFromExportMap(false);
+		if (MetaDataIndex >= 0)
+		{
+			FObjectExport& LinkerExport = LinkerLoadState->Linker->ExportMap[MetaDataIndex];
+			FExportObject& ExportObject = Data.Exports[MetaDataIndex];
+			ExportObject.Object = LinkerExport.Object;
+			if (ExportObject.Object)
+			{
+				ExportObject.bWasFoundInMemory = true; // Make sure that the async flags are cleared in ClearConstructedObjects
+			}
+			else
+			{
+				ExportObject.bExportLoadFailed = LinkerExport.bExportLoadFailed;
+				if (!ExportObject.bExportLoadFailed)
+				{
+					ExportObject.bFiltered = true;
+				}
+			}
+		}
+		LinkerLoadState->MetaDataIndex = MetaDataIndex;
+	}
+#endif
+
+	// Create exports
+	const int32 ExportCount = LinkerLoadState->Linker->ExportMap.Num();
+	while (LinkerLoadState->CreateExportIndex < ExportCount)
+	{
+		const int32 ExportIndex = LinkerLoadState->CreateExportIndex++;
+#if WITH_EDITORONLY_DATA
+		if (ExportIndex == MetaDataIndex)
+		{
+			continue;
+		}
+#endif
+		FObjectExport& LinkerExport = LinkerLoadState->Linker->ExportMap[ExportIndex];
+		FExportObject& ExportObject = Data.Exports[ExportIndex];
+		if (UObject* Object = LinkerLoadState->Linker->CreateExport(ExportIndex))
+		{
+			ExportObject.Object = Object;
+			ExportObject.bWasFoundInMemory = true; // Make sure that the async flags are cleared in ClearConstructedObjects
+			EInternalObjectFlags FlagsToSet = EInternalObjectFlags::Async;
+			uint64 PublicExportHash = LinkerLoadState->LinkerLoadHeaderData.PublicExportHashes[ExportIndex];
+			if (Desc.bCanBeImported && PublicExportHash)
+			{
+				FlagsToSet |= EInternalObjectFlags::LoaderImport;
+				ImportStore.StoreGlobalObject(Desc.UPackageId, PublicExportHash, Object);
+			}
+			Object->SetInternalFlags(FlagsToSet);
+		}
+		else
+		{
+			ExportObject.bExportLoadFailed = LinkerExport.bExportLoadFailed;
+			if (!ExportObject.bExportLoadFailed)
+			{
+				ExportObject.bFiltered = true;
+			}
+		}
+
+		if (ThreadState.IsTimeLimitExceeded(TEXT("CreateLinkerLoadExports")))
+		{
+			return false;
+		}
+	}
+
+	// Serialize exports
+	while (LinkerLoadState->SerializeExportIndex < ExportCount)
+	{
+		const int32 ExportIndex = LinkerLoadState->SerializeExportIndex++;
+		FExportObject& ExportObject = Data.Exports[ExportIndex];
+		if (UObject* Object = ExportObject.Object)
+		{
+			if (Object->HasAnyFlags(RF_NeedLoad))
+			{
+				LinkerLoadState->Linker->Preload(Object);
+			}
+		}
+		if (ThreadState.IsTimeLimitExceeded(TEXT("SerializeLinkerLoadExports")))
+		{
+			return false;
+		}
+	}
+
+	return true;
+}
+
+EEventLoadNodeExecutionResult FAsyncPackage2::PostLoadLinkerLoadPackageExports(FAsyncLoadingThreadState2& ThreadState)
+{
+	FAsyncLoadingTickScope2 InAsyncLoadingTick(AsyncLoadingThread);
+
+	const int32 ExportCount = LinkerLoadState->Linker->ExportMap.Num();
+	while (LinkerLoadState->PostLoadExportIndex < ExportCount)
+	{
+		const int32 ExportIndex = LinkerLoadState->PostLoadExportIndex++;
+		const FObjectExport& Export = LinkerLoadState->Linker->ExportMap[ExportIndex];
+		UObject* Object = Export.Object;
+		if (Object)
+		{
+			Object->ConditionalPostLoad();
+		}
+	}
+
+	FUObjectSerializeContext* LoadContext = GetSerializeContext();
+	TArray<UObject*>& ThreadObjLoaded = LoadContext->PRIVATE_GetObjectsLoadedInternalUseOnly();
+	ThreadObjLoaded.Empty();
+
+	AsyncPackageLoadingState = EAsyncPackageLoadingState2::DeferredPostLoadDone;
+	ConditionalFinishLoading(ThreadState);
+	return EEventLoadNodeExecutionResult::Complete;;
+}
+#endif
+
 EEventLoadNodeExecutionResult FAsyncPackage2::Event_ProcessPackageSummary(FAsyncLoadingThreadState2& ThreadState, FAsyncPackage2* Package, int32)
 {
 	TRACE_CPUPROFILER_EVENT_SCOPE(Event_ProcessPackageSummary);
@@ -4024,6 +4899,16 @@ EEventLoadNodeExecutionResult FAsyncPackage2::Event_ProcessPackageSummary(FAsync
 
 	FScopedAsyncPackageEvent2 Scope(Package);
 
+#if ALT2_ENABLE_LINKERLOAD_SUPPORT
+	if (Package->LinkerLoadState.IsSet())
+	{
+		EEventLoadNodeExecutionResult Result = Package->ProcessLinkerLoadPackageSummary(ThreadState);
+		if (!Package->bLoadHasFailed)
+		{
+			return Result;
+		}
+	}
+#endif
 	if (Package->bLoadHasFailed)
 	{
 		if (Package->Desc.bCanBeImported)
@@ -4109,7 +4994,22 @@ EEventLoadNodeExecutionResult FAsyncPackage2::Event_ProcessPackageSummary(FAsync
 			Package->Desc.UPackageName = Package->HeaderData.PackageName;
 		}
 		check(Package->Desc.UPackageId == FPackageId::FromName(Package->Desc.UPackageName));
-		Package->CreateUPackage(Package->HeaderData.PackageSummary, Package->HeaderData.VersioningInfo.GetPtrOrNull());
+		Package->CreateUPackage();
+		Package->LinkerRoot->SetPackageFlags(Package->HeaderData.PackageSummary->PackageFlags);
+#if WITH_EDITOR
+		Package->LinkerRoot->bIsCookedForEditor = !!(Package->HeaderData.PackageSummary->PackageFlags & PKG_FilterEditorOnly);
+#endif
+		if (const FZenPackageVersioningInfo* VersioningInfo = Package->HeaderData.VersioningInfo.GetPtrOrNull())
+		{
+			Package->LinkerRoot->SetLinkerPackageVersion(VersioningInfo->PackageVersion);
+			Package->LinkerRoot->SetLinkerLicenseeVersion(VersioningInfo->LicenseeVersion);
+			Package->LinkerRoot->SetLinkerCustomVersions(VersioningInfo->CustomVersions);
+		}
+		else
+		{
+			Package->LinkerRoot->SetLinkerPackageVersion(GPackageFileUEVersion);
+			Package->LinkerRoot->SetLinkerLicenseeVersion(GPackageFileLicenseeUEVersion);
+		}
 
 		TRACE_LOADTIME_PACKAGE_SUMMARY(Package, Package->HeaderData.PackageName, Package->HeaderData.PackageSummary->HeaderSize, Package->HeaderData.ImportMap.Num(), Package->HeaderData.ExportMap.Num());
 	}
@@ -4123,7 +5023,20 @@ EEventLoadNodeExecutionResult FAsyncPackage2::Event_SetupDependencies(FAsyncLoad
 {
 	TRACE_CPUPROFILER_EVENT_SCOPE(Event_SetupDependencies);
 	check(Package->AsyncPackageLoadingState == EAsyncPackageLoadingState2::SetupDependencies);
-	
+
+#if ALT2_ENABLE_LINKERLOAD_SUPPORT
+	if (Package->LinkerLoadState.IsSet())
+	{
+		if (!Package->AsyncLoadingThread.bHasRegisteredAllScriptObjects)
+		{
+			Package->SetupScriptDependencies();
+		}
+		Package->AsyncPackageLoadingState = EAsyncPackageLoadingState2::ProcessExportBundles;
+		Package->ConditionalProcessLinkerLoadPackageImportsAndExports(ThreadState);
+		return EEventLoadNodeExecutionResult::Complete;
+	}
+#endif
+
 	if (!Package->bLoadHasFailed)
 	{
 		if (!Package->AsyncLoadingThread.bHasRegisteredAllScriptObjects)
@@ -4156,6 +5069,13 @@ EEventLoadNodeExecutionResult FAsyncPackage2::Event_ProcessExportBundle(FAsyncLo
 
 	FScopedAsyncPackageEvent2 Scope(Package);
 
+#if ALT2_ENABLE_LINKERLOAD_SUPPORT
+	if (Package->LinkerLoadState.IsSet())
+	{
+		return Package->ProcessLinkerLoadPackageImportsAndExports(ThreadState);
+	}
+#endif
+
 	auto FilterExport = [](const EExportFilterFlags FilterFlags) -> bool
 	{
 #if WITH_EDITOR
@@ -4701,12 +5621,18 @@ EEventLoadNodeExecutionResult FAsyncPackage2::Event_ExportsDone(FAsyncLoadingThr
 	return EEventLoadNodeExecutionResult::Complete;
 }
 
-bool FAsyncPackage2::HaveAllDependenciesReachedStateDebug(FAsyncLoadingThreadState2& ThreadState, FAsyncPackage2* Package, TSet<FAsyncPackage2*>& VisitedPackages, EAsyncPackageLoadingState2 WaitForPackageState)
+bool FAsyncPackage2::HaveAllDependenciesReachedStateDebug(FAsyncLoadingThreadState2& ThreadState, FAsyncPackage2* Package, TSet<FAsyncPackage2*>& VisitedPackages, EAsyncPackageLoadingState2 WaitForPackageState, bool bExcludeZenPackages)
 {
 	for (FAsyncPackage2* ImportedPackage : Package->Data.ImportedAsyncPackages)
 	{
 		if (ImportedPackage)
 		{
+#if ALT2_ENABLE_LINKERLOAD_SUPPORT
+			if (bExcludeZenPackages && !ImportedPackage->LinkerLoadState.IsSet())
+			{
+				continue;
+			}
+#endif
 			if (ThreadState.PackagesExcludedFromStateDependencyChecks.Contains(ImportedPackage))
 			{
 				continue;
@@ -4722,7 +5648,7 @@ bool FAsyncPackage2::HaveAllDependenciesReachedStateDebug(FAsyncLoadingThreadSta
 				return false;
 			}
 
-			if (!HaveAllDependenciesReachedStateDebug(ThreadState, ImportedPackage, VisitedPackages, WaitForPackageState))
+			if (!HaveAllDependenciesReachedStateDebug(ThreadState, ImportedPackage, VisitedPackages, WaitForPackageState, bExcludeZenPackages))
 			{
 				return false;
 			}
@@ -4731,7 +5657,7 @@ bool FAsyncPackage2::HaveAllDependenciesReachedStateDebug(FAsyncLoadingThreadSta
 	return true;
 }
 
-bool FAsyncPackage2::HaveAllDependenciesReachedState(FAsyncLoadingThreadState2& ThreadState, FAllDependenciesState FAsyncPackage2::* StateMemberPtr, EAsyncPackageLoadingState2 WaitForPackageState, uint32 CurrentTick)
+bool FAsyncPackage2::HaveAllDependenciesReachedState(FAsyncLoadingThreadState2& ThreadState, FAllDependenciesState FAsyncPackage2::* StateMemberPtr, EAsyncPackageLoadingState2 WaitForPackageState, uint32 CurrentTick, bool bExcludeZenPackages)
 {
 	FAllDependenciesState& ThisState = this->*StateMemberPtr;
 	if (ThisState.bAllDone)
@@ -4743,12 +5669,12 @@ bool FAsyncPackage2::HaveAllDependenciesReachedState(FAsyncLoadingThreadState2&
 		return false;
 	}
 	ThisState.UpdateTick(CurrentTick);
-	UpdateDependenciesStateRecursive(ThreadState, StateMemberPtr, WaitForPackageState, CurrentTick, this);
+	UpdateDependenciesStateRecursive(ThreadState, StateMemberPtr, WaitForPackageState, CurrentTick, this, bExcludeZenPackages);
 	check(ThisState.bAllDone || (ThisState.WaitingForPackage && ThisState.WaitingForPackage->AsyncPackageLoadingState <= WaitForPackageState));
 	return ThisState.bAllDone;
 }
 
-void FAsyncPackage2::UpdateDependenciesStateRecursive(FAsyncLoadingThreadState2& ThreadState, FAllDependenciesState FAsyncPackage2::* StateMemberPtr, EAsyncPackageLoadingState2 WaitForPackageState, uint32 CurrentTick, FAsyncPackage2* Root)
+void FAsyncPackage2::UpdateDependenciesStateRecursive(FAsyncLoadingThreadState2& ThreadState, FAllDependenciesState FAsyncPackage2::* StateMemberPtr, EAsyncPackageLoadingState2 WaitForPackageState, uint32 CurrentTick, FAsyncPackage2* Root, bool bExcludeZenPackages)
 {
 	FAllDependenciesState& ThisState = this->*StateMemberPtr;
 
@@ -4777,7 +5703,7 @@ void FAsyncPackage2::UpdateDependenciesStateRecursive(FAsyncLoadingThreadState2&
 				WaitingForPackageState.UpdateTick(CurrentTick);
 				if (!WaitingForPackageState.bVisitedMark)
 				{
-					WaitingForPackage->UpdateDependenciesStateRecursive(ThreadState, StateMemberPtr, WaitForPackageState, CurrentTick, Root);
+					WaitingForPackage->UpdateDependenciesStateRecursive(ThreadState, StateMemberPtr, WaitForPackageState, CurrentTick, Root, bExcludeZenPackages);
 				}
 				if (WaitingForPackageState.bAnyNotDone)
 				{
@@ -4797,6 +5723,13 @@ void FAsyncPackage2::UpdateDependenciesStateRecursive(FAsyncLoadingThreadState2&
 			continue;
 		}
 
+#if ALT2_ENABLE_LINKERLOAD_SUPPORT
+		if (bExcludeZenPackages && !ImportedPackage->LinkerLoadState.IsSet())
+		{
+			continue;
+		}
+#endif
+
 		if (ThreadState.PackagesExcludedFromStateDependencyChecks.Contains(ImportedPackage))
 		{
 			continue;
@@ -4817,7 +5750,7 @@ void FAsyncPackage2::UpdateDependenciesStateRecursive(FAsyncLoadingThreadState2&
 		}
 		else if (!ImportedPackageState.bVisitedMark)
 		{
-			ImportedPackage->UpdateDependenciesStateRecursive(ThreadState, StateMemberPtr, WaitForPackageState, CurrentTick, Root);
+			ImportedPackage->UpdateDependenciesStateRecursive(ThreadState, StateMemberPtr, WaitForPackageState, CurrentTick, Root, bExcludeZenPackages);
 		}
 
 		if (ImportedPackageState.bAnyNotDone)
@@ -4849,9 +5782,9 @@ void FAsyncPackage2::UpdateDependenciesStateRecursive(FAsyncLoadingThreadState2&
 	}
 }
 
-void FAsyncPackage2::WaitForAllDependenciesToReachState(FAsyncLoadingThreadState2& ThreadState, FAllDependenciesState FAsyncPackage2::* StateMemberPtr, EAsyncPackageLoadingState2 WaitForPackageState, uint32& CurrentTickVariable, TFunctionRef<void(FAsyncPackage2*)> OnStateReached)
+void FAsyncPackage2::WaitForAllDependenciesToReachState(FAsyncLoadingThreadState2& ThreadState, FAllDependenciesState FAsyncPackage2::* StateMemberPtr, EAsyncPackageLoadingState2 WaitForPackageState, uint32& CurrentTickVariable, bool bExcludeZenPackages, TFunctionRef<void(FAsyncPackage2*)> OnStateReached)
 {
-	if (HaveAllDependenciesReachedState(ThreadState, StateMemberPtr, WaitForPackageState, CurrentTickVariable++))
+	if (HaveAllDependenciesReachedState(ThreadState, StateMemberPtr, WaitForPackageState, CurrentTickVariable++, bExcludeZenPackages))
 	{
 		FAsyncPackage2* FirstPackageReadyToProceed = this;
 
@@ -4868,14 +5801,14 @@ void FAsyncPackage2::WaitForAllDependenciesToReachState(FAsyncLoadingThreadState
 
 #if DO_CHECK
 			TSet<FAsyncPackage2*> VisitedPackages;
-			check(HaveAllDependenciesReachedStateDebug(ThreadState, this, VisitedPackages, WaitForPackageState));
+			check(HaveAllDependenciesReachedStateDebug(ThreadState, this, VisitedPackages, WaitForPackageState, bExcludeZenPackages));
 #endif
 
 			while (FAsyncPackage2* WaitingPackage = PackageReadyToProceedState.PackagesWaitingForThisHead)
 			{
 				FAllDependenciesState& WaitingPackageState = WaitingPackage->*StateMemberPtr;
 				check(WaitingPackageState.WaitingForPackage == PackageReadyToProceed);
-				if (WaitingPackage->HaveAllDependenciesReachedState(ThreadState, StateMemberPtr, WaitForPackageState, CurrentTickVariable++))
+				if (WaitingPackage->HaveAllDependenciesReachedState(ThreadState, StateMemberPtr, WaitForPackageState, CurrentTickVariable++, bExcludeZenPackages))
 				{
 					FAllDependenciesState::RemoveFromWaitList(StateMemberPtr, PackageReadyToProceed, WaitingPackage);
 					WaitingPackageState.NextLink = FirstPackageReadyToProceed;
@@ -4893,7 +5826,8 @@ void FAsyncPackage2::WaitForAllDependenciesToReachState(FAsyncLoadingThreadState
 void FAsyncPackage2::ConditionalBeginPostLoad(FAsyncLoadingThreadState2& ThreadState)
 {
 	TRACE_CPUPROFILER_EVENT_SCOPE(ConditionalBeginPostLoad);
-	WaitForAllDependenciesToReachState(ThreadState, &FAsyncPackage2::AllDependenciesSerializedState, EAsyncPackageLoadingState2::ExportsDone, AsyncLoadingThread.ConditionalBeginPostLoadTick,
+	constexpr bool bExcludeZenPackages = false;
+	WaitForAllDependenciesToReachState(ThreadState, &FAsyncPackage2::AllDependenciesSerializedState, EAsyncPackageLoadingState2::ExportsDone, AsyncLoadingThread.ConditionalBeginPostLoadTick, bExcludeZenPackages,
 		[](FAsyncPackage2* Package)
 		{
 			for (int32 ExportBundleIndex = 0; ExportBundleIndex < Package->Data.ExportBundleCount; ++ExportBundleIndex)
@@ -4906,7 +5840,8 @@ void FAsyncPackage2::ConditionalBeginPostLoad(FAsyncLoadingThreadState2& ThreadS
 void FAsyncPackage2::ConditionalFinishLoading(FAsyncLoadingThreadState2& ThreadState)
 {
 	TRACE_CPUPROFILER_EVENT_SCOPE(ConditionalFinishLoading);
-	WaitForAllDependenciesToReachState(ThreadState, &FAsyncPackage2::AllDependenciesFullyLoadedState, EAsyncPackageLoadingState2::DeferredPostLoadDone, AsyncLoadingThread.ConditionalFinishLoadingTick,
+	constexpr bool bExcludeZenPackages = false;
+	WaitForAllDependenciesToReachState(ThreadState, &FAsyncPackage2::AllDependenciesFullyLoadedState, EAsyncPackageLoadingState2::DeferredPostLoadDone, AsyncLoadingThread.ConditionalFinishLoadingTick, bExcludeZenPackages,
 		[](FAsyncPackage2* Package)
 		{
 			Package->AsyncLoadingThread.LoadedPackagesToProcess.Add(Package);
@@ -4922,6 +5857,16 @@ EEventLoadNodeExecutionResult FAsyncPackage2::Event_PostLoadExportBundle(FAsyncL
 	
 	FAsyncPackageScope2 PackageScope(Package);
 
+#if ALT2_ENABLE_LINKERLOAD_SUPPORT
+	if (Package->LinkerLoadState.IsSet())
+	{
+		// No async postload for now
+		Package->AsyncPackageLoadingState = EAsyncPackageLoadingState2::DeferredPostLoad;
+		Package->GetExportBundleNode(ExportBundle_DeferredPostLoad, 0).ReleaseBarrier();
+		return EEventLoadNodeExecutionResult::Complete;
+	}
+#endif
+
 	/*TSet<FAsyncPackage2*> Visited;
 	TArray<FAsyncPackage2*> ProcessQueue;
 	ProcessQueue.Push(Package);
@@ -5068,6 +6013,13 @@ EEventLoadNodeExecutionResult FAsyncPackage2::Event_DeferredPostLoadExportBundle
 
 	FAsyncPackageScope2 PackageScope(Package);
 
+#if ALT2_ENABLE_LINKERLOAD_SUPPORT
+	if (Package->LinkerLoadState.IsSet())
+	{
+		return Package->PostLoadLinkerLoadPackageExports(ThreadState);
+	}
+#endif
+
 	check(InExportBundleIndex < Package->Data.ExportBundleCount);
 	EEventLoadNodeExecutionResult LoadingState = EEventLoadNodeExecutionResult::Complete;
 
@@ -5163,7 +6115,7 @@ EEventLoadNodeExecutionResult FAsyncPackage2::Event_DeferredPostLoadExportBundle
 FEventLoadNode2& FAsyncPackage2::GetPackageNode(EEventLoadNode2 Phase)
 {
 	check(Phase < EEventLoadNode2::Package_NumPhases);
-	return *(reinterpret_cast<FEventLoadNode2*>(PackageNodesMemory) + Phase);
+	return PackageNodes[Phase];
 }
 
 FEventLoadNode2& FAsyncPackage2::GetExportBundleNode(EEventLoadNode2 Phase, uint32 ExportBundleIndex)
@@ -5449,6 +6401,9 @@ EAsyncPackageState::Type FAsyncLoadingThread2::ProcessLoadedPackagesFromGameThre
 #endif
 					Package->ClearConstructedObjects();
 				}
+#if ALT2_ENABLE_LINKERLOAD_SUPPORT
+				Package->DetachLinker();
+#endif
 			}
 
 			// Remove the package from the list before we trigger the callbacks, 
@@ -5539,7 +6494,12 @@ EAsyncPackageState::Type FAsyncLoadingThread2::ProcessLoadedPackagesFromGameThre
 
 	if (Result == EAsyncPackageState::Complete)
 	{
-#if WITH_IOSTORE_IN_EDITOR
+#if WITH_EDITORONLY_DATA
+		// This needs to happen after loading new blueprints in the editor, and this is handled in EndLoad for synchronous loads
+		FBlueprintSupport::FlushReinstancingQueue();
+#endif
+
+#if WITH_EDITOR
 		// In editor builds, call the asset load callback. This happens in both editor and standalone to match EndLoad
 		TSet<FWeakObjectPtr> TempLoadedAssets = LoadedAssets;
 		LoadedAssets.Reset();
@@ -5578,6 +6538,12 @@ EAsyncPackageState::Type FAsyncLoadingThread2::TickAsyncLoadingFromGameThread(FA
 	check(IsInGameThread());
 	check(!IsGarbageCollecting());
 
+#if WITH_EDITOR
+	// In the editor loading cannot be part of a transaction as it cannot be undone, and may result in recording half-loaded objects. So we suppress any active transaction while in this stack, and set the editor loading flag
+	TGuardValue<ITransaction*> SuppressTransaction(GUndo, nullptr);
+	TGuardValue<bool> IsEditorLoadingPackage(GIsEditorLoadingPackage, GIsEditor || GIsEditorLoadingPackage);
+#endif
+
 	const bool bLoadingSuspended = IsAsyncLoadingSuspended();
 	EAsyncPackageState::Type Result = bLoadingSuspended ? EAsyncPackageState::PendingImports : EAsyncPackageState::Complete;
 
@@ -6269,6 +7235,8 @@ FAsyncPackage2::FAsyncPackage2(
 	AddRequestID(Desc.RequestID);
 
 	CreatePackageNodes(EventSpecs);
+
+	ImportStore.AddPackageReference(Desc);
 }
 
 void FAsyncPackage2::CreatePackageNodes(const FAsyncLoadEventSpec* EventSpecs)
@@ -6278,9 +7246,9 @@ void FAsyncPackage2::CreatePackageNodes(const FAsyncLoadEventSpec* EventSpecs)
 	FEventLoadNode2* Node = reinterpret_cast<FEventLoadNode2*>(PackageNodesMemory);
 	for (int32 Phase = 0; Phase < EEventLoadNode2::Package_NumPhases; ++Phase)
 	{
-		new (Node) FEventLoadNode2(EventSpecs + Phase, this, -1, BarrierCount);
-		++Node;
+		new (Node + Phase) FEventLoadNode2(EventSpecs + Phase, this, -1, BarrierCount);
 	}
+	PackageNodes = MakeArrayView(Node, EEventLoadNode2::Package_NumPhases);
 }
 
 void FAsyncPackage2::CreateExportBundleNodes(const FAsyncLoadEventSpec* EventSpecs)
@@ -6302,6 +7270,12 @@ FAsyncPackage2::~FAsyncPackage2()
 	UE_ASYNC_PACKAGE_LOG(Verbose, Desc, TEXT("AsyncThread: Deleted"), TEXT("Package deleted."));
 
 	ImportStore.ReleaseImportedPackageReferences(Desc, HeaderData.ImportedPackageIds);
+#if WITH_EDITOR
+	if (OptionalSegmentHeaderData.IsSet())
+	{
+		ImportStore.ReleaseImportedPackageReferences(Desc, OptionalSegmentHeaderData->ImportedPackageIds);
+	}
+#endif
 	ImportStore.ReleasePackageReference(Desc);
 
 	checkf(RefCount == 0, TEXT("RefCount is not 0 when deleting package %s"),
@@ -6446,7 +7420,7 @@ void FAsyncPackage2::EndAsyncLoad()
 	}
 }
 
-void FAsyncPackage2::CreateUPackage(const FZenPackageSummary* PackageSummary, const FZenPackageVersioningInfo* VersioningInfo)
+void FAsyncPackage2::CreateUPackage()
 {
 	check(!LinkerRoot);
 
@@ -6490,39 +7464,20 @@ void FAsyncPackage2::CreateUPackage(const FZenPackageSummary* PackageSummary, co
 			LinkerRoot = NewObject<UPackage>(/*Outer*/nullptr, Desc.UPackageName);
 			bCreatedLinkerRoot = true;
 		}
-		LinkerRoot->SetFlags(RF_Public | RF_WasLoaded);
-		LinkerRoot->SetLoadedPath(Desc.PackagePathToLoad);
-		LinkerRoot->SetCanBeImportedFlag(Desc.bCanBeImported);
-		LinkerRoot->SetPackageId(Desc.UPackageId);
-		LinkerRoot->SetPackageFlagsTo(PackageSummary->PackageFlags | PKG_Cooked);
-		if (VersioningInfo)
-		{
-			LinkerRoot->SetLinkerPackageVersion(VersioningInfo->PackageVersion);
-			LinkerRoot->SetLinkerLicenseeVersion(VersioningInfo->LicenseeVersion);
-			LinkerRoot->SetLinkerCustomVersions(VersioningInfo->CustomVersions);
-		}
-		else
-		{
-			LinkerRoot->SetLinkerPackageVersion(GPackageFileUEVersion);
-			LinkerRoot->SetLinkerLicenseeVersion(GPackageFileLicenseeUEVersion);
-		}
-#if WITH_IOSTORE_IN_EDITOR
-		LinkerRoot->bIsCookedForEditor = !!(PackageSummary->PackageFlags & PKG_FilterEditorOnly);
-#endif
 		if (PackageRef)
 		{
 			PackageRef->SetPackage(LinkerRoot);
 		}
 	}
-	else
-	{
-		LinkerRoot->SetPackageFlags(PackageSummary->PackageFlags | PKG_Cooked);
-		check(LinkerRoot->CanBeImported() == Desc.bCanBeImported);
-		check(LinkerRoot->GetPackageId() == Desc.UPackageId);
-		check(LinkerRoot->GetLinkerPackageVersion() == GPackageFileUEVersion);
-		check(LinkerRoot->GetLinkerLicenseeVersion() == GPackageFileLicenseeUEVersion);
-		check(LinkerRoot->HasAnyFlags(RF_WasLoaded));
-	}
+
+	LinkerRoot->SetFlags(RF_Public | RF_WasLoaded);
+	LinkerRoot->SetLoadedPath(Desc.PackagePathToLoad);
+	LinkerRoot->SetCanBeImportedFlag(Desc.bCanBeImported);
+	LinkerRoot->SetPackageId(Desc.UPackageId);
+	LinkerRoot->SetPackageFlags(Desc.PackageFlags);
+#if WITH_EDITOR
+	LinkerRoot->SetPIEInstanceID(Desc.PIEInstanceID);
+#endif
 
 	EInternalObjectFlags FlagsToSet = EInternalObjectFlags::Async;
 	if (Desc.bCanBeImported)
@@ -6673,7 +7628,7 @@ void FAsyncPackage2::AddCompletionCallback(TUniquePtr<FLoadPackageAsyncDelegate>
 	CompletionCallbacks.Emplace(MoveTemp(Callback));
 }
 
-int32 FAsyncLoadingThread2::LoadPackage(const FPackagePath& InPackagePath, FName InCustomName, FLoadPackageAsyncDelegate InCompletionDelegate, EPackageFlags InPackageFlags, int32 InPIEInstanceID, int32 InPackagePriority, const FLinkerInstancingContext*)
+int32 FAsyncLoadingThread2::LoadPackage(const FPackagePath& InPackagePath, FName InCustomName, FLoadPackageAsyncDelegate InCompletionDelegate, EPackageFlags InPackageFlags, int32 InPIEInstanceID, int32 InPackagePriority, const FLinkerInstancingContext* InInstancingContext)
 {
 	TRACE_CPUPROFILER_EVENT_SCOPE(LoadPackage);
 
@@ -6699,7 +7654,7 @@ int32 FAsyncLoadingThread2::LoadPackage(const FPackagePath& InPackagePath, FName
 		? MakeUnique<FLoadPackageAsyncDelegate>(MoveTemp(InCompletionDelegate))
 		: TUniquePtr<FLoadPackageAsyncDelegate>();
 
-	PackageRequestQueue.Enqueue(FPackageRequest::Create(RequestId, InPackagePriority, InPackagePath, InCustomName, MoveTemp(CompletionDelegate)));
+	PackageRequestQueue.Enqueue(FPackageRequest::Create(RequestId, InPackageFlags, InPIEInstanceID, InPackagePriority, InInstancingContext, InPackagePath, InCustomName, MoveTemp(CompletionDelegate)));
 	++QueuedPackagesCounter;
 	++PackagesWithRemainingWorkCounter;
 
@@ -6908,8 +7863,6 @@ IAsyncPackageLoader* MakeAsyncPackageLoader2(FIoDispatcher& InIoDispatcher, IAsy
 	return new FAsyncLoadingThread2(InIoDispatcher, InUncookedPackageLoader);
 }
 
-#endif //WITH_ASYNCLOADING2
-
 #if UE_BUILD_DEVELOPMENT || UE_BUILD_DEBUG
 PRAGMA_ENABLE_OPTIMIZATION
 #endif
diff --git a/Engine/Source/Runtime/CoreUObject/Private/Serialization/AsyncPackageLoader.cpp b/Engine/Source/Runtime/CoreUObject/Private/Serialization/AsyncPackageLoader.cpp
--- a/Engine/Source/Runtime/CoreUObject/Private/Serialization/AsyncPackageLoader.cpp
+++ b/Engine/Source/Runtime/CoreUObject/Private/Serialization/AsyncPackageLoader.cpp
@@ -185,22 +185,35 @@ void SetAsyncLoadingAllowed(bool bAllowAsyncLoading)
 void InitAsyncThread()
 {
 	LLM_SCOPE(ELLMTag::AsyncLoading);
-#if WITH_ASYNCLOADING2
 	if (FIoDispatcher::IsInitialized())
 	{
+		bool bSettingsEnabled = false;
+		bool bCommandLineEnabled = false;
+		bool bHasUseIoStoreParamInEditor = false;
+#if WITH_EDITOR
+		bCommandLineEnabled = FParse::Param(FCommandLine::Get(), TEXT("ZenLoader"));
+		check(GConfig);
+		GConfig->GetBool(TEXT("/Script/Engine.EditorStreamingSettings"), TEXT("s.ZenLoaderEnabled"), bSettingsEnabled, GEngineIni);
+		bHasUseIoStoreParamInEditor = FParse::Param(FCommandLine::Get(), TEXT("UseIoStore"));
+#endif
 		FIoDispatcher& IoDispatcher = FIoDispatcher::Get();
 		bool bHasScriptObjectsChunk = IoDispatcher.DoesChunkExist(CreateIoChunkId(0, 0, EIoChunkType::ScriptObjects));
-		bool bHasUseIoStoreParamInEditor = WITH_EDITOR && FParse::Param(FCommandLine::Get(), TEXT("UseIoStore"));
-		if (bHasScriptObjectsChunk || bHasUseIoStoreParamInEditor)
+		if (bSettingsEnabled || bCommandLineEnabled)
 		{
+			GPackageLoader.Reset(MakeAsyncPackageLoader2(IoDispatcher));
+		}
 #if WITH_EDITOR
+		else if (bHasScriptObjectsChunk || bHasUseIoStoreParamInEditor)
+		{
 			GPackageLoader = MakeEditorPackageLoader(IoDispatcher);
+		}
 #else
+		else if (bHasScriptObjectsChunk)
+		{
 			GPackageLoader.Reset(MakeAsyncPackageLoader2(IoDispatcher));
-#endif
 		}
-	}
 #endif
+	}
 	if (!GPackageLoader.IsValid())
 	{
 		GPackageLoader = MakeUnique<FAsyncLoadingThread>(/** ThreadIndex = */ 0);
diff --git a/Engine/Source/Runtime/CoreUObject/Private/UObject/LinkerLoad.cpp b/Engine/Source/Runtime/CoreUObject/Private/UObject/LinkerLoad.cpp
--- a/Engine/Source/Runtime/CoreUObject/Private/UObject/LinkerLoad.cpp
+++ b/Engine/Source/Runtime/CoreUObject/Private/UObject/LinkerLoad.cpp
@@ -699,6 +699,116 @@ FUObjectSerializeContext* FLinkerLoad::GetSerializeContext()
 	return FUObjectThreadContext::Get().GetSerializeContext();
 }
 
+FLinkerLoad::ELinkerStatus FLinkerLoad::ProcessPackageSummary(TMap<TPair<FName, FPackageIndex>, FPackageIndex>* ObjectNameWithOuterToExportMap)
+{
+	ELinkerStatus Status = LINKER_Loaded;
+	{
+		SCOPED_LOADTIMER(LinkerLoad_SerializePackageFileSummary);
+		Status = SerializePackageFileSummary();
+	}
+
+	// Serialize the header for the package trailer
+	if (Status == LINKER_Loaded)
+	{
+		SCOPED_LOADTIMER(LinkerLoad_SerializePackageTrailer);
+		Status = SerializePackageTrailer();
+	}
+
+	// Serialize the name map and register the names.
+	if( Status == LINKER_Loaded )
+	{
+		SCOPED_LOADTIMER(LinkerLoad_SerializeNameMap);
+		Status = SerializeNameMap();
+	}
+
+	// Serialize the gatherable text data map.
+	if( Status == LINKER_Loaded )
+	{
+		SCOPED_LOADTIMER(LinkerLoad_SerializeGatherableTextDataMap);
+		Status = SerializeGatherableTextDataMap();
+	}
+
+	// Serialize the import map.
+	if( Status == LINKER_Loaded )
+	{
+		SCOPED_LOADTIMER(LinkerLoad_SerializeImportMap);
+		Status = SerializeImportMap();
+	}
+
+	// Serialize the export map.
+	if( Status == LINKER_Loaded )
+	{
+		SCOPED_LOADTIMER(LinkerLoad_SerializeExportMap);
+		Status = SerializeExportMap();
+	}
+
+#if WITH_TEXT_ARCHIVE_SUPPORT
+	// Construct the exports readers
+	if (Status == LINKER_Loaded)
+	{
+		SCOPED_LOADTIMER(LinkerLoad_ConstructExportsReaders);
+		Status = ConstructExportsReaders();
+	}
+#endif
+
+	// Fix up import map for backward compatible serialization.
+	if( Status == LINKER_Loaded )
+	{	
+		SCOPED_LOADTIMER(LinkerLoad_FixupImportMap);
+		Status = FixupImportMap();
+	}
+
+	// Populate the linker instancing context for instance loading if needed.
+	if (Status == LINKER_Loaded)
+	{
+		SCOPED_LOADTIMER(LinkerLoad_PopulateInstancingContext);
+		Status = PopulateInstancingContext();
+	}
+
+	// Fix up export map for object class conversion 
+	if( Status == LINKER_Loaded )
+	{	
+		SCOPED_LOADTIMER(LinkerLoad_FixupExportMap);
+		Status = FixupExportMap();
+	}
+
+	// Serialize the dependency map.
+	if( Status == LINKER_Loaded )
+	{
+		SCOPED_LOADTIMER(LinkerLoad_SerializeDependsMap);
+		Status = SerializeDependsMap();
+	}
+
+	// Hash exports.
+	if( Status == LINKER_Loaded )
+	{
+		SCOPED_LOADTIMER(LinkerLoad_CreateExportHash);
+		Status = CreateExportHash();
+	}
+
+	// Find existing objects matching exports and associate them with this linker.
+	if( Status == LINKER_Loaded )
+	{
+		SCOPED_LOADTIMER(LinkerLoad_FindExistingExports);
+		Status = FindExistingExports();
+	}
+
+	if (Status == LINKER_Loaded)
+	{
+		SCOPED_LOADTIMER(LinkerLoad_SerializePreloadDependencies);
+		Status = SerializePreloadDependencies();
+	}
+
+	// Finalize creation process.
+	if( Status == LINKER_Loaded )
+	{
+		SCOPED_LOADTIMER(LinkerLoad_FinalizeCreation);
+		Status = FinalizeCreation(ObjectNameWithOuterToExportMap);
+	}
+
+	return Status;
+}
+
 /**
  * Ticks an in-flight linker and spends InTimeLimit seconds on creation. This is a soft time limit used
  * if bInUseTimeLimit is true.
@@ -743,107 +853,7 @@ FLinkerLoad::ELinkerStatus FLinkerLoad::Tick( float InTimeLimit, bool bInUseTime
 			// Serialize the package file summary and presize the various arrays (name, import & export map)
 			if (bCanSerializePackageFileSummary)
 			{
-				SCOPED_LOADTIMER(LinkerLoad_SerializePackageFileSummary);
-				Status = SerializePackageFileSummary();
-			}
-
-			// Serialize the header for the package trailer
-			if (Status == LINKER_Loaded)
-			{
-				SCOPED_LOADTIMER(LinkerLoad_SerializePackageTrailer);
-				Status = SerializePackageTrailer();
-			}
-
-			// Serialize the name map and register the names.
-			if( Status == LINKER_Loaded )
-			{
-				SCOPED_LOADTIMER(LinkerLoad_SerializeNameMap);
-				Status = SerializeNameMap();
-			}
-
-			// Serialize the gatherable text data map.
-			if( Status == LINKER_Loaded )
-			{
-				SCOPED_LOADTIMER(LinkerLoad_SerializeGatherableTextDataMap);
-				Status = SerializeGatherableTextDataMap();
-			}
-
-			// Serialize the import map.
-			if( Status == LINKER_Loaded )
-			{
-				SCOPED_LOADTIMER(LinkerLoad_SerializeImportMap);
-				Status = SerializeImportMap();
-			}
-
-			// Serialize the export map.
-			if( Status == LINKER_Loaded )
-			{
-				SCOPED_LOADTIMER(LinkerLoad_SerializeExportMap);
-				Status = SerializeExportMap();
-			}
-
-#if WITH_TEXT_ARCHIVE_SUPPORT
-			// Construct the exports readers
-			if (Status == LINKER_Loaded)
-			{
-				SCOPED_LOADTIMER(LinkerLoad_ConstructExportsReaders);
-				Status = ConstructExportsReaders();
-			}
-#endif
-
-			// Fix up import map for backward compatible serialization.
-			if( Status == LINKER_Loaded )
-			{	
-				SCOPED_LOADTIMER(LinkerLoad_FixupImportMap);
-				Status = FixupImportMap();
-			}
-
-			// Populate the linker instancing context for instance loading if needed.
-			if (Status == LINKER_Loaded)
-			{
-				SCOPED_LOADTIMER(LinkerLoad_PopulateInstancingContext);
-				Status = PopulateInstancingContext();
-			}
-
-			// Fix up export map for object class conversion 
-			if( Status == LINKER_Loaded )
-			{	
-				SCOPED_LOADTIMER(LinkerLoad_FixupExportMap);
-				Status = FixupExportMap();
-			}
-
-			// Serialize the dependency map.
-			if( Status == LINKER_Loaded )
-			{
-				SCOPED_LOADTIMER(LinkerLoad_SerializeDependsMap);
-				Status = SerializeDependsMap();
-			}
-
-			// Hash exports.
-			if( Status == LINKER_Loaded )
-			{
-				SCOPED_LOADTIMER(LinkerLoad_CreateExportHash);
-				Status = CreateExportHash();
-			}
-
-			// Find existing objects matching exports and associate them with this linker.
-			if( Status == LINKER_Loaded )
-			{
-				SCOPED_LOADTIMER(LinkerLoad_FindExistingExports);
-				Status = FindExistingExports();
-			}
-
-			if (Status == LINKER_Loaded)
-			{
-				SCOPED_LOADTIMER(LinkerLoad_SerializePreloadDependencies);
-				Status = SerializePreloadDependencies();
-			}
-
-			// Finalize creation process.
-			if( Status == LINKER_Loaded )
-			{
-				SCOPED_LOADTIMER(LinkerLoad_FinalizeCreation);
-				Status = FinalizeCreation(ObjectNameWithOuterToExportMap);
+				Status = ProcessPackageSummary(ObjectNameWithOuterToExportMap);
 			}
 		}
 		// Loop till we are done if no time limit is specified, or loop until the real time limit is up if we want to use full time
@@ -3076,6 +3086,15 @@ bool FLinkerLoad::VerifyImportInner(const int32 ImportIndex, FString& WarningSuf
 		// In the other case we do not want to trigger another load of the objects in that import, in case they contain dependencies to the package we are currently loading
 		// and the current loader doesn't have the LOAD_DeferDependencyLoads flag
 		Package = FindObjectFast<UPackage>(nullptr, PackageToLoadInto);
+		if (LoadFlags & LOAD_SkipLoadImportedPackages)
+		{
+			if (!Package)
+			{
+				return nullptr;
+			}
+			Import.SourceLinker = FindExistingLinkerForPackage(Package);
+			return Package;
+		}
 		if (Package == nullptr || !Package->IsFullyLoaded())
 		{
 			{
diff --git a/Engine/Source/Runtime/CoreUObject/Private/UObject/UObjectLinker.cpp b/Engine/Source/Runtime/CoreUObject/Private/UObject/UObjectLinker.cpp
--- a/Engine/Source/Runtime/CoreUObject/Private/UObject/UObjectLinker.cpp
+++ b/Engine/Source/Runtime/CoreUObject/Private/UObject/UObjectLinker.cpp
@@ -105,6 +105,13 @@ void UObject::SetLinker( FLinkerLoad* LinkerLoad, int32 LinkerIndex, bool bShoul
 {
 	FLinkerIndexPair Existing = LinkerAnnotation.GetAnnotation(this);
 	Existing.CheckInvariants();
+
+	if (Existing.Linker == LinkerLoad && Existing.LinkerIndex == LinkerIndex)
+	{
+		// Nothing to do
+		return;
+	}
+
 	// Detach from existing linker.
 	if( Existing.Linker && bShouldDetachExisting )
 	{
diff --git a/Engine/Source/Runtime/CoreUObject/Public/Serialization/AsyncLoading2.h b/Engine/Source/Runtime/CoreUObject/Public/Serialization/AsyncLoading2.h
--- a/Engine/Source/Runtime/CoreUObject/Public/Serialization/AsyncLoading2.h
+++ b/Engine/Source/Runtime/CoreUObject/Public/Serialization/AsyncLoading2.h
@@ -340,12 +340,6 @@ struct FExportMapEntry
 
 COREUOBJECT_API void FindAllRuntimeScriptPackages(TArray<UPackage*>& OutPackages);
 
-#ifndef WITH_ASYNCLOADING2
-#define WITH_ASYNCLOADING2 (WITH_IOSTORE_IN_EDITOR || !WITH_EDITORONLY_DATA)
-#endif
-
-#if WITH_ASYNCLOADING2
-
 /**
  * Creates a new instance of the AsyncPackageLoader #2.
  *
@@ -354,5 +348,3 @@ COREUOBJECT_API void FindAllRuntimeScriptPackages(TArray<UPackage*>& OutPackages
  * @return The async package loader.
  */
 IAsyncPackageLoader* MakeAsyncPackageLoader2(FIoDispatcher& InIoDispatcher, IAsyncPackageLoader* UncookedPackageLoader = nullptr);
-
-#endif
diff --git a/Engine/Source/Runtime/CoreUObject/Public/UObject/LinkerLoad.h b/Engine/Source/Runtime/CoreUObject/Public/UObject/LinkerLoad.h
--- a/Engine/Source/Runtime/CoreUObject/Public/UObject/LinkerLoad.h
+++ b/Engine/Source/Runtime/CoreUObject/Public/UObject/LinkerLoad.h
@@ -120,6 +120,7 @@ class FLinkerLoad
 	friend class UPackageMap;
 	friend struct FAsyncPackage;
 	friend struct FAsyncPackage2;
+	friend class FAsyncLoadingThread2;
 	friend struct FResolvingExportTracker;
 protected:
 	/** Linker loading status. */
@@ -1067,6 +1068,8 @@ protected: // Daniel L: Made this protected so I can override the constructor an
 		TFunction<void()>&& InSummaryReadyCallback
 	);
 private:
+	ELinkerStatus ProcessPackageSummary(TMap<TPair<FName, FPackageIndex>, FPackageIndex>* ObjectNameWithOuterToExportMap);
+
 	/**
 	 * Start the process of serializing the package file summary if needed
 	 */
diff --git a/Engine/Source/Runtime/CoreUObject/Public/UObject/ObjectMacros.h b/Engine/Source/Runtime/CoreUObject/Public/UObject/ObjectMacros.h
--- a/Engine/Source/Runtime/CoreUObject/Public/UObject/ObjectMacros.h
+++ b/Engine/Source/Runtime/CoreUObject/Public/UObject/ObjectMacros.h
@@ -65,7 +65,7 @@ enum ELoadFlags
 //	LOAD_Unused						= 0x00000040
 	LOAD_NoVerify					= 0x00000080,   ///< Don't verify imports yet.
 	LOAD_IsVerifying				= 0x00000100,	///< Is verifying imports
-//	LOAD_Unused						= 0x00000200,
+	LOAD_SkipLoadImportedPackages	= 0x00000200,	///< Assume that all import packages are already loaded and don't call LoadPackage when creating imports 
 //	LOAD_Unused						= 0x00000400,
 //	LOAD_Unused						= 0x00000800,
 	LOAD_DisableDependencyPreloading = 0x00001000,	///< Bypass dependency preloading system
