Compile fix to match modified API

diff --git a/Engine/Source/Runtime/CoreUObject/Private/Serialization/AsyncLoading2.cpp b/Engine/Source/Runtime/CoreUObject/Private/Serialization/AsyncLoading2.cpp
--- a/Engine/Source/Runtime/CoreUObject/Private/Serialization/AsyncLoading2.cpp
+++ b/Engine/Source/Runtime/CoreUObject/Private/Serialization/AsyncLoading2.cpp
@@ -4613,7 +4613,7 @@ EEventLoadNodeExecutionResult FAsyncPackage2::ProcessLinkerLoadPackageSummary(FA
 	for (int32 PackageImportIndex = 0; PackageImportIndex < ImportedPackageNames.Num(); ++PackageImportIndex)
 	{
 		FAsyncPackageHeaderData::FImportPackageNames& ImportedPackage = ImportedPackageNames[PackageImportIndex];
-		ImportedPackage.UPackageName = InstancingContext.Remap(ImportedPackage.PackageNameToLoad);
+		ImportedPackage.UPackageName = InstancingContext.RemapPackage(ImportedPackage.PackageNameToLoad);
 		if (ImportedPackage.UPackageName != ImportedPackage.PackageNameToLoad)
 		{
 			ImportedPackageIds[PackageImportIndex] = FPackageId::FromName(ImportedPackage.UPackageName);
