ZenLoader: Removed incorrect check

diff --git a/Engine/Source/Runtime/CoreUObject/Private/Serialization/AsyncLoading2.cpp b/Engine/Source/Runtime/CoreUObject/Private/Serialization/AsyncLoading2.cpp
--- a/Engine/Source/Runtime/CoreUObject/Private/Serialization/AsyncLoading2.cpp
+++ b/Engine/Source/Runtime/CoreUObject/Private/Serialization/AsyncLoading2.cpp
@@ -3553,7 +3553,6 @@ void FEventLoadNode2::DependsOn(FEventLoadNode2* Other)
 	uint8 Expected = 0;
 	while (!Other->DependencyWriterCount.compare_exchange_weak(Expected, 1))
 	{
-		check(Expected == 1);
 		Expected = 0;
 	}
 	if (!Other->IsDone())
