Public change split off from internal change.

diff --git a/Engine/Source/Runtime/CoreUObject/Private/UObject/UObjectGlobals.cpp b/Engine/Source/Runtime/CoreUObject/Private/UObject/UObjectGlobals.cpp
--- a/Engine/Source/Runtime/CoreUObject/Private/UObject/UObjectGlobals.cpp
+++ b/Engine/Source/Runtime/CoreUObject/Private/UObject/UObjectGlobals.cpp
@@ -3667,16 +3667,17 @@ void FObjectInitializer::PostConstructInit()
 			}
 		}
 	}
-	if (bNeedInstancing || bNeedSubobjectInstancing)
-	{
-		InstanceSubobjects(Class, bNeedInstancing, bNeedSubobjectInstancing);
-	}
 
 	// Allow custom property initialization to happen before PostInitProperties is called
 	if (PropertyInitCallback)
 	{
 		PropertyInitCallback();
 	}
+	// After the call to `PropertyInitCallback` to allow the callback to modify the instancing graph
+	if (bNeedInstancing || bNeedSubobjectInstancing)
+	{
+		InstanceSubobjects(Class, bNeedInstancing, bNeedSubobjectInstancing);
+	}
 
 	// Make sure subobjects knows that they had their properties overwritten
 	for (int32 Index = 0; Index < ComponentInits.SubobjectInits.Num(); Index++)
diff --git a/Engine/Source/Runtime/CoreUObject/Public/UObject/Class.h b/Engine/Source/Runtime/CoreUObject/Public/UObject/Class.h
--- a/Engine/Source/Runtime/CoreUObject/Public/UObject/Class.h
+++ b/Engine/Source/Runtime/CoreUObject/Public/UObject/Class.h
@@ -3512,6 +3512,22 @@ public:
 	 */
 	FObjectInstancingGraph( class UObject* DestinationSubobjectRoot );
 
+	/**
+	 * Returns whether this instancing graph has a valid destination root.
+	 */
+	bool HasDestinationRoot() const
+	{
+		return GetDestinationRoot() != nullptr;
+	}
+
+	/**
+	 * Returns the DestinationRoot for this instancing graph.
+	 */
+	const UObject* GetDestinationRoot() const
+	{
+		return DestinationRoot;
+	}
+
 	/**
 	 * Sets the DestinationRoot for this instancing graph.
 	 *
@@ -3624,14 +3640,6 @@ public:
 	}
 
 private:
-	/**
-	 * Returns whether this instancing graph has a valid destination root.
-	 */
-	bool HasDestinationRoot() const
-	{
-		return DestinationRoot != nullptr;
-	}
-
 	/**
 	 * Returns whether DestinationRoot corresponds to an archetype object.
 	 *
@@ -3715,6 +3723,10 @@ private:
 	TSet<const FProperty*>						SubobjectInstantiationExclusionList;
 };
 
+Expose_TNameOf(FObjectInstancingGraph);
+
+Expose_TNameOf(FObjectInstancingGraph*);
+
 // UFunction interface.
 
 inline UFunction* UFunction::GetSuperFunction() const
