[CoreUObject] fix for PrintScriptCallstack not printing callstack to the immediate window
logging buffers any call not from the logging thread, added a call to Flush() to force the logger write.

added to improve efficiency in building a stacktrace

FFrame::GetStackDescription(FStringBuilderBase& StringBuilder)
FFrame::GetStackTrace(FStringBuilderBase& StringBuilder)
FFrame::GetScriptCallstack(FStringBuilderBase& StringBuilder, bool bReturnEmpty = false, bool bTopOfStackOnly = false)
UObjectBaseUtility::GetFullName(FStringBuilderBase& ResultString, const UObject* StopOuter = NULL, EObjectFullNameFlags Flags = EObjectFullNameFlags::None) const

diff --git a/Engine/Source/Runtime/CoreUObject/Private/UObject/ScriptCore.cpp b/Engine/Source/Runtime/CoreUObject/Private/UObject/ScriptCore.cpp
--- a/Engine/Source/Runtime/CoreUObject/Private/UObject/ScriptCore.cpp
+++ b/Engine/Source/Runtime/CoreUObject/Private/UObject/ScriptCore.cpp
@@ -337,12 +337,16 @@ void PrintScriptCallStackImpl()
 	if (BlueprintExceptionTracker)
 	{
 		TArrayView<const FFrame* const> RawStack = BlueprintExceptionTracker->GetCurrentScriptStack();
-		FString ScriptStack = FString::Printf(TEXT("\n\nScript Stack (%d frames):\n"), RawStack.Num());
+		TStringBuilder<4096> ScriptStack;
+		ScriptStack << TEXT("\n\nScript Stack (") << RawStack.Num() << TEXT(" frames) :\n");
+
 		for (int32 FrameIdx = RawStack.Num() - 1; FrameIdx >= 0; --FrameIdx)
 		{
-			ScriptStack += RawStack[FrameIdx]->GetStackDescription() + TEXT("\n");
+			RawStack[FrameIdx]->GetStackDescription(ScriptStack);
+			ScriptStack << TEXT("\n");
 		}
 		UE_LOG(LogOutputDevice, Warning, TEXT("%s"), *ScriptStack);
+		GLog->Flush(); // flush is needed because logging is buffered 
 	}
 }
 
@@ -521,8 +525,13 @@ static bool ShowKismetScriptStackOnWarnings()
 
 FString FFrame::GetScriptCallstack(bool bReturnEmpty, bool bTopOfStackOnly)
 {
-	FString ScriptStack;
+	TStringBuilder<4096> ScriptStack;
+	GetScriptCallstack(ScriptStack, bReturnEmpty, bTopOfStackOnly);
+	return FString(ScriptStack);
+}
 
+void FFrame::GetScriptCallstack(FStringBuilderBase& ScriptStack, bool bReturnEmpty, bool bTopOfStackOnly)
+{
 #if DO_BLUEPRINT_GUARD
 	FBlueprintContextTracker& BlueprintExceptionTracker = FBlueprintContextTracker::Get();
 	if (BlueprintExceptionTracker.ScriptStack.Num() > 0)
@@ -533,32 +542,39 @@ FString FFrame::GetScriptCallstack(bool bReturnEmpty, bool bTopOfStackOnly)
 
 		do
 		{
-			ScriptStack += TEXT("\t") + BlueprintExceptionTracker.ScriptStack[i]->GetStackDescription();
+			ScriptStack << TEXT("\t");
+			BlueprintExceptionTracker.ScriptStack[i]->GetStackDescription(ScriptStack);
 			if ((i == TopOfStackIndex) && bDisplayArrow)
 			{
-				ScriptStack += TEXT(" <---");
+				ScriptStack << TEXT(" <---");
 			}
-			ScriptStack += TEXT("\n");
+			ScriptStack << TEXT("\n");
 			--i;
 		} while ((i >= 0) && !bTopOfStackOnly);
 	}
 	else if (!bReturnEmpty)
 	{
-		ScriptStack += TEXT("\t[Empty] (FFrame::GetScriptCallstack() called from native code)");
+		ScriptStack << TEXT("\t[Empty] (FFrame::GetScriptCallstack() called from native code)");
 	}
 #else
 	if (!bReturnEmpty)
 	{
-		ScriptStack = TEXT("Unable to display Script Callstack. Compile with DO_BLUEPRINT_GUARD=1");
+		ScriptStack << TEXT("Unable to display Script Callstack. Compile with DO_BLUEPRINT_GUARD=1");
 	}
 #endif
-
-	return ScriptStack;
 }
 
 FString FFrame::GetStackDescription() const
 {
-	return Node->GetOuter()->GetPathName() + TEXT(".") + Node->GetName();
+	TStringBuilder<256> StringBuilder;
+	GetStackDescription(StringBuilder);
+	return FString(StringBuilder);
+}
+
+void FFrame::GetStackDescription(FStringBuilderBase& StringBuilder) const
+{
+	Node->GetOuter()->GetPathName(nullptr, StringBuilder);
+	StringBuilder << TEXT(".") << Node->GetName();
 }
 
 #if DO_BLUEPRINT_GUARD
@@ -610,7 +626,7 @@ void FFrame::KismetExecutionMessage(const TCHAR* Message, ELogVerbosity::Type Ve
 	}
 #endif
 
-	FString ScriptStack;
+	TStringBuilder<4096> ScriptStack;
 
 	// Tracking down some places that display warnings but no message..
 	ensureAlways(Verbosity > ELogVerbosity::Warning || FCString::Strlen(Message) > 0);
@@ -620,12 +636,12 @@ void FFrame::KismetExecutionMessage(const TCHAR* Message, ELogVerbosity::Type Ve
 	if (Verbosity <= ELogVerbosity::Error || (ShowKismetScriptStackOnWarnings() && Verbosity == ELogVerbosity::Warning))
 	{
 		ScriptStack = TEXT("Script call stack:\n");
-		ScriptStack += GetScriptCallstack();
+		GetScriptCallstack(ScriptStack);
 	}
 	else if (Verbosity == ELogVerbosity::Warning)
 	{
 		ScriptStack = TEXT("Last function called:\n");
-		ScriptStack += GetScriptCallstack(false, true);
+		GetScriptCallstack(ScriptStack, false, true);
 	}
 #endif
 
@@ -694,11 +710,15 @@ void FFrame::Serialize( const TCHAR* V, ELogVerbosity::Type Verbosity, const cla
 #endif
 	}
 }
-
 FString FFrame::GetStackTrace() const
 {
-	FString Result;
+	TStringBuilder<4096> Result;
+	GetStackTrace(Result);
+	return FString(Result);
+}
 
+void FFrame::GetStackTrace(FStringBuilderBase& Result) const
+{
 	// travel down the stack recording the frames
 	TArray<const FFrame*> FrameStack;
 	const FFrame* CurrFrame = this;
@@ -711,18 +731,18 @@ FString FFrame::GetStackTrace() const
 	// and then dump them to a string
 	if (FrameStack.Num() > 0)
 	{
-		Result += FString(TEXT("Script call stack:\n"));
+		Result << TEXT("Script call stack:\n");
 		for (int32 Index = FrameStack.Num() - 1; Index >= 0; Index--)
 		{
-			Result += FString::Printf(TEXT("\t%s\n"), *FrameStack[Index]->Node->GetFullName());
+			Result << TEXT("\t");
+			FrameStack[Index]->Node->GetFullName(Result);
+			Result << TEXT("\n");
 		}
 	}
 	else
 	{
-		Result += FString(TEXT("Script call stack: [Empty] (FFrame::GetStackTrace() called from native code)"));
+		Result << TEXT("Script call stack: [Empty] (FFrame::GetStackTrace() called from native code)");
 	}
-
-	return Result;
 }
 
 //////////////////////////////////////////////////////////////////////////
diff --git a/Engine/Source/Runtime/CoreUObject/Private/UObject/UObjectBaseUtility.cpp b/Engine/Source/Runtime/CoreUObject/Private/UObject/UObjectBaseUtility.cpp
--- a/Engine/Source/Runtime/CoreUObject/Private/UObject/UObjectBaseUtility.cpp
+++ b/Engine/Source/Runtime/CoreUObject/Private/UObject/UObjectBaseUtility.cpp
@@ -95,16 +95,23 @@ FString UObjectBaseUtility::GetFullName(const UObject* StopOuter, EObjectFullNam
   * Version of GetFullName() that eliminates unnecessary copies and appends an existing string.
   */
 void UObjectBaseUtility::GetFullName(const UObject* StopOuter, FString& ResultString, EObjectFullNameFlags Flags) const
+{
+	TStringBuilder<256> StringBuilder;
+	GetFullName(StringBuilder, StopOuter, Flags);
+	ResultString.Append(StringBuilder);
+}
+
+void UObjectBaseUtility::GetFullName(FStringBuilderBase& ResultString, const UObject* StopOuter, EObjectFullNameFlags Flags) const
 {
 	if (this != nullptr)
 	{
 		if (EnumHasAllFlags(Flags, EObjectFullNameFlags::IncludeClassPackage))
 		{
-			ResultString += GetClass()->GetPathName();
+			GetClass()->GetPathName(nullptr, ResultString);
 		}
 		else
 		{
-			GetClass()->AppendName(ResultString);
+			GetClass()->GetFName().AppendString(ResultString);
 		}
 		ResultString += TEXT(' ');
 		GetPathName(StopOuter, ResultString);
@@ -112,7 +119,7 @@ void UObjectBaseUtility::GetFullName(const UObject* StopOuter, FString& ResultSt
 	else
 	{
 		ResultString += TEXT("None");
-	} 
+	}
 }
 
 
diff --git a/Engine/Source/Runtime/CoreUObject/Public/UObject/Stack.h b/Engine/Source/Runtime/CoreUObject/Public/UObject/Stack.h
--- a/Engine/Source/Runtime/CoreUObject/Public/UObject/Stack.h
+++ b/Engine/Source/Runtime/CoreUObject/Public/UObject/Stack.h
@@ -235,18 +235,43 @@ public:
 	 **/
 	COREUOBJECT_API FString GetStackTrace() const;
 
+	/**
+	 * This will return the StackTrace of the current callstack from the last native entry point
+	 * 
+	 * @param StringBuilder to populate
+	 **/
+	COREUOBJECT_API void GetStackTrace(FStringBuilderBase& StringBuilder) const;
+
 	/**
 	* This will return the StackTrace of the all script frames currently active
 	* 
 	* @param	bReturnEmpty if true, returns empty string when no script callstack found
+	* @param	bTopOfStackOnly if true only returns the top of the callstack
 	**/
 	COREUOBJECT_API static FString GetScriptCallstack(bool bReturnEmpty = false, bool bTopOfStackOnly = false);
 
+	/**
+	* This will return the StackTrace of the all script frames currently active
+	*
+	* @param	StringBuilder to populate
+	* @param	bReturnEmpty if true, returns empty string when no script callstack found
+	* @param	bTopOfStackOnly if true only returns the top of the callstack
+	**/
+	COREUOBJECT_API static void GetScriptCallstack(FStringBuilderBase& StringBuilder, bool bReturnEmpty = false, bool bTopOfStackOnly = false);
+		
 	/** 
 	 * This will return a string of the form "ScopeName.FunctionName" associated with this stack frame:
 	 */
+	UE_DEPRECATED(5.2, "Please use GetStackDescription(FStringBuilderBase&).")
 	COREUOBJECT_API FString GetStackDescription() const;
 
+	/**
+	* This will append a string of the form "ScopeName.FunctionName" associated with this stack frame
+	*
+	* @param	StringBuilder to populate
+	**/
+	COREUOBJECT_API void GetStackDescription(FStringBuilderBase& StringBuilder) const;
+
 #if DO_BLUEPRINT_GUARD
 	static void InitPrintScriptCallstack();
 #endif
diff --git a/Engine/Source/Runtime/CoreUObject/Public/UObject/UObjectBaseUtility.h b/Engine/Source/Runtime/CoreUObject/Public/UObject/UObjectBaseUtility.h
--- a/Engine/Source/Runtime/CoreUObject/Public/UObject/UObjectBaseUtility.h
+++ b/Engine/Source/Runtime/CoreUObject/Public/UObject/UObjectBaseUtility.h
@@ -433,6 +433,19 @@ public:
 	 */
 	void GetFullName( const UObject* StopOuter, FString& ResultString, EObjectFullNameFlags Flags = EObjectFullNameFlags::None ) const;
 
+	/**
+	 * Returns the fully qualified pathname for this object as well as the name of the class, in the format:
+	 * 'ClassName Outermost[.Outer].Name'.
+	 *
+	 * @param	ResultString StringBuilder to populate
+	 * @param	StopOuter	if specified, indicates that the output string should be relative to this object.  if StopOuter
+	 *						does not exist in this object's Outer chain, the result would be the same as passing NULL.
+	 * @param	Flags		flags that control the behavior of full name generation
+	 *
+	 * @note	safe to call on NULL object pointers!
+	 */
+	void GetFullName(FStringBuilderBase& ResultString, const UObject* StopOuter = NULL, EObjectFullNameFlags Flags = EObjectFullNameFlags::None) const;
+
 	/**
 	 * Returns the fully qualified pathname for this object, in the format:
 	 * 'Outermost[.Outer].Name'
diff --git a/Engine/Source/Runtime/SlateCore/Private/Trace/SlateTrace.cpp b/Engine/Source/Runtime/SlateCore/Private/Trace/SlateTrace.cpp
--- a/Engine/Source/Runtime/SlateCore/Private/Trace/SlateTrace.cpp
+++ b/Engine/Source/Runtime/SlateCore/Private/Trace/SlateTrace.cpp
@@ -358,7 +358,7 @@ void FSlateTrace::WidgetInvalidated(const SWidget* Widget, const SWidget* Invest
 
 		static_assert(sizeof(EInvalidateWidgetReason) == sizeof(uint8), "EInvalidateWidgetReason is not a uint8");
 
-		FString ScriptTrace;
+		TStringBuilder<4096> ScriptTrace;
 		constexpr int MAX_DEPTH = 64;
 		uint64 StackTrace[MAX_DEPTH] = { 0 };
 		uint32 StackTraceDepth = 0;
@@ -371,10 +371,10 @@ void FSlateTrace::WidgetInvalidated(const SWidget* Widget, const SWidget* Invest
 			FSlowHeartBeatScope SuspendHeartBeat;
 			FDisableHitchDetectorScope SuspendGameThreadHitch;
 
-			ScriptTrace = FFrame::GetScriptCallstack(true /* bReturnEmpty */);
-			if (!ScriptTrace.IsEmpty())
+			FFrame::GetScriptCallstack(ScriptTrace, true /* bReturnEmpty */);
+			if (ScriptTrace.Len() != 0)
 			{
-				ScriptTrace = "ScriptTrace: \n" + ScriptTrace;
+				ScriptTrace.InsertAt(0, TEXT("ScriptTrace: \n"));
 			}
 
 			// Walk the stack and dump it to the allocated memory.
