Core - Adding a new console command "LazyLoad.PrintUnresolvedObjects", Prints a list of all unresolved objects from the object handle index.  This is a great way to examine all the things that *should be loaded* that were avoided being loaded by the lazy loading system.

diff --git a/Engine/Source/Runtime/CoreUObject/Private/UObject/ObjectHandle.cpp b/Engine/Source/Runtime/CoreUObject/Private/UObject/ObjectHandle.cpp
--- a/Engine/Source/Runtime/CoreUObject/Private/UObject/ObjectHandle.cpp
+++ b/Engine/Source/Runtime/CoreUObject/Private/UObject/ObjectHandle.cpp
@@ -100,6 +100,55 @@ namespace ObjectHandle_Private
 		TArray<FObjectHandlePackageData> PackageData;
 	} GObjectHandleIndex;
 
+#if UE_WITH_OBJECT_HANDLE_LATE_RESOLVE
+
+	static FAutoConsoleCommand CmdPrintUnresolvedObjects(
+		TEXT("LazyLoad.PrintUnresolvedObjects"),
+		TEXT("Prints a list of all unresolved objects from the object handle index."),
+		FConsoleCommandWithArgsDelegate::CreateLambda([](const TArray<FString>& Args)
+			{
+				FRWScopeLock GlobalLockScope(GObjectHandleIndex.Lock, SLT_ReadOnly);
+				const int32 TotalObjectHandlePackages = GObjectHandleIndex.PackageData.Num();
+				TArray<FName> LoadedPackages;
+				TArray<FName> UnloadedPackages;
+				for (int32 PackageIndex = 0; PackageIndex < TotalObjectHandlePackages; PackageIndex++)
+				{
+					FName PackageName;
+					{
+						FObjectHandlePackageData& PackageData = GObjectHandleIndex.PackageData[PackageIndex];
+						FRWScopeLock LocalLockScope(PackageData.Lock, SLT_ReadOnly);
+						PackageName = MinimalNameToName(PackageData.PackageName);
+					}
+
+					UPackage* Package = FindObjectFast<UPackage>(nullptr, PackageName);
+					if (Package)
+					{
+						LoadedPackages.Add(PackageName);
+					}
+					else
+					{
+						UnloadedPackages.Add(PackageName);
+					}
+				}
+
+				UnloadedPackages.Sort();
+				for (FName PackageName : UnloadedPackages)
+				{
+					UE_LOG(LogUObjectGlobals, Log, TEXT("Unloaded %s"), *PackageName.ToString());
+				}
+
+				LoadedPackages.Sort();
+				for (FName PackageName : LoadedPackages)
+				{
+					UE_LOG(LogUObjectGlobals, Log, TEXT("Loaded %s"), *PackageName.ToString());
+				}
+
+				UE_LOG(LogUObjectGlobals, Log, TEXT("Unloaded Packages (%d out of %d)"), UnloadedPackages.Num(), LoadedPackages.Num() + UnloadedPackages.Num());
+			})
+	);
+
+#endif
+
 	static inline FPackedObjectRef Pack(FPackageId PackageId, FObjectId ObjectId)
 	{
 #if UE_WITH_OBJECT_HANDLE_LATE_RESOLVE || UE_WITH_OBJECT_HANDLE_TRACKING
