Core - LazyLoad.PrintUnresolvedObjects, now with 100% fewer warnings about FName sorting.

diff --git a/Engine/Source/Runtime/CoreUObject/Private/UObject/ObjectHandle.cpp b/Engine/Source/Runtime/CoreUObject/Private/UObject/ObjectHandle.cpp
--- a/Engine/Source/Runtime/CoreUObject/Private/UObject/ObjectHandle.cpp
+++ b/Engine/Source/Runtime/CoreUObject/Private/UObject/ObjectHandle.cpp
@@ -131,13 +131,13 @@ namespace ObjectHandle_Private
 					}
 				}
 
-				UnloadedPackages.Sort();
+				UnloadedPackages.Sort(FNameLexicalLess());
 				for (FName PackageName : UnloadedPackages)
 				{
 					UE_LOG(LogUObjectGlobals, Log, TEXT("Unloaded %s"), *PackageName.ToString());
 				}
 
-				LoadedPackages.Sort();
+				LoadedPackages.Sort(FNameLexicalLess());
 				for (FName PackageName : LoadedPackages)
 				{
 					UE_LOG(LogUObjectGlobals, Log, TEXT("Loaded %s"), *PackageName.ToString());
