Fix for nonunity compile issue.

diff --git a/Engine/Source/Runtime/CoreUObject/Public/AssetRegistry/ARFilter.h b/Engine/Source/Runtime/CoreUObject/Public/AssetRegistry/ARFilter.h
--- a/Engine/Source/Runtime/CoreUObject/Public/AssetRegistry/ARFilter.h
+++ b/Engine/Source/Runtime/CoreUObject/Public/AssetRegistry/ARFilter.h
@@ -5,6 +5,7 @@
 #include "CoreMinimal.h"
 #include "UObject/Class.h"
 #include "UObject/ObjectMacros.h"
+#include "UObject/SoftObjectPath.h"
 #include "UObject/TopLevelAssetPath.h"
 
 struct FAssetData;
