fixed non-unity compilation by including the correct header
https://horde.devtools.epicgames.com/job/6318b22ba31879aea8d85c9b?step=8e3f&issue=209198

diff --git a/Engine/Source/Runtime/CoreUObject/Public/Serialization/BulkDataRegistry.h b/Engine/Source/Runtime/CoreUObject/Public/Serialization/BulkDataRegistry.h
--- a/Engine/Source/Runtime/CoreUObject/Public/Serialization/BulkDataRegistry.h
+++ b/Engine/Source/Runtime/CoreUObject/Public/Serialization/BulkDataRegistry.h
@@ -3,7 +3,7 @@
 #pragma once
 
 #include "Containers/Map.h"
-#include "Delegates/DelegateBase.h"
+#include "Delegates/Delegate.h"
 #include "HAL/Platform.h"
 #include "Logging/LogMacros.h"
 #include "UObject/NameTypes.h"
