[CoreUObject] change to Blueprint loading to only resolve imports that are needed. prevents loading of unnecessary data

diff --git a/Engine/Source/Runtime/CoreUObject/Private/Blueprint/BlueprintSupport.cpp b/Engine/Source/Runtime/CoreUObject/Private/Blueprint/BlueprintSupport.cpp
--- a/Engine/Source/Runtime/CoreUObject/Private/Blueprint/BlueprintSupport.cpp
+++ b/Engine/Source/Runtime/CoreUObject/Private/Blueprint/BlueprintSupport.cpp
@@ -16,6 +16,8 @@
 #include "Misc/PackageName.h"
 #include "UObject/ObjectResource.h"
 #include "UObject/GCObject.h"
+#include "UObject/LinkerLoad.h"
+#include "UObject/LinkerLoadImportBehavior.h"
 #include "UObject/LinkerPlaceholderClass.h"
 #include "UObject/LinkerPlaceholderExportObject.h"
 #include "UObject/LinkerPlaceholderFunction.h"
@@ -1765,6 +1767,18 @@ void FLinkerLoad::ResolveAllImports()
 		//       in turn, could end us back in this function before we ever  
 		//       returned from this
 		FObjectImport& Import = ImportMap[ImportIndex];
+
+#if UE_WITH_OBJECT_HANDLE_LATE_RESOLVE
+		if (FLinkerLoad::IsImportLazyLoadEnabled())
+		{
+			using namespace UE::LinkerLoad;
+			if (GetPropertyImportLoadBehavior(Import, *this) != EImportBehavior::Eager)
+			{
+				continue;
+			}
+		}
+#endif
+
 		UObject* ImportObject = CreateImport(ImportIndex);
 
 		// see if this import is currently being resolved (presumably somewhere 
@@ -1846,11 +1860,7 @@ void FLinkerLoad::FinalizeBlueprint(UClass* LoadClass)
 	// have to)... we do however need it here in FinalizeBlueprint(), because
 	// we need it ran for any super-classes before we regen
 
-	if (!IsImportLazyLoadEnabled())
-	{
-		// @TODO: OBJPTR: Need to find other options for solving this issue of placeholder classes during blueprint compile without forcing all imports to resolve always
-		ResolveAllImports();
-	}
+	ResolveAllImports();
 
 	// Now that imports have been resolved we optionally flush the compilation
 	// queue. This is only done for level blueprints, which will have instances
