Remove ensure from FSoftObjectPath::SetPath.

diff --git a/Engine/Source/Runtime/CoreUObject/Private/UObject/SoftObjectPath.cpp b/Engine/Source/Runtime/CoreUObject/Private/UObject/SoftObjectPath.cpp
--- a/Engine/Source/Runtime/CoreUObject/Private/UObject/SoftObjectPath.cpp
+++ b/Engine/Source/Runtime/CoreUObject/Private/UObject/SoftObjectPath.cpp
@@ -113,11 +113,18 @@ void FSoftObjectPath::SetPath(FWideStringView Path)
 		// Empty path, just empty the pathname.
 		Reset();
 	}
-	else if (ensureMsgf(!FPackageName::IsShortPackageName(Path), TEXT("Cannot create SoftObjectPath with short package name '%.*s'%s! You must pass in fully qualified package names"), Path.Len(), Path.GetData(), *GetObjectBeingSerializedForSoftObjectPath()))
+	else 
 	{
 		// Possibly an ExportText path. Trim the ClassName.
 		Path = FPackageName::ExportTextPathToObjectPath(Path);
 
+		if (!Path.StartsWith('/'))
+		{
+			// Not a recognized path. No ensure/logging here because many things attempt to construct paths from user input. 
+			Reset();
+			return;
+		}
+
 		int32 ColonIndex;
 		if (Path.FindChar(SUBOBJECT_DELIMITER_CHAR, ColonIndex))
 		{
