Asset: AppendObjectPath should append to the given string

diff --git a/Engine/Source/Runtime/CoreUObject/Public/AssetRegistry/AssetData.h b/Engine/Source/Runtime/CoreUObject/Public/AssetRegistry/AssetData.h
--- a/Engine/Source/Runtime/CoreUObject/Public/AssetRegistry/AssetData.h
+++ b/Engine/Source/Runtime/CoreUObject/Public/AssetRegistry/AssetData.h
@@ -445,7 +445,7 @@ PRAGMA_ENABLE_DEPRECATION_WARNINGS
 	{
 		TStringBuilder<FName::StringBufferSize> Builder;
 		AppendObjectPath(Builder);
-		String = FString(Builder);
+		String.Append(FString(Builder));
 	}
 
 
