Removing unacceptable words in Runtime/CoreUObject

diff --git a/Engine/Source/Runtime/CoreUObject/Private/UObject/DeferredRegistry.cpp b/Engine/Source/Runtime/CoreUObject/Private/UObject/DeferredRegistry.cpp
--- a/Engine/Source/Runtime/CoreUObject/Private/UObject/DeferredRegistry.cpp
+++ b/Engine/Source/Runtime/CoreUObject/Private/UObject/DeferredRegistry.cpp
@@ -17,7 +17,7 @@ void ReloadProcessObject(UScriptStruct* ScriptStruct, const TCHAR* RenamePrefix)
 
 void ReloadProcessObject(UEnum* Enum, const TCHAR* RenamePrefix)
 {	
-	// Before renaming the enum (moving the enum to transient package) remove its enum value names from the enum master list as the list is populated with enum package name
+	// Before renaming the enum (moving the enum to transient package) remove its enum value names from the enum Primary list as the list is populated with enum package name
 	Enum->RemoveNamesFromPrimaryList();
 	// Make sure the old enum is not used by anything
 	Enum->ClearFlags(RF_Standalone | RF_Public);
diff --git a/Engine/Source/Runtime/CoreUObject/Private/UObject/Enum.cpp b/Engine/Source/Runtime/CoreUObject/Private/UObject/Enum.cpp
--- a/Engine/Source/Runtime/CoreUObject/Private/UObject/Enum.cpp
+++ b/Engine/Source/Runtime/CoreUObject/Private/UObject/Enum.cpp
@@ -259,27 +259,9 @@ bool UEnum::IsValidEnumName(FName InName) const
 	return false;
 }
 
-void UEnum::AddNamesToMasterList()
-{
-	FWriteScopeLock ScopeLock(AllEnumNamesLock);
-	EnumPackage = GetPackage()->GetFName();
-	TMap<FName, UEnum*>& PackageEnumValues = AllEnumNames.FindOrAdd(EnumPackage);
-	for (TPair<FName, int64> Kvp : Names)
-	{
-		UEnum* Enum = PackageEnumValues.FindRef(Kvp.Key);
-		if (Enum == nullptr || Enum->HasAnyFlags(RF_NewerVersionExists))
-		{
-			PackageEnumValues.Add(Kvp.Key, this);
-		}
-		else if (Enum != this && Enum->GetOutermost() != GetTransientPackage())
-		{
-			UE_LOG(LogEnum, Warning, TEXT("Enum name collision: '%s' is in both '%s' and '%s'"), *Kvp.Key.ToString(), *GetPathName(), *Enum->GetPathName());
-		}
-	}
-}
-
 void UEnum::AddNamesToPrimaryList()
 {
+	FWriteScopeLock ScopeLock(AllEnumNamesLock);
 	EnumPackage = GetPackage()->GetFName();
 	TMap<FName, UEnum*>& PackageEnumValues = AllEnumNames.FindOrAdd(EnumPackage);
 	for (TPair<FName, int64> Kvp : Names)
@@ -296,29 +278,9 @@ void UEnum::AddNamesToPrimaryList()
 	}
 }
 
-void UEnum::RemoveNamesFromMasterList()
-{
-	FWriteScopeLock ScopeLock(AllEnumNamesLock);
-	TMap<FName, UEnum*>* PackageEnumValues = AllEnumNames.Find(EnumPackage);
-	if (PackageEnumValues != nullptr)
-	{
-		for (TPair<FName, int64> Kvp : Names)
-		{
-			UEnum* Enum = PackageEnumValues->FindRef(Kvp.Key);
-			if (Enum == this)
-			{
-				PackageEnumValues->Remove(Kvp.Key);
-			}
-		}
-		if (PackageEnumValues->Num() == 0)
-		{
-			AllEnumNames.Remove(EnumPackage);
-		}
-	}
-}
-
 void UEnum::RemoveNamesFromPrimaryList()
 {
+	FWriteScopeLock ScopeLock(AllEnumNamesLock);
 	TMap<FName, UEnum*>* PackageEnumValues = AllEnumNames.Find(EnumPackage);
 	if (PackageEnumValues != nullptr)
 	{
diff --git a/Engine/Source/Runtime/CoreUObject/Public/UObject/Class.h b/Engine/Source/Runtime/CoreUObject/Public/UObject/Class.h
--- a/Engine/Source/Runtime/CoreUObject/Public/UObject/Class.h
+++ b/Engine/Source/Runtime/CoreUObject/Public/UObject/Class.h
@@ -1998,9 +1998,12 @@ public:
 	/** Checks if enum has entry with given name. Includes autogenerated _MAX entry. */
 	bool IsValidEnumName(FName InName) const;
 
-	/** Removes the Names in this enum from the master AllEnumNames list */
+	/** Removes the Names in this enum from the primary AllEnumNames list */
 	UE_DEPRECATED(5.1, "RemoveNamesFromMasterList is deprecated, please use RemoveNamesFromPrimaryList instead.")
-	void RemoveNamesFromMasterList();
+	void RemoveNamesFromMasterList()
+	{
+		RemoveNamesFromPrimaryList();
+	}
 
 	/** Removes the Names in this enum from the primary AllEnumNames list */
 	void RemoveNamesFromPrimaryList();
@@ -2333,7 +2336,7 @@ protected:
 	/** pointer to function used to look up the enum's display name. Currently only assigned for UEnums generated for nativized blueprints */
 	FEnumDisplayNameFn EnumDisplayNameFn;
 
-	/** Package name this enum was in when its names were being added to the master list */
+	/** Package name this enum was in when its names were being added to the primary list */
 	FName EnumPackage;
 
 	/** lock to be taken when accessing AllEnumNames */
@@ -2342,9 +2345,12 @@ protected:
 	/** global list of all value names used by all enums in memory, used for property text import */
 	static TMap<FName, TMap<FName, UEnum*> > AllEnumNames;
 
-	/** adds the Names in this enum to the master AllEnumNames list */
+	/** adds the Names in this enum to the primary AllEnumNames list */
 	UE_DEPRECATED(5.1, "AddNamesToMasterList is deprecated, please use AddNamesToPrimaryList instead.")
-	void AddNamesToMasterList();
+	void AddNamesToMasterList()
+	{
+		AddNamesToPrimaryList();
+	}
 
 	/** adds the Names in this enum to the primary AllEnumNames list */
 	void AddNamesToPrimaryList();
