Fix GetLongPackageName and GetAssetName returning the string 'None' rather than empty strings when those components are empty.

diff --git a/Engine/Source/Runtime/CoreUObject/Public/UObject/SoftObjectPath.h b/Engine/Source/Runtime/CoreUObject/Public/UObject/SoftObjectPath.h
--- a/Engine/Source/Runtime/CoreUObject/Public/UObject/SoftObjectPath.h
+++ b/Engine/Source/Runtime/CoreUObject/Public/UObject/SoftObjectPath.h
@@ -178,7 +178,8 @@ PRAGMA_ENABLE_DEPRECATION_WARNINGS;
 	/** Returns /package/path, leaving off the asset name and sub object */
 	FString GetLongPackageName() const
 	{
-		return GetAssetPath().GetPackageName().ToString();
+		FName PackageName = GetAssetPath().GetPackageName();
+		return PackageName.IsNone() ? FString() : PackageName.ToString();
 	}
 
 	/** Returns /package/path, leaving off the asset name and sub object */
@@ -190,7 +191,8 @@ PRAGMA_ENABLE_DEPRECATION_WARNINGS;
 	/** Returns assetname string, leaving off the /package/path part and sub object */
 	FString GetAssetName() const
 	{
-		return GetAssetPath().GetAssetName().ToString();
+		FName AssetName = GetAssetPath().GetAssetName();
+		return AssetName.IsNone() ? FString() : AssetName.ToString();
 	}
 
 	/** Returns assetname string, leaving off the /package/path part and sub object */
