Allow methods that can only be invoked from a constructor be available when fetching vtable address for HotReload and Live Coding.

diff --git a/Engine/Source/Runtime/CoreUObject/Private/UObject/Class.cpp b/Engine/Source/Runtime/CoreUObject/Private/UObject/Class.cpp
--- a/Engine/Source/Runtime/CoreUObject/Private/UObject/Class.cpp
+++ b/Engine/Source/Runtime/CoreUObject/Private/UObject/Class.cpp
@@ -5735,6 +5735,11 @@ bool UClass::HotReloadPrivateStaticClass(
 	UObject* TempObjectForVTable = nullptr;
 	{
 		TGuardValue<bool> Guard(GIsRetrievingVTablePtr, true);
+
+		// Mark we're in the constructor now.
+		FUObjectThreadContext& ThreadContext = FUObjectThreadContext::Get();
+		TScopeCounter<int32> InConstructor(ThreadContext.IsInConstructor);
+
 		FVTableHelper Helper;
 		TempObjectForVTable = ClassVTableHelperCtorCaller(Helper);
 		TempObjectForVTable->AtomicallyClearInternalFlags(EInternalObjectFlags::PendingConstruction);
