Removed redundant private include paths from build.cs files.
Fixed include paths to be relative to the private or public folders.
Hid or removed includes that reached into other private module folders.
Updated PublicInclude paths when necessary.

diff --git a/Engine/Plugins/Experimental/Water/Source/Runtime/Private/WaterInfoRendering.cpp b/Engine/Plugins/Experimental/Water/Source/Runtime/Private/WaterInfoRendering.cpp
--- a/Engine/Plugins/Experimental/Water/Source/Runtime/Private/WaterInfoRendering.cpp
+++ b/Engine/Plugins/Experimental/Water/Source/Runtime/Private/WaterInfoRendering.cpp
@@ -8,10 +8,10 @@
 #include "WaterBodyActor.h"
 #include "WaterBodySceneProxy.h"
 #include "Engine/TextureRenderTarget2D.h"
-#include "Renderer/Private/ScenePrivate.h"
-#include "Renderer/Private/RendererModule.h"
-#include "Renderer/Private/SceneCaptureRendering.h"
-#include "Renderer/Private/PostProcess/SceneFilterRendering.h"
+#include "ScenePrivate.h"
+#include "RendererModule.h"
+#include "SceneCaptureRendering.h"
+#include "PostProcess/SceneFilterRendering.h"
 #include "Math/OrthoMatrix.h"
 #include "GameFramework/WorldSettings.h"
 #include "WaterZoneActor.h"
diff --git a/Engine/Plugins/Experimental/Water/Source/Runtime/Water.Build.cs b/Engine/Plugins/Experimental/Water/Source/Runtime/Water.Build.cs
--- a/Engine/Plugins/Experimental/Water/Source/Runtime/Water.Build.cs
+++ b/Engine/Plugins/Experimental/Water/Source/Runtime/Water.Build.cs
@@ -8,7 +8,11 @@ public class Water : ModuleRules
 	{
 		PCHUsage = ModuleRules.PCHUsageMode.UseExplicitOrSharedPCHs;
 
-		PrivateIncludePaths.Add("Runtime/Private");
+		PrivateIncludePaths.AddRange(
+			new string[] {
+				System.IO.Path.GetFullPath(Target.RelativeEnginePath) + "/Source/Runtime/Renderer/Private",
+			}
+		);
 
 		PublicDependencyModuleNames.AddRange(
 			new string[] {
diff --git a/Engine/Plugins/FX/Niagara/Source/Niagara/Niagara.Build.cs b/Engine/Plugins/FX/Niagara/Source/Niagara/Niagara.Build.cs
--- a/Engine/Plugins/FX/Niagara/Source/Niagara/Niagara.Build.cs
+++ b/Engine/Plugins/FX/Niagara/Source/Niagara/Niagara.Build.cs
@@ -1,5 +1,6 @@
 // Copyright Epic Games, Inc. All Rights Reserved.
 
+using System.IO;
 using UnrealBuildTool;
 
 public class Niagara : ModuleRules
@@ -60,7 +61,8 @@ public class Niagara : ModuleRules
         PrivateIncludePaths.AddRange(
             new string[] {
                 "Niagara/Private",
-            });
+				Path.Combine(Path.GetFullPath(Target.RelativeEnginePath), "Source/Runtime/Renderer/Private")
+			});
 
         if (Target.bBuildEditor == true)
         {
diff --git a/Engine/Plugins/FX/Niagara/Source/NiagaraEditor/NiagaraEditor.Build.cs b/Engine/Plugins/FX/Niagara/Source/NiagaraEditor/NiagaraEditor.Build.cs
--- a/Engine/Plugins/FX/Niagara/Source/NiagaraEditor/NiagaraEditor.Build.cs
+++ b/Engine/Plugins/FX/Niagara/Source/NiagaraEditor/NiagaraEditor.Build.cs
@@ -1,5 +1,6 @@
 // Copyright Epic Games, Inc. All Rights Reserved.
 
+using System.IO;
 using UnrealBuildTool;
 
 public class NiagaraEditor : ModuleRules
@@ -12,7 +13,8 @@ public class NiagaraEditor : ModuleRules
 			"NiagaraEditor/Private/Widgets",
 			"NiagaraEditor/Private/Sequencer/NiagaraSequence",
 			"NiagaraEditor/Private/ViewModels",
-			"NiagaraEditor/Private/TypeEditorUtilities"
+			"NiagaraEditor/Private/TypeEditorUtilities",
+			Path.Combine(Path.GetFullPath(Target.RelativeEnginePath), "Source/Runtime/Renderer/Private")
 		});
 
 		PrivateDependencyModuleNames.AddRange(
diff --git a/Engine/Plugins/FX/Niagara/Source/NiagaraShader/NiagaraShader.Build.cs b/Engine/Plugins/FX/Niagara/Source/NiagaraShader/NiagaraShader.Build.cs
--- a/Engine/Plugins/FX/Niagara/Source/NiagaraShader/NiagaraShader.Build.cs
+++ b/Engine/Plugins/FX/Niagara/Source/NiagaraShader/NiagaraShader.Build.cs
@@ -1,5 +1,6 @@
 // Copyright Epic Games, Inc. All Rights Reserved.
 
+using System.IO;
 using UnrealBuildTool;
 
 public class NiagaraShader : ModuleRules
@@ -45,7 +46,11 @@ public class NiagaraShader : ModuleRules
             new string[] {
             });
 
-        PrivateIncludePathModuleNames.AddRange(
+		PrivateIncludePaths.AddRange(new string[] {
+			Path.Combine(Path.GetFullPath(Target.RelativeEnginePath), "Source/Runtime/Renderer/Private")
+		});
+
+		PrivateIncludePathModuleNames.AddRange(
             new string[] {
                 "DerivedDataCache",
                 "Niagara",
diff --git a/Engine/Source/Developer/LogVisualizer/Private/LogVisualizer.h b/Engine/Source/Developer/LogVisualizer/Private/LogVisualizer.h
--- a/Engine/Source/Developer/LogVisualizer/Private/LogVisualizer.h
+++ b/Engine/Source/Developer/LogVisualizer/Private/LogVisualizer.h
@@ -13,7 +13,6 @@
 #include "Dom/JsonObject.h"
 #include "Framework/MultiBox/MultiBox.h"
 #include "Framework/MultiBox/MultiBoxBuilder.h"
-#include "Runtime/Slate/Private/Framework/Docking/SDockingTabStack.h"
 
 /* Private includes
 *****************************************************************************/
diff --git a/Engine/Source/Runtime/AIModule/Private/AIModule.cpp b/Engine/Source/Runtime/AIModule/Private/AIModule.cpp
--- a/Engine/Source/Runtime/AIModule/Private/AIModule.cpp
+++ b/Engine/Source/Runtime/AIModule/Private/AIModule.cpp
@@ -8,8 +8,8 @@
 #include "DefaultManagerInstanceTracker.h"
 
 #if WITH_EDITOR
-#include "Developer/AssetTools/Public/IAssetTools.h"
-#include "Developer/AssetTools/Public/AssetToolsModule.h"
+#include "IAssetTools.h"
+#include "AssetToolsModule.h"
 #if ENABLE_VISUAL_LOG
 	#include "VisualLoggerExtension.h"
 #endif // ENABLE_VISUAL_LOG
diff --git a/Engine/Source/Runtime/AIModule/Public/AIModule.h b/Engine/Source/Runtime/AIModule/Public/AIModule.h
--- a/Engine/Source/Runtime/AIModule/Public/AIModule.h
+++ b/Engine/Source/Runtime/AIModule/Public/AIModule.h
@@ -6,7 +6,7 @@
 #include "Modules/ModuleManager.h"
 #include "AI/AISystemBase.h"
 #if WITH_EDITOR
-#include "Developer/AssetTools/Public/AssetTypeCategories.h"
+#include "AssetTypeCategories.h"
 #endif // WITH_EDITOR
 
 /**
diff --git a/Engine/Source/Runtime/AppFramework/AppFramework.Build.cs b/Engine/Source/Runtime/AppFramework/AppFramework.Build.cs
--- a/Engine/Source/Runtime/AppFramework/AppFramework.Build.cs
+++ b/Engine/Source/Runtime/AppFramework/AppFramework.Build.cs
@@ -22,18 +22,6 @@ public class AppFramework : ModuleRules
 			}
 		);
 
-		PrivateIncludePaths.AddRange(
-			new string[] {
-				"Runtime/AppFramework/Private",
-                "Runtime/AppFramework/Private/Framework",
-                "Runtime/AppFramework/Private/Framework/Testing",
-				"Runtime/AppFramework/Private/Widgets",
-                "Runtime/AppFramework/Private/Widgets/Colors",
-                "Runtime/AppFramework/Private/Widgets/Testing",
-                "Runtime/AppFramework/Private/Widgets/Workflow",
-			}
-		);
-
 		if (Target.Configuration != UnrealTargetConfiguration.Shipping)
 		{
 			PrivateIncludePathModuleNames.AddRange(
diff --git a/Engine/Source/Runtime/ApplicationCore/Private/Windows/WindowsApplication.cpp b/Engine/Source/Runtime/ApplicationCore/Private/Windows/WindowsApplication.cpp
--- a/Engine/Source/Runtime/ApplicationCore/Private/Windows/WindowsApplication.cpp
+++ b/Engine/Source/Runtime/ApplicationCore/Private/Windows/WindowsApplication.cpp
@@ -23,7 +23,6 @@
 
 #if WITH_EDITOR
 #include "Modules/ModuleManager.h"
-#include "Developer/SourceCodeAccess/Public/ISourceCodeAccessModule.h"
 #endif
 
 #if WITH_ACCESSIBILITY && UE_WINDOWS_USING_UIA
diff --git a/Engine/Source/Runtime/AudioExtensions/Public/IWaveformTransformation.h b/Engine/Source/Runtime/AudioExtensions/Public/IWaveformTransformation.h
--- a/Engine/Source/Runtime/AudioExtensions/Public/IWaveformTransformation.h
+++ b/Engine/Source/Runtime/AudioExtensions/Public/IWaveformTransformation.h
@@ -4,7 +4,7 @@
 
 #include "Containers/Array.h"
 #include "HAL/Platform.h"
-#include "SignalProcessing/Public/DSP/BufferVectorOperations.h"
+#include "DSP/BufferVectorOperations.h"
 #include "Templates/SharedPointer.h"
 #include "Templates/UniquePtr.h"
 #include "UObject/Object.h"
@@ -70,4 +70,4 @@ public:
 	TArray<TObjectPtr<UWaveformTransformationBase>> Transformations;
 
 	TArray<Audio::FTransformationPtr> CreateTransformations() const;
-};
\ No newline at end of file
+};
diff --git a/Engine/Source/Runtime/AudioLink/AudioLinkEngine/Private/BufferedListenerBase.h b/Engine/Source/Runtime/AudioLink/AudioLinkEngine/Private/BufferedListenerBase.h
--- a/Engine/Source/Runtime/AudioLink/AudioLinkEngine/Private/BufferedListenerBase.h
+++ b/Engine/Source/Runtime/AudioLink/AudioLinkEngine/Private/BufferedListenerBase.h
@@ -3,7 +3,7 @@
 #pragma once
 
 #include "AudioDevice.h"
-#include "SignalProcessing/Public/DSP/Dsp.h"
+#include "DSP/Dsp.h"
 
 #include "IBufferedAudioOutput.h"
 #include "Misc/Optional.h"
diff --git a/Engine/Source/Runtime/AudioMixer/AudioMixer.Build.cs b/Engine/Source/Runtime/AudioMixer/AudioMixer.Build.cs
--- a/Engine/Source/Runtime/AudioMixer/AudioMixer.Build.cs
+++ b/Engine/Source/Runtime/AudioMixer/AudioMixer.Build.cs
@@ -11,13 +11,6 @@ namespace UnrealBuildTool.Rules
 
 			PublicIncludePathModuleNames.Add("Engine");
 
-			PrivateIncludePaths.AddRange(
-				new string[]
-				{
-					"Runtime/AudioMixer/Private",
-				}
-			);
-
 			PublicDependencyModuleNames.AddRange(
 				new string[]
 				{
@@ -38,7 +31,8 @@ namespace UnrealBuildTool.Rules
 					"AudioPlatformConfiguration",
 					"SoundFieldRendering",
 					"AudioExtensions",
-					"AudioLinkEngine"
+					"AudioLinkEngine",
+					"HeadMountedDisplay"
 				}
 			);
 
diff --git a/Engine/Source/Runtime/AudioMixer/Private/AudioMixerDevice.cpp b/Engine/Source/Runtime/AudioMixer/Private/AudioMixerDevice.cpp
--- a/Engine/Source/Runtime/AudioMixer/Private/AudioMixerDevice.cpp
+++ b/Engine/Source/Runtime/AudioMixer/Private/AudioMixerDevice.cpp
@@ -20,7 +20,7 @@
 #include "UObject/StrongObjectPtr.h"
 #include "UObject/UObjectHash.h"
 #include "UObject/UObjectIterator.h"
-#include "Runtime/HeadMountedDisplay/Public/IHeadMountedDisplayModule.h"
+#include "IHeadMountedDisplayModule.h"
 #include "Misc/App.h"
 #include "ProfilingDebugging/CsvProfiler.h"
 #include "AssetRegistry/IAssetRegistry.h"
diff --git a/Engine/Source/Runtime/AudioMixer/Public/SoundWaveDecoder.h b/Engine/Source/Runtime/AudioMixer/Public/SoundWaveDecoder.h
--- a/Engine/Source/Runtime/AudioMixer/Public/SoundWaveDecoder.h
+++ b/Engine/Source/Runtime/AudioMixer/Public/SoundWaveDecoder.h
@@ -3,8 +3,8 @@
 #pragma once
 
 #include "CoreMinimal.h"
-#include "Engine/Public/AudioDevice.h"
-#include "Engine/Classes/Sound/SoundWave.h"
+#include "AudioDevice.h"
+#include "Sound/SoundWave.h"
 #include "DSP/SinOsc.h"
 #include "DSP/ParamInterpolator.h"
 #include "Misc/ScopeLock.h"
diff --git a/Engine/Source/Runtime/Core/Core.Build.cs b/Engine/Source/Runtime/Core/Core.Build.cs
--- a/Engine/Source/Runtime/Core/Core.Build.cs
+++ b/Engine/Source/Runtime/Core/Core.Build.cs
@@ -28,10 +28,6 @@ public class Core : ModuleRules
 		PrivateIncludePaths.AddRange(
 			new string[] {
 				"Runtime/SynthBenchmark/Public",
-				"Runtime/Core/Private",
-				"Runtime/Core/Private/Misc",
-				"Runtime/Core/Private/Internationalization",
-				"Runtime/Core/Private/Internationalization/Cultures",
 				"Runtime/Engine/Public",
 			}
 			);
diff --git a/Engine/Source/Runtime/Core/Private/HAL/UnrealMemory.cpp b/Engine/Source/Runtime/Core/Private/HAL/UnrealMemory.cpp
--- a/Engine/Source/Runtime/Core/Private/HAL/UnrealMemory.cpp
+++ b/Engine/Source/Runtime/Core/Private/HAL/UnrealMemory.cpp
@@ -11,7 +11,7 @@
 #include "Misc/CoreDelegates.h"
 
 #if USE_MALLOC_PROFILER && WITH_ENGINE && IS_MONOLITHIC
-	#include "Runtime/Engine/Public/MallocProfilerEx.h"
+	#include "MallocProfilerEx.h"
 #endif
 
 /*-----------------------------------------------------------------------------
@@ -725,4 +725,4 @@ SIZE_T FMemory::GetUsedPersistentAuxiliary()
 
 #if !INLINE_FMEMORY_OPERATION && !PLATFORM_USES_FIXED_GMalloc_CLASS
 #include "HAL/FMemory.inl"
-#endif
\ No newline at end of file
+#endif
diff --git a/Engine/Source/Runtime/Core/Private/Internationalization/ICURegex.cpp b/Engine/Source/Runtime/Core/Private/Internationalization/ICURegex.cpp
--- a/Engine/Source/Runtime/Core/Private/Internationalization/ICURegex.cpp
+++ b/Engine/Source/Runtime/Core/Private/Internationalization/ICURegex.cpp
@@ -1,6 +1,6 @@
 // Copyright Epic Games, Inc. All Rights Reserved.
 
-#include "ICURegex.h"
+#include "Internationalization/ICURegex.h"
 #include "Misc/ScopeLock.h"
 
 #if UE_ENABLE_ICU
diff --git a/Engine/Source/Runtime/Core/Private/Internationalization/Internationalization.cpp b/Engine/Source/Runtime/Core/Private/Internationalization/Internationalization.cpp
--- a/Engine/Source/Runtime/Core/Private/Internationalization/Internationalization.cpp
+++ b/Engine/Source/Runtime/Core/Private/Internationalization/Internationalization.cpp
@@ -14,7 +14,7 @@
 #if UE_ENABLE_ICU
 #include "Internationalization/ICUInternationalization.h"
 #else
-#include "LegacyInternationalization.h"
+#include "Internationalization/LegacyInternationalization.h"
 #endif
 
 #define LOCTEXT_NAMESPACE "Internationalization"
diff --git a/Engine/Source/Runtime/Core/Private/Internationalization/LegacyInternationalization.cpp b/Engine/Source/Runtime/Core/Private/Internationalization/LegacyInternationalization.cpp
--- a/Engine/Source/Runtime/Core/Private/Internationalization/LegacyInternationalization.cpp
+++ b/Engine/Source/Runtime/Core/Private/Internationalization/LegacyInternationalization.cpp
@@ -5,7 +5,7 @@
 
 #if !UE_ENABLE_ICU
 
-#include "InvariantCulture.h"
+#include "Internationalization/Cultures/InvariantCulture.h"
 
 FLegacyInternationalization::FLegacyInternationalization(FInternationalization* const InI18N)
 	: I18N(InI18N)
diff --git a/Engine/Source/Runtime/Core/Public/ProfilingDebugging/CallstackTrace.h b/Engine/Source/Runtime/Core/Public/ProfilingDebugging/CallstackTrace.h
--- a/Engine/Source/Runtime/Core/Public/ProfilingDebugging/CallstackTrace.h
+++ b/Engine/Source/Runtime/Core/Public/ProfilingDebugging/CallstackTrace.h
@@ -4,7 +4,7 @@
 #include "HAL/PlatformTLS.h"
 #include "Math/NumericLimits.h"
 #include "Misc/Build.h"
-#include "TraceLog/Public/Trace/Config.h"
+#include "Trace/Config.h"
 
 ////////////////////////////////////////////////////////////////////////////////
 #if !defined(UE_CALLSTACK_TRACE_ENABLED)
diff --git a/Engine/Source/Runtime/Core/Public/Tests/TestHarness.h b/Engine/Source/Runtime/Core/Public/Tests/TestHarness.h
--- a/Engine/Source/Runtime/Core/Public/Tests/TestHarness.h
+++ b/Engine/Source/Runtime/Core/Public/Tests/TestHarness.h
@@ -10,7 +10,7 @@
  */
 
 #if WITH_LOW_LEVEL_TESTS
-#include "LowLevelTestsRunner/Public/TestHarness.h"
+#include "TestHarness.h"
 #elif defined(WITH_AUTOMATION_TESTS) || (WITH_DEV_AUTOMATION_TESTS || WITH_PERF_AUTOMATION_TESTS)
 #include "Misc/AutomationTest.h"
 #include "Misc/LowLevelTestAdapter.h"
diff --git a/Engine/Source/Runtime/CoreUObject/Private/CoreUObjectPrivatePCH.h b/Engine/Source/Runtime/CoreUObject/Private/CoreUObjectPrivatePCH.h
--- a/Engine/Source/Runtime/CoreUObject/Private/CoreUObjectPrivatePCH.h
+++ b/Engine/Source/Runtime/CoreUObject/Private/CoreUObjectPrivatePCH.h
@@ -243,19 +243,19 @@
 #include "UObject/LinkerLoad.h"
 #include "UObject/ObjectRedirector.h"
 #include "Misc/PackageName.h"
-#include "Runtime/CoreUObject/Private/UObject/LinkerPlaceholderBase.h"
+#include "UObject/LinkerPlaceholderBase.h"
 #include "UObject/UObjectIterator.h"
 #include "UObject/CoreNet.h"
 #include "UObject/UObjectHash.h"
 #include "UObject/UObjectThreadContext.h"
-#include "Runtime/CoreUObject/Private/UObject/PropertyHelper.h"
-#include "Runtime/CoreUObject/Private/UObject/LinkerPlaceholderClass.h"
+#include "UObject/PropertyHelper.h"
+#include "UObject/LinkerPlaceholderClass.h"
 #include "UObject/UObjectAnnotation.h"
 #include "Internationalization/TextPackageNamespaceUtil.h"
 #include "Blueprint/BlueprintSupport.h"
 #include "UObject/MetaData.h"
-#include "Runtime/CoreUObject/Private/UObject/LinkerPlaceholderFunction.h"
-#include "Runtime/CoreUObject/Private/UObject/LinkerManager.h"
+#include "UObject/LinkerPlaceholderFunction.h"
+#include "UObject/LinkerManager.h"
 #include "UObject/UObjectAllocator.h"
 #include "UObject/GCObject.h"
-#include "Runtime/CoreUObject/Private/UObject/LinkerPlaceholderExportObject.h"
+#include "UObject/LinkerPlaceholderExportObject.h"
diff --git a/Engine/Source/Runtime/D3D12RHI/D3D12RHI.Build.cs b/Engine/Source/Runtime/D3D12RHI/D3D12RHI.Build.cs
--- a/Engine/Source/Runtime/D3D12RHI/D3D12RHI.Build.cs
+++ b/Engine/Source/Runtime/D3D12RHI/D3D12RHI.Build.cs
@@ -8,7 +8,6 @@ public class D3D12RHI : ModuleRules
 
 	public D3D12RHI(ReadOnlyTargetRules Target) : base(Target)
 	{
-		PrivateIncludePaths.Add("Runtime/D3D12RHI/Private");
 		PrivateIncludePaths.Add("../Shaders/Shared");
 
 		PrivateDependencyModuleNames.AddRange(
@@ -22,6 +21,12 @@ public class D3D12RHI : ModuleRules
 				}
 			);
 
+		PublicIncludePathModuleNames.AddRange(
+			new string[] {
+					"HeadMountedDisplay"
+			}
+			);
+
 		if (Target.Configuration != UnrealTargetConfiguration.Shipping)
 		{
 			PrivateIncludePathModuleNames.AddRange(new string[] { "TaskGraph" });
diff --git a/Engine/Source/Runtime/D3D12RHI/Private/D3D12RHIPrivate.h b/Engine/Source/Runtime/D3D12RHI/Private/D3D12RHIPrivate.h
--- a/Engine/Source/Runtime/D3D12RHI/Private/D3D12RHIPrivate.h
+++ b/Engine/Source/Runtime/D3D12RHI/Private/D3D12RHIPrivate.h
@@ -64,17 +64,17 @@ extern int32 GDX12NVAfterMathMarkers;
 #include "D3D12Residency.h"
 
 // D3D RHI public headers.
-#include "../Public/D3D12Util.h"
-#include "../Public/D3D12State.h"
-#include "../Public/D3D12Resources.h"
+#include "D3D12Util.h"
+#include "D3D12State.h"
+#include "D3D12Resources.h"
 #include "D3D12RootSignature.h"
 #include "D3D12Shader.h"
 #include "D3D12View.h"
 #include "D3D12CommandList.h"
 #include "D3D12Texture.h"
 #include "D3D12DirectCommandListManager.h"
-#include "../Public/D3D12Viewport.h"
-#include "../Public/D3D12ConstantBuffer.h"
+#include "D3D12Viewport.h"
+#include "D3D12ConstantBuffer.h"
 #include "D3D12Query.h"
 #include "D3D12DescriptorCache.h"
 #include "D3D12StateCachePrivate.h"
diff --git a/Engine/Source/Runtime/D3D12RHI/Private/Windows/D3D12RHIBasePrivate.h b/Engine/Source/Runtime/D3D12RHI/Private/Windows/D3D12RHIBasePrivate.h
--- a/Engine/Source/Runtime/D3D12RHI/Private/Windows/D3D12RHIBasePrivate.h
+++ b/Engine/Source/Runtime/D3D12RHI/Private/Windows/D3D12RHIBasePrivate.h
@@ -14,6 +14,6 @@
 #define D3D12RHI_NEEDS_VENDOR_EXTENSIONS     1
 #define D3D12RHI_NEEDS_SHADER_FEATURE_CHECKS 1
 
-#include "../Public/D3D12Util.h"
+#include "D3D12Util.h"
 #include "WindowsD3D12DiskCache.h"
 #include "WindowsD3D12PipelineState.h"
diff --git a/Engine/Source/Runtime/D3D12RHI/Private/Windows/WindowsD3D12Device.cpp b/Engine/Source/Runtime/D3D12RHI/Private/Windows/WindowsD3D12Device.cpp
--- a/Engine/Source/Runtime/D3D12RHI/Private/Windows/WindowsD3D12Device.cpp
+++ b/Engine/Source/Runtime/D3D12RHI/Private/Windows/WindowsD3D12Device.cpp
@@ -34,7 +34,7 @@
 #include "Windows/HideWindowsPlatformTypes.h"
 
 #include "HardwareInfo.h"
-#include "Runtime/HeadMountedDisplay/Public/IHeadMountedDisplayModule.h"
+#include "IHeadMountedDisplayModule.h"
 #include "GenericPlatform/GenericPlatformDriver.h"			// FGPUDriverInfo
 #include "RHIValidation.h"
 
diff --git a/Engine/Source/Runtime/Engine/Classes/Engine/GameEngine.h b/Engine/Source/Runtime/Engine/Classes/Engine/GameEngine.h
--- a/Engine/Source/Runtime/Engine/Classes/Engine/GameEngine.h
+++ b/Engine/Source/Runtime/Engine/Classes/Engine/GameEngine.h
@@ -7,7 +7,7 @@
 #include "Widgets/SWindow.h"
 #include "Widgets/SViewport.h"
 #include "Engine/Engine.h"
-#include "Runtime/MovieSceneCapture/Public/MovieSceneCaptureHandle.h"
+#include "MovieSceneCaptureHandle.h"
 #include "Templates/PimplPtr.h"
 #include "GameEngine.generated.h"
 
diff --git a/Engine/Source/Runtime/Engine/Classes/PhysicsEngine/ShapeElem.h b/Engine/Source/Runtime/Engine/Classes/PhysicsEngine/ShapeElem.h
--- a/Engine/Source/Runtime/Engine/Classes/PhysicsEngine/ShapeElem.h
+++ b/Engine/Source/Runtime/Engine/Classes/PhysicsEngine/ShapeElem.h
@@ -6,7 +6,7 @@
 #include "UObject/ObjectMacros.h"
 #include "EngineDefines.h"
 #include "PhysxUserData.h"
-#include "PhysicsCore/Public/PhysicsInterfaceTypesCore.h"
+#include "PhysicsInterfaceTypesCore.h"
 #include "Engine/EngineTypes.h"
 #include "ShapeElem.generated.h"
 
diff --git a/Engine/Source/Runtime/Engine/Classes/Sound/QuartzQuantizationUtilities.cpp b/Engine/Source/Runtime/Engine/Classes/Sound/QuartzQuantizationUtilities.cpp
--- a/Engine/Source/Runtime/Engine/Classes/Sound/QuartzQuantizationUtilities.cpp
+++ b/Engine/Source/Runtime/Engine/Classes/Sound/QuartzQuantizationUtilities.cpp
@@ -2,9 +2,9 @@
 
 #include "Sound/QuartzQuantizationUtilities.h"
 
-#include "Core/Public/CoreGlobals.h"
-#include "Core/Public/Math/NumericLimits.h"
-#include "Core/Public/Math/UnrealMathUtility.h"
+#include "CoreGlobals.h"
+#include "Math/NumericLimits.h"
+#include "Math/UnrealMathUtility.h"
 #include "AudioMixerDevice.h"
 
 #define INVALID_DURATION -1
diff --git a/Engine/Source/Runtime/Engine/Engine.Build.cs b/Engine/Source/Runtime/Engine/Engine.Build.cs
--- a/Engine/Source/Runtime/Engine/Engine.Build.cs
+++ b/Engine/Source/Runtime/Engine/Engine.Build.cs
@@ -13,21 +13,24 @@ public class Engine : ModuleRules
 
 		SharedPCHHeaderFile = "Public/EngineSharedPCH.h";
 
-		PublicIncludePathModuleNames.AddRange(new string[] {
-			"Renderer", 
-			"PacketHandler", 
-			"AudioMixer", 
-			"AudioMixerCore",
-			"AnimationCore"
+		PublicIncludePathModuleNames.AddRange(
+			new string[] {
+				"AnimationCore",
+				"AudioMixer", 
+				"AudioMixerCore",
+				"MovieSceneCapture", 
+				"PacketHandler", 
+				"Renderer", 
 			}
 		);
 
 		PrivateIncludePaths.AddRange(
 			new string[] {
-				"Runtime/SynthBenchmark/Public",
+				"Developer/Virtualization/Private",
 				"Runtime/Engine/Private",
-				"Runtime/Net/Core/Private/Net/Core/PushModel/Types",
-				"Developer/Virtualization/Private"
+				"Runtime/Net/Core/Private",
+				"Runtime/Renderer/Private",
+				"Runtime/SynthBenchmark/Public",
 			}
 		);
 
diff --git a/Engine/Source/Runtime/Engine/Private/Atmosphere/AtmosphericFog.cpp b/Engine/Source/Runtime/Engine/Private/Atmosphere/AtmosphericFog.cpp
--- a/Engine/Source/Runtime/Engine/Private/Atmosphere/AtmosphericFog.cpp
+++ b/Engine/Source/Runtime/Engine/Private/Atmosphere/AtmosphericFog.cpp
@@ -12,7 +12,7 @@
 #include "ComponentReregisterContext.h"
 #include "Atmosphere/AtmosphericFog.h"
 #include "Components/BillboardComponent.h"
-#include "Runtime/Renderer/Private/ScenePrivate.h"
+#include "ScenePrivate.h"
 
 #if WITH_EDITOR
 #include "ObjectEditorUtils.h"
diff --git a/Engine/Source/Runtime/Engine/Private/Atmosphere/AtmosphericFogComponent.cpp b/Engine/Source/Runtime/Engine/Private/Atmosphere/AtmosphericFogComponent.cpp
--- a/Engine/Source/Runtime/Engine/Private/Atmosphere/AtmosphericFogComponent.cpp
+++ b/Engine/Source/Runtime/Engine/Private/Atmosphere/AtmosphericFogComponent.cpp
@@ -11,7 +11,7 @@
 #include "ComponentReregisterContext.h"
 #include "Atmosphere/AtmosphericFog.h"
 #include "Components/BillboardComponent.h"
-#include "Runtime/Renderer/Private/ScenePrivate.h"
+#include "ScenePrivate.h"
 #include "UObject/UE5MainStreamObjectVersion.h"
 
 #if WITH_EDITOR
diff --git a/Engine/Source/Runtime/Engine/Private/AudioDerivedData.cpp b/Engine/Source/Runtime/Engine/Private/AudioDerivedData.cpp
--- a/Engine/Source/Runtime/Engine/Private/AudioDerivedData.cpp
+++ b/Engine/Source/Runtime/Engine/Private/AudioDerivedData.cpp
@@ -23,7 +23,7 @@
 #include "Sound/SoundSourceBus.h"
 #include "Sound/SoundWave.h"
 #include "Sound/SoundWaveProcedural.h"
-#include "AudioExtensions/Public/IWaveformTransformation.h"
+#include "IWaveformTransformation.h"
 #include "DSP/FloatArrayMath.h"
 #include "DSP/BufferVectorOperations.h"
 #include "DSP/MultichannelBuffer.h"
diff --git a/Engine/Source/Runtime/Engine/Private/AudioDevice.cpp b/Engine/Source/Runtime/Engine/Private/AudioDevice.cpp
--- a/Engine/Source/Runtime/Engine/Private/AudioDevice.cpp
+++ b/Engine/Source/Runtime/Engine/Private/AudioDevice.cpp
@@ -48,8 +48,8 @@
 #if WITH_EDITOR
 #include "AudioEditorModule.h"
 #include "AssetRegistry/AssetRegistryModule.h"
-#include "Developer/AssetTools/Public/AssetToolsModule.h"
-#include "Developer/AssetTools/Public/IAssetTools.h"
+#include "AssetToolsModule.h"
+#include "IAssetTools.h"
 #include "Editor/EditorEngine.h"
 #endif // WITH_EDITOR
 
diff --git a/Engine/Source/Runtime/Engine/Private/DataReplication.cpp b/Engine/Source/Runtime/Engine/Private/DataReplication.cpp
--- a/Engine/Source/Runtime/Engine/Private/DataReplication.cpp
+++ b/Engine/Source/Runtime/Engine/Private/DataReplication.cpp
@@ -22,7 +22,7 @@
 #include "Net/Core/Trace/NetTrace.h"
 #include "ProfilingDebugging/CsvProfiler.h"
 #include "HAL/LowLevelMemStats.h"
-#include "PushModelPerNetDriverState.h"
+#include "Net/Core/PushModel/Types/PushModelPerNetDriverState.h"
 
 DECLARE_LLM_MEMORY_STAT(TEXT("NetObjReplicator"), STAT_NetObjReplicatorLLM, STATGROUP_LLMFULL);
 LLM_DEFINE_TAG(NetObjReplicator, NAME_None, TEXT("Networking"), GET_STATFNAME(STAT_NetObjReplicatorLLM), GET_STATFNAME(STAT_NetworkingSummaryLLM));
diff --git a/Engine/Source/Runtime/Engine/Private/Materials/Material.cpp b/Engine/Source/Runtime/Engine/Private/Materials/Material.cpp
--- a/Engine/Source/Runtime/Engine/Private/Materials/Material.cpp
+++ b/Engine/Source/Runtime/Engine/Private/Materials/Material.cpp
@@ -70,7 +70,7 @@
 #include "Materials/MaterialExpressionComment.h"
 #include "UObject/EditorObjectVersion.h"
 #include "UObject/ReleaseObjectVersion.h"
-#include "RenderCore/Public/RenderUtils.h"
+#include "RenderUtils.h"
 #include "Engine/Font.h"
 #include "UObject/UE5MainStreamObjectVersion.h"
 #include "UObject/UE5ReleaseStreamObjectVersion.h"
diff --git a/Engine/Source/Runtime/Engine/Private/MeshSimplificationSettings.cpp b/Engine/Source/Runtime/Engine/Private/MeshSimplificationSettings.cpp
--- a/Engine/Source/Runtime/Engine/Private/MeshSimplificationSettings.cpp
+++ b/Engine/Source/Runtime/Engine/Private/MeshSimplificationSettings.cpp
@@ -6,7 +6,7 @@
 #include "Modules/ModuleManager.h"
 
 #if WITH_EDITOR
-#include "Developer/MeshReductionInterface/Public/IMeshReductionManagerModule.h"
+#include "IMeshReductionManagerModule.h"
 #endif
 
 UMeshSimplificationSettings::UMeshSimplificationSettings(const FObjectInitializer& ObjectInitializer)
diff --git a/Engine/Source/Runtime/Engine/Private/NetDriver.cpp b/Engine/Source/Runtime/Engine/Private/NetDriver.cpp
--- a/Engine/Source/Runtime/Engine/Private/NetDriver.cpp
+++ b/Engine/Source/Runtime/Engine/Private/NetDriver.cpp
@@ -70,7 +70,7 @@
 #include "Engine/LevelScriptActor.h"
 #include "Engine/NetworkSettings.h"
 #include "Engine/NetworkDelegates.h"
-#include "Engine/Private/Net/NetSubObjectRegistryGetter.h"
+#include "Net/NetSubObjectRegistryGetter.h"
 #include "Net/NetworkGranularMemoryLogging.h"
 #include "SocketSubsystem.h"
 #include "AddressInfoTypes.h"
diff --git a/Engine/Source/Runtime/Engine/Private/NetworkProfiler.cpp b/Engine/Source/Runtime/Engine/Private/NetworkProfiler.cpp
--- a/Engine/Source/Runtime/Engine/Private/NetworkProfiler.cpp
+++ b/Engine/Source/Runtime/Engine/Private/NetworkProfiler.cpp
@@ -13,7 +13,7 @@
 #include "Engine/World.h"
 #include "Serialization/MemoryWriter.h"
 #include "HAL/IConsoleManager.h"
-#include "Engine/Public/TimerManager.h"
+#include "TimerManager.h"
 
 #if USE_NETWORK_PROFILER
 
diff --git a/Engine/Source/Runtime/Engine/Private/Particles/FXSystem.cpp b/Engine/Source/Runtime/Engine/Private/Particles/FXSystem.cpp
--- a/Engine/Source/Runtime/Engine/Private/Particles/FXSystem.cpp
+++ b/Engine/Source/Runtime/Engine/Private/Particles/FXSystem.cpp
@@ -14,7 +14,7 @@
 #include "VectorField/VectorField.h"
 #include "Components/VectorFieldComponent.h"
 #include "SceneUtils.h"
-#include "Renderer/Private/SceneRendering.h" // needed for STATGROUP_CommandListMarkers
+#include "SceneRendering.h" // needed for STATGROUP_CommandListMarkers
 #include "GPUSortManager.h"
 
 
diff --git a/Engine/Source/Runtime/Engine/Private/Particles/ParticleSystemRender.cpp b/Engine/Source/Runtime/Engine/Private/Particles/ParticleSystemRender.cpp
--- a/Engine/Source/Runtime/Engine/Private/Particles/ParticleSystemRender.cpp
+++ b/Engine/Source/Runtime/Engine/Private/Particles/ParticleSystemRender.cpp
@@ -33,7 +33,7 @@
 #include "Particles/TypeData/ParticleModuleTypeDataRibbon.h"
 #include "Particles/ParticleModuleRequired.h"
 #include "ParticleBeamTrailVertexFactory.h"
-#include "Renderer/Private/SceneRendering.h"
+#include "SceneRendering.h"
 #include "Particles/ParticleLODLevel.h"
 #include "Engine/StaticMesh.h"
 #include "UnrealEngine.h"
diff --git a/Engine/Source/Runtime/Engine/Private/PhysicsEngine/AsyncPhysicsInputComponent.cpp b/Engine/Source/Runtime/Engine/Private/PhysicsEngine/AsyncPhysicsInputComponent.cpp
--- a/Engine/Source/Runtime/Engine/Private/PhysicsEngine/AsyncPhysicsInputComponent.cpp
+++ b/Engine/Source/Runtime/Engine/Private/PhysicsEngine/AsyncPhysicsInputComponent.cpp
@@ -7,7 +7,7 @@
 #include "Components/PrimitiveComponent.h"
 #include "Engine/World.h"
 #include "EngineUtils.h"
-#include "Experimental/Chaos/Public/RewindData.h"
+#include "RewindData.h"
 #include "GameFramework/Actor.h"
 #include "GameFramework/Pawn.h"
 #include "GameFramework/PlayerController.h"
diff --git a/Engine/Source/Runtime/Engine/Private/PhysicsEngine/PhysicsReplication.cpp b/Engine/Source/Runtime/Engine/Private/PhysicsEngine/PhysicsReplication.cpp
--- a/Engine/Source/Runtime/Engine/Private/PhysicsEngine/PhysicsReplication.cpp
+++ b/Engine/Source/Runtime/Engine/Private/PhysicsEngine/PhysicsReplication.cpp
@@ -12,9 +12,9 @@
 #include "DrawDebugHelpers.h"
 #include "PhysicsEngine/PhysicsSettings.h"
 #include "PhysicsPublic.h"
-#include "Engine/Classes/GameFramework/PlayerController.h"
-#include "Engine/Classes/GameFramework/PlayerState.h"
-#include "Engine/Classes/Engine/Player.h"
+#include "GameFramework/PlayerController.h"
+#include "GameFramework/PlayerState.h"
+#include "Engine/Player.h"
 #include "Physics/PhysicsInterfaceCore.h"
 #include "Physics/PhysScene_PhysX.h"
 #include "Components/SkeletalMeshComponent.h"
diff --git a/Engine/Source/Runtime/Engine/Private/SceneManagement.cpp b/Engine/Source/Runtime/Engine/Private/SceneManagement.cpp
--- a/Engine/Source/Runtime/Engine/Private/SceneManagement.cpp
+++ b/Engine/Source/Runtime/Engine/Private/SceneManagement.cpp
@@ -6,8 +6,8 @@
 #include "DeviceProfiles/DeviceProfile.h"
 #include "DeviceProfiles/DeviceProfileManager.h"
 #include "StaticMeshResources.h"
-#include "Runtime/Renderer/Private/SceneRendering.h"
-#include "Runtime/Renderer/Private/SceneCore.h"
+#include "SceneRendering.h"
+#include "SceneCore.h"
 #include "Async/ParallelFor.h"
 #include "LightMap.h"
 #include "ShadowMap.h"
diff --git a/Engine/Source/Runtime/Engine/Private/Tests/AutomationCommon.cpp b/Engine/Source/Runtime/Engine/Private/Tests/AutomationCommon.cpp
--- a/Engine/Source/Runtime/Engine/Private/Tests/AutomationCommon.cpp
+++ b/Engine/Source/Runtime/Engine/Private/Tests/AutomationCommon.cpp
@@ -24,7 +24,7 @@
 #include "Tests/AutomationTestSettings.h"
 #include "GameMapsSettings.h"
 #include "IRenderCaptureProvider.h"
-#include <Runtime/Core/Public/Algo/Accumulate.h>
+#include "Algo/Accumulate.h"
 
 #if WITH_AUTOMATION_TESTS
 
diff --git a/Engine/Source/Runtime/Engine/Private/Texture2D.cpp b/Engine/Source/Runtime/Engine/Private/Texture2D.cpp
--- a/Engine/Source/Runtime/Engine/Private/Texture2D.cpp
+++ b/Engine/Source/Runtime/Engine/Private/Texture2D.cpp
@@ -48,7 +48,7 @@
 #include "VT/VirtualTexturePoolConfig.h"
 #include "ProfilingDebugging/LoadTimeTracker.h"
 #include "VT/VirtualTextureScalability.h"
-#include "Engine/Public/ImageUtils.h"
+#include "ImageUtils.h"
 #include "TextureCompiler.h"
 #include "Misc/ScopedSlowTask.h"
 #include "UObject/ObjectSaveContext.h"
diff --git a/Engine/Source/Runtime/Engine/Private/TextureCube.cpp b/Engine/Source/Runtime/Engine/Private/TextureCube.cpp
--- a/Engine/Source/Runtime/Engine/Private/TextureCube.cpp
+++ b/Engine/Source/Runtime/Engine/Private/TextureCube.cpp
@@ -14,7 +14,7 @@
 #include "TextureCompiler.h"
 #include "Misc/ScopedSlowTask.h"
 #include "UObject/StrongObjectPtr.h"
-#include "Engine/Public/ImageUtils.h"
+#include "ImageUtils.h"
 #include "UObject/ReleaseObjectVersion.h"
 
 #define LOCTEXT_NAMESPACE "UTextureCube"
@@ -676,4 +676,4 @@ uint32 UTextureCube::GetMaximumDimension() const
 }
 #endif // #if WITH_EDITOR
 
-#undef LOCTEXT_NAMESPACE
\ No newline at end of file
+#undef LOCTEXT_NAMESPACE
diff --git a/Engine/Source/Runtime/Engine/Private/VorbisAudioInfo.cpp b/Engine/Source/Runtime/Engine/Private/VorbisAudioInfo.cpp
--- a/Engine/Source/Runtime/Engine/Private/VorbisAudioInfo.cpp
+++ b/Engine/Source/Runtime/Engine/Private/VorbisAudioInfo.cpp
@@ -5,7 +5,7 @@
 #include "Misc/Paths.h"
 #include "Interfaces/IAudioFormat.h"
 #include "ContentStreaming.h"
-#include "Engine/Public/Audio.h"
+#include "Audio.h"
 #if PLATFORM_WINDOWS
 #include "Windows/WindowsHWrapper.h"
 #endif
diff --git a/Engine/Source/Runtime/Engine/Private/WorldPartition/WorldPartition.cpp b/Engine/Source/Runtime/Engine/Private/WorldPartition/WorldPartition.cpp
--- a/Engine/Source/Runtime/Engine/Private/WorldPartition/WorldPartition.cpp
+++ b/Engine/Source/Runtime/Engine/Private/WorldPartition/WorldPartition.cpp
@@ -36,7 +36,7 @@
 #include "UnrealEdMisc.h"
 #include "LocationVolume.h"
 #include "Engine/LevelScriptBlueprint.h"
-#include "Engine/Public/ActorReferencesUtils.h"
+#include "ActorReferencesUtils.h"
 #include "WorldPartition/IWorldPartitionEditorModule.h"
 #include "WorldPartition/WorldPartitionLevelHelper.h"
 #include "WorldPartition/WorldPartitionLevelStreamingDynamic.h"
diff --git a/Engine/Source/Runtime/Engine/Private/WorldPartition/WorldPartitionActorDesc.cpp b/Engine/Source/Runtime/Engine/Private/WorldPartition/WorldPartitionActorDesc.cpp
--- a/Engine/Source/Runtime/Engine/Private/WorldPartition/WorldPartitionActorDesc.cpp
+++ b/Engine/Source/Runtime/Engine/Private/WorldPartition/WorldPartitionActorDesc.cpp
@@ -24,7 +24,7 @@
 #include "WorldPartition/DataLayer/DataLayerAsset.h"
 #include "WorldPartition/DataLayer/DataLayerUtils.h"
 #include "WorldPartition/ErrorHandling/WorldPartitionStreamingGenerationErrorHandler.h"
-#include "Engine/Public/ActorReferencesUtils.h"
+#include "ActorReferencesUtils.h"
 
 TMap<TSubclassOf<AActor>, FWorldPartitionActorDesc::FActorDescDeprecator> FWorldPartitionActorDesc::Deprecators;
 
@@ -607,4 +607,4 @@ void FWorldPartitionActorDesc::Unload()
 		ActorPtr = nullptr;
 	}
 }
-#endif
\ No newline at end of file
+#endif
diff --git a/Engine/Source/Runtime/Engine/Public/Animation/SkinWeightProfile.h b/Engine/Source/Runtime/Engine/Public/Animation/SkinWeightProfile.h
--- a/Engine/Source/Runtime/Engine/Public/Animation/SkinWeightProfile.h
+++ b/Engine/Source/Runtime/Engine/Public/Animation/SkinWeightProfile.h
@@ -9,8 +9,8 @@
 #include "RenderingThread.h"
 #include "HAL/UnrealMemory.h"
 #include "BoneIndices.h"
-#include "RHI/Public/RHIGPUReadback.h"
-#include "Core/Public/Templates/UniquePtr.h"
+#include "RHIGPUReadback.h"
+#include "Templates/UniquePtr.h"
 
 #include "SkinWeightProfile.generated.h"
 
diff --git a/Engine/Source/Runtime/Engine/Public/GeneratedCodeHelpers.h b/Engine/Source/Runtime/Engine/Public/GeneratedCodeHelpers.h
--- a/Engine/Source/Runtime/Engine/Public/GeneratedCodeHelpers.h
+++ b/Engine/Source/Runtime/Engine/Public/GeneratedCodeHelpers.h
@@ -28,9 +28,9 @@
 //For DOREPLIFETIME macros
 #include "Net/UnrealNetwork.h"
 
-#include "Runtime/Core/Public/Math/Box2D.h"
-#include "Runtime/Core/Public/Math/InterpCurvePoint.h"
-#include "Runtime/Core/Public/Containers/EnumAsByte.h"
+#include "Math/Box2D.h"
+#include "Math/InterpCurvePoint.h"
+#include "Containers/EnumAsByte.h"
 
 inline FBox2D CreateFBox2D(FVector2D InMin, FVector2D InMax, bool InIsValid)
 {
@@ -769,4 +769,4 @@ T ConstructTInlineValue(UScriptStruct* Struct)
 	void* Allocation = Value.Reserve(StructOps->GetSize(), StructOps->GetAlignment());
 	Struct->InitializeStruct(Allocation);
 	return Value;
-}
\ No newline at end of file
+}
diff --git a/Engine/Source/Runtime/Engine/Public/Subsystems/WorldSubsystem.h b/Engine/Source/Runtime/Engine/Public/Subsystems/WorldSubsystem.h
--- a/Engine/Source/Runtime/Engine/Public/Subsystems/WorldSubsystem.h
+++ b/Engine/Source/Runtime/Engine/Public/Subsystems/WorldSubsystem.h
@@ -4,7 +4,7 @@
 
 #include "Subsystem.h"
 #include "Engine/EngineTypes.h"
-#include "Engine/Public/Tickable.h"
+#include "Tickable.h"
 #include "WorldSubsystem.generated.h"
 
 /**
@@ -74,4 +74,4 @@ public:
 
 private:
 	bool bInitialized = false;
-};
\ No newline at end of file
+};
diff --git a/Engine/Source/Runtime/Experimental/Animation/Constraints/Private/ConstraintsScripting.cpp b/Engine/Source/Runtime/Experimental/Animation/Constraints/Private/ConstraintsScripting.cpp
--- a/Engine/Source/Runtime/Experimental/Animation/Constraints/Private/ConstraintsScripting.cpp
+++ b/Engine/Source/Runtime/Experimental/Animation/Constraints/Private/ConstraintsScripting.cpp
@@ -3,7 +3,7 @@
 
 #include "ConstraintsScripting.h"
 #include "TransformConstraint.h"
-#include "Animation/Constraints/Public/TransformableHandle.h"
+#include "TransformableHandle.h"
 #include "ConstraintsManager.h"
 
 UConstraintsManager* UConstraintsScriptingLibrary::GetManager(UWorld* InWorld)
@@ -57,4 +57,4 @@ bool UConstraintsScriptingLibrary::RemoveConstraint(UWorld* InWorld, int32 InInd
 {
 	const FConstraintsManagerController& Controller = FConstraintsManagerController::Get(InWorld);
 	return Controller.RemoveConstraint(InIndex);
-}
\ No newline at end of file
+}
diff --git a/Engine/Source/Runtime/Experimental/Animation/Constraints/Private/TransformConstraint.cpp b/Engine/Source/Runtime/Experimental/Animation/Constraints/Private/TransformConstraint.cpp
--- a/Engine/Source/Runtime/Experimental/Animation/Constraints/Private/TransformConstraint.cpp
+++ b/Engine/Source/Runtime/Experimental/Animation/Constraints/Private/TransformConstraint.cpp
@@ -1,7 +1,7 @@
 ﻿// Copyright Epic Games, Inc. All Rights Reserved.
 
 #include "TransformConstraint.h"
-#include "Animation/Constraints/Public/TransformableHandle.h"
+#include "TransformableHandle.h"
 #include "ConstraintsManager.h"
 #include "TransformableRegistry.h"
 #include "GameFramework/Actor.h"
@@ -1269,4 +1269,4 @@ void FTransformConstraintUtils::GetChildrenConstraints(
 	} );
 
 	OutConstraints.Append(FilteredConstraints);
-}
\ No newline at end of file
+}
diff --git a/Engine/Source/Runtime/Experimental/ChaosSolverEngine/Private/Chaos/ChaosGameplayEventDispatcher.cpp b/Engine/Source/Runtime/Experimental/ChaosSolverEngine/Private/Chaos/ChaosGameplayEventDispatcher.cpp
--- a/Engine/Source/Runtime/Experimental/ChaosSolverEngine/Private/Chaos/ChaosGameplayEventDispatcher.cpp
+++ b/Engine/Source/Runtime/Experimental/ChaosSolverEngine/Private/Chaos/ChaosGameplayEventDispatcher.cpp
@@ -5,7 +5,7 @@
 #include "Components/SkeletalMeshComponent.h"
 #include "PhysicsEngine/BodyInstance.h"
 #include "Chaos/Framework/PhysicsProxy.h"
-#include "Chaos/Public/PhysicsProxy/SingleParticlePhysicsProxy.h"
+#include "PhysicsProxy/SingleParticlePhysicsProxy.h"
 #include "PhysicsSolver.h"
 #include "Physics/Experimental/PhysScene_Chaos.h"
 #include "Engine/World.h"
diff --git a/Engine/Source/Runtime/Experimental/GeometryCollectionEngine/Private/GeometryCollection/GeometryCollectionSceneProxy.cpp b/Engine/Source/Runtime/Experimental/GeometryCollectionEngine/Private/GeometryCollection/GeometryCollectionSceneProxy.cpp
--- a/Engine/Source/Runtime/Experimental/GeometryCollectionEngine/Private/GeometryCollection/GeometryCollectionSceneProxy.cpp
+++ b/Engine/Source/Runtime/Experimental/GeometryCollectionEngine/Private/GeometryCollection/GeometryCollectionSceneProxy.cpp
@@ -6,7 +6,7 @@
 #include "Engine/Engine.h"
 #include "MaterialShaderType.h"
 #include "Materials/Material.h"
-#include "RenderCore/Public/CommonRenderResources.h"
+#include "CommonRenderResources.h"
 #include "Rendering/NaniteResources.h"
 #include "PrimitiveSceneInfo.h"
 #include "GeometryCollection/GeometryCollectionComponent.h"
diff --git a/Engine/Source/Runtime/ImageWrapper/ImageWrapper.Build.cs b/Engine/Source/Runtime/ImageWrapper/ImageWrapper.Build.cs
--- a/Engine/Source/Runtime/ImageWrapper/ImageWrapper.Build.cs
+++ b/Engine/Source/Runtime/ImageWrapper/ImageWrapper.Build.cs
@@ -6,12 +6,6 @@ public class ImageWrapper : ModuleRules
 {
 	public ImageWrapper(ReadOnlyTargetRules Target) : base(Target)
 	{
-		PrivateIncludePaths.AddRange(
-			new string[] {
-				"Runtime/ImageWrapper/Private",
-				"Runtime/ImageWrapper/Private/Formats",
-			});
-
 		PublicIncludePathModuleNames.AddRange(new string[] {
 			"ImageCore"
 			}
diff --git a/Engine/Source/Runtime/ImageWrapper/Private/Formats/BmpImageWrapper.cpp b/Engine/Source/Runtime/ImageWrapper/Private/Formats/BmpImageWrapper.cpp
--- a/Engine/Source/Runtime/ImageWrapper/Private/Formats/BmpImageWrapper.cpp
+++ b/Engine/Source/Runtime/ImageWrapper/Private/Formats/BmpImageWrapper.cpp
@@ -1,6 +1,6 @@
 // Copyright Epic Games, Inc. All Rights Reserved.
 
-#include "BmpImageWrapper.h"
+#include "Formats/BmpImageWrapper.h"
 #include "ImageWrapperPrivate.h"
 
 #include "BmpImageSupport.h"
@@ -573,4 +573,4 @@ void FBmpImageWrapper::Compress(int32 Quality)
 	}
 	
 	check( PayloadPtr == CompressedData.GetData() + CompressedData.Num() );
-}
\ No newline at end of file
+}
diff --git a/Engine/Source/Runtime/ImageWrapper/Private/Formats/DdsImageWrapper.cpp b/Engine/Source/Runtime/ImageWrapper/Private/Formats/DdsImageWrapper.cpp
--- a/Engine/Source/Runtime/ImageWrapper/Private/Formats/DdsImageWrapper.cpp
+++ b/Engine/Source/Runtime/ImageWrapper/Private/Formats/DdsImageWrapper.cpp
@@ -1,6 +1,6 @@
 // Copyright Epic Games, Inc. All Rights Reserved.
 
-#include "DdsImageWrapper.h"
+#include "Formats/DdsImageWrapper.h"
 #include "ImageWrapperPrivate.h"
 
 void FDdsImageWrapper::Reset()
diff --git a/Engine/Source/Runtime/ImageWrapper/Private/Formats/ExrImageWrapper.cpp b/Engine/Source/Runtime/ImageWrapper/Private/Formats/ExrImageWrapper.cpp
--- a/Engine/Source/Runtime/ImageWrapper/Private/Formats/ExrImageWrapper.cpp
+++ b/Engine/Source/Runtime/ImageWrapper/Private/Formats/ExrImageWrapper.cpp
@@ -1,6 +1,6 @@
 // Copyright Epic Games, Inc. All Rights Reserved.
 
-#include "ExrImageWrapper.h"
+#include "Formats/ExrImageWrapper.h"
 #include "ImageWrapperPrivate.h"
 
 #include "Containers/StringConv.h"
diff --git a/Engine/Source/Runtime/ImageWrapper/Private/Formats/IcnsImageWrapper.cpp b/Engine/Source/Runtime/ImageWrapper/Private/Formats/IcnsImageWrapper.cpp
--- a/Engine/Source/Runtime/ImageWrapper/Private/Formats/IcnsImageWrapper.cpp
+++ b/Engine/Source/Runtime/ImageWrapper/Private/Formats/IcnsImageWrapper.cpp
@@ -1,6 +1,6 @@
 // Copyright Epic Games, Inc. All Rights Reserved.
 
-#include "IcnsImageWrapper.h"
+#include "Formats/IcnsImageWrapper.h"
 
 
 /* FIcnsImageWrapper structors
diff --git a/Engine/Source/Runtime/ImageWrapper/Private/Formats/IcoImageWrapper.cpp b/Engine/Source/Runtime/ImageWrapper/Private/Formats/IcoImageWrapper.cpp
--- a/Engine/Source/Runtime/ImageWrapper/Private/Formats/IcoImageWrapper.cpp
+++ b/Engine/Source/Runtime/ImageWrapper/Private/Formats/IcoImageWrapper.cpp
@@ -1,6 +1,6 @@
 // Copyright Epic Games, Inc. All Rights Reserved.
 
-#include "IcoImageWrapper.h"
+#include "Formats/IcoImageWrapper.h"
 
 #include "BmpImageSupport.h"
 #include "Formats/PngImageWrapper.h"
diff --git a/Engine/Source/Runtime/ImageWrapper/Private/Formats/JpegImageWrapper.cpp b/Engine/Source/Runtime/ImageWrapper/Private/Formats/JpegImageWrapper.cpp
--- a/Engine/Source/Runtime/ImageWrapper/Private/Formats/JpegImageWrapper.cpp
+++ b/Engine/Source/Runtime/ImageWrapper/Private/Formats/JpegImageWrapper.cpp
@@ -1,6 +1,6 @@
 // Copyright Epic Games, Inc. All Rights Reserved.
 
-#include "JpegImageWrapper.h"
+#include "Formats/JpegImageWrapper.h"
 
 #include "Math/Color.h"
 #include "Misc/ScopeLock.h"
diff --git a/Engine/Source/Runtime/ImageWrapper/Private/Formats/PngImageWrapper.cpp b/Engine/Source/Runtime/ImageWrapper/Private/Formats/PngImageWrapper.cpp
--- a/Engine/Source/Runtime/ImageWrapper/Private/Formats/PngImageWrapper.cpp
+++ b/Engine/Source/Runtime/ImageWrapper/Private/Formats/PngImageWrapper.cpp
@@ -1,6 +1,6 @@
 // Copyright Epic Games, Inc. All Rights Reserved.
 
-#include "PngImageWrapper.h"
+#include "Formats/PngImageWrapper.h"
 #include "ImageWrapperPrivate.h"
 
 #include "Misc/ScopeLock.h"
diff --git a/Engine/Source/Runtime/ImageWrapper/Private/Formats/TgaImageWrapper.cpp b/Engine/Source/Runtime/ImageWrapper/Private/Formats/TgaImageWrapper.cpp
--- a/Engine/Source/Runtime/ImageWrapper/Private/Formats/TgaImageWrapper.cpp
+++ b/Engine/Source/Runtime/ImageWrapper/Private/Formats/TgaImageWrapper.cpp
@@ -1,6 +1,6 @@
 // Copyright Epic Games, Inc. All Rights Reserved.
 
-#include "TgaImageWrapper.h"
+#include "Formats/TgaImageWrapper.h"
 #include "ImageWrapperPrivate.h"
 
 #include "TgaImageSupport.h"
diff --git a/Engine/Source/Runtime/Landscape/Landscape.Build.cs b/Engine/Source/Runtime/Landscape/Landscape.Build.cs
--- a/Engine/Source/Runtime/Landscape/Landscape.Build.cs
+++ b/Engine/Source/Runtime/Landscape/Landscape.Build.cs
@@ -10,8 +10,8 @@ public class Landscape : ModuleRules
 		PrivateIncludePaths.AddRange(
 			new string[] {
 				"Runtime/Engine/Private", // for Engine/Private/Collision/PhysXCollision.h
-				"Runtime/Landscape/Private",
-                "../Shaders/Shared"
+				"Runtime/Renderer/Private",
+				"../Shaders/Shared"
             }
 		);
 
diff --git a/Engine/Source/Runtime/Landscape/Private/LandscapeEditLayers.cpp b/Engine/Source/Runtime/Landscape/Private/LandscapeEditLayers.cpp
--- a/Engine/Source/Runtime/Landscape/Private/LandscapeEditLayers.cpp
+++ b/Engine/Source/Runtime/Landscape/Private/LandscapeEditLayers.cpp
@@ -26,7 +26,7 @@ LandscapeEditLayers.cpp: Landscape editing layers mode
 #include "RenderGraph.h"
 #include "RenderGraphUtils.h"
 #include "PixelShaderUtils.h"
-#include "Runtime/Renderer/Private/SystemTextures.h"
+#include "SystemTextures.h"
 #include "Rendering/Texture2DResource.h"
 #include "MaterialCachedData.h"
 
diff --git a/Engine/Source/Runtime/Landscape/Private/LandscapeRender.cpp b/Engine/Source/Runtime/Landscape/Private/LandscapeRender.cpp
--- a/Engine/Source/Runtime/Landscape/Private/LandscapeRender.cpp
+++ b/Engine/Source/Runtime/Landscape/Private/LandscapeRender.cpp
@@ -35,7 +35,7 @@ LandscapeRender.cpp: New terrain rendering
 #include "DrawDebugHelpers.h"
 #include "PrimitiveSceneInfo.h"
 #include "SceneView.h"
-#include "Runtime/Renderer/Private/SceneCore.h"
+#include "SceneCore.h"
 #include "LandscapeProxy.h"
 #include "HAL/LowLevelMemTracker.h"
 #include "MeshMaterialShader.h"
diff --git a/Engine/Source/Runtime/Launch/Public/RequiredProgramMainCPPInclude.h b/Engine/Source/Runtime/Launch/Public/RequiredProgramMainCPPInclude.h
--- a/Engine/Source/Runtime/Launch/Public/RequiredProgramMainCPPInclude.h
+++ b/Engine/Source/Runtime/Launch/Public/RequiredProgramMainCPPInclude.h
@@ -6,4 +6,4 @@
 #include "LaunchEngineLoop.h"
 
 // this is highly sketchy, but we need some stuff from launchengineloop.cpp
-#include "Runtime/Launch/Private/LaunchEngineLoop.cpp"
+#include "LaunchEngineLoop.cpp"
diff --git a/Engine/Source/Runtime/LevelSequence/Private/LevelSequencePlayer.cpp b/Engine/Source/Runtime/LevelSequence/Private/LevelSequencePlayer.cpp
--- a/Engine/Source/Runtime/LevelSequence/Private/LevelSequencePlayer.cpp
+++ b/Engine/Source/Runtime/LevelSequence/Private/LevelSequencePlayer.cpp
@@ -24,7 +24,7 @@
 #include "LevelSequenceActor.h"
 #include "Modules/ModuleManager.h"
 #include "LevelUtils.h"
-#include "Core/Public/ProfilingDebugging/CsvProfiler.h"
+#include "ProfilingDebugging/CsvProfiler.h"
 #include "LevelSequenceModule.h"
 #include "Generators/MovieSceneEasingCurves.h"
 
diff --git a/Engine/Source/Runtime/MRMesh/MRMesh.build.cs b/Engine/Source/Runtime/MRMesh/MRMesh.build.cs
--- a/Engine/Source/Runtime/MRMesh/MRMesh.build.cs
+++ b/Engine/Source/Runtime/MRMesh/MRMesh.build.cs
@@ -8,9 +8,6 @@ namespace UnrealBuildTool.Rules
 	{
 		public MRMesh(ReadOnlyTargetRules Target) : base(Target)
 		{
-            PrivateIncludePaths.Add("Runtime/MRMesh/Private");
-            PublicIncludePaths.Add("Runtime/MRMesh/Public");
-
 			PrivateDependencyModuleNames.AddRange(
 				new string[]
 				{
@@ -31,7 +28,7 @@ namespace UnrealBuildTool.Rules
 			
 			// Used for including the private Chaos headers
 			string EnginePath = Path.GetFullPath(Target.RelativeEnginePath);
-			PrivateIncludePaths.Add(Path.Combine(EnginePath, "Source/Runtime/Engine/Private/PhysicsEngine"));
+			PrivateIncludePaths.Add(Path.Combine(EnginePath, "Source/Runtime/Engine/Private"));
 
 			PrivateIncludePathModuleNames.Add("DerivedDataCache");
 		}
diff --git a/Engine/Source/Runtime/MRMesh/Private/MRMeshComponent.cpp b/Engine/Source/Runtime/MRMesh/Private/MRMeshComponent.cpp
--- a/Engine/Source/Runtime/MRMesh/Private/MRMeshComponent.cpp
+++ b/Engine/Source/Runtime/MRMesh/Private/MRMeshComponent.cpp
@@ -25,7 +25,7 @@
 #include "UObject/ConstructorHelpers.h"
 #include "Materials/MaterialInstanceDynamic.h"
 
-#include "Experimental/ChaosDerivedData.h"
+#include "PhysicsEngine/Experimental/ChaosDerivedData.h"
 
 // See ISpatialAcceleration.h, the header is not included to avoid having a dependency on Chaos
 #ifndef CHAOS_SERIALIZE_OUT
diff --git a/Engine/Source/Runtime/MediaAssets/MediaAssets.Build.cs b/Engine/Source/Runtime/MediaAssets/MediaAssets.Build.cs
--- a/Engine/Source/Runtime/MediaAssets/MediaAssets.Build.cs
+++ b/Engine/Source/Runtime/MediaAssets/MediaAssets.Build.cs
@@ -23,13 +23,6 @@ namespace UnrealBuildTool.Rules
 					"RHI"
 				});
 
-			PrivateIncludePaths.AddRange(
-				new string[] {
-					"Runtime/MediaAssets/Private",
-					"Runtime/MediaAssets/Private/Assets",
-					"Runtime/MediaAssets/Private/Misc",
-				});
-
 			if (Target.bBuildEditor)
 			{
 				PrivateIncludePathModuleNames.Add("TargetPlatform");
diff --git a/Engine/Source/Runtime/MessagingCommon/MessagingCommon.Build.cs b/Engine/Source/Runtime/MessagingCommon/MessagingCommon.Build.cs
--- a/Engine/Source/Runtime/MessagingCommon/MessagingCommon.Build.cs
+++ b/Engine/Source/Runtime/MessagingCommon/MessagingCommon.Build.cs
@@ -18,11 +18,6 @@ namespace UnrealBuildTool.Rules
 				{
 					"Messaging",
 				});
-
-			PrivateIncludePaths.AddRange(
-				new string[] {
-					"Runtime/MessagingCommon/Private",
-				});
 		}
 	}
 }
diff --git a/Engine/Source/Runtime/MovieSceneTracks/MovieSceneTracks.Build.cs b/Engine/Source/Runtime/MovieSceneTracks/MovieSceneTracks.Build.cs
--- a/Engine/Source/Runtime/MovieSceneTracks/MovieSceneTracks.Build.cs
+++ b/Engine/Source/Runtime/MovieSceneTracks/MovieSceneTracks.Build.cs
@@ -6,13 +6,6 @@ public class MovieSceneTracks : ModuleRules
 {
 	public MovieSceneTracks(ReadOnlyTargetRules Target) : base(Target)
 	{
-		PrivateIncludePaths.AddRange(new string[]
-			{ 
-				"Runtime/MovieSceneTracks/Private",
-				"Runtime/MovieSceneTracks/Private/Sections",
-				"Runtime/MovieSceneTracks/Private/Tracks",
-			});
-
 		PublicDependencyModuleNames.AddRange(new string[]
 			{
 				"Core",
diff --git a/Engine/Source/Runtime/MovieSceneTracks/Private/Evaluation/MovieSceneAudioTemplate.h b/Engine/Source/Runtime/MovieSceneTracks/Private/Evaluation/MovieSceneAudioTemplate.h
--- a/Engine/Source/Runtime/MovieSceneTracks/Private/Evaluation/MovieSceneAudioTemplate.h
+++ b/Engine/Source/Runtime/MovieSceneTracks/Private/Evaluation/MovieSceneAudioTemplate.h
@@ -5,7 +5,7 @@
 #include "CoreMinimal.h"
 #include "UObject/ObjectMacros.h"
 #include "Evaluation/MovieSceneEvalTemplate.h"
-#include "Runtime/Engine/Classes/Components/AudioComponent.h"
+#include "Components/AudioComponent.h"
 #include "Sound/SoundAttenuation.h"
 #include "Channels/MovieSceneFloatChannel.h"
 #include "MovieSceneAudioTemplate.generated.h"
diff --git a/Engine/Source/Runtime/MovieSceneTracks/Private/Evaluation/MovieSceneSkeletalAnimationTemplate.cpp b/Engine/Source/Runtime/MovieSceneTracks/Private/Evaluation/MovieSceneSkeletalAnimationTemplate.cpp
--- a/Engine/Source/Runtime/MovieSceneTracks/Private/Evaluation/MovieSceneSkeletalAnimationTemplate.cpp
+++ b/Engine/Source/Runtime/MovieSceneTracks/Private/Evaluation/MovieSceneSkeletalAnimationTemplate.cpp
@@ -7,7 +7,7 @@
 #include "Animation/AnimSingleNodeInstance.h"
 #include "Components/SkeletalMeshComponent.h"
 #include "Animation/AnimSingleNodeInstance.h"
-#include "Runtime/AnimGraphRuntime/Public/AnimSequencerInstance.h"
+#include "AnimSequencerInstance.h"
 #include "AnimCustomInstanceHelper.h"
 #include "Evaluation/MovieSceneEvaluation.h"
 #include "IMovieScenePlayer.h"
diff --git a/Engine/Source/Runtime/MovieSceneTracks/Private/Sections/MovieSceneSkeletalAnimationSection.cpp b/Engine/Source/Runtime/MovieSceneTracks/Private/Sections/MovieSceneSkeletalAnimationSection.cpp
--- a/Engine/Source/Runtime/MovieSceneTracks/Private/Sections/MovieSceneSkeletalAnimationSection.cpp
+++ b/Engine/Source/Runtime/MovieSceneTracks/Private/Sections/MovieSceneSkeletalAnimationSection.cpp
@@ -3,7 +3,7 @@
 #include "Sections/MovieSceneSkeletalAnimationSection.h"
 #include "Channels/MovieSceneChannelProxy.h"
 #include "Animation/AnimSequence.h"
-#include "Runtime/AnimGraphRuntime/Public/AnimSequencerInstanceProxy.h"
+#include "AnimSequencerInstanceProxy.h"
 #include "Evaluation/MovieSceneSkeletalAnimationTemplate.h"
 #include "Logging/MessageLog.h"
 #include "MovieScene.h"
diff --git a/Engine/Source/Runtime/MovieSceneTracks/Public/Sections/MovieSceneAudioSection.h b/Engine/Source/Runtime/MovieSceneTracks/Public/Sections/MovieSceneAudioSection.h
--- a/Engine/Source/Runtime/MovieSceneTracks/Public/Sections/MovieSceneAudioSection.h
+++ b/Engine/Source/Runtime/MovieSceneTracks/Public/Sections/MovieSceneAudioSection.h
@@ -5,7 +5,7 @@
 #include "CoreMinimal.h"
 #include "UObject/ObjectMacros.h"
 #include "MovieSceneSection.h"
-#include "Runtime/Engine/Classes/Components/AudioComponent.h"
+#include "Components/AudioComponent.h"
 #include "Sound/SoundAttenuation.h"
 #include "Channels/MovieSceneFloatChannel.h"
 #include "Channels/MovieSceneStringChannel.h"
diff --git a/Engine/Source/Runtime/NavigationSystem/Private/NavigationSystem.cpp b/Engine/Source/Runtime/NavigationSystem/Private/NavigationSystem.cpp
--- a/Engine/Source/Runtime/NavigationSystem/Private/NavigationSystem.cpp
+++ b/Engine/Source/Runtime/NavigationSystem/Private/NavigationSystem.cpp
@@ -33,7 +33,7 @@
 #if WITH_EDITOR
 #include "EditorModeManager.h"
 #include "EditorModes.h"
-#include "Editor/LevelEditor/Public/LevelEditor.h"
+#include "LevelEditor.h"
 #include "Misc/MessageDialog.h"
 #endif
 
diff --git a/Engine/Source/Runtime/NavigationSystem/Public/NavigationSystemModule.h b/Engine/Source/Runtime/NavigationSystem/Public/NavigationSystemModule.h
--- a/Engine/Source/Runtime/NavigationSystem/Public/NavigationSystemModule.h
+++ b/Engine/Source/Runtime/NavigationSystem/Public/NavigationSystemModule.h
@@ -6,7 +6,7 @@
 #include "Modules/ModuleManager.h"
 //#include "AI/NavigationSystemBase.h"
 #if WITH_EDITOR
-#include "Developer/AssetTools/Public/AssetTypeCategories.h"
+#include "AssetTypeCategories.h"
 #endif // WITH_EDITOR
 
 /**
diff --git a/Engine/Source/Runtime/Online/HTTP/Private/HttpManager.cpp b/Engine/Source/Runtime/Online/HTTP/Private/HttpManager.cpp
--- a/Engine/Source/Runtime/Online/HTTP/Private/HttpManager.cpp
+++ b/Engine/Source/Runtime/Online/HTTP/Private/HttpManager.cpp
@@ -9,8 +9,8 @@
 #include "Misc/App.h"
 #include "Misc/Guid.h"
 #include "Misc/Fork.h"
-
 #include "HttpThread.h"
+#include "IHttpThreadedRequest.h"
 #include "Misc/ConfigCacheIni.h"
 #include "Misc/CommandLine.h"
 
diff --git a/Engine/Source/Runtime/Online/HTTP/Public/HttpManager.h b/Engine/Source/Runtime/Online/HTTP/Public/HttpManager.h
--- a/Engine/Source/Runtime/Online/HTTP/Public/HttpManager.h
+++ b/Engine/Source/Runtime/Online/HTTP/Public/HttpManager.h
@@ -14,7 +14,6 @@
 #include "HttpPackage.h"
 #include "Interfaces/IHttpRequest.h"
 #include "Misc/EnumRange.h"
-#include "Runtime/Online/HTTP/Private/IHttpThreadedRequest.h"
 #include "Templates/Function.h"
 #include "Templates/SharedPointer.h"
 
diff --git a/Engine/Source/Runtime/Online/Stomp/Private/StompMessage.h b/Engine/Source/Runtime/Online/Stomp/Private/StompMessage.h
--- a/Engine/Source/Runtime/Online/Stomp/Private/StompMessage.h
+++ b/Engine/Source/Runtime/Online/Stomp/Private/StompMessage.h
@@ -3,7 +3,7 @@
 #pragma once
 
 #include "CoreMinimal.h"
-#include "Runtime/Online/Stomp/Private/StompFrame.h"
+#include "StompFrame.h"
 
 #if WITH_STOMP
 
diff --git a/Engine/Source/Runtime/Online/Voice/Private/Windows/VoiceModuleWindows.cpp b/Engine/Source/Runtime/Online/Voice/Private/Windows/VoiceModuleWindows.cpp
--- a/Engine/Source/Runtime/Online/Voice/Private/Windows/VoiceModuleWindows.cpp
+++ b/Engine/Source/Runtime/Online/Voice/Private/Windows/VoiceModuleWindows.cpp
@@ -6,7 +6,7 @@
 #include "VoiceCaptureWindows.h"
 #include "VoiceCodecOpus.h"
 #include "Voice.h"
-#include "Runtime/HeadMountedDisplay/Public/IHeadMountedDisplayModule.h"
+#include "IHeadMountedDisplayModule.h"
 #include "HAL/IConsoleManager.h"
 #if PLATFORM_SUPPORTS_VOICE_CAPTURE
 
@@ -378,4 +378,4 @@ IVoiceCapture* CreateVoiceCaptureObject(const FString& DeviceName, int32 SampleR
 	return VoiceCaptureDev ? VoiceCaptureDev->CreateVoiceCaptureObject(DeviceName, SampleRate, NumChannels) : nullptr;
 }
 
-#endif // PLATFORM_SUPPORTS_VOICE_CAPTURE
\ No newline at end of file
+#endif // PLATFORM_SUPPORTS_VOICE_CAPTURE
diff --git a/Engine/Source/Runtime/Online/Voice/Voice.Build.cs b/Engine/Source/Runtime/Online/Voice/Voice.Build.cs
--- a/Engine/Source/Runtime/Online/Voice/Voice.Build.cs
+++ b/Engine/Source/Runtime/Online/Voice/Voice.Build.cs
@@ -20,6 +20,12 @@ public class Voice : ModuleRules
 			PublicDefinitions.Add("VOICE_MODULE_WITH_CAPTURE=1");
 		}
 
+		PublicIncludePathModuleNames.AddRange(
+			new string[] {
+					"HeadMountedDisplay"
+			}
+			);
+
 		if (Target.Platform == UnrealTargetPlatform.Android)
 		{
 			PublicIncludePathModuleNames.AddRange(
@@ -35,12 +41,6 @@ public class Voice : ModuleRules
 			}
 			);
 
-		PrivateIncludePaths.AddRange(
-			new string[] {
-				"Runtime/Online/Voice/Private",
-			}
-			);
-
 		PrivateDependencyModuleNames.AddRange(
 			new string[] { 
 				"Core",
diff --git a/Engine/Source/Runtime/Overlay/Overlay.Build.cs b/Engine/Source/Runtime/Overlay/Overlay.Build.cs
--- a/Engine/Source/Runtime/Overlay/Overlay.Build.cs
+++ b/Engine/Source/Runtime/Overlay/Overlay.Build.cs
@@ -28,14 +28,6 @@ namespace UnrealBuildTool.Rules
                     "Runtime/Overlay/Public",
                 }
             );
-
-            PrivateIncludePaths.AddRange(
-				new string[] {
-					"Runtime/Overlay/Private",
-					"Runtime/Overlay/Private/Assets",
-					"Runtime/Overlay/Private/Factories",
-				}
-			);
         }
 	}
 }
\ No newline at end of file
diff --git a/Engine/Source/Runtime/PhysicsCore/Private/ChaosScene.cpp b/Engine/Source/Runtime/PhysicsCore/Private/ChaosScene.cpp
--- a/Engine/Source/Runtime/PhysicsCore/Private/ChaosScene.cpp
+++ b/Engine/Source/Runtime/PhysicsCore/Private/ChaosScene.cpp
@@ -27,9 +27,9 @@
 #include "PBDRigidActiveParticlesBuffer.h"
 #include "Chaos/GeometryParticlesfwd.h"
 #include "Chaos/Box.h"
-#include "Chaos/Public/EventsData.h"
-#include "Chaos/Public/EventManager.h"
-#include "Chaos/Public/RewindData.h"
+#include "EventsData.h"
+#include "EventManager.h"
+#include "RewindData.h"
 #include "PhysicsSettingsCore.h"
 #include "Chaos/PhysicsSolverBaseImpl.h"
 
diff --git a/Engine/Source/Runtime/PhysicsCore/Private/PhysicsInitialization.cpp b/Engine/Source/Runtime/PhysicsCore/Private/PhysicsInitialization.cpp
--- a/Engine/Source/Runtime/PhysicsCore/Private/PhysicsInitialization.cpp
+++ b/Engine/Source/Runtime/PhysicsCore/Private/PhysicsInitialization.cpp
@@ -8,7 +8,7 @@
 
 #include "Modules/ModuleManager.h"
 #include "Misc/Paths.h"
-#include "Core/Public/HAL/IConsoleManager.h"
+#include "HAL/IConsoleManager.h"
 
 #ifndef APEX_STATICALLY_LINKED
 #define APEX_STATICALLY_LINKED	0
diff --git a/Engine/Source/Runtime/PhysicsCore/Private/SQAccelerator.cpp b/Engine/Source/Runtime/PhysicsCore/Private/SQAccelerator.cpp
--- a/Engine/Source/Runtime/PhysicsCore/Private/SQAccelerator.cpp
+++ b/Engine/Source/Runtime/PhysicsCore/Private/SQAccelerator.cpp
@@ -2,7 +2,7 @@
 
 #include "SQAccelerator.h"
 #include "CollisionQueryFilterCallbackCore.h"
-#include "PhysicsCore/Public/PhysicsInterfaceUtilsCore.h"
+#include "PhysicsInterfaceUtilsCore.h"
 
 #include "SceneQueryChaosImp.h"
 
diff --git a/Engine/Source/Runtime/PhysicsCore/Public/SQVerifier.h b/Engine/Source/Runtime/PhysicsCore/Public/SQVerifier.h
--- a/Engine/Source/Runtime/PhysicsCore/Public/SQVerifier.h
+++ b/Engine/Source/Runtime/PhysicsCore/Public/SQVerifier.h
@@ -5,7 +5,7 @@
 #include "ProfilingDebugging/ScopedTimers.h"
 
 #include "Chaos/PBDRigidsEvolutionGBF.h"
-#include "PhysicsCore/Public/PhysicsInterfaceUtilsCore.h"
+#include "PhysicsInterfaceUtilsCore.h"
 #include "PhysTestSerializer.h"
 
 #ifndef SQ_REPLAY_TEST
diff --git a/Engine/Source/Runtime/Renderer/Private/ClusteredDeferredShadingPass.cpp b/Engine/Source/Runtime/Renderer/Private/ClusteredDeferredShadingPass.cpp
--- a/Engine/Source/Runtime/Renderer/Private/ClusteredDeferredShadingPass.cpp
+++ b/Engine/Source/Runtime/Renderer/Private/ClusteredDeferredShadingPass.cpp
@@ -24,8 +24,8 @@
 #include "ShaderPrintParameters.h"
 #include "ShaderPrint.h"
 #include "VirtualShadowMaps/VirtualShadowMapArray.h"
-#include "SceneFilterRendering.h"
-#include "PostProcessing.h"
+#include "PostProcess/SceneFilterRendering.h"
+#include "PostProcess/PostProcessing.h"
 #include "Strata/Strata.h"
 
 // This is used to switch on and off the clustered deferred shading implementation, that uses the light grid to perform shading.
diff --git a/Engine/Source/Runtime/Renderer/Private/CompositionLighting/CompositionLighting.h b/Engine/Source/Runtime/Renderer/Private/CompositionLighting/CompositionLighting.h
--- a/Engine/Source/Runtime/Renderer/Private/CompositionLighting/CompositionLighting.h
+++ b/Engine/Source/Runtime/Renderer/Private/CompositionLighting/CompositionLighting.h
@@ -8,8 +8,8 @@
 
 #include "CoreMinimal.h"
 #include "SceneRendering.h"
-#include "PostProcessDeferredDecals.h"
-#include "PostProcessAmbientOcclusion.h"
+#include "CompositionLighting/PostProcessDeferredDecals.h"
+#include "CompositionLighting/PostProcessAmbientOcclusion.h"
 
 /**
  * The center for all screen space processing activities (e.g. G-buffer manipulation, lighting).
@@ -56,4 +56,4 @@ private:
 	bool bInitialized = false;
 };
 
-extern bool ShouldRenderScreenSpaceAmbientOcclusion(const FViewInfo& View);
\ No newline at end of file
+extern bool ShouldRenderScreenSpaceAmbientOcclusion(const FViewInfo& View);
diff --git a/Engine/Source/Runtime/Renderer/Private/DebugViewModeRendering.cpp b/Engine/Source/Runtime/Renderer/Private/DebugViewModeRendering.cpp
--- a/Engine/Source/Runtime/Renderer/Private/DebugViewModeRendering.cpp
+++ b/Engine/Source/Runtime/Renderer/Private/DebugViewModeRendering.cpp
@@ -8,7 +8,7 @@ DebugViewModeRendering.cpp: Contains definitions for rendering debug viewmodes.
 #include "Materials/Material.h"
 #include "PrimitiveSceneInfo.h"
 #include "ScenePrivate.h"
-#include "PostProcessing.h"
+#include "PostProcess/PostProcessing.h"
 #include "PostProcess/PostProcessVisualizeComplexity.h"
 #include "PostProcess/PostProcessStreamingAccuracyLegend.h"
 #include "PostProcess/PostProcessSelectionOutline.h"
@@ -569,4 +569,4 @@ void RenderDebugViewMode(
 	const FRenderTargetBindingSlots& RenderTargets)
 {}
 
-#endif // WITH_DEBUG_VIEW_MODES
\ No newline at end of file
+#endif // WITH_DEBUG_VIEW_MODES
diff --git a/Engine/Source/Runtime/Renderer/Private/GPUScene.cpp b/Engine/Source/Runtime/Renderer/Private/GPUScene.cpp
--- a/Engine/Source/Runtime/Renderer/Private/GPUScene.cpp
+++ b/Engine/Source/Runtime/Renderer/Private/GPUScene.cpp
@@ -11,7 +11,7 @@
 #include "ScenePrivate.h"
 #include "UnifiedBuffer.h"
 #include "SpriteIndexBuffer.h"
-#include "SceneFilterRendering.h"
+#include "PostProcess/SceneFilterRendering.h"
 #include "ClearQuad.h"
 #include "RendererModule.h"
 #include "Rendering/NaniteResources.h"
diff --git a/Engine/Source/Runtime/Renderer/Private/HairStrands/HairStrandsClusters.h b/Engine/Source/Runtime/Renderer/Private/HairStrands/HairStrandsClusters.h
--- a/Engine/Source/Runtime/Renderer/Private/HairStrands/HairStrandsClusters.h
+++ b/Engine/Source/Runtime/Renderer/Private/HairStrands/HairStrandsClusters.h
@@ -9,7 +9,7 @@
 #include "CoreMinimal.h"
 #include "RendererInterface.h"
 #include "HairStrandsInterface.h"
-#include "Renderer/Private/SceneRendering.h"
+#include "SceneRendering.h"
 
 struct FHairCullingParams
 {
diff --git a/Engine/Source/Runtime/Renderer/Private/HairStrands/HairStrandsComposition.cpp b/Engine/Source/Runtime/Renderer/Private/HairStrands/HairStrandsComposition.cpp
--- a/Engine/Source/Runtime/Renderer/Private/HairStrands/HairStrandsComposition.cpp
+++ b/Engine/Source/Runtime/Renderer/Private/HairStrands/HairStrandsComposition.cpp
@@ -7,7 +7,7 @@
 #include "ShaderParameterStruct.h"
 #include "ShaderParameterMacros.h"
 #include "RenderGraphUtils.h"
-#include "PostProcessing.h"
+#include "PostProcess/PostProcessing.h"
 #include "HairStrandsRendering.h"
 #include "HairStrandsTile.h"
 #include "FogRendering.h"
@@ -645,4 +645,4 @@ void RenderHairComposition(
 			SceneDepthTexture,
 			SceneVelocityTexture);
 	}
-}
\ No newline at end of file
+}
diff --git a/Engine/Source/Runtime/Renderer/Private/HairStrands/HairStrandsDebug.cpp b/Engine/Source/Runtime/Renderer/Private/HairStrands/HairStrandsDebug.cpp
--- a/Engine/Source/Runtime/Renderer/Private/HairStrands/HairStrandsDebug.cpp
+++ b/Engine/Source/Runtime/Renderer/Private/HairStrands/HairStrandsDebug.cpp
@@ -18,7 +18,7 @@
 #include "ShaderParameters.h"
 #include "ShaderParameterStruct.h"
 #include "RenderGraphUtils.h"
-#include "PostProcessing.h"
+#include "PostProcess/PostProcessing.h"
 #include "SceneTextureParameters.h"
 #include "DynamicPrimitiveDrawing.h"
 #include "RenderTargetTemp.h"
@@ -1552,4 +1552,4 @@ void RenderHairStrandsDebugInfo(
 	{
 		InternalRenderHairStrandsDebugInfo(GraphBuilder, Scene, View, HairClusterData, SceneColorTexture, SceneDepthTexture);
 	}
-}
\ No newline at end of file
+}
diff --git a/Engine/Source/Runtime/Renderer/Private/HairStrands/HairStrandsDebug.h b/Engine/Source/Runtime/Renderer/Private/HairStrands/HairStrandsDebug.h
--- a/Engine/Source/Runtime/Renderer/Private/HairStrands/HairStrandsDebug.h
+++ b/Engine/Source/Runtime/Renderer/Private/HairStrands/HairStrandsDebug.h
@@ -7,7 +7,7 @@
 #pragma once
 
 #include "RenderGraph.h"
-#include "Renderer/Private/SceneRendering.h"
+#include "SceneRendering.h"
 
 void RenderHairStrandsDebugInfo(
 	FRDGBuilder& GraphBuilder,
@@ -15,4 +15,4 @@ void RenderHairStrandsDebugInfo(
 	TArrayView<FViewInfo> Views,
 	const struct FHairStrandClusterData& HairClusterData,
 	FRDGTextureRef SceneColorTexture, 
-	FRDGTextureRef SceneDepthTexture);
\ No newline at end of file
+	FRDGTextureRef SceneDepthTexture);
diff --git a/Engine/Source/Runtime/Renderer/Private/HairStrands/HairStrandsDeepShadow.h b/Engine/Source/Runtime/Renderer/Private/HairStrands/HairStrandsDeepShadow.h
--- a/Engine/Source/Runtime/Renderer/Private/HairStrands/HairStrandsDeepShadow.h
+++ b/Engine/Source/Runtime/Renderer/Private/HairStrands/HairStrandsDeepShadow.h
@@ -11,10 +11,10 @@
 #include "HairStrandsInterface.h"
 #include "RenderGraphResources.h"
 #include "SceneTypes.h"
-#include "Renderer/Private/SceneRendering.h"
+#include "SceneRendering.h"
 
 void RenderHairStrandsDeepShadows(
 	FRDGBuilder& GraphBuilder,
 	const class FScene* Scene,
 	FViewInfo& View,
-	FInstanceCullingManager& InstanceCullingManager);
\ No newline at end of file
+	FInstanceCullingManager& InstanceCullingManager);
diff --git a/Engine/Source/Runtime/Renderer/Private/HairStrands/HairStrandsTransmittance.cpp b/Engine/Source/Runtime/Renderer/Private/HairStrands/HairStrandsTransmittance.cpp
--- a/Engine/Source/Runtime/Renderer/Private/HairStrands/HairStrandsTransmittance.cpp
+++ b/Engine/Source/Runtime/Renderer/Private/HairStrands/HairStrandsTransmittance.cpp
@@ -13,7 +13,7 @@
 #include "ShaderParameterStruct.h"
 #include "SceneTextureParameters.h"
 #include "RenderGraphUtils.h"
-#include "PostProcessing.h"
+#include "PostProcess/PostProcessing.h"
 #include "ShaderPrintParameters.h"
 #include "LightSceneInfo.h"
 #include "ShaderPrint.h"
diff --git a/Engine/Source/Runtime/Renderer/Private/HairStrands/HairStrandsVisibility.cpp b/Engine/Source/Runtime/Renderer/Private/HairStrands/HairStrandsVisibility.cpp
--- a/Engine/Source/Runtime/Renderer/Private/HairStrands/HairStrandsVisibility.cpp
+++ b/Engine/Source/Runtime/Renderer/Private/HairStrands/HairStrandsVisibility.cpp
@@ -12,7 +12,7 @@
 #include "ShaderParameterStruct.h"
 #include "SceneTextureParameters.h"
 #include "RenderGraphUtils.h"
-#include "PostProcessing.h"
+#include "PostProcess/PostProcessing.h"
 #include "MeshPassProcessor.inl"
 #include "ScenePrivate.h"
 #include "SceneTextureReductions.h"
diff --git a/Engine/Source/Runtime/Renderer/Private/HairStrands/HairStrandsVoxelization.h b/Engine/Source/Runtime/Renderer/Private/HairStrands/HairStrandsVoxelization.h
--- a/Engine/Source/Runtime/Renderer/Private/HairStrands/HairStrandsVoxelization.h
+++ b/Engine/Source/Runtime/Renderer/Private/HairStrands/HairStrandsVoxelization.h
@@ -9,7 +9,7 @@
 #include "CoreMinimal.h"
 #include "RendererInterface.h"
 #include "RenderGraphResources.h"
-#include "Renderer/Private/SceneRendering.h"
+#include "SceneRendering.h"
 
 
 /// Global enable/disable for hair voxelization
@@ -20,4 +20,4 @@ void VoxelizeHairStrands(
 	FRDGBuilder& GraphBuilder,
 	const class FScene* Scene,
 	FViewInfo& View,
-	FInstanceCullingManager& InstanceCullingManager);
\ No newline at end of file
+	FInstanceCullingManager& InstanceCullingManager);
diff --git a/Engine/Source/Runtime/Renderer/Private/IndirectLightRendering.cpp b/Engine/Source/Runtime/Renderer/Private/IndirectLightRendering.cpp
--- a/Engine/Source/Runtime/Renderer/Private/IndirectLightRendering.cpp
+++ b/Engine/Source/Runtime/Renderer/Private/IndirectLightRendering.cpp
@@ -11,7 +11,7 @@
 #include "PostProcess/PostProcessSubsurface.h"
 #include "PostProcess/TemporalAA.h"
 #include "CompositionLighting/PostProcessAmbientOcclusion.h"
-#include "PostProcessing.h" // for FPostProcessVS
+#include "PostProcess/PostProcessing.h" // for FPostProcessVS
 #include "RendererModule.h" 
 #include "RayTracing/RaytracingOptions.h"
 #include "RayTracing/RayTracingReflections.h"
diff --git a/Engine/Source/Runtime/Renderer/Private/MobileReflectionEnvironmentCapture.cpp b/Engine/Source/Runtime/Renderer/Private/MobileReflectionEnvironmentCapture.cpp
--- a/Engine/Source/Runtime/Renderer/Private/MobileReflectionEnvironmentCapture.cpp
+++ b/Engine/Source/Runtime/Renderer/Private/MobileReflectionEnvironmentCapture.cpp
@@ -7,11 +7,11 @@
 #include "ReflectionEnvironmentCapture.h"
 #include "ShaderParameterUtils.h"
 #include "RHIStaticStates.h"
-#include "SceneRenderTargets.h"
+#include "PostProcess/SceneRenderTargets.h"
 #include "SceneUtils.h"
 #include "ScreenRendering.h"
 #include "PipelineStateCache.h"
-#include "SceneFilterRendering.h"
+#include "PostProcess/SceneFilterRendering.h"
 #include "OneColorShader.h"
 #include "PixelShaderUtils.h"
 
@@ -98,4 +98,4 @@ namespace MobileReflectionEnvironmentCapture
 
 		return FilterCubeMap(GraphBuilder, ShaderMap, CubemapTexture);
 	}
-}
\ No newline at end of file
+}
diff --git a/Engine/Source/Runtime/Renderer/Private/PathTracing.cpp b/Engine/Source/Runtime/Renderer/Private/PathTracing.cpp
--- a/Engine/Source/Runtime/Renderer/Private/PathTracing.cpp
+++ b/Engine/Source/Runtime/Renderer/Private/PathTracing.cpp
@@ -22,7 +22,7 @@ TAutoConsoleVariable<int32> CVarPathTracing(
 #include "PathTracingDefinitions.h"
 #include "RayTracing/RayTracingMaterialHitShaders.h"
 #include "RayTracing/RayTracingDecals.h"
-#include "RenderCore/Public/GenerateMips.h"
+#include "GenerateMips.h"
 #include "HairStrands/HairStrandsData.h"
 #include "Modules/ModuleManager.h"
 #include <limits>
diff --git a/Engine/Source/Runtime/Renderer/Private/PostProcess/DiaphragmDOF.cpp b/Engine/Source/Runtime/Renderer/Private/PostProcess/DiaphragmDOF.cpp
--- a/Engine/Source/Runtime/Renderer/Private/PostProcess/DiaphragmDOF.cpp
+++ b/Engine/Source/Runtime/Renderer/Private/PostProcess/DiaphragmDOF.cpp
@@ -17,7 +17,7 @@
 #include "ScenePrivate.h"
 #include "ClearQuad.h"
 #include "SpriteIndexBuffer.h"
-#include "TemporalAA.h"
+#include "PostProcess/TemporalAA.h"
 #include "SceneTextureParameters.h"
 #include "TranslucentRendering.h"
 
diff --git a/Engine/Source/Runtime/Renderer/Private/PostProcess/PostProcessAmbientOcclusionMobile.cpp b/Engine/Source/Runtime/Renderer/Private/PostProcess/PostProcessAmbientOcclusionMobile.cpp
--- a/Engine/Source/Runtime/Renderer/Private/PostProcess/PostProcessAmbientOcclusionMobile.cpp
+++ b/Engine/Source/Runtime/Renderer/Private/PostProcess/PostProcessAmbientOcclusionMobile.cpp
@@ -9,7 +9,7 @@
 #include "ShaderParameterStruct.h"
 #include "SceneRendering.h"
 #include "RenderTargetPool.h"
-#include "SceneRenderTargets.h"
+#include "PostProcess/SceneRenderTargets.h"
 #include "SystemTextures.h"
 #include "ScreenPass.h"
 #include "ScenePrivate.h"
diff --git a/Engine/Source/Runtime/Renderer/Private/PostProcess/PostProcessBokehDOF.cpp b/Engine/Source/Runtime/Renderer/Private/PostProcess/PostProcessBokehDOF.cpp
--- a/Engine/Source/Runtime/Renderer/Private/PostProcess/PostProcessBokehDOF.cpp
+++ b/Engine/Source/Runtime/Renderer/Private/PostProcess/PostProcessBokehDOF.cpp
@@ -4,7 +4,7 @@
 #include "PostProcess/PostProcessDOF.h"
 #include "CanvasTypes.h"
 #include "RenderTargetTemp.h"
-#include "DiaphragmDOF.h"
+#include "PostProcess/DiaphragmDOF.h"
 
 class FVisualizeDOFPS : public FGlobalShader
 {
@@ -114,4 +114,4 @@ FScreenPassTexture AddVisualizeDOFPass(FRDGBuilder& GraphBuilder, const FViewInf
 	});
 
 	return MoveTemp(Output);
-}
\ No newline at end of file
+}
diff --git a/Engine/Source/Runtime/Renderer/Private/PostProcess/PostProcessLensFlares.cpp b/Engine/Source/Runtime/Renderer/Private/PostProcess/PostProcessLensFlares.cpp
--- a/Engine/Source/Runtime/Renderer/Private/PostProcess/PostProcessLensFlares.cpp
+++ b/Engine/Source/Runtime/Renderer/Private/PostProcess/PostProcessLensFlares.cpp
@@ -1,7 +1,7 @@
 // Copyright Epic Games, Inc. All Rights Reserved.
 
-#include "PostProcessLensFlares.h"
-#include "PostProcessDownsample.h"
+#include "PostProcess/PostProcessLensFlares.h"
+#include "PostProcess/PostProcessDownsample.h"
 #include "PixelShaderUtils.h"
 
 DECLARE_GPU_STAT(LensFlare);
@@ -410,4 +410,4 @@ FScreenPassTexture AddLensFlaresPass(
 	}
 
 	return AddLensFlaresPass(GraphBuilder, View, LensFlareInputs);
-}
\ No newline at end of file
+}
diff --git a/Engine/Source/Runtime/Renderer/Private/PostProcess/PostProcessMaterial.cpp b/Engine/Source/Runtime/Renderer/Private/PostProcess/PostProcessMaterial.cpp
--- a/Engine/Source/Runtime/Renderer/Private/PostProcess/PostProcessMaterial.cpp
+++ b/Engine/Source/Runtime/Renderer/Private/PostProcess/PostProcessMaterial.cpp
@@ -15,9 +15,9 @@
 #include "SceneRendering.h"
 #include "ClearQuad.h"
 #include "Materials/MaterialExpressionSceneTexture.h"
-#include "RHI/Public/PipelineStateCache.h"
-#include "PostProcessing.h"
-#include "PostProcessMobile.h"
+#include "PipelineStateCache.h"
+#include "PostProcess/PostProcessing.h"
+#include "PostProcess/PostProcessMobile.h"
 #include "BufferVisualizationData.h"
 #include "SceneTextureParameters.h"
 #include "SystemTextures.h"
diff --git a/Engine/Source/Runtime/Renderer/Private/PostProcess/PostProcessSelectionOutline.cpp b/Engine/Source/Runtime/Renderer/Private/PostProcess/PostProcessSelectionOutline.cpp
--- a/Engine/Source/Runtime/Renderer/Private/PostProcess/PostProcessSelectionOutline.cpp
+++ b/Engine/Source/Runtime/Renderer/Private/PostProcess/PostProcessSelectionOutline.cpp
@@ -10,7 +10,7 @@
 #include "RenderTargetTemp.h"
 #include "ClearQuad.h"
 #include "ScenePrivate.h"
-#include "SceneRenderTargets.h"
+#include "PostProcess/SceneRenderTargets.h"
 
 namespace
 {
@@ -206,4 +206,4 @@ FScreenPassTexture AddSelectionOutlinePass(FRDGBuilder& GraphBuilder, const FVie
 	return MoveTemp(Output);
 }
 
-#endif
\ No newline at end of file
+#endif
diff --git a/Engine/Source/Runtime/Renderer/Private/PostProcess/PostProcessSubsurface.cpp b/Engine/Source/Runtime/Renderer/Private/PostProcess/PostProcessSubsurface.cpp
--- a/Engine/Source/Runtime/Renderer/Private/PostProcess/PostProcessSubsurface.cpp
+++ b/Engine/Source/Runtime/Renderer/Private/PostProcess/PostProcessSubsurface.cpp
@@ -19,7 +19,7 @@
 #include "GenerateMips.h"
 #include "ClearQuad.h"
 #include "Strata/Strata.h"
-#include "TemporalAA.h"
+#include "PostProcess/TemporalAA.h"
 #include "SubsurfaceTiles.h"
 
 namespace
@@ -1545,4 +1545,4 @@ void AddSubsurfacePass(
 
 	AddSubsurfacePass(GraphBuilder, Views, ViewMask, SceneTextures, SceneTextures.Color.Target);
 	SceneTextures.Color.Resolve = SceneTextures.Color.Target;
-}
\ No newline at end of file
+}
diff --git a/Engine/Source/Runtime/Renderer/Private/PostProcess/PostProcessUpscale.cpp b/Engine/Source/Runtime/Renderer/Private/PostProcess/PostProcessUpscale.cpp
--- a/Engine/Source/Runtime/Renderer/Private/PostProcess/PostProcessUpscale.cpp
+++ b/Engine/Source/Runtime/Renderer/Private/PostProcess/PostProcessUpscale.cpp
@@ -1,7 +1,7 @@
 // Copyright Epic Games, Inc. All Rights Reserved.
 
 #include "PostProcess/PostProcessUpscale.h"
-#include "SceneFilterRendering.h"
+#include "PostProcess/SceneFilterRendering.h"
 
 namespace
 {
@@ -271,4 +271,4 @@ FScreenPassTexture ISpatialUpscaler::AddDefaultUpscalePass(
 	});
 
 	return MoveTemp(Output);
-}
\ No newline at end of file
+}
diff --git a/Engine/Source/Runtime/Renderer/Private/PostProcess/PostProcessVisualizeBuffer.cpp b/Engine/Source/Runtime/Renderer/Private/PostProcess/PostProcessVisualizeBuffer.cpp
--- a/Engine/Source/Runtime/Renderer/Private/PostProcess/PostProcessVisualizeBuffer.cpp
+++ b/Engine/Source/Runtime/Renderer/Private/PostProcess/PostProcessVisualizeBuffer.cpp
@@ -2,8 +2,8 @@
 
 #include "PostProcess/PostProcessVisualizeBuffer.h"
 #include "HighResScreenshot.h"
-#include "PostProcessMaterial.h"
-#include "PostProcessDownsample.h"
+#include "PostProcess/PostProcessMaterial.h"
+#include "PostProcess/PostProcessDownsample.h"
 #include "ImagePixelData.h"
 #include "ImageWriteStream.h"
 #include "ImageWriteTask.h"
@@ -448,4 +448,4 @@ FScreenPassTexture AddVisualizeGBufferOverviewPass(
 
 		return Inputs.SceneColor;
 	}
-}
\ No newline at end of file
+}
diff --git a/Engine/Source/Runtime/Renderer/Private/PostProcess/SceneFilterRendering.h b/Engine/Source/Runtime/Renderer/Private/PostProcess/SceneFilterRendering.h
--- a/Engine/Source/Runtime/Renderer/Private/PostProcess/SceneFilterRendering.h
+++ b/Engine/Source/Runtime/Renderer/Private/PostProcess/SceneFilterRendering.h
@@ -11,7 +11,7 @@
 #include "RendererInterface.h"
 #include "ShaderParameters.h"
 #include "Shader.h"
-#include "SceneRenderTargets.h"
+#include "PostProcess/SceneRenderTargets.h"
 #include "ShaderParameterUtils.h"
 
 /** Uniform buffer for computing the vertex positional and UV adjustments in the vertex shader. */
diff --git a/Engine/Source/Runtime/Renderer/Private/PostProcess/TemporalAA.h b/Engine/Source/Runtime/Renderer/Private/PostProcess/TemporalAA.h
--- a/Engine/Source/Runtime/Renderer/Private/PostProcess/TemporalAA.h
+++ b/Engine/Source/Runtime/Renderer/Private/PostProcess/TemporalAA.h
@@ -3,7 +3,7 @@
 #pragma once
 
 #include "ScreenPass.h"
-#include "PostProcessMotionBlur.h"
+#include "PostProcess/PostProcessMotionBlur.h"
 
 struct FTemporalAAHistory;
 struct FTranslucencyPassResources;
diff --git a/Engine/Source/Runtime/Renderer/Private/RayTracing/RayTracingAmbientOcclusion.cpp b/Engine/Source/Runtime/Renderer/Private/RayTracing/RayTracingAmbientOcclusion.cpp
--- a/Engine/Source/Runtime/Renderer/Private/RayTracing/RayTracingAmbientOcclusion.cpp
+++ b/Engine/Source/Runtime/Renderer/Private/RayTracing/RayTracingAmbientOcclusion.cpp
@@ -6,12 +6,12 @@
 
 #include "ClearQuad.h"
 #include "SceneRendering.h"
-#include "SceneRenderTargets.h"
+#include "PostProcess/SceneRenderTargets.h"
 #include "SceneUtils.h"
 #include "RenderTargetPool.h"
 #include "RHIResources.h"
 #include "UniformBuffer.h"
-#include "RHI/Public/PipelineStateCache.h"
+#include "PipelineStateCache.h"
 #include "RayTracing/RaytracingOptions.h"
 #include "RayTracingMaterialHitShaders.h"
 #include "RayTracingDefinitions.h"
diff --git a/Engine/Source/Runtime/Renderer/Private/RayTracing/RayTracingBarycentrics.cpp b/Engine/Source/Runtime/Renderer/Private/RayTracing/RayTracingBarycentrics.cpp
--- a/Engine/Source/Runtime/Renderer/Private/RayTracing/RayTracingBarycentrics.cpp
+++ b/Engine/Source/Runtime/Renderer/Private/RayTracing/RayTracingBarycentrics.cpp
@@ -6,9 +6,9 @@
 
 #include "DeferredShadingRenderer.h"
 #include "GlobalShader.h"
-#include "SceneRenderTargets.h"
+#include "PostProcess/SceneRenderTargets.h"
 #include "RenderGraphBuilder.h"
-#include "RHI/Public/PipelineStateCache.h"
+#include "PipelineStateCache.h"
 #include "RayTracing/RaytracingOptions.h"
 
 #include "Rendering/NaniteStreamingManager.h"
@@ -171,4 +171,4 @@ void FDeferredShadingSceneRenderer::RenderRayTracingBarycentrics(FRDGBuilder& Gr
 		RenderRayTracingBarycentricsRGS(GraphBuilder, View, SceneColor);
 	}
 }
-#endif
\ No newline at end of file
+#endif
diff --git a/Engine/Source/Runtime/Renderer/Private/RayTracing/RayTracingDecals.cpp b/Engine/Source/Runtime/Renderer/Private/RayTracing/RayTracingDecals.cpp
--- a/Engine/Source/Runtime/Renderer/Private/RayTracing/RayTracingDecals.cpp
+++ b/Engine/Source/Runtime/Renderer/Private/RayTracing/RayTracingDecals.cpp
@@ -2,7 +2,7 @@
 
 #include "RayTracingDecals.h"
 
-#include "RHI/Public/RHIDefinitions.h"
+#include "RHIDefinitions.h"
 
 #if RHI_RAYTRACING
 
diff --git a/Engine/Source/Runtime/Renderer/Private/RayTracing/RayTracingLighting.cpp b/Engine/Source/Runtime/Renderer/Private/RayTracing/RayTracingLighting.cpp
--- a/Engine/Source/Runtime/Renderer/Private/RayTracing/RayTracingLighting.cpp
+++ b/Engine/Source/Runtime/Renderer/Private/RayTracing/RayTracingLighting.cpp
@@ -1,7 +1,7 @@
 // Copyright Epic Games, Inc. All Rights Reserved.
 
 #include "RayTracingLighting.h"
-#include "RHI/Public/RHIDefinitions.h"
+#include "RHIDefinitions.h"
 #include "RendererPrivate.h"
 
 #if RHI_RAYTRACING
diff --git a/Engine/Source/Runtime/Renderer/Private/SceneRendering.cpp b/Engine/Source/Runtime/Renderer/Private/SceneRendering.cpp
--- a/Engine/Source/Runtime/Renderer/Private/SceneRendering.cpp
+++ b/Engine/Source/Runtime/Renderer/Private/SceneRendering.cpp
@@ -56,7 +56,7 @@
 #include "IXRTrackingSystem.h"
 #include "IXRCamera.h"
 #include "IHeadMountedDisplay.h"
-#include "DiaphragmDOF.h" 
+#include "PostProcess/DiaphragmDOF.h" 
 #include "SingleLayerWaterRendering.h"
 #include "HairStrands/HairStrandsVisibility.h"
 #include "SystemTextures.h"
diff --git a/Engine/Source/Runtime/Renderer/Private/SceneTextures.cpp b/Engine/Source/Runtime/Renderer/Private/SceneTextures.cpp
--- a/Engine/Source/Runtime/Renderer/Private/SceneTextures.cpp
+++ b/Engine/Source/Runtime/Renderer/Private/SceneTextures.cpp
@@ -13,7 +13,7 @@
 #include "StereoRendering.h"
 #include "StereoRenderTargetManager.h"
 #include "CompositionLighting/PostProcessAmbientOcclusion.h"
-#include "PostProcessCompositeEditorPrimitives.h"
+#include "PostProcess/PostProcessCompositeEditorPrimitives.h"
 #include "ShaderCompiler.h"
 #include "SystemTextures.h"
 #include "PostProcess/PostProcessAmbientOcclusionMobile.h"
diff --git a/Engine/Source/Runtime/Renderer/Private/SceneVisibility.cpp b/Engine/Source/Runtime/Renderer/Private/SceneVisibility.cpp
--- a/Engine/Source/Runtime/Renderer/Private/SceneVisibility.cpp
+++ b/Engine/Source/Runtime/Renderer/Private/SceneVisibility.cpp
@@ -43,7 +43,7 @@
 #include "ProfilingDebugging/DiagnosticTable.h"
 #include "Algo/Unique.h"
 #include "InstanceCulling/InstanceCullingManager.h"
-#include "TemporalAA.h"
+#include "PostProcess/TemporalAA.h"
 #include "RayTracing/RayTracingInstanceCulling.h"
 #include "RendererModule.h"
 #include "SceneViewExtension.h"
diff --git a/Engine/Source/Runtime/Renderer/Private/VT/VirtualTextureSpace.cpp b/Engine/Source/Runtime/Renderer/Private/VT/VirtualTextureSpace.cpp
--- a/Engine/Source/Runtime/Renderer/Private/VT/VirtualTextureSpace.cpp
+++ b/Engine/Source/Runtime/Renderer/Private/VT/VirtualTextureSpace.cpp
@@ -4,7 +4,7 @@
 #include "VirtualTexturePhysicalSpace.h"
 #include "VirtualTextureSystem.h"
 #include "SpriteIndexBuffer.h"
-#include "SceneFilterRendering.h"
+#include "PostProcess/SceneFilterRendering.h"
 #include "RenderTargetPool.h"
 #include "VisualizeTexture.h"
 #include "CommonRenderResources.h"
@@ -561,4 +561,4 @@ void FVirtualTextureSpace::SaveAllocatorDebugImage() const
 	const FString ImageName = FString::Printf(TEXT("Space%dAllocator.png"), ID);
 	Allocator.SaveDebugImage(*ImageName);
 }
-#endif
\ No newline at end of file
+#endif
diff --git a/Engine/Source/Runtime/Renderer/Private/VT/VirtualTextureTest.cpp b/Engine/Source/Runtime/Renderer/Private/VT/VirtualTextureTest.cpp
--- a/Engine/Source/Runtime/Renderer/Private/VT/VirtualTextureTest.cpp
+++ b/Engine/Source/Runtime/Renderer/Private/VT/VirtualTextureTest.cpp
@@ -5,7 +5,7 @@
 #if 0
 #include "SpriteIndexBuffer.h"
 #include "ScreenRendering.h"
-#include "SceneFilterRendering.h"
+#include "PostProcess/SceneFilterRendering.h"
 #include "RenderTargetPool.h"
 #include "GlobalShader.h"
 #include "PipelineStateCache.h"
@@ -230,4 +230,4 @@ void FVirtualTextureTest::DumpToConsole(bool verbose)
 	UE_LOG(LogConsoleResponse, Display, TEXT("Test Virtual Texture"));
 }
 
-#endif
\ No newline at end of file
+#endif
diff --git a/Engine/Source/Runtime/Renderer/Private/VisualizeVolumetricLightmap.cpp b/Engine/Source/Runtime/Renderer/Private/VisualizeVolumetricLightmap.cpp
--- a/Engine/Source/Runtime/Renderer/Private/VisualizeVolumetricLightmap.cpp
+++ b/Engine/Source/Runtime/Renderer/Private/VisualizeVolumetricLightmap.cpp
@@ -22,7 +22,7 @@
 #include "ClearQuad.h"
 #include "ScenePrivate.h"
 #include "SpriteIndexBuffer.h"
-#include "SceneFilterRendering.h"
+#include "PostProcess/SceneFilterRendering.h"
 #include "PrecomputedVolumetricLightmap.h"
 
 float GVolumetricLightmapVisualizationRadiusScale = .01f;
@@ -175,4 +175,4 @@ void FDeferredShadingSceneRenderer::VisualizeVolumetricLightmap(
 			RHICmdList.DrawIndexedPrimitive(GVisualizeQuadIndexBuffer.IndexBufferRHI, 0, 0, 4 * GQuadsPerVisualizeInstance, 0, 2 * GQuadsPerVisualizeInstance, FMath::DivideAndRoundUp(FMath::Min(NumQuads, 0x7FFFFFFFu / 4), (uint32)GQuadsPerVisualizeInstance));
 		});
 	}
-}
\ No newline at end of file
+}
diff --git a/Engine/Source/Runtime/Renderer/Private/VolumetricCloudRendering.cpp b/Engine/Source/Runtime/Renderer/Private/VolumetricCloudRendering.cpp
--- a/Engine/Source/Runtime/Renderer/Private/VolumetricCloudRendering.cpp
+++ b/Engine/Source/Runtime/Renderer/Private/VolumetricCloudRendering.cpp
@@ -9,7 +9,7 @@
 #include "VolumetricCloudProxy.h"
 #include "DeferredShadingRenderer.h"
 #include "PixelShaderUtils.h"
-#include "RenderCore/Public/RenderGraphUtils.h"
+#include "RenderGraphUtils.h"
 #include "ScenePrivate.h"
 #include "MeshPassProcessor.inl"
 #include "StaticMeshResources.h"
@@ -2561,4 +2561,4 @@ bool SetupLightCloudTransmittanceParameters(FRDGBuilder& GraphBuilder, const FSc
 	}
 
 	return bLight0CloudShadowEnabled || bLight1CloudShadowEnabled;
-}
\ No newline at end of file
+}
diff --git a/Engine/Source/Runtime/Renderer/Private/VolumetricFogLightFunction.cpp b/Engine/Source/Runtime/Renderer/Private/VolumetricFogLightFunction.cpp
--- a/Engine/Source/Runtime/Renderer/Private/VolumetricFogLightFunction.cpp
+++ b/Engine/Source/Runtime/Renderer/Private/VolumetricFogLightFunction.cpp
@@ -9,9 +9,9 @@
 #include "ScenePrivate.h"
 #include "SceneUtils.h"
 #include "LightRendering.h"
-#include "SceneFilterRendering.h"
-#include "PostProcessing.h"
-#include "RHI/Public/PipelineStateCache.h"
+#include "PostProcess/SceneFilterRendering.h"
+#include "PostProcess/PostProcessing.h"
+#include "PipelineStateCache.h"
 
 int GVolumetricFogLightFunctionResolution = 128;
 FAutoConsoleVariableRef CVarVolumetricFogLightFunctionResolution(
diff --git a/Engine/Source/Runtime/Renderer/Private/VolumetricRenderTarget.cpp b/Engine/Source/Runtime/Renderer/Private/VolumetricRenderTarget.cpp
--- a/Engine/Source/Runtime/Renderer/Private/VolumetricRenderTarget.cpp
+++ b/Engine/Source/Runtime/Renderer/Private/VolumetricRenderTarget.cpp
@@ -6,8 +6,8 @@
 
 #include "VolumetricRenderTarget.h"
 #include "DeferredShadingRenderer.h"
-#include "RenderCore/Public/RenderGraphUtils.h"
-#include "RenderCore/Public/PixelShaderUtils.h"
+#include "RenderGraphUtils.h"
+#include "PixelShaderUtils.h"
 #include "ScenePrivate.h"
 #include "SceneTextureParameters.h"
 #include "SingleLayerWaterRendering.h"
diff --git a/Engine/Source/Runtime/Renderer/Renderer.Build.cs b/Engine/Source/Runtime/Renderer/Renderer.Build.cs
--- a/Engine/Source/Runtime/Renderer/Renderer.Build.cs
+++ b/Engine/Source/Runtime/Renderer/Renderer.Build.cs
@@ -8,9 +8,6 @@ public class Renderer : ModuleRules
 	{
 		PrivateIncludePaths.AddRange(
 			new string[] {
-				"Runtime/Renderer/Private",
-				"Runtime/Renderer/Private/CompositionLighting",
-				"Runtime/Renderer/Private/PostProcess",
 				"../Shaders/Shared"
 				}
 			);
diff --git a/Engine/Source/Runtime/SignalProcessing/Public/DSP/DynamicStateVariableFilter.h b/Engine/Source/Runtime/SignalProcessing/Public/DSP/DynamicStateVariableFilter.h
--- a/Engine/Source/Runtime/SignalProcessing/Public/DSP/DynamicStateVariableFilter.h
+++ b/Engine/Source/Runtime/SignalProcessing/Public/DSP/DynamicStateVariableFilter.h
@@ -3,7 +3,7 @@
 #pragma once
 
 #include "CoreMinimal.h"
-#include "SignalProcessing/Public/DSP/EnvelopeFollower.h"
+#include "DSP/EnvelopeFollower.h"
 
 namespace Audio
 {
@@ -128,4 +128,4 @@ namespace Audio
 		
 		void ClearFilterDenormals();
 	};
-}
\ No newline at end of file
+}
diff --git a/Engine/Source/Runtime/Slate/Private/Framework/Docking/SDockingArea.cpp b/Engine/Source/Runtime/Slate/Private/Framework/Docking/SDockingArea.cpp
--- a/Engine/Source/Runtime/Slate/Private/Framework/Docking/SDockingArea.cpp
+++ b/Engine/Source/Runtime/Slate/Private/Framework/Docking/SDockingArea.cpp
@@ -7,7 +7,7 @@
 #include "Framework/Docking/SDockingTarget.h"
 #include "Framework/Docking/FDockingDragOperation.h"
 #include "HAL/PlatformApplicationMisc.h"
-#include "STabSidebar.h"
+#include "Framework/Docking/STabSidebar.h"
 
 
 void SDockingArea::Construct( const FArguments& InArgs, const TSharedRef<FTabManager>& InTabManager, const TSharedRef<FTabManager::FArea>& PersistentNode )
diff --git a/Engine/Source/Runtime/Slate/Private/Framework/Docking/SDockingArea.h b/Engine/Source/Runtime/Slate/Private/Framework/Docking/SDockingArea.h
--- a/Engine/Source/Runtime/Slate/Private/Framework/Docking/SDockingArea.h
+++ b/Engine/Source/Runtime/Slate/Private/Framework/Docking/SDockingArea.h
@@ -9,8 +9,8 @@
 #include "Widgets/SOverlay.h"
 #include "Widgets/SWindow.h"
 #include "Framework/Docking/TabManager.h"
-#include "SDockingNode.h"
-#include "SDockingSplitter.h"
+#include "Framework/Docking/SDockingNode.h"
+#include "Framework/Docking/SDockingSplitter.h"
 
 class STabSidebar;
 
diff --git a/Engine/Source/Runtime/Slate/Private/Framework/Docking/SDockingSplitter.h b/Engine/Source/Runtime/Slate/Private/Framework/Docking/SDockingSplitter.h
--- a/Engine/Source/Runtime/Slate/Private/Framework/Docking/SDockingSplitter.h
+++ b/Engine/Source/Runtime/Slate/Private/Framework/Docking/SDockingSplitter.h
@@ -5,7 +5,7 @@
 #include "CoreMinimal.h"
 #include "Widgets/DeclarativeSyntaxSupport.h"
 #include "Framework/Docking/TabManager.h"
-#include "SDockingNode.h"
+#include "Framework/Docking/SDockingNode.h"
 
 class SDockingTabStack;
 
diff --git a/Engine/Source/Runtime/Slate/Private/Framework/Docking/SDockingTabStack.h b/Engine/Source/Runtime/Slate/Private/Framework/Docking/SDockingTabStack.h
--- a/Engine/Source/Runtime/Slate/Private/Framework/Docking/SDockingTabStack.h
+++ b/Engine/Source/Runtime/Slate/Private/Framework/Docking/SDockingTabStack.h
@@ -12,7 +12,7 @@
 #include "Widgets/SWidget.h"
 #include "Animation/CurveSequence.h"
 #include "Framework/Docking/TabManager.h"
-#include "SDockingNode.h"
+#include "Framework/Docking/SDockingNode.h"
 
 #define DEBUG_TAB_MANAGEMENT 0
 
diff --git a/Engine/Source/Runtime/Slate/Private/Framework/Docking/STabSidebar.cpp b/Engine/Source/Runtime/Slate/Private/Framework/Docking/STabSidebar.cpp
--- a/Engine/Source/Runtime/Slate/Private/Framework/Docking/STabSidebar.cpp
+++ b/Engine/Source/Runtime/Slate/Private/Framework/Docking/STabSidebar.cpp
@@ -1,9 +1,9 @@
 // Copyright Epic Games, Inc. All Rights Reserved.
 
-#include "STabSidebar.h"
+#include "Framework/Docking/STabSidebar.h"
 #include "Widgets/Input/SButton.h"
 #include "Widgets/Input/SCheckBox.h"
-#include "SDockingTabWell.h"
+#include "Framework/Docking/SDockingTabWell.h"
 #include "Widgets/Layout/SBox.h"
 #include "Widgets/Images/SImage.h"
 #include "Framework/Application/SlateApplication.h"
diff --git a/Engine/Source/Runtime/Slate/Private/Framework/MultiBox/MultiBox.cpp b/Engine/Source/Runtime/Slate/Private/Framework/MultiBox/MultiBox.cpp
--- a/Engine/Source/Runtime/Slate/Private/Framework/MultiBox/MultiBox.cpp
+++ b/Engine/Source/Runtime/Slate/Private/Framework/MultiBox/MultiBox.cpp
@@ -26,7 +26,7 @@
 #include "Framework/MultiBox/SToolBarComboButtonBlock.h"
 
 #include "Framework/Commands/UICommandDragDropOp.h"
-#include "SUniformToolbarPanel.h"
+#include "Framework/MultiBox/SUniformToolbarPanel.h"
 #include "Styling/ToolBarStyle.h"
 
 #define LOCTEXT_NAMESPACE "MultiBox"
diff --git a/Engine/Source/Runtime/Slate/Private/Framework/MultiBox/SUniformToolbarPanel.cpp b/Engine/Source/Runtime/Slate/Private/Framework/MultiBox/SUniformToolbarPanel.cpp
--- a/Engine/Source/Runtime/Slate/Private/Framework/MultiBox/SUniformToolbarPanel.cpp
+++ b/Engine/Source/Runtime/Slate/Private/Framework/MultiBox/SUniformToolbarPanel.cpp
@@ -1,6 +1,6 @@
 // Copyright Epic Games, Inc. All Rights Reserved.
 
-#include "SUniformToolbarPanel.h"
+#include "Framework/MultiBox/SUniformToolbarPanel.h"
 #include "Layout/LayoutUtils.h"
 #include "Widgets/Input/SComboButton.h"
 #include "Widgets/Images/SImage.h"
@@ -252,4 +252,4 @@ bool SUniformToolbarPanel::RemoveSlot( const TSharedRef<SWidget>& SlotWidget )
 	}
 	
 	return false;
-}
\ No newline at end of file
+}
diff --git a/Engine/Source/Runtime/Slate/Private/Framework/Notifications/SlateAsyncTaskNotificationImpl.h b/Engine/Source/Runtime/Slate/Private/Framework/Notifications/SlateAsyncTaskNotificationImpl.h
--- a/Engine/Source/Runtime/Slate/Private/Framework/Notifications/SlateAsyncTaskNotificationImpl.h
+++ b/Engine/Source/Runtime/Slate/Private/Framework/Notifications/SlateAsyncTaskNotificationImpl.h
@@ -4,7 +4,7 @@
 
 #include "Misc/CoreAsyncTaskNotificationImpl.h"
 #include "Containers/Ticker.h"
-#include "SlateCore/Public/Layout/Visibility.h"
+#include "Layout/Visibility.h"
 #include "Widgets/Notifications/SNotificationList.h"
 
 class SSlateAsyncTaskNotificationWidget;
diff --git a/Engine/Source/Runtime/Slate/Slate.Build.cs b/Engine/Source/Runtime/Slate/Slate.Build.cs
--- a/Engine/Source/Runtime/Slate/Slate.Build.cs
+++ b/Engine/Source/Runtime/Slate/Slate.Build.cs
@@ -20,29 +20,6 @@ public class Slate : ModuleRules
 				"ImageWrapper"
 			});
 
-		PrivateIncludePaths.AddRange(
-			new string[] {
-				"Runtime/Slate/Private",
-				"Runtime/Slate/Private/Framework",
-				"Runtime/Slate/Private/Framework/Application",
-				"Runtime/Slate/Private/Framework/Commands",
-				"Runtime/Slate/Private/Framework/Docking",
-				"Runtime/Slate/Private/Framework/Layout",
-				"Runtime/Slate/Private/Framework/MultiBox",
-				"Runtime/Slate/Private/Framework/Notifications",
-				"Runtime/Slate/Private/Framework/Styling",
-				"Runtime/Slate/Private/Framework/Text",
-				"Runtime/Slate/Private/Framework/Text/IOS",
-				"Runtime/Slate/Private/Widgets/Colors",
-				"Runtime/Slate/Private/Widgets/Docking",
-				"Runtime/Slate/Private/Widgets/Images",
-				"Runtime/Slate/Private/Widgets/Input",
-				"Runtime/Slate/Private/Widgets/Layout",
-				"Runtime/Slate/Private/Widgets/Notifications",
-				"Runtime/Slate/Private/Widgets/Text",
-				"Runtime/Slate/Private/Widgets/Views",
-			});
-
 		if (Target.Platform == UnrealTargetPlatform.Win64)
 		{
 			AddEngineThirdPartyPrivateStaticDependencies(Target, "XInput");
diff --git a/Engine/Source/Runtime/SlateCore/SlateCore.Build.cs b/Engine/Source/Runtime/SlateCore/SlateCore.Build.cs
--- a/Engine/Source/Runtime/SlateCore/SlateCore.Build.cs
+++ b/Engine/Source/Runtime/SlateCore/SlateCore.Build.cs
@@ -16,23 +16,6 @@ public class SlateCore : ModuleRules
 				"TraceLog",
 			});
 
-		PrivateIncludePaths.AddRange(
-			new string[] {
-				"Runtime/SlateCore/Private",
-				"Runtime/SlateCore/Private/Animation",
-				"Runtime/SlateCore/Private/Application",
-				"Runtime/SlateCore/Private/Brushes",
-				"Runtime/SlateCore/Private/Fonts",
-				"Runtime/SlateCore/Private/Input",
-				"Runtime/SlateCore/Private/Layout",
-				"Runtime/SlateCore/Private/Rendering",
-				"Runtime/SlateCore/Private/Sound",
-				"Runtime/SlateCore/Private/Styling",
-				"Runtime/SlateCore/Private/Textures",
-				"Runtime/SlateCore/Private/Types",
-				"Runtime/SlateCore/Private/Widgets",
-			});
-
         if (Target.Type != TargetType.Server)
 		{
 			if (Target.bCompileFreeType)
diff --git a/Engine/Source/Runtime/SlateRHIRenderer/Private/SlateRHIRenderer.cpp b/Engine/Source/Runtime/SlateRHIRenderer/Private/SlateRHIRenderer.cpp
--- a/Engine/Source/Runtime/SlateRHIRenderer/Private/SlateRHIRenderer.cpp
+++ b/Engine/Source/Runtime/SlateRHIRenderer/Private/SlateRHIRenderer.cpp
@@ -21,7 +21,7 @@
 #include "StereoRendering.h"
 #include "SlateNativeTextureResource.h"
 #include "SceneUtils.h"
-#include "Runtime/Renderer/Public/VolumeRendering.h"
+#include "VolumeRendering.h"
 #include "ShaderCompiler.h"
 #include "PipelineStateCache.h"
 #include "EngineModule.h"
diff --git a/Engine/Source/Runtime/UMG/Private/Slate/SMeshWidget.cpp b/Engine/Source/Runtime/UMG/Private/Slate/SMeshWidget.cpp
--- a/Engine/Source/Runtime/UMG/Private/Slate/SMeshWidget.cpp
+++ b/Engine/Source/Runtime/UMG/Private/Slate/SMeshWidget.cpp
@@ -9,7 +9,7 @@
 #include "UMGPrivate.h"
 
 #include "SlateMaterialBrush.h"
-#include "Runtime/SlateRHIRenderer/Public/Interfaces/ISlateRHIRendererModule.h"
+#include "Interfaces/ISlateRHIRendererModule.h"
 #include "Slate/SlateVectorArtData.h"
 #include "Slate/SlateVectorArtInstanceData.h"
 
diff --git a/Engine/Source/Runtime/UMG/Private/SlateBlueprintLibrary.cpp b/Engine/Source/Runtime/UMG/Private/SlateBlueprintLibrary.cpp
--- a/Engine/Source/Runtime/UMG/Private/SlateBlueprintLibrary.cpp
+++ b/Engine/Source/Runtime/UMG/Private/SlateBlueprintLibrary.cpp
@@ -7,7 +7,7 @@
 
 #include "Engine/UserInterfaceSettings.h"
 #include "Slate/SlateBrushAsset.h"
-#include "Runtime/Engine/Classes/Engine/UserInterfaceSettings.h"
+#include "Engine/UserInterfaceSettings.h"
 #include "Blueprint/SlateBlueprintLibrary.h"
 #include "Slate/SceneViewport.h"
 #include "Blueprint/WidgetLayoutLibrary.h"
