ZenLoader: Fixed recently introdced atomic ordering problem

diff --git a/Engine/Source/Runtime/CoreUObject/Private/Serialization/AsyncLoading2.cpp b/Engine/Source/Runtime/CoreUObject/Private/Serialization/AsyncLoading2.cpp
--- a/Engine/Source/Runtime/CoreUObject/Private/Serialization/AsyncLoading2.cpp
+++ b/Engine/Source/Runtime/CoreUObject/Private/Serialization/AsyncLoading2.cpp
@@ -1680,16 +1680,6 @@ public:
 	void ReleaseBarrier(FAsyncLoadingThreadState2* ThreadState = nullptr);
 	EEventLoadNodeExecutionResult Execute(FAsyncLoadingThreadState2& ThreadState);
 
-	inline bool IsDone() const
-	{
-		return bIsDone.load(std::memory_order_acquire) == 1;
-	}
-
-	inline void MarkAsDone()
-	{
-		bIsDone.store(1, std::memory_order_release);
-	}
-
 	FAsyncPackage2* GetPackage() const
 	{
 		return Package;
@@ -1717,8 +1707,8 @@ private:
 	uint32 DependenciesCount = 0;
 	uint32 DependenciesCapacity = 0;
 	std::atomic<int32> BarrierCount { 0 };
-	std::atomic<uint8> DependencyWriterCount { 0 };
-	std::atomic<uint8> bIsDone { 0 };
+	std::atomic<bool> bIsUpdatingDependencies { false };
+	std::atomic<bool> bIsDone { false };
 
 	const FAsyncLoadEventSpec* Spec = nullptr;
 	FAsyncPackage2* Package = nullptr;
@@ -3547,15 +3537,17 @@ FEventLoadNode2::FEventLoadNode2(const FAsyncLoadEventSpec* InSpec, FAsyncPackag
 void FEventLoadNode2::DependsOn(FEventLoadNode2* Other)
 {
 	TRACE_CPUPROFILER_EVENT_SCOPE(DependsOn);
-#if UE_BUILD_DEBUG || UE_BUILD_DEVELOPMENT
-	check(!IsDone());
-#endif
-	uint8 Expected = 0;
-	while (!Other->DependencyWriterCount.compare_exchange_weak(Expected, 1))
+	check(!bIsDone.load(std::memory_order_relaxed));
+	bool bExpected = false;
+	// Set modification flag before checking the done flag
+	// If we're currently in ProcessDependencies the done flag will have been set and we won't do anything
+	// If ProcessDependencies is called during this call it will wait for the modification flag to be cleared
+	while (!Other->bIsUpdatingDependencies.compare_exchange_strong(bExpected, true))
 	{
-		Expected = 0;
+		// Note: Currently only the async loading thread is calling DependsOn so this will never be contested
+		bExpected = false;
 	}
-	if (!Other->IsDone())
+	if (!Other->bIsDone.load())
 	{
 		++BarrierCount;
 		if (Other->DependenciesCount == 0)
@@ -3589,22 +3581,18 @@ void FEventLoadNode2::DependsOn(FEventLoadNode2* Other)
 			Other->MultipleDependents[Other->DependenciesCount++] = this;
 		}
 	}
-	Other->DependencyWriterCount.store(0, std::memory_order_release);
+	Other->bIsUpdatingDependencies.store(false);
 }
 
 void FEventLoadNode2::AddBarrier()
 {
-#if UE_BUILD_DEBUG || UE_BUILD_DEVELOPMENT
-	check(!IsDone());
-#endif
+	check(!bIsDone.load(std::memory_order_relaxed));
 	++BarrierCount;
 }
 
 void FEventLoadNode2::AddBarrier(int32 Count)
 {
-#if UE_BUILD_DEBUG || UE_BUILD_DEVELOPMENT
-	check(!IsDone());
-#endif
+	check(!bIsDone.load(std::memory_order_relaxed));
 	BarrierCount += Count;
 }
 
@@ -3642,7 +3630,6 @@ EEventLoadNodeExecutionResult FEventLoadNode2::Execute(FAsyncLoadingThreadState2
 	}
 	if (Result == EEventLoadNodeExecutionResult::Complete)
 	{
-		MarkAsDone();
 		ProcessDependencies(ThreadState);
 	}
 	return Result;
@@ -3651,13 +3638,11 @@ EEventLoadNodeExecutionResult FEventLoadNode2::Execute(FAsyncLoadingThreadState2
 void FEventLoadNode2::ProcessDependencies(FAsyncLoadingThreadState2& ThreadState)
 {
 	//TRACE_CPUPROFILER_EVENT_SCOPE(ProcessDependencies);
-	if (DependencyWriterCount.load(std::memory_order_acquire) != 0)
+	// Set the done flag before checking the modification flag
+	bIsDone.store(true);
+	while (bIsUpdatingDependencies.load())
 	{
-		TRACE_CPUPROFILER_EVENT_SCOPE(ConcurrentWriter);
-		do
-		{
-			FPlatformProcess::Sleep(0);
-		} while (DependencyWriterCount.load(std::memory_order_acquire) != 0);
+		FPlatformProcess::Sleep(0);
 	}
 
 	if (DependenciesCount == 1)
