DerivedData: Added a way to request compressed data from derived data references

- FDerivedDataIoBatch::Compress requests the entire compressed value.
- FDerivedDataIoResponse::GetCompressedData gets the compressed buffer from the response when available.
- FDerivedDataIoResponse::GetCacheValueId has been added to get the FValueId when available.
- FDerivedDataIoResponse::GetCacheKey has been fixed to return a non-null cache key when appropriate.
- Asserts about FDerivedData now log the derived data reference too.

diff --git a/Engine/Source/Runtime/CoreUObject/Private/Serialization/DerivedData.cpp b/Engine/Source/Runtime/CoreUObject/Private/Serialization/DerivedData.cpp
--- a/Engine/Source/Runtime/CoreUObject/Private/Serialization/DerivedData.cpp
+++ b/Engine/Source/Runtime/CoreUObject/Private/Serialization/DerivedData.cpp
@@ -189,15 +189,14 @@ uint32 FEditorData::ReferenceHash() const
 
 void FEditorData::Serialize(FArchive& Ar, UObject* Owner) const
 {
-	checkf(Ar.IsSaving() && Ar.IsCooking(), TEXT("FEditorData for FDerivedData only supports saving to cooked packages."));
+	checkf(Ar.IsSaving() && Ar.IsCooking(),
+		TEXT("FEditorData for FDerivedData only supports saving to cooked packages. %s"), *WriteToString<256>(*this));
 
 	struct FVisitor
 	{
 		inline explicit FVisitor(FArchive& Ar)
 			: Linker(Cast<FLinkerSave>(Ar.GetLinker()))
 		{
-			checkf(Linker, TEXT("Serializing FDerivedData requires a linker."));
-			checkf(Ar.IsCooking(), TEXT("Serializing FDerivedData is only supported for cooked packages."));
 		}
 
 		inline void operator()(const FCompositeBufferWithHash& BufferWithHash)
@@ -223,6 +222,7 @@ void FEditorData::Serialize(FArchive& Ar, UObject* Owner) const
 
 	static_assert(sizeof(FIoChunkId) == 12);
 	FVisitor Visitor(Ar);
+	checkf(Visitor.Linker, TEXT("Serializing FDerivedData requires a linker. %s"), *WriteToString<256>(*this));
 	Visit(Visitor);
 
 	Ar << Visitor.ChunkOffset;
@@ -438,6 +438,9 @@ enum class EIoRequestType : uint8
 	Read,
 	Cache,
 	Exists,
+#if WITH_EDITORONLY_DATA
+	Compress,
+#endif
 };
 
 struct FIoChunkState
@@ -454,7 +457,9 @@ struct FIoEditorState
 #if WITH_EDITORONLY_DATA
 	FEditorData EditorData;
 	FCacheKey CacheKey;
+	FValueId CacheValueId;
 	FIoHash Hash;
+	FCompressedBuffer CompressedData;
 
 	inline explicit FIoEditorState(const FEditorData& InEditorData)
 		: EditorData(InEditorData)
@@ -474,8 +479,7 @@ struct FIoRequestState
 
 	inline void SetStatus(EDerivedDataIoStatus NewStatus)
 	{
-		const bool bOk = NewStatus == EDerivedDataIoStatus::Ok;
-		Status.store(NewStatus, bOk ? std::memory_order_release : std::memory_order_relaxed);
+		Status.store(NewStatus, std::memory_order_release);
 	}
 };
 
@@ -537,9 +541,13 @@ private:
 
 #if WITH_EDITORONLY_DATA
 	void DispatchEditor(FIoResponse& Response, FIoRequestState& Request, FIoEditorState& Editor, int32 RequestIndex);
-	static void OnCacheRequestComplete(FIoResponse& Response, FCacheGetChunkResponse&& Chunk);
+	static void OnCacheRecordRequestComplete(FIoResponse& Response, FCacheGetResponse&& CacheResponse);
+	static void OnCacheValueRequestComplete(FIoResponse& Response, FCacheGetValueResponse&& CacheResponse);
+	static void OnCacheChunkRequestComplete(FIoResponse& Response, FCacheGetChunkResponse&& CacheResponse);
 
-	TArray<FCacheGetChunkRequest> CacheRequests;
+	TArray<FCacheGetRequest> CacheRecordRequests;
+	TArray<FCacheGetValueRequest> CacheValueRequests;
+	TArray<FCacheGetChunkRequest> CacheChunkRequests;
 #endif
 
 	FIoBatch Batch;
@@ -730,12 +738,30 @@ void FIoResponseDispatcher::Dispatch(FIoResponse& Response, FDerivedDataIoPriori
 	Dispatcher.Batch.Issue();
 
 #if WITH_EDITORONLY_DATA
-	if (!Dispatcher.CacheRequests.IsEmpty())
+	if (!Dispatcher.CacheRecordRequests.IsEmpty())
+	{
+		GetCache().Get(Dispatcher.CacheRecordRequests, Response.Owner,
+			[Response = &Response](FCacheGetResponse&& CacheResponse)
+			{
+				OnCacheRecordRequestComplete(*Response, MoveTemp(CacheResponse));
+			});
+	}
+
+	if (!Dispatcher.CacheValueRequests.IsEmpty())
 	{
-		GetCache().GetChunks(Dispatcher.CacheRequests, Response.Owner,
-			[Response = &Response](FCacheGetChunkResponse&& Chunk)
+		GetCache().GetValue(Dispatcher.CacheValueRequests, Response.Owner,
+			[Response = &Response](FCacheGetValueResponse&& CacheResponse)
 			{
-				OnCacheRequestComplete(*Response, MoveTemp(Chunk));
+				OnCacheValueRequestComplete(*Response, MoveTemp(CacheResponse));
+			});
+	}
+
+	if (!Dispatcher.CacheChunkRequests.IsEmpty())
+	{
+		GetCache().GetChunks(Dispatcher.CacheChunkRequests, Response.Owner,
+			[Response = &Response](FCacheGetChunkResponse&& CacheResponse)
+			{
+				OnCacheChunkRequestComplete(*Response, MoveTemp(CacheResponse));
 			});
 	}
 #endif
@@ -888,6 +914,17 @@ void FIoResponseDispatcher::DispatchEditor(
 					Response->Owner.LaunchTask(TEXT("DerivedDataCopy"), MoveTemp(Execute));
 				}
 			}
+			else if (Request->Type == EIoRequestType::Compress)
+			{
+				Response->BeginRequest();
+				FRequestBarrier Barrier(Response->Owner);
+				Response->Owner.LaunchTask(TEXT("DerivedDataCompress"), [Response = Response, Request = Request, &BufferWithHash]
+				{
+					Request->State.Get<FIoEditorState>().CompressedData = FValue::Compress(BufferWithHash.Buffer).GetData();
+					Request->SetStatus(EDerivedDataIoStatus::Ok);
+					Response->EndRequest();
+				});
+			}
 			else
 			{
 				Request->SetStatus(EDerivedDataIoStatus::Ok);
@@ -937,6 +974,11 @@ void FIoResponseDispatcher::DispatchEditor(
 					Response->Owner.LaunchTask(TEXT("DerivedDataDecompress"), MoveTemp(Execute));
 				}
 			}
+			else if (Request->Type == EIoRequestType::Compress)
+			{
+				Request->State.Get<FIoEditorState>().CompressedData = Buffer;
+				Request->SetStatus(EDerivedDataIoStatus::Ok);
+			}
 			else
 			{
 				Request->SetStatus(EDerivedDataIoStatus::Ok);
@@ -945,17 +987,49 @@ void FIoResponseDispatcher::DispatchEditor(
 
 		void operator()(const FCacheKeyWithId& CacheKeyWithId) const
 		{
-			FCacheGetChunkRequest& Chunk = Dispatcher->CacheRequests.AddDefaulted_GetRef();
-			Chunk.Name = Editor->EditorData.GetName();
-			Chunk.Key = CacheKeyWithId.Key;
-			Chunk.Id = CacheKeyWithId.Id;
-			Chunk.RawOffset = Request->Options.GetOffset();
-			Chunk.RawSize = Request->Options.GetSize();
-			Chunk.Policy =
-				Request->Type == EIoRequestType::Read  ? (ECachePolicy::Default) :
-				Request->Type == EIoRequestType::Cache ? (ECachePolicy::Default | ECachePolicy::SkipData) :
-				                                         (ECachePolicy::Query   | ECachePolicy::SkipData);
-			Chunk.UserData = uint64(RequestIndex);
+			Editor->CacheKey = CacheKeyWithId.Key;
+			Editor->CacheValueId = CacheKeyWithId.Id;
+
+			if (Request->Type == EIoRequestType::Compress)
+			{
+				if (CacheKeyWithId.Id.IsValid())
+				{
+					// Requesting multiple values from the same record can be optimized by making
+					// one request for the record with a policy that requests each of the values,
+					// rather than one record request for each value.
+
+					FCacheRecordPolicyBuilder Policy(ECachePolicy::None | ECachePolicy::SkipMeta);
+					Policy.AddValuePolicy(CacheKeyWithId.Id, ECachePolicy::Default);
+
+					FCacheGetRequest& CacheRequest = Dispatcher->CacheRecordRequests.AddDefaulted_GetRef();
+					CacheRequest.Name = Editor->EditorData.GetName();
+					CacheRequest.Key = CacheKeyWithId.Key;
+					CacheRequest.Policy = Policy.Build();
+					CacheRequest.UserData = uint64(RequestIndex);
+				}
+				else
+				{
+					FCacheGetValueRequest& CacheRequest = Dispatcher->CacheValueRequests.AddDefaulted_GetRef();
+					CacheRequest.Name = Editor->EditorData.GetName();
+					CacheRequest.Key = CacheKeyWithId.Key;
+					CacheRequest.UserData = uint64(RequestIndex);
+				}
+			}
+			else
+			{
+				FCacheGetChunkRequest& CacheRequest = Dispatcher->CacheChunkRequests.AddDefaulted_GetRef();
+				CacheRequest.Name = Editor->EditorData.GetName();
+				CacheRequest.Key = CacheKeyWithId.Key;
+				CacheRequest.Id = CacheKeyWithId.Id;
+				CacheRequest.RawOffset = Request->Options.GetOffset();
+				CacheRequest.RawSize = Request->Options.GetSize();
+				CacheRequest.Policy =
+					Request->Type == EIoRequestType::Read  ? (ECachePolicy::Default) :
+					Request->Type == EIoRequestType::Cache ? (ECachePolicy::Default | ECachePolicy::SkipData) :
+															 (ECachePolicy::Query   | ECachePolicy::SkipData);
+				CacheRequest.UserData = uint64(RequestIndex);
+			}
+
 			Response->BeginRequest();
 		}
 
@@ -970,14 +1044,37 @@ void FIoResponseDispatcher::DispatchEditor(
 	Editor.EditorData.Visit(Visitor);
 }
 
-void FIoResponseDispatcher::OnCacheRequestComplete(FIoResponse& Response, FCacheGetChunkResponse&& Chunk)
+void FIoResponseDispatcher::OnCacheRecordRequestComplete(FIoResponse& Response, FCacheGetResponse&& CacheResponse)
+{
+	FIoRequestState& Request = Response.Requests[int32(CacheResponse.UserData)];
+	FIoEditorState& Editor = Request.State.Get<FIoEditorState>();
+	const FValue& Value = CacheResponse.Record.GetValue(Editor.CacheValueId);
+	Editor.CompressedData = Value.GetData();
+	Editor.Hash = Value.GetRawHash();
+	Request.Size = Value.GetRawSize();
+	Request.SetStatus(ConvertToIoStatus(CacheResponse.Status));
+	Response.EndRequest();
+}
+
+void FIoResponseDispatcher::OnCacheValueRequestComplete(FIoResponse& Response, FCacheGetValueResponse&& CacheResponse)
+{
+	FIoRequestState& Request = Response.Requests[int32(CacheResponse.UserData)];
+	FIoEditorState& Editor = Request.State.Get<FIoEditorState>();
+	Editor.CompressedData = CacheResponse.Value.GetData();
+	Editor.Hash = CacheResponse.Value.GetRawHash();
+	Request.Size = CacheResponse.Value.GetRawSize();
+	Request.SetStatus(ConvertToIoStatus(CacheResponse.Status));
+	Response.EndRequest();
+}
+
+void FIoResponseDispatcher::OnCacheChunkRequestComplete(FIoResponse& Response, FCacheGetChunkResponse&& CacheResponse)
 {
-	FIoRequestState& Request = Response.Requests[int32(Chunk.UserData)];
+	FIoRequestState& Request = Response.Requests[int32(CacheResponse.UserData)];
 	FIoEditorState& Editor = Request.State.Get<FIoEditorState>();
-	Editor.Hash = Chunk.RawHash;
-	Request.Size = Chunk.RawSize;
-	Request.Data = MoveTemp(Chunk.RawData);
-	Request.SetStatus(ConvertToIoStatus(Chunk.Status));
+	Editor.Hash = CacheResponse.RawHash;
+	Request.Size = CacheResponse.RawSize;
+	Request.Data = MoveTemp(CacheResponse.RawData);
+	Request.SetStatus(ConvertToIoStatus(CacheResponse.Status));
 	Response.EndRequest();
 }
 
@@ -1055,7 +1152,7 @@ const FIoHash* FDerivedDataIoResponse::GetHash(FDerivedDataIoRequest Handle) con
 	using namespace DerivedData::Private;
 	if (FIoRequestState* Request = FIoResponse::TryGetRequest(Response, Handle))
 	{
-		if (Request->Status.load(std::memory_order_relaxed) == EDerivedDataIoStatus::Ok)
+		if (Request->Status.load(std::memory_order_acquire) == EDerivedDataIoStatus::Ok)
 		{
 			if (FIoEditorState* EditorState = Request->State.TryGet<FIoEditorState>())
 			{
@@ -1074,7 +1171,7 @@ const DerivedData::FCacheKey* FDerivedDataIoResponse::GetCacheKey(FDerivedDataIo
 	using namespace DerivedData::Private;
 	if (FIoRequestState* Request = FIoResponse::TryGetRequest(Response, Handle))
 	{
-		if (Request->Status.load(std::memory_order_relaxed) == EDerivedDataIoStatus::Ok)
+		if (Request->Status.load(std::memory_order_acquire) == EDerivedDataIoStatus::Ok)
 		{
 			if (FIoEditorState* EditorState = Request->State.TryGet<FIoEditorState>())
 			{
@@ -1089,6 +1186,49 @@ const DerivedData::FCacheKey* FDerivedDataIoResponse::GetCacheKey(FDerivedDataIo
 }
 #endif // WITH_EDITORONLY_DATA
 
+#if WITH_EDITORONLY_DATA
+const DerivedData::FValueId* FDerivedDataIoResponse::GetCacheValueId(FDerivedDataIoRequest Handle) const
+{
+	using namespace DerivedData;
+	using namespace DerivedData::Private;
+	if (FIoRequestState* Request = FIoResponse::TryGetRequest(Response, Handle))
+	{
+		if (Request->Status.load(std::memory_order_acquire) == EDerivedDataIoStatus::Ok)
+		{
+			if (FIoEditorState* EditorState = Request->State.TryGet<FIoEditorState>())
+			{
+				if (EditorState->CacheValueId != FValueId::Null)
+				{
+					return &EditorState->CacheValueId;
+				}
+			}
+		}
+	}
+	return nullptr;
+}
+#endif // WITH_EDITORONLY_DATA
+
+#if WITH_EDITORONLY_DATA
+const FCompressedBuffer* FDerivedDataIoResponse::GetCompressedData(FDerivedDataIoRequest Handle) const
+{
+	using namespace DerivedData::Private;
+	if (FIoRequestState* Request = FIoResponse::TryGetRequest(Response, Handle))
+	{
+		if (Request->Status.load(std::memory_order_acquire) == EDerivedDataIoStatus::Ok)
+		{
+			if (FIoEditorState* EditorState = Request->State.TryGet<FIoEditorState>())
+			{
+				if (EditorState->CompressedData)
+				{
+					return &EditorState->CompressedData;
+				}
+			}
+		}
+	}
+	return nullptr;
+}
+#endif // WITH_EDITORONLY_DATA
+
 ///////////////////////////////////////////////////////////////////////////////////////////////////////////////////////
 
 FDerivedDataIoRequest FDerivedDataIoBatch::Read(const FDerivedData& Data, const FDerivedDataIoOptions& Options)
@@ -1100,15 +1240,40 @@ FDerivedDataIoRequest FDerivedDataIoBatch::Read(const FDerivedData& Data, const
 FDerivedDataIoRequest FDerivedDataIoBatch::Cache(const FDerivedData& Data, const FDerivedDataIoOptions& Options)
 {
 	using namespace DerivedData::Private;
+	checkf(!Options.GetTarget(), TEXT("Target is not supported by Cache. %s"), *WriteToString<256>(Data));
 	return FIoResponse::Queue(Response, Data, Options, EIoRequestType::Cache);
 }
 
 FDerivedDataIoRequest FDerivedDataIoBatch::Exists(const FDerivedData& Data, const FDerivedDataIoOptions& Options)
 {
 	using namespace DerivedData::Private;
+	checkf(!Options.GetTarget(), TEXT("Target is not supported by Exists. %s"), *WriteToString<256>(Data));
 	return FIoResponse::Queue(Response, Data, Options, EIoRequestType::Exists);
 }
 
+#if WITH_EDITORONLY_DATA
+FDerivedDataIoRequest FDerivedDataIoBatch::Compress(const FDerivedData& Data)
+{
+	using namespace DerivedData::Private;
+	return FIoResponse::Queue(Response, Data, {}, EIoRequestType::Compress);
+}
+#endif // WITH_EDITORONLY_DATA
+
+void FDerivedDataIoBatch::Dispatch(FDerivedDataIoResponse& OutResponse)
+{
+	Dispatch(OutResponse, {}, {});
+}
+
+void FDerivedDataIoBatch::Dispatch(FDerivedDataIoResponse& OutResponse, FDerivedDataIoPriority Priority)
+{
+	Dispatch(OutResponse, Priority, {});
+}
+
+void FDerivedDataIoBatch::Dispatch(FDerivedDataIoResponse& OutResponse, FDerivedDataIoComplete&& OnComplete)
+{
+	Dispatch(OutResponse, {}, MoveTemp(OnComplete));
+}
+
 void FDerivedDataIoBatch::Dispatch(
 	FDerivedDataIoResponse& OutResponse,
 	FDerivedDataIoPriority Priority,
diff --git a/Engine/Source/Runtime/CoreUObject/Public/Serialization/DerivedData.h b/Engine/Source/Runtime/CoreUObject/Public/Serialization/DerivedData.h
--- a/Engine/Source/Runtime/CoreUObject/Public/Serialization/DerivedData.h
+++ b/Engine/Source/Runtime/CoreUObject/Public/Serialization/DerivedData.h
@@ -8,7 +8,6 @@
 #include "Memory/MemoryFwd.h"
 #include "Memory/MemoryView.h"
 #include "Misc/EnumClassFlags.h"
-#include "Templates/Function.h"
 #include "Templates/PimplPtr.h"
 
 #define UE_API COREUOBJECT_API
@@ -17,6 +16,7 @@ class FArchive;
 class FCompressedBuffer;
 class UObject;
 struct FIoHash;
+template <typename FuncType> class TUniqueFunction;
 
 namespace UE::DerivedData { struct FCacheKey; }
 namespace UE::DerivedData { struct FValueId; }
@@ -265,7 +265,13 @@ public:
 
 	/** Returns the cache key for the request. Null if not applicable, executing, failed, or canceled. */
 	UE_API const DerivedData::FCacheKey* GetCacheKey(FDerivedDataIoRequest Handle) const;
-#endif
+
+	/** Returns the cache value ID for the request. Null if not applicable, executing, failed, or canceled. */
+	UE_API const DerivedData::FValueId* GetCacheValueId(FDerivedDataIoRequest Handle) const;
+
+	/** Returns the compressed data for the request. Null if not applicable, executing, failed, or canceled. */
+	UE_API const FCompressedBuffer* GetCompressedData(FDerivedDataIoRequest Handle) const;
+#endif // WITH_EDITORONLY_DATA
 
 private:
 	TPimplPtr<DerivedData::Private::FIoResponse> Response;
@@ -361,6 +367,17 @@ public:
 	 */
 	UE_API FDerivedDataIoRequest Exists(const FDerivedData& Data, const FDerivedDataIoOptions& Options = {});
 
+#if WITH_EDITORONLY_DATA
+	/**
+	 * Compresses the derived data. Not executed until Dispatch() is invoked.
+	 *
+	 * The compressed data, size, and status will be available on the response.
+	 *
+	 * @return A request to query the response with. May be discarded by the caller if not needed.
+	 */
+	UE_API FDerivedDataIoRequest Compress(const FDerivedData& Data);
+#endif // WITH_EDITORONLY_DATA
+
 	/**
 	 * Dispatches the requests that have been queued to this batch and resets this batch to empty.
 	 *
@@ -368,9 +385,12 @@ public:
 	 *
 	 * @param OutResponse   Response that is assigned before any requests begin executing.
 	 * @param Priority      Priority of this batch relative to other batches.
-	 * @param OnComplete    Optional completion callback invoked when the entire batch is complete or canceled.
+	 * @param OnComplete    Completion callback invoked when the entire batch is complete or canceled.
 	 */
-	UE_API void Dispatch(FDerivedDataIoResponse& OutResponse, FDerivedDataIoPriority Priority = {}, FDerivedDataIoComplete&& OnComplete = {});
+	UE_API void Dispatch(FDerivedDataIoResponse& OutResponse);
+	UE_API void Dispatch(FDerivedDataIoResponse& OutResponse, FDerivedDataIoPriority Priority);
+	UE_API void Dispatch(FDerivedDataIoResponse& OutResponse, FDerivedDataIoPriority Priority, FDerivedDataIoComplete&& OnComplete);
+	UE_API void Dispatch(FDerivedDataIoResponse& OutResponse, FDerivedDataIoComplete&& OnComplete);
 
 private:
 	TPimplPtr<DerivedData::Private::FIoResponse> Response;
diff --git a/Engine/Source/Runtime/CoreUObject/Tests/Serialization/DerivedDataTest.cpp b/Engine/Source/Runtime/CoreUObject/Tests/Serialization/DerivedDataTest.cpp
--- a/Engine/Source/Runtime/CoreUObject/Tests/Serialization/DerivedDataTest.cpp
+++ b/Engine/Source/Runtime/CoreUObject/Tests/Serialization/DerivedDataTest.cpp
@@ -4,6 +4,7 @@
 
 #include "Serialization/DerivedData.h"
 
+#include "Algo/Compare.h"
 #include "HAL/Event.h"
 #include "IO/IoHash.h"
 #include "Memory/CompositeBuffer.h"
@@ -46,6 +47,13 @@ TEST_CASE("CoreUObject::Serialization::DerivedData", "[CoreUObject][Serializatio
 	using namespace DerivedDataTest::Private;
 	uint64 Counter = 0;
 
+	SECTION("Null Request")
+	{
+		FDerivedDataIoRequest Request;
+		CHECK(Request.IsNull());
+		CHECK_FALSE(Request);
+	}
+
 	SECTION("Null")
 	{
 		const FDerivedData DerivedData;
@@ -66,16 +74,16 @@ TEST_CASE("CoreUObject::Serialization::DerivedData", "[CoreUObject][Serializatio
 		CHECK(Response.IsNull());
 		CHECK_FALSE(Response);
 
-		FDerivedDataIoRequest Requests[3];
-		CHECK(Requests[0].IsNull());
-		CHECK_FALSE(Requests[0]);
-
+		TArray<FDerivedDataIoRequest, TInlineAllocator<4>> Requests;
 		{
 			FDerivedDataIoBatch Batch;
 			CHECK(Batch.IsEmpty());
-			Requests[0] = Batch.Read(DerivedData, {});
-			Requests[1] = Batch.Cache(DerivedData, {});
-			Requests[2] = Batch.Exists(DerivedData, {});
+			Requests.Add(Batch.Read(DerivedData, {}));
+			Requests.Add(Batch.Cache(DerivedData, {}));
+			Requests.Add(Batch.Exists(DerivedData, {}));
+		#if WITH_EDITORONLY_DATA
+			Requests.Add(Batch.Compress(DerivedData));
+		#endif
 			CHECK_FALSE(Batch.IsEmpty());
 			Batch.Dispatch(Response, {}, [&Event] { Event->Trigger(); });
 			CHECK(Batch.IsEmpty());
@@ -98,6 +106,8 @@ TEST_CASE("CoreUObject::Serialization::DerivedData", "[CoreUObject][Serializatio
 		#if WITH_EDITORONLY_DATA
 			CHECK(Response.GetHash(Request) == nullptr);
 			CHECK(Response.GetCacheKey(Request) == nullptr);
+			CHECK(Response.GetCacheValueId(Request) == nullptr);
+			CHECK(Response.GetCompressedData(Request) == nullptr);
 		#endif
 		}
 	}
@@ -124,15 +134,18 @@ TEST_CASE("CoreUObject::Serialization::DerivedData", "[CoreUObject][Serializatio
 
 		FUniqueBuffer Target = FUniqueBuffer::Alloc(HalfSize);
 		FDerivedDataIoResponse Response;
-		FDerivedDataIoRequest Requests[6];
+		TArray<FDerivedDataIoRequest, TInlineAllocator<7>> Requests;
 		{
 			FDerivedDataIoBatch Batch;
-			Requests[0] = Batch.Read(DerivedData, {});
-			Requests[1] = Batch.Cache(DerivedData, {});
-			Requests[2] = Batch.Exists(DerivedData, {});
-			Requests[3] = Batch.Exists(DerivedData, FDerivedDataIoOptions(HalfSize));
-			Requests[4] = Batch.Read(DerivedData, FDerivedDataIoOptions(HalfSize, HalfSize));
-			Requests[5] = Batch.Read(DerivedData, FDerivedDataIoOptions(Target, HalfSize / 2));
+			Requests.Add(Batch.Read(DerivedData, {}));
+			Requests.Add(Batch.Cache(DerivedData, {}));
+			Requests.Add(Batch.Exists(DerivedData, {}));
+			Requests.Add(Batch.Exists(DerivedData, FDerivedDataIoOptions(HalfSize)));
+			Requests.Add(Batch.Read(DerivedData, FDerivedDataIoOptions(HalfSize, HalfSize)));
+			Requests.Add(Batch.Read(DerivedData, FDerivedDataIoOptions(Target, HalfSize / 2)));
+		#if WITH_EDITORONLY_DATA
+			Requests.Add(Batch.Compress(DerivedData));
+		#endif
 			Batch.Dispatch(Response, {}, [&Event] { Event->Trigger(); });
 		}
 		Event->Wait();
@@ -150,6 +163,7 @@ TEST_CASE("CoreUObject::Serialization::DerivedData", "[CoreUObject][Serializatio
 				CHECK(*RequestHash == SourceHash);
 			}
 			CHECK(Response.GetCacheKey(Request) == nullptr);
+			CHECK(Response.GetCacheValueId(Request) == nullptr);
 		}
 
 		CHECK(Response.GetData(Requests[0]).GetSize() == HalfSize * 2);
@@ -158,6 +172,9 @@ TEST_CASE("CoreUObject::Serialization::DerivedData", "[CoreUObject][Serializatio
 		CHECK(Response.GetData(Requests[3]).IsNull());
 		CHECK(Response.GetData(Requests[4]) == Source.GetSegments()[1]);
 		CHECK(FCompositeBuffer(Response.GetData(Requests[5])).EqualBytes(Source.Mid(HalfSize / 2, HalfSize)));
+	#if WITH_EDITORONLY_DATA
+		CHECK(Response.GetData(Requests[6]).IsNull());
+	#endif
 
 		CHECK(Response.GetSize(Requests[0]) == HalfSize * 2);
 		CHECK(Response.GetSize(Requests[1]) == HalfSize * 2);
@@ -165,6 +182,18 @@ TEST_CASE("CoreUObject::Serialization::DerivedData", "[CoreUObject][Serializatio
 		CHECK(Response.GetSize(Requests[3]) == HalfSize);
 		CHECK(Response.GetSize(Requests[4]) == HalfSize);
 		CHECK(Response.GetSize(Requests[5]) == HalfSize);
+	#if WITH_EDITORONLY_DATA
+		CHECK(Response.GetSize(Requests[6]) == HalfSize * 2);
+	#endif
+
+	#if WITH_EDITORONLY_DATA
+		const FCompressedBuffer* Compressed = Response.GetCompressedData(Requests[6]);
+		CHECKED_IF(Compressed)
+		{
+			CHECK(Compressed->GetRawHash() == SourceHash);
+			CHECK(Compressed->GetRawSize() == HalfSize * 2);
+		}
+	#endif
 
 		Response.Reset();
 		CHECK(Response.IsNull());
@@ -194,15 +223,18 @@ TEST_CASE("CoreUObject::Serialization::DerivedData", "[CoreUObject][Serializatio
 
 		FUniqueBuffer Target = FUniqueBuffer::Alloc(HalfSize);
 		FDerivedDataIoResponse Response;
-		FDerivedDataIoRequest Requests[6];
+		TArray<FDerivedDataIoRequest, TInlineAllocator<7>> Requests;
 		{
 			FDerivedDataIoBatch Batch;
-			Requests[0] = Batch.Read(DerivedData, {});
-			Requests[1] = Batch.Cache(DerivedData, {});
-			Requests[2] = Batch.Exists(DerivedData, {});
-			Requests[3] = Batch.Exists(DerivedData, FDerivedDataIoOptions(HalfSize));
-			Requests[4] = Batch.Read(DerivedData, FDerivedDataIoOptions(HalfSize, HalfSize));
-			Requests[5] = Batch.Read(DerivedData, FDerivedDataIoOptions(Target, HalfSize / 2));
+			Requests.Add(Batch.Read(DerivedData, {}));
+			Requests.Add(Batch.Cache(DerivedData, {}));
+			Requests.Add(Batch.Exists(DerivedData, {}));
+			Requests.Add(Batch.Exists(DerivedData, FDerivedDataIoOptions(HalfSize)));
+			Requests.Add(Batch.Read(DerivedData, FDerivedDataIoOptions(HalfSize, HalfSize)));
+			Requests.Add(Batch.Read(DerivedData, FDerivedDataIoOptions(Target, HalfSize / 2)));
+		#if WITH_EDITORONLY_DATA
+			Requests.Add(Batch.Compress(DerivedData));
+		#endif
 			Batch.Dispatch(Response, {}, [&Event] { Event->Trigger(); });
 		}
 		Event->Wait();
@@ -220,6 +252,7 @@ TEST_CASE("CoreUObject::Serialization::DerivedData", "[CoreUObject][Serializatio
 				CHECK(*RequestHash == Source.GetRawHash());
 			}
 			CHECK(Response.GetCacheKey(Request) == nullptr);
+			CHECK(Response.GetCacheValueId(Request) == nullptr);
 		}
 
 		CHECK(Response.GetData(Requests[0]).GetSize() == HalfSize * 2);
@@ -228,6 +261,9 @@ TEST_CASE("CoreUObject::Serialization::DerivedData", "[CoreUObject][Serializatio
 		CHECK(Response.GetData(Requests[3]).IsNull());
 		CHECK(Response.GetData(Requests[4]).GetView().EqualBytes(RawSource.GetView().Right(HalfSize)));
 		CHECK(Response.GetData(Requests[5]).GetView().EqualBytes(RawSource.GetView().Mid(HalfSize / 2, HalfSize)));
+	#if WITH_EDITORONLY_DATA
+		CHECK(Response.GetData(Requests[6]).IsNull());
+	#endif
 
 		CHECK(Response.GetSize(Requests[0]) == HalfSize * 2);
 		CHECK(Response.GetSize(Requests[1]) == HalfSize * 2);
@@ -235,9 +271,127 @@ TEST_CASE("CoreUObject::Serialization::DerivedData", "[CoreUObject][Serializatio
 		CHECK(Response.GetSize(Requests[3]) == HalfSize);
 		CHECK(Response.GetSize(Requests[4]) == HalfSize);
 		CHECK(Response.GetSize(Requests[5]) == HalfSize);
+	#if WITH_EDITORONLY_DATA
+		CHECK(Response.GetSize(Requests[6]) == HalfSize * 2);
+	#endif
+
+	#if WITH_EDITORONLY_DATA
+		const FCompressedBuffer* Compressed = Response.GetCompressedData(Requests[6]);
+		CHECKED_IF(Compressed)
+		{
+			CHECK(Compressed->GetRawHash() == Source.GetRawHash());
+			CHECK(Compressed->GetRawSize() == Source.GetRawSize());
+			CHECK(Compressed->GetCompressedSize() == Source.GetCompressedSize());
+			CHECK(Algo::Compare(Compressed->GetCompressed().GetSegments(), Source.GetCompressed().GetSegments()));
+		}
+	#endif
 	}
 
 #if !PLATFORM_MAC
+	SECTION("CacheRecord")
+	{
+		constexpr uint64 HalfSize = 8 * 1024;
+		const FSharedBuffer RawValue = MakeSharedBuffer(Counter, HalfSize * 2);
+		const FValue Value = FValue::Compress(RawValue, HalfSize);
+		const FCacheKey Key{FCacheBucket(ANSITEXTVIEW("Test")), Value.GetRawHash()};
+		const FValueId ValueId = FValueId::FromName(ANSITEXTVIEW("Value"));
+
+		FCacheRecordBuilder RecordBuilder(Key);
+		RecordBuilder.AddValue(ValueId, Value);
+
+		FRequestOwner PutOwner(EPriority::Blocking);
+		GetCache().Put({{{TEXTVIEW("CacheRecord")}, RecordBuilder.Build()}}, PutOwner);
+		PutOwner.Wait();
+
+		const FDerivedData DerivedData(TEXTVIEW("CacheRecord"), Key, ValueId);
+		const FDerivedData DerivedDataCopy(TEXTVIEW("CacheRecordCopy"), Key, ValueId);
+
+		CHECK_FALSE(DerivedData.IsNull());
+		CHECK(DerivedData);
+		CHECK(DerivedData.GetFlags() == EDerivedDataFlags::Required);
+		CHECK(WriteToString<128>(DerivedData).ToView() == WriteToString<128>(TEXTVIEW("Cache: Key "), Key,
+			TEXTVIEW(" ID "), ValueId, TEXTVIEW(" for CacheRecord")));
+		CHECK(DerivedData.ReferenceEquals(DerivedDataCopy));
+		CHECK(DerivedData.ReferenceHash() == DerivedDataCopy.ReferenceHash());
+		CHECK(DerivedData.GetName() == TEXTVIEW("CacheRecord"));
+
+		FEventRef Event;
+
+		FUniqueBuffer Target = FUniqueBuffer::Alloc(HalfSize);
+		FDerivedDataIoResponse Response;
+		TArray<FDerivedDataIoRequest, TInlineAllocator<7>> Requests;
+		{
+			FDerivedDataIoBatch Batch;
+			Requests.Add(Batch.Read(DerivedData, {}));
+			Requests.Add(Batch.Cache(DerivedData, {}));
+			Requests.Add(Batch.Exists(DerivedData, {}));
+			Requests.Add(Batch.Exists(DerivedData, FDerivedDataIoOptions(HalfSize)));
+			Requests.Add(Batch.Read(DerivedData, FDerivedDataIoOptions(HalfSize, HalfSize)));
+			Requests.Add(Batch.Read(DerivedData, FDerivedDataIoOptions(Target, HalfSize / 2)));
+		#if WITH_EDITORONLY_DATA
+			Requests.Add(Batch.Compress(DerivedData));
+		#endif
+			Batch.Dispatch(Response, {}, [&Event] { Event->Trigger(); });
+		}
+		Event->Wait();
+
+		CHECK(Response.Poll());
+		CHECK(Response.Cancel());
+		CHECK(Response.GetOverallStatus() == EDerivedDataIoStatus::Ok);
+
+		for (const FDerivedDataIoRequest& Request : Requests)
+		{
+			CHECK(Response.GetStatus(Request) == EDerivedDataIoStatus::Ok);
+
+			const FIoHash* RequestHash = Response.GetHash(Request);
+			CHECKED_IF(RequestHash)
+			{
+				CHECK(*RequestHash == Value.GetRawHash());
+			}
+
+			const FCacheKey* CacheKey = Response.GetCacheKey(Request);
+			CHECKED_IF(CacheKey)
+			{
+				CHECK(*CacheKey == Key);
+			}
+
+			const FValueId* CacheValueId = Response.GetCacheValueId(Request);
+			CHECKED_IF(CacheValueId)
+			{
+				CHECK(*CacheValueId == ValueId);
+			}
+		}
+
+		CHECK(Response.GetData(Requests[0]).GetSize() == HalfSize * 2);
+		CHECK(Response.GetData(Requests[1]).IsNull());
+		CHECK(Response.GetData(Requests[2]).IsNull());
+		CHECK(Response.GetData(Requests[3]).IsNull());
+		CHECK(Response.GetData(Requests[4]).GetView().EqualBytes(RawValue.GetView().Right(HalfSize)));
+		CHECK(Response.GetData(Requests[5]).GetView().EqualBytes(RawValue.GetView().Mid(HalfSize / 2, HalfSize)));
+	#if WITH_EDITORONLY_DATA
+		CHECK(Response.GetData(Requests[6]).IsNull());
+	#endif
+
+		CHECK(Response.GetSize(Requests[0]) == HalfSize * 2);
+		CHECK(Response.GetSize(Requests[1]) == HalfSize * 2);
+		CHECK(Response.GetSize(Requests[2]) == HalfSize * 2);
+		CHECK(Response.GetSize(Requests[3]) == HalfSize);
+		CHECK(Response.GetSize(Requests[4]) == HalfSize);
+		CHECK(Response.GetSize(Requests[5]) == HalfSize);
+	#if WITH_EDITORONLY_DATA
+		CHECK(Response.GetSize(Requests[6]) == HalfSize * 2);
+	#endif
+
+	#if WITH_EDITORONLY_DATA
+		const FCompressedBuffer* Compressed = Response.GetCompressedData(Requests[6]);
+		CHECKED_IF(Compressed)
+		{
+			CHECK(Compressed->GetRawHash() == Value.GetRawHash());
+			CHECK(Compressed->GetRawSize() == Value.GetRawSize());
+		}
+	#endif
+	}
+
 	SECTION("CacheValue")
 	{
 		constexpr uint64 HalfSize = 4 * 1024;
@@ -265,15 +419,18 @@ TEST_CASE("CoreUObject::Serialization::DerivedData", "[CoreUObject][Serializatio
 
 		FUniqueBuffer Target = FUniqueBuffer::Alloc(HalfSize);
 		FDerivedDataIoResponse Response;
-		FDerivedDataIoRequest Requests[6];
+		TArray<FDerivedDataIoRequest, TInlineAllocator<7>> Requests;
 		{
 			FDerivedDataIoBatch Batch;
-			Requests[0] = Batch.Read(DerivedData, {});
-			Requests[1] = Batch.Cache(DerivedData, {});
-			Requests[2] = Batch.Exists(DerivedData, {});
-			Requests[3] = Batch.Exists(DerivedData, FDerivedDataIoOptions(HalfSize));
-			Requests[4] = Batch.Read(DerivedData, FDerivedDataIoOptions(HalfSize, HalfSize));
-			Requests[5] = Batch.Read(DerivedData, FDerivedDataIoOptions(Target, HalfSize / 2));
+			Requests.Add(Batch.Read(DerivedData, {}));
+			Requests.Add(Batch.Cache(DerivedData, {}));
+			Requests.Add(Batch.Exists(DerivedData, {}));
+			Requests.Add(Batch.Exists(DerivedData, FDerivedDataIoOptions(HalfSize)));
+			Requests.Add(Batch.Read(DerivedData, FDerivedDataIoOptions(HalfSize, HalfSize)));
+			Requests.Add(Batch.Read(DerivedData, FDerivedDataIoOptions(Target, HalfSize / 2)));
+		#if WITH_EDITORONLY_DATA
+			Requests.Add(Batch.Compress(DerivedData));
+		#endif
 			Batch.Dispatch(Response, {}, [&Event] { Event->Trigger(); });
 		}
 		Event->Wait();
@@ -285,12 +442,20 @@ TEST_CASE("CoreUObject::Serialization::DerivedData", "[CoreUObject][Serializatio
 		for (const FDerivedDataIoRequest& Request : Requests)
 		{
 			CHECK(Response.GetStatus(Request) == EDerivedDataIoStatus::Ok);
+
 			const FIoHash* RequestHash = Response.GetHash(Request);
 			CHECKED_IF(RequestHash)
 			{
 				CHECK(*RequestHash == Value.GetRawHash());
 			}
-			CHECK(Response.GetCacheKey(Request) == nullptr);
+
+			const FCacheKey* CacheKey = Response.GetCacheKey(Request);
+			CHECKED_IF(CacheKey)
+			{
+				CHECK(*CacheKey == Key);
+			}
+
+			CHECK(Response.GetCacheValueId(Request) == nullptr);
 		}
 
 		CHECK(Response.GetData(Requests[0]).GetSize() == HalfSize * 2);
@@ -299,6 +464,9 @@ TEST_CASE("CoreUObject::Serialization::DerivedData", "[CoreUObject][Serializatio
 		CHECK(Response.GetData(Requests[3]).IsNull());
 		CHECK(Response.GetData(Requests[4]).GetView().EqualBytes(RawValue.GetView().Right(HalfSize)));
 		CHECK(Response.GetData(Requests[5]).GetView().EqualBytes(RawValue.GetView().Mid(HalfSize / 2, HalfSize)));
+	#if WITH_EDITORONLY_DATA
+		CHECK(Response.GetData(Requests[6]).IsNull());
+	#endif
 
 		CHECK(Response.GetSize(Requests[0]) == HalfSize * 2);
 		CHECK(Response.GetSize(Requests[1]) == HalfSize * 2);
@@ -306,6 +474,18 @@ TEST_CASE("CoreUObject::Serialization::DerivedData", "[CoreUObject][Serializatio
 		CHECK(Response.GetSize(Requests[3]) == HalfSize);
 		CHECK(Response.GetSize(Requests[4]) == HalfSize);
 		CHECK(Response.GetSize(Requests[5]) == HalfSize);
+	#if WITH_EDITORONLY_DATA
+		CHECK(Response.GetSize(Requests[6]) == HalfSize * 2);
+	#endif
+
+	#if WITH_EDITORONLY_DATA
+		const FCompressedBuffer* Compressed = Response.GetCompressedData(Requests[6]);
+		CHECKED_IF(Compressed)
+		{
+			CHECK(Compressed->GetRawHash() == Value.GetRawHash());
+			CHECK(Compressed->GetRawSize() == Value.GetRawSize());
+		}
+	#endif
 	}
 #endif // !PLATFORM_MAC
 #endif // WITH_EDITORONLY_DATA
