Defer Blueprint CDO export creation on load so that placeholder dependencies can be resolved prior to invoking the native base class ctor.

Notes:
- This change is needed to supplement the feature that allows users to set a subobject override class from the Blueprint editor; i.e. overriding a component class with a non-native Blueprint component class will otherwise fail on reload.

Change summary:
- Modified FLinkerLoad::CreateExport() to move the deferred CDO serialization check to occur prior to CDO construction (part of the linker's circular dependency mitigation code path). This means we now defer its construction in addition to its serialization.
- FLinkerLoad::CreateExport() will continue to bypass class regeneration post-construction when Blueprint finalization is still pending. This is because CDO construction, while now deferred until finalization, still goes through CreateExport() for consistency.
- Modified FLinkerLoad::ResolveDeferredExports() to handle the deferred construction of the CDO during Blueprint finalization, prior to its deferred serialization.
- Modified UClass::Serialize() to bypass the NULL check for failed CDO serialization when deferred dependency loading is active and Blueprint finalization is still pending.

diff --git a/Engine/Source/Runtime/CoreUObject/Private/Blueprint/BlueprintSupport.cpp b/Engine/Source/Runtime/CoreUObject/Private/Blueprint/BlueprintSupport.cpp
--- a/Engine/Source/Runtime/CoreUObject/Private/Blueprint/BlueprintSupport.cpp
+++ b/Engine/Source/Runtime/CoreUObject/Private/Blueprint/BlueprintSupport.cpp
@@ -1985,6 +1985,16 @@ void FLinkerLoad::ResolveDeferredExports(UClass* LoadClass)
 
 	DEFERRED_DEPENDENCY_CHECK(DeferredCDOIndex != INDEX_NONE || bForceBlueprintFinalization);
 
+	// Handle deferred construction of the CDO and patch it into the export table. Any deferred ctor
+	// initializer dependencies (e.g. subobject class overrides) should now be resolved at this point.
+	if (DeferredCDOIndex != INDEX_NONE && !ExportMap[DeferredCDOIndex].Object)
+	{
+		// Note: We could just call GetDefaultObject() here, but then we'd also need to set object
+		// flags on it to ensure that it gets serialized later. For consistency/safety, this routes
+		// through CreateExport() so that we don't have to worry about keeping object flags in sync.
+		CreateExport(DeferredCDOIndex);
+	}
+
 	UObject* BlueprintCDO = DeferredCDOIndex != INDEX_NONE ? ExportMap[DeferredCDOIndex].Object : LoadClass->ClassDefaultObject;
 	DEFERRED_DEPENDENCY_CHECK(BlueprintCDO != nullptr);
 	
diff --git a/Engine/Source/Runtime/CoreUObject/Private/UObject/Class.cpp b/Engine/Source/Runtime/CoreUObject/Private/UObject/Class.cpp
--- a/Engine/Source/Runtime/CoreUObject/Private/UObject/Class.cpp
+++ b/Engine/Source/Runtime/CoreUObject/Private/UObject/Class.cpp
@@ -5136,7 +5136,7 @@ void UClass::Serialize( FArchive& Ar )
 
 	if( Ar.IsLoading() )
 	{
-		if (ClassDefaultObject == NULL)
+		if (ClassDefaultObject == nullptr)
 		{
 			check(GConfig);
 			if (GEventDrivenLoaderEnabled || Ar.IsLoadingFromCookedPackage())
@@ -5147,10 +5147,20 @@ void UClass::Serialize( FArchive& Ar )
 			}
 			else if( !Ar.HasAnyPortFlags(PPF_DuplicateForPIE|PPF_Duplicate) )
 			{
-				UE_LOG(LogClass, Error, TEXT("CDO for class %s did not load!"), *GetPathName());
-				ensure(ClassDefaultObject != NULL);
-				ClassDefaultObject = GetDefaultObject();
-				Ar.ForceBlueprintFinalization();
+#if USE_CIRCULAR_DEPENDENCY_LOAD_DEFERRING
+				// Blueprint CDO serialization is deferred (see FLinkerLoad::CreateExport).
+				const FLinkerLoad* LinkerLoad = Cast<FLinkerLoad>(Ar.GetLinker());
+				const bool bHasDeferredCDOSerialization = HasAnyClassFlags(CLASS_CompiledFromBlueprint)
+					&& GetClass()->HasAnyClassFlags(CLASS_NeedsDeferredDependencyLoading)
+					&& LinkerLoad && LinkerLoad->IsBlueprintFinalizationPending();
+				if (!bHasDeferredCDOSerialization)
+#endif	// USE_CIRCULAR_DEPENDENCY_LOAD_DEFERRING
+				{
+					UE_LOG(LogClass, Error, TEXT("CDO for class %s did not load!"), *GetPathName());
+					ensure(ClassDefaultObject != nullptr);
+					ClassDefaultObject = GetDefaultObject();
+					Ar.ForceBlueprintFinalization();
+				}
 			}
 		}
 	}
diff --git a/Engine/Source/Runtime/CoreUObject/Private/UObject/LinkerLoad.cpp b/Engine/Source/Runtime/CoreUObject/Private/UObject/LinkerLoad.cpp
--- a/Engine/Source/Runtime/CoreUObject/Private/UObject/LinkerLoad.cpp
+++ b/Engine/Source/Runtime/CoreUObject/Private/UObject/LinkerLoad.cpp
@@ -4923,6 +4923,24 @@ UObject* FLinkerLoad::CreateExport( int32 Index )
 			}
 		}
 
+		const bool bIsBlueprintCDO = ((Export.ObjectFlags & RF_ClassDefaultObject) != 0) && LoadClass->HasAnyClassFlags(CLASS_CompiledFromBlueprint) &&
+			LoadClass->GetClass()->HasAnyClassFlags(CLASS_NeedsDeferredDependencyLoading);
+
+#if USE_CIRCULAR_DEPENDENCY_LOAD_DEFERRING
+		const bool bDeferCDOSerialization = bIsBlueprintCDO && ((LoadFlags & LOAD_DeferDependencyLoads) != 0);
+		if (bDeferCDOSerialization)
+		{
+			// if LOAD_DeferDependencyLoads is set, then we're already
+			// serializing the blueprint's class somewhere up the chain... 
+			// we don't want the class regenerated while it in the middle of
+			// serializing. we also don't want to construct the CDO yet,
+			// as that may depend on other deferred imports which may not
+			// be fully resolved (e.g. default subobject class overrides).
+			DeferredCDOIndex = Index;
+			return Export.Object;
+		}
+#endif // USE_CIRCULAR_DEPENDENCY_LOAD_DEFERRING
+
 		LoadClass->GetDefaultObject();
 
 		FStaticConstructObjectParameters Params(LoadClass);
@@ -4954,18 +4972,10 @@ UObject* FLinkerLoad::CreateExport( int32 Index )
 		
 		if( Export.Object )
 		{
-			bool const bIsBlueprintCDO = ((Export.ObjectFlags & RF_ClassDefaultObject) != 0) && LoadClass->HasAnyClassFlags(CLASS_CompiledFromBlueprint) &&
-				LoadClass->GetClass()->HasAnyClassFlags(CLASS_NeedsDeferredDependencyLoading);
-
 #if USE_CIRCULAR_DEPENDENCY_LOAD_DEFERRING
-			const bool bDeferCDOSerialization = bIsBlueprintCDO && ((LoadFlags & LOAD_DeferDependencyLoads) != 0);
-			if (bDeferCDOSerialization)			
+			if (bIsBlueprintCDO && IsBlueprintFinalizationPending())
 			{
-				// if LOAD_DeferDependencyLoads is set, then we're already
-				// serializing the blueprint's class somewhere up the chain... 
-				// we don't want the class regenerated while it in the middle of
-				// serializing
-				DeferredCDOIndex = Index;
+				// Class regeneration is deferred until Blueprint finalization, so just return the CDO.
 				return Export.Object;
 			}
 			else 
