Replaced identical code in TCppStructOps to help older clang fold it.

diff --git a/Engine/Source/Runtime/CoreUObject/Public/UObject/Class.h b/Engine/Source/Runtime/CoreUObject/Public/UObject/Class.h
--- a/Engine/Source/Runtime/CoreUObject/Public/UObject/Class.h
+++ b/Engine/Source/Runtime/CoreUObject/Public/UObject/Class.h
@@ -1065,22 +1065,7 @@ public:
 		}
 		virtual void ConstructForTests(void* Dest) override
 		{
-			check(!TTraits::WithZeroConstructor); // don't call this if we have indicated it is not necessary
-			// that could have been an if statement, but we might as well force optimization above the virtual call
-			// could also not attempt to call the constructor for types where this is not possible, but I didn't do that here
-#if CHECK_PUREVIRTUALS
-			if constexpr (!TStructOpsTypeTraits<CPPSTRUCT>::WithPureVirtual)
-#endif
-			{
-				if constexpr (TStructOpsTypeTraits<CPPSTRUCT>::WithNoInitConstructor)
-				{
-					new (Dest) CPPSTRUCT(ForceInit);
-				}
-				else
-				{
-					new (Dest) CPPSTRUCT;
-				}
-			}
+			Construct(Dest);
 		}
 		virtual bool HasDestructor() override
 		{
