Report CustomVersion warnings/errors to the MessageLog

#rnx

- This will allow the messages to show up when the user clicks the link on the "Failed to load assets" notification in the editor. If the warnings/errors are only printed via UE_LOG, then no messages will be displayed when the user clicks that link which can be confusing.

diff --git a/Engine/Source/Runtime/CoreUObject/Private/UObject/LinkerLoad.cpp b/Engine/Source/Runtime/CoreUObject/Private/UObject/LinkerLoad.cpp
--- a/Engine/Source/Runtime/CoreUObject/Private/UObject/LinkerLoad.cpp
+++ b/Engine/Source/Runtime/CoreUObject/Private/UObject/LinkerLoad.cpp
@@ -1463,14 +1463,27 @@ FLinkerLoad::ELinkerStatus FLinkerLoad::UpdateFromPackageFileSummary()
 			else if (Diff.Type == ECustomVersionDifference::Invalid)
 			{
 				UE_ASSET_LOG(LogLinker, Error, PackagePath, TEXT("Package was saved with an invalid custom version. Tag %s  Version %d"), *Diff.Version->Key.ToString(), Diff.Version->Version);
+
+				FMessageLog("LoadErrors")
+					.SuppressLoggingToOutputLog(true)
+					.Error(FText::Format(NSLOCTEXT("Core", "LinkerLoad_InvalidCustomVersion", "Package {0} was saved with an invalid custom version and cannot be loaded, see output log for details"),
+						FText::FromString(GetDebugName())));
+
 				return LINKER_Failed;
 			}
 			else if (Diff.Type == ECustomVersionDifference::Newer)
 			{
 				FCustomVersion LatestVersion = FCurrentCustomVersions::Get(Diff.Version->Key).GetValue();
+				
 				// Loading a package with a newer custom version than the current one.
 				UE_ASSET_LOG(LogLinker, Error, PackagePath, TEXT("Package was saved with a newer custom version than the current. Tag %s Name '%s' PackageVersion %d  MaxExpected %d"),
 					*Diff.Version->Key.ToString(), *LatestVersion.GetFriendlyName().ToString(), Diff.Version->Version, LatestVersion.Version);
+
+				FMessageLog("LoadErrors")
+					.SuppressLoggingToOutputLog(true)
+					.Error(FText::Format(NSLOCTEXT("Core", "LinkerLoad_NewCustomVersion", "Package {0} was saved with a newer custom version than the current engine and cannot be loaded, see output log for details"),
+						FText::FromString(GetDebugName())));
+
 				return LINKER_Failed;
 			}
 		}
