TopLevelAssetPath: Fix constructors to allow UObject subclasses to use the UObject* constructor instead of trying and failing to use the CharType* constructor.
Allow UObject constructor to record a valid path to non-asset toplevelobjects, for consistency with the other constructors which are based on path strings rather than the virtual IsAsset function (which e.g. returns false for UClass).
#rnx

diff --git a/Engine/Source/Runtime/CoreUObject/Private/UObject/TopLevelAssetPath.cpp b/Engine/Source/Runtime/CoreUObject/Private/UObject/TopLevelAssetPath.cpp
--- a/Engine/Source/Runtime/CoreUObject/Private/UObject/TopLevelAssetPath.cpp
+++ b/Engine/Source/Runtime/CoreUObject/Private/UObject/TopLevelAssetPath.cpp
@@ -19,29 +19,6 @@
 UE_IMPLEMENT_STRUCT("/Script/CoreUObject", TopLevelAssetPath)
 #endif
 
-FTopLevelAssetPath::FTopLevelAssetPath(const UObject* InObject)
-{
-	if (InObject == nullptr)
-	{
-		PackageName = AssetName = FName();
-	}
-	else if (!InObject->GetOuter())
-	{
-		check(Cast<UPackage>(InObject) != nullptr);
-		PackageName = InObject->GetFName();
-		AssetName = FName();
-	}
-	else if (!InObject->IsAsset() || InObject->GetOuter()->GetOuter() != nullptr)
-	{
-		PackageName = AssetName = FName();
-	}
-	else
-	{
-		PackageName = InObject->GetOuter()->GetFName();
-		AssetName = InObject->GetFName();
-	}
-}
-
 void FTopLevelAssetPath::AppendString(FStringBuilderBase& Builder) const
 {
 	if (!IsNull())
@@ -87,6 +64,33 @@ bool FTopLevelAssetPath::TrySetPath(FName InPackageName, FName InAssetName)
 	return !PackageName.IsNone();
 }
 
+bool FTopLevelAssetPath::TrySetPath(const UObject* InObject)
+{
+	if (InObject == nullptr)
+	{
+		Reset();
+		return false;
+	}
+	else if (!InObject->GetOuter())
+	{
+		check(Cast<UPackage>(InObject) != nullptr);
+		PackageName = InObject->GetFName();
+		AssetName = FName();
+		return true;
+	}
+	else if (InObject->GetOuter()->GetOuter() != nullptr)
+	{
+		Reset();
+		return false;
+	}
+	else
+	{
+		PackageName = InObject->GetOuter()->GetFName();
+		AssetName = InObject->GetFName();
+		return true;
+	}
+}
+
 bool FTopLevelAssetPath::TrySetPath(FWideStringView Path)
 {
 	if (Path.IsEmpty() || Path.Equals(TEXT("None"), ESearchCase::CaseSensitive))
diff --git a/Engine/Source/Runtime/CoreUObject/Public/UObject/TopLevelAssetPath.h b/Engine/Source/Runtime/CoreUObject/Public/UObject/TopLevelAssetPath.h
--- a/Engine/Source/Runtime/CoreUObject/Public/UObject/TopLevelAssetPath.h
+++ b/Engine/Source/Runtime/CoreUObject/Public/UObject/TopLevelAssetPath.h
@@ -11,8 +11,11 @@
 #include "Serialization/StructuredArchiveNameHelpers.h"
 #include "Serialization/StructuredArchiveSlots.h"
 #include "Templates/TypeHash.h"
+#include "Traits/IsCharType.h"
 #include "UObject/NameTypes.h"
 
+#include <type_traits>
+
 class CORE_API FString;
 class UObject;
 
@@ -46,29 +49,31 @@ struct COREUOBJECT_API FTopLevelAssetPath
 
 	/** Construct from string / string view / raw string of a supported character type. */
 	explicit FTopLevelAssetPath(const FString& Path) { TrySetPath(FStringView(Path)); }
-	template<typename CharType>
+	template<typename CharType, typename = typename std::enable_if<TIsCharType<CharType>::Value>::type>
 	explicit FTopLevelAssetPath(TStringView<CharType> Path) { TrySetPath(Path); }
-	template<typename CharType>
+	template<typename CharType, typename = typename std::enable_if<TIsCharType<CharType>::Value>::type>
 	explicit FTopLevelAssetPath(const CharType* Path) { TrySetPath(TStringView<CharType>(Path)); }
 
-	/** Construct from an existing object in memory */
-	explicit FTopLevelAssetPath(const UObject* InObject);
+	/** Construct from an existing object. Creates an empty path if the object is not a package or the direct inner of a package. */
+	explicit FTopLevelAssetPath(const UObject* InObject) { TrySetPath(InObject); }
 
 	/** Assign from the same types we can construct from */
 	FTopLevelAssetPath& operator=(const FString& Path) { TrySetPath(FStringView(Path)); return *this; }
 	template<typename CharType>
 	FTopLevelAssetPath& operator=(TStringView<CharType> Path) { TrySetPath(Path); return *this; }
-	template<typename CharType>
+	template<typename CharType, typename = typename std::enable_if<TIsCharType<CharType>::Value>::type>
 	FTopLevelAssetPath& operator=(const CharType* Path) { TrySetPath(TStringView<CharType>(Path)); return *this; }
 	FTopLevelAssetPath& operator=(TYPE_OF_NULLPTR) { Reset(); return *this; }
 
 	/** Sets asset path of this reference based on components. */
 	bool TrySetPath(FName InPackageName, FName InAssetName);
+	/** Sets asset path to path of existing object. Sets an empty path and returns false if the object is not a package or the direct inner of a package. */
+	bool TrySetPath(const UObject* InObject);
 	/** Sets asset path of this reference based on a string path. Resets this object and returns false if the string is empty or does not represent a top level asset path. */
 	bool TrySetPath(FWideStringView Path);
 	bool TrySetPath(FUtf8StringView Path);
 	bool TrySetPath(FAnsiStringView Path);
-	template<typename CharType>
+	template<typename CharType,	typename = typename std::enable_if<TIsCharType<CharType>::Value>::type>
 	bool TrySetPath(const CharType* Path) { return TrySetPath(TStringView<CharType>(Path)); }
 	bool TrySetPath(const FString& Path) { return TrySetPath(FStringView(Path)); }
 
