Set the external package when linker load finds the package in memory.

diff --git a/Engine/Source/Runtime/CoreUObject/Private/UObject/LinkerLoad.cpp b/Engine/Source/Runtime/CoreUObject/Private/UObject/LinkerLoad.cpp
--- a/Engine/Source/Runtime/CoreUObject/Private/UObject/LinkerLoad.cpp
+++ b/Engine/Source/Runtime/CoreUObject/Private/UObject/LinkerLoad.cpp
@@ -4804,6 +4804,10 @@ UObject* FLinkerLoad::CreateExport( int32 Index )
 			else
 			{
 				Export.Object->SetLinker(this, Index);
+				if (Export.OuterIndex.IsImport())
+				{
+					Export.Object->SetExternalPackage(LinkerRoot);
+				}
 
 				// If this object was allocated but never loaded (components created by a constructor) make sure it gets loaded
 				// Don't do this for any packages that have previously fully loaded as they may have in memory changes
