Made FAssetData ~16 bytes smaller (technically 20, but there's an odd number of 32-bit structures now, so we'll shrink 8 when deprecated AssetClass is deleted). This saves about 5MB if the project has 300K+ assets.

There's a lot of duplication of very small arrays by storing the inline array in each FAssetData entry, and this isn't used in a critical path anyway, so this changes ChunkIDs to be an index/reference into a chunk array registry.

diff --git a/Engine/Plugins/Editor/AssetManagerEditor/Source/AssetManagerEditor/Private/AssetManagerEditorModule.cpp b/Engine/Plugins/Editor/AssetManagerEditor/Source/AssetManagerEditor/Private/AssetManagerEditorModule.cpp
--- a/Engine/Plugins/Editor/AssetManagerEditor/Source/AssetManagerEditor/Private/AssetManagerEditorModule.cpp
+++ b/Engine/Plugins/Editor/AssetManagerEditor/Source/AssetManagerEditor/Private/AssetManagerEditorModule.cpp
@@ -1215,14 +1215,14 @@ bool FAssetManagerEditorModule::GetStringValueForCustomColumn(const FAssetData&
 		if (CurrentRegistrySource->bIsEditor)
 		{
 			// The in-memory data is wrong, ask the asset manager
-			AssetManager.GetPackageChunkIds(AssetData.PackageName, CurrentRegistrySource->TargetPlatform, AssetData.ChunkIDs, FoundChunks);
+			AssetManager.GetPackageChunkIds(AssetData.PackageName, CurrentRegistrySource->TargetPlatform, AssetData.GetChunkIDs(), FoundChunks);
 		}
 		else
 		{
 			FAssetData PlatformData = CurrentRegistrySource->GetAssetByObjectPath(AssetData.GetSoftObjectPath());
 			if (PlatformData.IsValid())
 			{
-				FoundChunks = MoveTemp(PlatformData.ChunkIDs);
+				FoundChunks = PlatformData.GetChunkIDs();
 			}
 		}
 		
@@ -1625,12 +1625,13 @@ void FAssetManagerEditorModule::SetCurrentRegistrySource(const FString& SourceNa
 
 				RegistryState->EnumerateAllAssets([&RegistryState, this](const FAssetData& AssetData)
 				{
-					if (AssetData.ChunkIDs.Num() > 0)
+					const FAssetData::FChunkArrayView ChunkIDs = AssetData.GetChunkIDs();
+					if (!ChunkIDs.IsEmpty())
 					{
 						TArray<FAssetIdentifier> ManagerAssets;
 						RegistryState->GetReferencers(AssetData.PackageName, ManagerAssets, UE::AssetRegistry::EDependencyCategory::Manage);
 
-						for (int32 ChunkId : AssetData.ChunkIDs)
+						for (int32 ChunkId : ChunkIDs)
 						{
 							FPrimaryAssetId ChunkAssetId = UAssetManager::CreatePrimaryAssetIdFromChunkId(ChunkId);
 							
diff --git a/Engine/Source/Editor/UnrealEd/Private/Commandlets/AssetRegistryGenerator.cpp b/Engine/Source/Editor/UnrealEd/Private/Commandlets/AssetRegistryGenerator.cpp
--- a/Engine/Source/Editor/UnrealEd/Private/Commandlets/AssetRegistryGenerator.cpp
+++ b/Engine/Source/Editor/UnrealEd/Private/Commandlets/AssetRegistryGenerator.cpp
@@ -982,14 +982,14 @@ void FAssetRegistryGenerator::InjectEncryptionData(FAssetRegistryState& TargetSt
 					if (AssetData)
 					{
 						FString GuidString;
-
-						if (AssetData->ChunkIDs.Num() > 1)
+						const FAssetData::FChunkArrayView ChunkIDs = AssetData->GetChunkIDs();
+						if (ChunkIDs.Num() > 1)
 						{
 							UE_LOG(LogAssetRegistryGenerator, Error, TEXT("Encrypted root asset '%s' exists in two chunks. Only secondary assets should be shared between chunks."), *AssetData->GetObjectPathString());
 						}
-						else if (AssetData->ChunkIDs.Num() == 1)
+						else if (ChunkIDs.Num() == 1)
 						{
-							int32 ChunkID = AssetData->ChunkIDs[0];
+							int32 ChunkID = ChunkIDs[0];
 							FGuid Guid;
 
 							if (GuidCache.Contains(ChunkID))
@@ -1005,7 +1005,7 @@ void FAssetRegistryGenerator::InjectEncryptionData(FAssetRegistryState& TargetSt
 							{
 								FAssetDataTagMap TagsAndValues = AssetData->TagsAndValues.CopyMap();
 								TagsAndValues.Add(UAssetManager::GetEncryptionKeyAssetTagName(), Guid.ToString());
-								FAssetData NewAssetData = FAssetData(AssetData->PackageName, AssetData->PackagePath, AssetData->AssetName, AssetData->AssetClassPath, TagsAndValues, AssetData->ChunkIDs, AssetData->PackageFlags);
+								FAssetData NewAssetData = FAssetData(AssetData->PackageName, AssetData->PackagePath, AssetData->AssetName, AssetData->AssetClassPath, TagsAndValues, ChunkIDs, AssetData->PackageFlags);
 								NewAssetData.TaggedAssetBundles = AssetData->TaggedAssetBundles;
 								TargetState.UpdateAssetData(AssetData, MoveTemp(NewAssetData));
 							}
@@ -1377,7 +1377,7 @@ void FAssetRegistryGenerator::FinalizeChunkIDs(const TSet<FName>& InCookedPackag
 	// Copy ExplicitChunkIDs and other data from the AssetRegistry into the maps we use during finalization
 	State.EnumerateAllAssets([&](const FAssetData& AssetData)
 	{
-		for (int32 ChunkID : AssetData.ChunkIDs)
+		for (int32 ChunkID : AssetData.GetChunkIDs())
 		{
 			if (ChunkID < 0)
 			{
@@ -1390,7 +1390,7 @@ void FAssetRegistryGenerator::FinalizeChunkIDs(const TSet<FName>& InCookedPackag
 
 		// Clear the Asset's chunk id list. We will fill it with the final IDs to use later on.
 		// Chunk Ids are safe to modify in place so do a const cast
-		const_cast<FAssetData&>(AssetData).ChunkIDs.Empty();
+		const_cast<FAssetData&>(AssetData).ClearChunkIDs();
 
 		// Update whether the owner package contains a map
 		if ((AssetData.PackageFlags & PKG_ContainsMap) != 0)
@@ -2008,9 +2008,10 @@ bool FAssetRegistryGenerator::GenerateAssetChunkInformationCSV(const FString& Ou
 		const FAssetPackageData* PackageData = State.GetAssetPackageData(AssetData.PackageName);
 
 		// Add only assets that have actually been cooked and belong to any chunk and that have a file size
-		if (PackageData != nullptr && AssetData.ChunkIDs.Num() > 0 && PackageData->DiskSize > 0)
+		const FAssetData::FChunkArrayView ChunkIDs = AssetData.GetChunkIDs();
+		if (PackageData != nullptr && ChunkIDs.Num() > 0 && PackageData->DiskSize > 0)
 		{
-			for (int32 PakchunkIndex : AssetData.ChunkIDs)
+			for (int32 PakchunkIndex : ChunkIDs)
 			{
 				const int64 FileSize = PackageData->DiskSize;
 				FString SoftChain;
@@ -2031,7 +2032,7 @@ bool FAssetRegistryGenerator::GenerateAssetChunkInformationCSV(const FString& Ou
 
 				// Build "other chunks" string or None if not part of
 				TmpStringChunks.Empty(64);
-				for (int32 OtherChunk : AssetData.ChunkIDs)
+				for (int32 OtherChunk : ChunkIDs)
 				{
 					if (OtherChunk != PakchunkIndex)
 					{
@@ -2046,7 +2047,7 @@ bool FAssetRegistryGenerator::GenerateAssetChunkInformationCSV(const FString& Ou
 					*AssetData.AssetClassPath.ToString(),
 					bHardChunk ? TEXT("Hard") : *SoftChain,
 					FileSize,
-					AssetData.ChunkIDs.Num() == 1 ? TEXT("None") : *TmpStringChunks
+					ChunkIDs.Num() == 1 ? TEXT("None") : *TmpStringChunks
 				);
 
 				// Write line to all chunks file and individual chunks files if requested
@@ -2323,7 +2324,7 @@ void FAssetRegistryGenerator::FixupPackageDependenciesForChunks(FSandboxPlatform
 			for (const FAssetData* AssetData : State.GetAssetsByPackageName(Asset.Key))
 			{
 				// Chunk Ids are safe to modify in place
-				const_cast<FAssetData*>(AssetData)->ChunkIDs.AddUnique(PakchunkIndex);
+				const_cast<FAssetData*>(AssetData)->AddChunkID(PakchunkIndex);
 			}
 		}
 	}
diff --git a/Engine/Source/Editor/UnrealEd/Private/Commandlets/DiffAssetRegistriesCommandlet.cpp b/Engine/Source/Editor/UnrealEd/Private/Commandlets/DiffAssetRegistriesCommandlet.cpp
--- a/Engine/Source/Editor/UnrealEd/Private/Commandlets/DiffAssetRegistriesCommandlet.cpp
+++ b/Engine/Source/Editor/UnrealEd/Private/Commandlets/DiffAssetRegistriesCommandlet.cpp
@@ -513,10 +513,13 @@ bool	UDiffAssetRegistriesCommandlet::IsInRelevantChunk(FAssetRegistryState& InRe
 
 	}
 	TArrayView<FAssetData const* const> Assets = InRegistryState.GetAssetsByPackageName(InAssetPath);
-
-	if (Assets.Num() && Assets[0]->ChunkIDs.Num())
+	if (!Assets.IsEmpty())
 	{
-		return Assets[0]->ChunkIDs.Contains(DiffChunkID);
+		const FAssetData::FChunkArrayView ChunkIDs = Assets[0]->GetChunkIDs();
+		if (!ChunkIDs.IsEmpty())
+		{
+			return ChunkIDs.Contains(DiffChunkID);
+		}
 	}
 
 	return true;
@@ -556,15 +559,15 @@ TArray<int32> UDiffAssetRegistriesCommandlet::GetAssetChunks(FAssetRegistryState
 	if (ChunkIdByAssetPath.Contains(InAssetPath) == false)
 	{
 		TArrayView<FAssetData const* const> Assets = InRegistryState.GetAssetsByPackageName(InAssetPath);
-
-		if (Assets.Num() > 0 && Assets[0]->ChunkIDs.Num() > 0)
+		const FAssetData::FChunkArrayView ChunkIDs = Assets.IsEmpty() ? FAssetData::FChunkArrayView() : Assets[0]->GetChunkIDs();
+		if (!ChunkIDs.IsEmpty())
 		{
-			if (Assets[0]->ChunkIDs.Num() > 1)
+			if (ChunkIDs.Num() > 1)
 			{
 				UE_LOG(LogDiffAssets, Log, TEXT("Multiple ChunkIds for asset %s"), *InAssetPath.ToString());
 			}
 
-			for (int32 id : Assets[0]->ChunkIDs)
+			for (int32 id : ChunkIDs)
 			{
 				ChangesByChunk.FindOrAdd(id).IncludedAssets.Add(InAssetPath);
 				ChunkIdByAssetPath.Add(InAssetPath, id);
diff --git a/Engine/Source/Runtime/AssetRegistry/Private/AssetDataGatherer.cpp b/Engine/Source/Runtime/AssetRegistry/Private/AssetDataGatherer.cpp
--- a/Engine/Source/Runtime/AssetRegistry/Private/AssetDataGatherer.cpp
+++ b/Engine/Source/Runtime/AssetRegistry/Private/AssetDataGatherer.cpp
@@ -3953,9 +3953,9 @@ SIZE_T FAssetDataGatherer::GetAllocatedSize() const
 	for (FAssetData* Value : AssetResults)
 	{
 		Result += sizeof(*Value);
-		Result += Value->ChunkIDs.GetAllocatedSize();
 		TagMemoryUsage.Include(Value->TagsAndValues);
 	}
+	Result += FAssetData::GetChunkArrayRegistryAllocatedSize();
 	Result += TagMemoryUsage.GetFixedSize() + TagMemoryUsage.GetLooseSize();
 
 	Result += GetArrayRecursiveAllocatedSize(DependencyResults);
diff --git a/Engine/Source/Runtime/AssetRegistry/Private/AssetRegistry.cpp b/Engine/Source/Runtime/AssetRegistry/Private/AssetRegistry.cpp
--- a/Engine/Source/Runtime/AssetRegistry/Private/AssetRegistry.cpp
+++ b/Engine/Source/Runtime/AssetRegistry/Private/AssetRegistry.cpp
@@ -3048,7 +3048,7 @@ EAssetAvailability::Type GetAssetAvailability(const FAssetData& AssetData)
 	EChunkLocation::Type BestLocation = EChunkLocation::DoesNotExist;
 
 	// check all chunks to see which has the best locality
-	for (int32 PakchunkId : AssetData.ChunkIDs)
+	for (int32 PakchunkId : AssetData.GetChunkIDs())
 	{
 		EChunkLocation::Type ChunkLocation = ChunkInstall->GetPakchunkLocation(PakchunkId);
 
@@ -3102,7 +3102,7 @@ float GetAssetAvailabilityProgress(const FAssetData& AssetData, EAssetAvailabili
 	float BestProgress = MAX_FLT;
 
 	// check all chunks to see which has the best time remaining
-	for (int32 PakchunkID : AssetData.ChunkIDs)
+	for (int32 PakchunkID : AssetData.GetChunkIDs())
 	{
 		float Progress = ChunkInstall->GetChunkProgress(PakchunkID, ChunkReportType);
 
@@ -3162,12 +3162,13 @@ void PrioritizeAssetInstall(const FAssetData& AssetData)
 {
 	IPlatformChunkInstall* ChunkInstall = FPlatformMisc::GetPlatformChunkInstall();
 
-	if (AssetData.ChunkIDs.Num() == 0)
+	const TConstArrayView<int32> ChunkIDs = AssetData.GetChunkIDs();
+	if (ChunkIDs.Num() == 0)
 	{
 		return;
 	}
 
-	ChunkInstall->PrioritizePakchunk(AssetData.ChunkIDs[0], EChunkPriority::Immediate);
+	ChunkInstall->PrioritizePakchunk(ChunkIDs[0], EChunkPriority::Immediate);
 }
 
 }
diff --git a/Engine/Source/Runtime/AssetRegistry/Private/AssetRegistryState.cpp b/Engine/Source/Runtime/AssetRegistry/Private/AssetRegistryState.cpp
--- a/Engine/Source/Runtime/AssetRegistry/Private/AssetRegistryState.cpp
+++ b/Engine/Source/Runtime/AssetRegistry/Private/AssetRegistryState.cpp
@@ -2,6 +2,7 @@
 
 #include "AssetRegistry/AssetRegistryState.h"
 
+#include "Algo/Compare.h"
 #include "Algo/Sort.h"
 #include "AssetRegistry/ARFilter.h"
 #include "AssetRegistryArchive.h"
@@ -192,7 +193,7 @@ void FAssetRegistryState::InitializeFromExistingAndPrune(const FAssetRegistrySta
 		bool bRemoveDependencyData = true;
 
 		if (bIsFilteredByChunkId &&
-			!AssetData->ChunkIDs.ContainsByPredicate([&](int32 ChunkId) { return ChunksToKeep.Contains(ChunkId); }))
+			!AssetData->GetChunkIDs().ContainsByPredicate([&](int32 ChunkId) { return ChunksToKeep.Contains(ChunkId); }))
 		{
 			bRemoveAssetData = true;
 		}
@@ -224,7 +225,7 @@ void FAssetRegistryState::InitializeFromExistingAndPrune(const FAssetRegistrySta
 		FAssetRegistryState::FilterTags(AssetData->TagsAndValues, NewTagsAndValues, Options.CookFilterlistTagsByClass.Find(AssetData->AssetClassPath), Options);
 
 		FAssetData* NewAssetData = new FAssetData(AssetData->PackageName, AssetData->PackagePath, AssetData->AssetName,
-			AssetData->AssetClassPath, NewTagsAndValues, AssetData->ChunkIDs, AssetData->PackageFlags);
+			AssetData->AssetClassPath, NewTagsAndValues, AssetData->GetChunkIDs(), AssetData->PackageFlags);
 
 		NewAssetData->TaggedAssetBundles = AssetData->TaggedAssetBundles;
 
@@ -486,7 +487,7 @@ void FAssetRegistryState::PruneAssetData(const TSet<FName>& RequiredPackages, co
 		bool bRemoveDependencyData = true;
 
 		if (bIsFilteredByChunkId &&
-			!AssetData->ChunkIDs.ContainsByPredicate([&](int32 ChunkId) { return ChunksToKeep.Contains(ChunkId); }))
+			!AssetData->GetChunkIDs().ContainsByPredicate([&](int32 ChunkId) { return ChunksToKeep.Contains(ChunkId); }))
 		{
 			bRemoveAssetData = true;
 		}
@@ -1440,7 +1441,6 @@ SIZE_T FAssetRegistryState::GetAllocatedSize(bool bLogDetailed) const
 	for (const FAssetData* AssetData : CachedAssets)
 	{
 		AssetDataSize += sizeof(*AssetData);
-		AssetDataSize += AssetData->ChunkIDs.GetAllocatedSize();
 		TagMemoryUsage.Include(AssetData->TagsAndValues);
 		if (AssetData->TaggedAssetBundles.IsValid())
 		{
@@ -1899,10 +1899,11 @@ void FAssetRegistryState::UpdateAssetData(FAssetData* AssetData, FAssetData&& Ne
 		// so we can skip the more expensive compares if the inequality is already known
 		// This is not possible for keyfields - we have to take action on those even if inequality is already known -
 		// so we start with whether bKeyFieldIsModified
+		// TODO: this is horrifying, don't submit like this
 		*bOutModified = bKeyFieldIsModified ||
 			AssetData->AssetName != NewAssetData.AssetName ||
 			AssetData->PackageFlags != NewAssetData.PackageFlags ||
-			AssetData->ChunkIDs != NewAssetData.ChunkIDs ||
+			!AssetData->HasSameChunkIDs(NewAssetData) ||
 			(AssetData->TaggedAssetBundles.IsValid() != NewAssetData.TaggedAssetBundles.IsValid() ||
 				(AssetData->TaggedAssetBundles.IsValid() &&
 					AssetData->TaggedAssetBundles.Get() != NewAssetData.TaggedAssetBundles.Get() && // First check whether the pointers are the same
diff --git a/Engine/Source/Runtime/CoreUObject/Private/AssetRegistry/AssetData.cpp b/Engine/Source/Runtime/CoreUObject/Private/AssetRegistry/AssetData.cpp
--- a/Engine/Source/Runtime/CoreUObject/Private/AssetRegistry/AssetData.cpp
+++ b/Engine/Source/Runtime/CoreUObject/Private/AssetRegistry/AssetData.cpp
@@ -4,6 +4,7 @@
 
 #include "Algo/IsSorted.h"
 #include "Algo/Sort.h"
+#include "Algo/Unique.h"
 #include "AssetRegistry/ARFilter.h"
 #include "Containers/Set.h"
 #include "HAL/CriticalSection.h"
@@ -99,6 +100,140 @@ static TSharedPtr<FAssetBundleData, ESPMode::ThreadSafe> ParseAssetBundles(const
 
 }}} // end namespace UE::AssetData::Private
 
+namespace UE::AssetRegistry
+{
+
+class FChunkArrayRegistryEntry
+{
+public:
+	FChunkArrayRegistryEntry(FAssetData::FChunkArray&& InChunkArray)
+		: ChunkArray(MoveTemp(InChunkArray))
+	{
+	}
+
+	bool operator==(const FChunkArrayRegistryEntry& OtherEntry) const
+	{
+		const int32 ThisNum = ChunkArray.Num();
+		if (ThisNum != OtherEntry.ChunkArray.Num())
+		{
+			return false;
+		}
+
+		// Arrays are guaranteed to be sorted/unique when an entry is created
+		for (int32 Index = 0; Index < ThisNum; ++Index)
+		{
+			if (ChunkArray[Index] != OtherEntry.ChunkArray[Index])
+			{
+				return false;
+			}
+		}
+
+		return true;
+	}
+
+	FAssetData::FChunkArrayView GetChunkIDs() const
+	{
+		return ChunkArray;
+	}
+
+	friend uint32 GetTypeHash(const FChunkArrayRegistryEntry& InEntry)
+	{
+		uint32 Hash = 0;
+		for (int32 ChunkID : InEntry.ChunkArray)
+		{
+			Hash = HashCombineFast(Hash, GetTypeHash(ChunkID));
+		}
+		return Hash;
+	}
+
+	SIZE_T GetAllocatedSize() const
+	{
+		return ChunkArray.GetAllocatedSize();
+	}
+
+private:
+	/* Array of chunk IDs that's guaranteed to be unique and sorted in ascending order. */
+	FAssetData::FChunkArray ChunkArray;
+};
+
+class FChunkArrayRegistry
+{
+public:
+	FChunkArrayRegistryHandle FindOrAdd(const FAssetData::FChunkArrayView& InChunkIDs)
+	{
+		// Make a copy so we can sort it inside; these are very small arrays generally
+		FAssetData::FChunkArray ChunkArray = FAssetData::FChunkArray(InChunkIDs);
+		return FindOrAdd(MoveTemp(ChunkArray));
+	}
+
+	FChunkArrayRegistryHandle FindOrAdd(FAssetData::FChunkArray&& InChunkIDs)
+	{
+		// Sort and remove duplicates before inserting
+		Algo::Sort(InChunkIDs);
+		return FindOrAddSorted(MoveTemp(InChunkIDs));
+	}
+
+	FChunkArrayRegistryHandle FindOrAddSorted(FAssetData::FChunkArray&& InChunkIDs)
+	{
+		// Make sure we have no duplicates on top of being sorted
+		InChunkIDs.SetNum(Algo::Unique(InChunkIDs));
+
+		FChunkArrayRegistryHandle Index;
+		if (!InChunkIDs.IsEmpty())
+		{
+			FChunkArrayRegistryEntry Entry(MoveTemp(InChunkIDs));
+			const uint32 Hash = GetTypeHash(Entry);
+			{
+				const FReadScopeLock ReadLock(Lock);
+				Index = ChunkArrays.FindIdByHash(Hash, Entry);
+			}
+
+			if (!Index.IsValidId())
+			{
+				// Two threads may hit this point simultaneously and race to insert, and if we just emplace, the second thread will
+				// replace the array inserted by the first. We never want this to happen since that would invalidate array views produced
+				// since, so we have to do the look-up again. Insertions should be rare in general, so the cost shouldn't be noticeable.
+				const FWriteScopeLock WriteLock(Lock);
+				Index = ChunkArrays.FindIdByHash(Hash, Entry);
+				if (!Index.IsValidId())
+				{
+					Index = ChunkArrays.EmplaceByHash(Hash, MoveTemp(Entry));
+				}
+			}
+		}
+		return Index;
+	}
+
+	FAssetData::FChunkArrayView GetChunkIDs(FChunkArrayRegistryHandle ChunkArrayRegistryHandle) const
+	{
+		FAssetData::FChunkArrayView ChunkIDs;
+		if (ChunkArrayRegistryHandle.IsValidId())
+		{
+			FReadScopeLock ReadLock(Lock);
+			ChunkIDs = ChunkArrays[ChunkArrayRegistryHandle].GetChunkIDs();
+		}
+		return ChunkIDs;
+	}
+
+	SIZE_T GetAllocatedSize() const
+	{
+		SIZE_T AllocatedSize = sizeof(*this);
+		AllocatedSize += ChunkArrays.GetAllocatedSize();
+		for (const FChunkArrayRegistryEntry& Entry : ChunkArrays)
+		{
+			AllocatedSize += Entry.GetAllocatedSize();
+		}
+		return AllocatedSize;
+	}
+
+private:
+	// We will only ever add elements to this set, so we can return persistent indices
+	TSet<FChunkArrayRegistryEntry> ChunkArrays;
+	mutable FRWLock Lock;
+} GChunkArrayRegistry;
+
+} // namespace UE::AssetRegistry
+
 FAssetData::FAssetData(FName InPackageName, FName InPackagePath, FName InAssetName, FName InAssetClass, FAssetDataTagMap InTags, TArrayView<const int32> InChunkIDs, uint32 InPackageFlags)
 	: FAssetData(InPackageName, InPackagePath, InAssetName, FAssetData::TryConvertShortClassNameToPathName(InAssetClass), InTags, InChunkIDs, InPackageFlags)
 {
@@ -115,7 +250,6 @@ FAssetData::FAssetData(FName InPackageName, FName InPackagePath, FName InAssetNa
 	, AssetName(InAssetName)
 	, AssetClassPath(InAssetClassPathName)
 	, PackageFlags(InPackageFlags)
-	, ChunkIDs(MoveTemp(InChunkIDs))
 {
 	SetTagsAndAssetBundles(MoveTemp(InTags));
 
@@ -125,13 +259,14 @@ FAssetData::FAssetData(FName InPackageName, FName InPackagePath, FName InAssetNa
 PRAGMA_DISABLE_DEPRECATION_WARNINGS
 	ObjectPath = FName(FStringView(ObjectPathStr));
 PRAGMA_ENABLE_DEPRECATION_WARNINGS
+
+	SetChunkIDs(InChunkIDs);
 }
 
 FAssetData::FAssetData(const FString& InLongPackageName, const FString& InObjectPath, FTopLevelAssetPath InAssetClassPathName, FAssetDataTagMap InTags, TArrayView<const int32> InChunkIDs, uint32 InPackageFlags)
 	: PackageName(*InLongPackageName)
 	, AssetClassPath(InAssetClassPathName)
 	, PackageFlags(InPackageFlags)
-	, ChunkIDs(MoveTemp(InChunkIDs))
 {
 	using namespace UE::AssetRegistry::Private;
 PRAGMA_DISABLE_DEPRECATION_WARNINGS
@@ -152,6 +287,8 @@ PRAGMA_ENABLE_DEPRECATION_WARNINGS
 		OptionalOuterPath = FName(Parts.OuterPath);
 	}
 #endif
+
+	SetChunkIDs(InChunkIDs);
 }
 
 FAssetData::FAssetData(const UObject* InAsset, FAssetData::ECreationFlags InCreationFlags)
@@ -189,8 +326,8 @@ PRAGMA_ENABLE_DEPRECATION_WARNINGS
 			InAsset->GetAssetRegistryTags(*this);
 		}
 
-		ChunkIDs = Package->GetChunkIDs();
 		PackageFlags = Package->GetPackageFlags();
+		SetChunkIDs(Package->GetChunkIDs());
 	}
 }
 
@@ -261,6 +398,68 @@ bool FAssetData::IsTopLevelAsset(UObject* Object)
 	return Outer->IsA<UPackage>();
 }
 
+FAssetData::FChunkArrayView FAssetData::GetChunkIDs() const
+{
+#if !UE_STRIP_DEPRECATED_PROPERTIES
+	PRAGMA_DISABLE_DEPRECATION_WARNINGS
+	// Modifying the chunk IDs directly is no longer supported; use AddChunkID/SetChunkIDs/ClearChunkIDs instead
+	if (!ChunkIDs.IsEmpty())
+	{
+		UE_LOG(LogAssetData, Error, TEXT("Modifying FAssetData::ChunkIDs directly is no longer supported; use AddChunkID/SetChunkIDs/ClearChunkIDs instead."));
+	}
+	PRAGMA_ENABLE_DEPRECATION_WARNINGS
+#endif
+
+	return UE::AssetRegistry::GChunkArrayRegistry.GetChunkIDs(ChunkArrayRegistryHandle);
+}
+
+void FAssetData::SetChunkIDs(FChunkArray&& InChunkIDs)
+{
+	ChunkArrayRegistryHandle = UE::AssetRegistry::GChunkArrayRegistry.FindOrAdd(MoveTemp(InChunkIDs));
+}
+
+void FAssetData::SetChunkIDs(const FChunkArrayView& InChunkIDs)
+{
+	ChunkArrayRegistryHandle = UE::AssetRegistry::GChunkArrayRegistry.FindOrAdd(InChunkIDs);
+}
+
+void FAssetData::AddChunkID(int32 ChunkID)
+{
+	// Chunk arrays are guaranteed to be sorted/unique when coming back from registry, so maintain that here
+	const FChunkArrayView CurrentChunkIDs = GetChunkIDs();
+	const int32 NumChunkIDs = CurrentChunkIDs.Num();
+	const int32 InsertIndex = Algo::LowerBound(CurrentChunkIDs, ChunkID);
+	if (CurrentChunkIDs.IsValidIndex(InsertIndex) && (CurrentChunkIDs[InsertIndex] == ChunkID))
+	{
+		return;
+	}
+
+	// Build the new array in parts to save having to shift/reallocate unnecessarily
+	const FChunkArrayView BeforeInserted = CurrentChunkIDs.Left(InsertIndex);
+	const FChunkArrayView AfterInserted = CurrentChunkIDs.RightChop(InsertIndex);
+	FChunkArray NewChunkIDs;
+	NewChunkIDs.Reserve(NumChunkIDs + 1);
+	NewChunkIDs.Append(BeforeInserted.GetData(), BeforeInserted.Num());
+	NewChunkIDs.Add(ChunkID);
+	NewChunkIDs.Append(AfterInserted.GetData(), AfterInserted.Num());
+	ChunkArrayRegistryHandle = UE::AssetRegistry::GChunkArrayRegistry.FindOrAddSorted(MoveTemp(NewChunkIDs));
+}
+
+void FAssetData::ClearChunkIDs()
+{
+	SetChunkIDs(FChunkArray());
+}
+
+bool FAssetData::HasSameChunkIDs(const FAssetData& OtherAssetData) const
+{
+	return ChunkArrayRegistryHandle == OtherAssetData.ChunkArrayRegistryHandle;
+}
+
+SIZE_T FAssetData::GetChunkArrayRegistryAllocatedSize()
+{
+	return UE::AssetRegistry::GChunkArrayRegistry.GetAllocatedSize();
+}
+
 void FAssetData::SetTagsAndAssetBundles(FAssetDataTagMap&& Tags)
 {
 	using namespace UE::AssetData::Private;
@@ -364,16 +563,19 @@ PRAGMA_ENABLE_DEPRECATION_WARNINGS
 
 	SerializeTagsAndBundles(Ar, *this, Version);
 
-	if (Ar.IsSaving() && ChunkIDs.Num() > 1)
+	FAssetData::FChunkArray SerializedChunkIDs;
+	if (Ar.IsSaving())
 	{
-		TArray<int32> SortedChunkIDs(ChunkIDs);
-		Algo::Sort(SortedChunkIDs);
-		Ar << SortedChunkIDs;
+		SerializedChunkIDs = GetChunkIDs();
+		Ar << SerializedChunkIDs;
 	}
 	else
 	{
-		Ar << ChunkIDs;
+		Ar << SerializedChunkIDs;
+		check(Algo::IsSorted(SerializedChunkIDs));
+		ChunkArrayRegistryHandle = UE::AssetRegistry::GChunkArrayRegistry.FindOrAddSorted(MoveTemp(SerializedChunkIDs));
 	}
+
 	Ar << PackageFlags;
 
 	// Rebuild deprecated ObjectPath field 
diff --git a/Engine/Source/Runtime/CoreUObject/Public/AssetRegistry/AssetData.h b/Engine/Source/Runtime/CoreUObject/Public/AssetRegistry/AssetData.h
--- a/Engine/Source/Runtime/CoreUObject/Public/AssetRegistry/AssetData.h
+++ b/Engine/Source/Runtime/CoreUObject/Public/AssetRegistry/AssetData.h
@@ -110,6 +110,11 @@ private:
 	FAssetRegistryVersion() {}
 };
 
+namespace UE::AssetRegistry
+{
+	using FChunkArrayRegistryHandle = FSetElementId;
+}
+
 namespace UE::AssetRegistry::Private
 {
 	struct FAssetPathParts
@@ -163,6 +168,11 @@ public:
 	FName AssetClass;
 	/** The path of the asset's class, e.g. /Script/Engine.StaticMesh */
 	FTopLevelAssetPath AssetClassPath;
+
+private:
+	UE::AssetRegistry::FChunkArrayRegistryHandle ChunkArrayRegistryHandle;
+
+public:
 	/** Asset package flags */
 	uint32 PackageFlags = 0;
 
@@ -188,10 +198,27 @@ public:
 	 */
 	TSharedPtr<FAssetBundleData, ESPMode::ThreadSafe> TaggedAssetBundles;
 
-	COREUOBJECT_API void SetTagsAndAssetBundles(FAssetDataTagMap&& Tags);
-
 	/** The IDs of the pakchunks this asset is located in for streaming install.  Empty if not assigned to a chunk */
+#if !UE_STRIP_DEPRECATED_PROPERTIES
+	UE_DEPRECATED(5.1, "Use SetChunkIDs/GetChunkIDs/AddChunkID instead.")
 	TArray<int32, TInlineAllocator<2>> ChunkIDs;
+#endif
+
+	COREUOBJECT_API void SetTagsAndAssetBundles(FAssetDataTagMap&& Tags);
+
+	// These are usually very small arrays and we can preallocate two elements for the same cost as one on 64-bit systems
+	using FChunkArray = TArray<int32, TInlineAllocator<2>>;
+	using FChunkArrayView = TConstArrayView<int32>;
+
+	COREUOBJECT_API FChunkArrayView GetChunkIDs() const;
+	COREUOBJECT_API void SetChunkIDs(FChunkArray&& InChunkIDs);
+	COREUOBJECT_API void SetChunkIDs(const FChunkArrayView& InChunkIDs);
+	COREUOBJECT_API void AddChunkID(int32 ChunkID);
+	COREUOBJECT_API void ClearChunkIDs();
+	COREUOBJECT_API bool HasSameChunkIDs(const FAssetData& OtherAssetData) const;
+
+	/** Returns overhead of the chunk array registry that's used to manage chunk ID arrays. */
+	static COREUOBJECT_API SIZE_T GetChunkArrayRegistryAllocatedSize();
 
 public:	
 PRAGMA_DISABLE_DEPRECATION_WARNINGS // Compilers can complain about deprecated members in compiler generated code
@@ -322,7 +349,6 @@ PRAGMA_ENABLE_DEPRECATION_WARNINGS
 
 	void Shrink()
 	{
-		ChunkIDs.Shrink();
 		TagsAndValues.Shrink();
 	}
 
@@ -598,9 +624,10 @@ PRAGMA_ENABLE_DEPRECATION_WARNINGS
 			UE_LOG(LogAssetData, Log, TEXT("            %s : %s"), *TagValue.Key.ToString(), *TagValue.Value.AsString());
 		}
 
-		UE_LOG(LogAssetData, Log, TEXT("        ChunkIDs: %d"), ChunkIDs.Num());
+		const TConstArrayView<int32> CurrentChunkIDs = GetChunkIDs();
+		UE_LOG(LogAssetData, Log, TEXT("        ChunkIDs: %d"), CurrentChunkIDs.Num());
 
-		for (int32 Chunk: ChunkIDs)
+		for (int32 Chunk : CurrentChunkIDs)
 		{
 			UE_LOG(LogAssetData, Log, TEXT("                 %d"), Chunk);
 		}
diff --git a/Engine/Source/Runtime/Engine/Private/AssetManager.cpp b/Engine/Source/Runtime/Engine/Private/AssetManager.cpp
--- a/Engine/Source/Runtime/Engine/Private/AssetManager.cpp
+++ b/Engine/Source/Runtime/Engine/Private/AssetManager.cpp
@@ -2124,7 +2124,7 @@ bool UAssetManager::FindMissingChunkList(const TArray<FSoftObjectPath>& AssetLis
 		GetAssetDataForPath(Asset, FoundData);
 		TSet<int32> FoundChunks, MissingChunks, ErrorChunks;
 
-		for (int32 PakchunkId : FoundData.ChunkIDs)
+		for (int32 PakchunkId : FoundData.GetChunkIDs())
 		{
 			if (!ChunkLocationCache.Contains(PakchunkId))
 			{
