[CoreUObject] fix for `TObjectPtr.GetClass()` crashing if `TObjectPtr` points to null.

#p4v-preflight-copy 21901752

diff --git a/Engine/Source/Runtime/CoreUObject/Public/UObject/ObjectHandle.h b/Engine/Source/Runtime/CoreUObject/Public/UObject/ObjectHandle.h
--- a/Engine/Source/Runtime/CoreUObject/Public/UObject/ObjectHandle.h
+++ b/Engine/Source/Runtime/CoreUObject/Public/UObject/ObjectHandle.h
@@ -338,7 +338,8 @@ inline UClass* ResolveObjectHandleClass(FObjectHandle Handle)
 {
 	if (IsObjectHandleResolved(Handle))
 	{
-		return ReadObjectHandlePointerNoCheck(Handle)->GetClass();
+		UObject* Obj = ReadObjectHandlePointerNoCheck(Handle);
+		return Obj != nullptr ? Obj->GetClass() : nullptr;
 	}
 	else
 	{
diff --git a/Engine/Source/Runtime/CoreUObject/Tests/ObjectHandleTest.cpp b/Engine/Source/Runtime/CoreUObject/Tests/ObjectHandleTest.cpp
--- a/Engine/Source/Runtime/CoreUObject/Tests/ObjectHandleTest.cpp
+++ b/Engine/Source/Runtime/CoreUObject/Tests/ObjectHandleTest.cpp
@@ -3,6 +3,7 @@
 #if WITH_LOW_LEVEL_TESTS
 
 #include "UObject/ObjectHandle.h"
+#include "UObject/ObjectPtr.h"
 
 #include "HAL/PlatformProperties.h"
 #include "ObjectRefTrackingTestBase.h"
@@ -192,13 +193,18 @@ DISABLED_TEST_CASE_METHOD(FObjectHandleTestBase, "CoreUObject::FObjectHandle::Re
 	TestResolveFailure("/Engine/Test/ObjectPtrDefaultSerialize/Transient", "DefaultSerializeObject_DoesNotExist");
 }
 
-
 TEST_CASE_METHOD(FObjectHandleTestBase, "CoreUObject::FObjectHandle::Resolve Script Target", "[CoreUObject][ObjectHandle]")
 {
 	// Confirm we successfully resolve a correct reference to engine content
 	TestResolvableNonNull("/Script/CoreUObject", "MetaData");
 }
 
+TEST_CASE_METHOD(FObjectHandleTestBase, "CoreUObject::TObjectPtr::HandleNullGetClass", "[CoreUObject][ObjectHandle][.Engine]")
+{
+	TObjectPtr<UObject> Ptr = nullptr;
+	TEST_TRUE(TEXT("TObjectPtr.GetClass should return null on a null object"), Ptr.GetClass() == nullptr);
+}
+
 #if UE_WITH_OBJECT_HANDLE_LATE_RESOLVE
 
 TEST_CASE_METHOD(FObjectHandleTestBase, "CoreUObject::FObjectHandle::Resolve Malformed Handle", "[CoreUObject][ObjectHandle]")
