Fixed build/deprecation error due to code that wasn't in the branch the change was developed in.

diff --git a/Engine/Source/Runtime/CoreUObject/Private/AssetRegistry/AssetData.cpp b/Engine/Source/Runtime/CoreUObject/Private/AssetRegistry/AssetData.cpp
--- a/Engine/Source/Runtime/CoreUObject/Private/AssetRegistry/AssetData.cpp
+++ b/Engine/Source/Runtime/CoreUObject/Private/AssetRegistry/AssetData.cpp
@@ -622,9 +622,12 @@ void FAssetData::NetworkWrite(FCbWriter& Writer, bool bWritePackageName) const
 		}
 		Writer.EndArray();
 	}
-	if (!ChunkIDs.IsEmpty())
+
+	FChunkArrayView CurrentChunkIDs = GetChunkIDs();
+	if (!CurrentChunkIDs.IsEmpty())
 	{
-		Writer << "I" << ChunkIDs;
+		FChunkArray SerializedChunkIDs(CurrentChunkIDs);
+		Writer << "I" << SerializedChunkIDs;
 	}
 	Writer.EndObject();
 }
@@ -689,7 +692,10 @@ bool FAssetData::TryNetworkRead(FCbFieldView Field, bool bReadPackageName, FName
 		SetTagsAndAssetBundles(FAssetDataTagMap());
 	}
 
-	LoadFromCompactBinary(Object["I"], ChunkIDs); // Ok if it does not exist
+	FChunkArray SerializedChunkIDs;
+	LoadFromCompactBinary(Object["I"], SerializedChunkIDs); // Ok if it does not exist
+	check(Algo::IsSorted(SerializedChunkIDs));
+	ChunkArrayRegistryHandle = UE::AssetRegistry::GChunkArrayRegistry.FindOrAddSorted(MoveTemp(SerializedChunkIDs));
 
 	// Rebuild deprecated ObjectPath field 
 	TStringBuilder<FName::StringBufferSize> Builder;
