Minor clean-up to my last change: use array comparator since the manual compare was done when it was an array view.

diff --git a/Engine/Source/Runtime/CoreUObject/Private/AssetRegistry/AssetData.cpp b/Engine/Source/Runtime/CoreUObject/Private/AssetRegistry/AssetData.cpp
--- a/Engine/Source/Runtime/CoreUObject/Private/AssetRegistry/AssetData.cpp
+++ b/Engine/Source/Runtime/CoreUObject/Private/AssetRegistry/AssetData.cpp
@@ -113,22 +113,8 @@ public:
 
 	bool operator==(const FChunkArrayRegistryEntry& OtherEntry) const
 	{
-		const int32 ThisNum = ChunkArray.Num();
-		if (ThisNum != OtherEntry.ChunkArray.Num())
-		{
-			return false;
-		}
-
 		// Arrays are guaranteed to be sorted/unique when an entry is created
-		for (int32 Index = 0; Index < ThisNum; ++Index)
-		{
-			if (ChunkArray[Index] != OtherEntry.ChunkArray[Index])
-			{
-				return false;
-			}
-		}
-
-		return true;
+		return ChunkArray == OtherEntry.ChunkArray;
 	}
 
 	FAssetData::FChunkArrayView GetChunkIDs() const
