Relax requirements for separators in FSoftObjectPath to match ResolveName in UObjectGlobals and add some tests.

diff --git a/Engine/Source/Runtime/CoreUObject/Private/UObject/SoftObjectPath.cpp b/Engine/Source/Runtime/CoreUObject/Private/UObject/SoftObjectPath.cpp
--- a/Engine/Source/Runtime/CoreUObject/Private/UObject/SoftObjectPath.cpp
+++ b/Engine/Source/Runtime/CoreUObject/Private/UObject/SoftObjectPath.cpp
@@ -4,6 +4,7 @@
 #include "UObject/PropertyPortFlags.h"
 #include "UObject/UnrealType.h"
 #include "UObject/ObjectRedirector.h"
+#include "Misc/AsciiSet.h"
 #include "Misc/PackageName.h"
 #include "Misc/StringBuilder.h"
 #include "UObject/LinkerLoad.h"
@@ -118,26 +119,55 @@ void FSoftObjectPath::SetPath(FWideStringView Path)
 		// Possibly an ExportText path. Trim the ClassName.
 		Path = FPackageName::ExportTextPathToObjectPath(Path);
 
-		if (!Path.StartsWith('/'))
+		constexpr FAsciiSet Delimiters = FAsciiSet(".") + (char)SUBOBJECT_DELIMITER_CHAR;
+		if (  !Path.StartsWith('/')  // Must start with a package path 
+			|| Delimiters.Contains(Path[Path.Len() - 1]) // Must not end with a trailing delimiter
+		)
 		{
 			// Not a recognized path. No ensure/logging here because many things attempt to construct paths from user input. 
 			Reset();
 			return;
 		}
 
-		int32 ColonIndex;
-		if (Path.FindChar(SUBOBJECT_DELIMITER_CHAR, ColonIndex))
+		
+		// Reject paths that contain two consecutive delimiters in any position 
+		for (int32 i=2; i < Path.Len(); ++i) // Start by comparing index 2 and index 1 because index 0 is known to be '/'
 		{
-			// Has a subobject, split on that then create a name from the temporary path
-			AssetPath = Path.Left(ColonIndex);
-			SubPathString = Path.Mid(ColonIndex + 1);
+			if (Delimiters.Contains(Path[i]) && Delimiters.Contains(Path[i-1]))
+			{
+				Reset();
+				return;
+			}
 		}
-		else
+
+		FWideStringView PackageNameView = FAsciiSet::FindPrefixWithout(Path, Delimiters);
+		if (PackageNameView.Len() == Path.Len())
+		{
+			// No delimiter, package name only
+			AssetPath = FTopLevelAssetPath(FName(PackageNameView), FName());
+			SubPathString.Empty();
+			return;
+		}
+
+		Path.RightChopInline(PackageNameView.Len() + 1);
+		check(!Path.IsEmpty() && !Delimiters.Contains(Path[0])); // Sanitized to avoid trailing delimiter or consecutive delimiters above
+
+		FWideStringView AssetNameView = FAsciiSet::FindPrefixWithout(Path, Delimiters);
+		if (AssetNameView.Len() == Path.Len())
 		{
-			// No Subobject
-			AssetPath = Path;
+			// No subobject path
+			AssetPath = FTopLevelAssetPath(FName(PackageNameView), FName(AssetNameView));
 			SubPathString.Empty();
+			return;
 		}
+
+		Path.RightChopInline(AssetNameView.Len() + 1);
+		check(!Path.IsEmpty() && !Delimiters.Contains(Path[0])); // Sanitized to avoid trailing delimiter or consecutive delimiters above
+
+		// Replace delimiters in subpath string with . to normalize
+		AssetPath = FTopLevelAssetPath(FName(PackageNameView), FName(AssetNameView));
+		SubPathString = Path;
+		SubPathString.ReplaceCharInline(SUBOBJECT_DELIMITER_CHAR, '.');
 	}
 }
 
@@ -810,8 +840,8 @@ TSet<FName> FSoftObjectPath::PIEPackageNames;
 
 
 
-IMPLEMENT_SIMPLE_AUTOMATION_TEST(FSoftObjectPathTests, "System.CoreUObject.SoftObjectPath", EAutomationTestFlags::ApplicationContextMask | EAutomationTestFlags::EngineFilter);
-bool FSoftObjectPathTests::RunTest(const FString& Parameters)
+IMPLEMENT_SIMPLE_AUTOMATION_TEST(FSoftObjectPathImportTextTests, "System.CoreUObject.SoftObjectPath.ImportText", EAutomationTestFlags::ApplicationContextMask | EAutomationTestFlags::EngineFilter);
+bool FSoftObjectPathImportTextTests::RunTest(const FString& Parameters)
 {
 	const TCHAR* PackageName = TEXT("/Game/Environments/Sets/Arid/Materials/M_Arid");
 	const TCHAR* AssetName = TEXT("M_Arid");
@@ -837,3 +867,130 @@ bool FSoftObjectPathTests::RunTest(const FString& Parameters)
 
 	return true;
 }
+
+IMPLEMENT_SIMPLE_AUTOMATION_TEST(FSoftObjectPathTrySetPathTests, "System.CoreUObject.SoftObjectPath.TrySetPath", EAutomationTestFlags::ApplicationContextMask | EAutomationTestFlags::EngineFilter);
+bool FSoftObjectPathTrySetPathTests::RunTest(const FString& Parameters)
+{
+	FSoftObjectPath Path;
+
+	const TCHAR* PackageName = TEXT("/Game/Maps/Arena");
+	const TCHAR* TopLevelPath = TEXT("/Game/Maps/Arena.Arena");
+	const TCHAR* TopLevelPathWrongSeparator = TEXT("/Game/Maps/Arena:Arena");
+
+	Path.SetPath(PackageName);
+	if (TestTrue("Package name: Is valid", Path.IsValid()))
+	{
+		TestEqual("Package name: Round trips equal", Path.ToString(), PackageName);
+		TestEqual("Package name: Package name part", Path.GetLongPackageName(), PackageName);
+		TestEqual("Package name: Asset name part", Path.GetAssetName(), FString());
+		TestEqual("Package name: Subobject path part", Path.GetSubPathString(), FString());
+	}
+
+	Path.SetPath(TopLevelPath);
+	if (TestTrue("Top level object path: Is valid", Path.IsValid()))
+	{
+		TestEqual("Top level object path: round trips equal", Path.ToString(), TopLevelPath);
+	}
+
+	const TCHAR* PathWithWideChars = TEXT("/Game/\u30ad\u30e3\u30e9\u30af\u30bf\u30fc/\u5c71\u672c.\u5c71\u672c");
+	Path.SetPath(PathWithWideChars);
+	if (TestTrue("Path with wide chars: Is valid", Path.IsValid()))
+	{
+		TestEqual("Path with wide chars: Round trips equal", Path.ToString(), PathWithWideChars);
+		TestEqual("Path with wide chars: Package name part", Path.GetLongPackageName(), TEXT("/Game/\u30ad\u30e3\u30e9\u30af\u30bf\u30fc/\u5c71\u672c"));
+		TestEqual("Path with wide chars: Asset name part", Path.GetAssetName(), TEXT("\u5c71\u672c"));
+		TestEqual("Path with wide chars: Subobject path part", Path.GetSubPathString(), FString());
+	}
+
+	Path.SetPath(TopLevelPathWrongSeparator);
+	// Round tripping replaces dot with subobject separator for second separator 
+	if (TestTrue("Top level object path with incorrect separator: is valid", Path.IsValid())) 
+	{ 
+		TestEqual("Top level object path with incorrect separator: Round trips with normalized separator", Path.ToString(), TopLevelPath);
+		TestEqual("Top level object path with incorrect separator: Package name part", Path.GetLongPackageName(), TEXT("/Game/Maps/Arena"));
+		TestEqual("Top level object path with incorrect separator: Asset name part", Path.GetAssetName(), TEXT("Arena"));
+		TestEqual("Top level object path with incorrect separator: Subobject path part", Path.GetSubPathString(), FString());
+	}
+
+	const TCHAR* PackageNameTrailingDot = TEXT("/Game/Maps/Arena.");
+	Path.SetPath(PackageNameTrailingDot);
+	TestFalse("Package name trailing dot: is not valid", Path.IsValid());
+
+	const TCHAR* PackageNameTrailingSeparator = TEXT("/Game/Maps/Arena:");
+	Path.SetPath(PackageNameTrailingSeparator);
+	TestFalse("Package name trailing separator: is not valid", Path.IsValid());
+
+	const TCHAR* ObjectPathTrailingDot = TEXT("/Game/Maps/Arena.Arena.");
+	Path.SetPath(ObjectPathTrailingDot);
+	TestFalse("Object path trailing dot: is not valid", Path.IsValid());
+
+	const TCHAR* ObjectPathTrailingSeparator = TEXT("/Game/Maps/Arena.Arena:");
+	Path.SetPath(ObjectPathTrailingSeparator);
+	TestFalse("Object path trailing separator: is not valid", Path.IsValid());
+
+	const TCHAR* PackageNameWithoutLeadingSlash = TEXT("Game/Maps/Arena");
+	Path.SetPath(PackageNameWithoutLeadingSlash);
+	TestFalse("Package name without leading slash: is not valid", Path.IsValid());
+
+	const TCHAR* ObjectPathWithoutLeadingSlash = TEXT("Game/Maps/Arena.Arena");
+	Path.SetPath(ObjectPathWithoutLeadingSlash);
+	TestFalse("Object name without leading slash: is not valid", Path.IsValid());
+
+	const TCHAR* SubObjectPathWithSeparator = TEXT("/Game/Characters/Steve.Steve_C:Root");
+	Path.SetPath(SubObjectPathWithSeparator);
+	if (TestTrue("Subobject path with separator: is valid", Path.IsValid()))
+	{
+		TestEqual("Subobject path with separator: round trip", Path.ToString(), SubObjectPathWithSeparator);
+		TestEqual("Subobject path with separator: package name", Path.GetLongPackageName(), TEXT("/Game/Characters/Steve"));
+		TestEqual("Subobject path with separator: asset name", Path.GetAssetName(), TEXT("Steve_C"));
+		TestEqual("Subobject path with separator: subobject path", Path.GetSubPathString(), TEXT("Root"));
+	}
+
+	const TCHAR* SubObjectPathWithTrailingDot = TEXT("/Game/Characters/Steve.Steve_C:Root.");
+	Path.SetPath(SubObjectPathWithTrailingDot);
+	TestFalse("Subobject path with trailing dot: is not valid", Path.IsValid());
+
+	const TCHAR* SubObjectPathWithTrailingSeparator = TEXT("/Game/Characters/Steve.Steve_C:Root:");
+	Path.SetPath(SubObjectPathWithTrailingSeparator);
+	TestFalse("Subobject path with trailing separator: is not valid", Path.IsValid());
+
+	const TCHAR* PathWithoutAssetName = TEXT("/Game/Characters/Steve.:Root");
+	Path.SetPath(PathWithoutAssetName );
+	TestFalse("Subobject path without asset name: is not valid", Path.IsValid());
+
+	const TCHAR* SubObjectPathWithDot = TEXT("/Game/Characters/Steve.Steve_C:Root");
+	Path.SetPath(SubObjectPathWithDot);
+	if (TestTrue("Subobject path with dot: is valid", Path.IsValid()))
+	{
+		TestEqual("Subobject path with dot: round trips with normalized separator", Path.ToString(), SubObjectPathWithDot); // Round tripping replaces dot with subobject separator for second separator 
+		TestEqual("Subobject path with dot: package name", Path.GetLongPackageName(), TEXT("/Game/Characters/Steve"));
+		TestEqual("Subobject path with dot: asset name", Path.GetAssetName(), TEXT("Steve_C"));
+		TestEqual("Subobject path with dot: subobject path", Path.GetSubPathString(), TEXT("Root"));
+	}
+
+	const TCHAR* LongPath = TEXT("/Game/Characters/Steve.Steve_C:Root.Inner.AnotherInner.FurtherInner");
+	Path.SetPath(LongPath);
+	if (TestTrue("Long path: is valid", Path.IsValid()))
+	{
+		TestEqual("Long path: round trip", Path.ToString(), LongPath);
+		TestEqual("Long path: Package name part", Path.GetLongPackageName(), TEXT("/Game/Characters/Steve"));
+		TestEqual("Long path: Asset name part", Path.GetAssetName(), TEXT("Steve_C"));
+		TestEqual("Long path: Subobject path part", Path.GetSubPathString(), TEXT("Root.Inner.AnotherInner.FurtherInner"));
+	}
+
+	const TCHAR* LongPathWithSeparatorInWrongPlace = TEXT("/Game/Characters/Steve.Steve_C.Root.Inner.AnotherInner:FurtherInner");
+	Path.SetPath(LongPathWithSeparatorInWrongPlace);
+	if (TestTrue("Long path with separator in wrong place: is valid", Path.IsValid()))
+	{
+		TestEqual("Long path with separator in wrong place: round trip with normalized separator", Path.ToString(), LongPath);
+		TestEqual("Long path with separator in wrong place: package name", Path.GetLongPackageName(), TEXT("/Game/Characters/Steve"));
+		TestEqual("Long path with separator in wrong place: asset name", Path.GetAssetName(), TEXT("Steve_C"));
+		TestEqual("Long path with separator in wrong place: subobject path", Path.GetSubPathString(), TEXT("Root.Inner.AnotherInner.FurtherInner"));
+	}
+
+	const TCHAR* LongPathWithConsecutiveDelimiters = TEXT("/Game/Characters/Steve.Steve_C:Root.Inner.AnotherInner..FurtherInner");
+	Path.SetPath(LongPathWithConsecutiveDelimiters );
+	TestFalse("Long path with consecutive delimiters: is not valid", Path.IsValid());
+
+	return true;
+}
