FAssetBundleEntry::ImportTextItem needs to be able to handle standard struct export for bundles created by PrimaryAssets

diff --git a/Engine/Source/Runtime/CoreUObject/Private/AssetRegistry/AssetBundleData.cpp b/Engine/Source/Runtime/CoreUObject/Private/AssetRegistry/AssetBundleData.cpp
--- a/Engine/Source/Runtime/CoreUObject/Private/AssetRegistry/AssetBundleData.cpp
+++ b/Engine/Source/Runtime/CoreUObject/Private/AssetRegistry/AssetBundleData.cpp
@@ -2,6 +2,7 @@
 
 #include "AssetRegistry/AssetBundleData.h"
 #include "AssetRegistry/AssetData.h"
+#include "Misc/OutputDeviceNull.h"
 #include "UObject/PropertyPortFlags.h"
 
 UE_IMPLEMENT_STRUCT("/Script/CoreUObject", AssetBundleData);
@@ -105,6 +106,9 @@ bool FAssetBundleEntry::ImportTextItem(const TCHAR*& Buffer, int32 PortFlags, UO
 		return false;
 	}
 
+	UScriptStruct* TopLevelAssetPathStruct = TBaseStructure<FTopLevelAssetPath>::Get();
+	static const FString TopLevelPathStructName = TopLevelAssetPathStruct->GetName(); 
+	FOutputDeviceNull SilenceErrors;
 	bool bBundleAssetsSeen = false;
 	bool bAssetPathsSeen = false;
 	while (true)
@@ -167,14 +171,38 @@ bool FAssetBundleEntry::ImportTextItem(const TCHAR*& Buffer, int32 PortFlags, UO
 				continue;
 			}
 
-			while (LocalObjectPaths.Emplace_GetRef().ImportTextItem(/* in-out */ BufferIt, PortFlags, Parent, ErrorText))
+			// Try multiple formats to parse our list of FTopLevelAssetPaths
+			//     - the native export converts FTopLevelAssetPaths to FSoftObjectPaths and exports the FSoftObjectPaths as strings
+			//     - the default export path (activated when defaults is not null) exports FTopLevelAssetpaths as formatted structs
+			FSoftObjectPath TempObjectPath;
+			FTopLevelAssetPath TempTopLevelPath;
+			while (true)
 			{
-				if (!LocalObjectPaths.Last().GetSubPathString().IsEmpty())
+				TempObjectPath.Reset();
+				// Import string e.g. /Path/To/Package.AssetName
+				if (TempObjectPath.ImportTextItem(BufferIt, PortFlags, Parent, &SilenceErrors))
 				{
-					// Parse error, these should only have top-level paths
-					ErrorText->Logf(ELogVerbosity::Error, TEXT("FAssetBundleEntry::ImportTextItem - Asset bundle entries should all be top level objects, but found a subobject path '%s'. When parsing import text '%s'."), 
-						*LocalObjectPaths.Last().ToString(), Buffer);
-					return false;
+					if (!TempObjectPath.GetSubPathString().IsEmpty())
+					{
+						// Parse error, these should only have top-level paths
+						ErrorText->Logf(ELogVerbosity::Error, TEXT("FAssetBundleEntry::ImportTextItem - Asset bundle entries should all be top level objects, but found a subobject path '%s'. When parsing import text '%s'."), 
+							*TempObjectPath.ToString(), Buffer);
+						return false;
+					}
+					LocalObjectPaths.Emplace(MoveTemp(TempObjectPath));
+				}
+				else
+				{
+					// Import formatted struct e.g. (PackageName="/Path/To/Package", AssetName="AssetName")
+					TempTopLevelPath.Reset();
+					const TCHAR* NewBuffer = TopLevelAssetPathStruct->ImportText(BufferIt, &TempTopLevelPath, Parent, PortFlags, &SilenceErrors, TopLevelPathStructName);;
+					if (!NewBuffer)
+					{
+						ErrorText->Logf(ELogVerbosity::Error, TEXT("FAssetBundleEntry::ImportTextItem - unterminated or ill formed AssetPaths list while importing asset bundle entry: '%s'. When parsing import text '%s'."), BufferIt, Buffer);
+						break;
+					}
+					BufferIt = NewBuffer;
+					LocalObjectPaths.Emplace(TempTopLevelPath);
 				}
 
 				if (*BufferIt == ',')
@@ -591,7 +619,7 @@ bool FAssetBundleEntryImportExportTextTest::RunTest(const FString& Parameters)
 	FString SubobjectEntryString = TEXT("(BundleName=\"Subobject\",AssetPaths=(\"/Game/Characters/Steve.Steve:Hat\"))");
 	const TCHAR* SubobjectBuffer = *SubobjectEntryString;
 	FAssetBundleEntry SubobjectImported;
-	TestFalse(TEXT("Subobject bunld does not import from text"), SubobjectImported.ImportTextItem(SubobjectBuffer, PPF_Delimited, nullptr, GLog->Get()));
+	TestFalse(TEXT("Subobject bundle does not import from text"), SubobjectImported.ImportTextItem(SubobjectBuffer, PPF_Delimited, nullptr, GLog->Get()));
 	TestEqual(TEXT("Subobject bundle does not import from text"), SubobjectImported, DefaultEntry);
 
 	AddExpectedError(TEXT("FAssetBundleEntry::ImportTextItem - Unterminated BundleName string"));
@@ -693,4 +721,76 @@ bool FLegacyAssetBundleEntryTest::RunTest(const FString& Parameters)
 
 	return true;
 }
+
+IMPLEMENT_SIMPLE_AUTOMATION_TEST(FAssetBundlDataImportExportTextTest, "Engine.AssetRegistry.AssetBundleData.ImportExportText", EAutomationTestFlags::ApplicationContextMask | EAutomationTestFlags::EngineFilter);
+bool FAssetBundlDataImportExportTextTest::RunTest(const FString& Parameters)
+{
+	UScriptStruct* Struct = FindObject<UScriptStruct>(nullptr, TEXT("/Script/CoreUObject.AssetBundleData"));
+	if (!TestTrue("FAssetBundleData::StaticStruct", Struct != nullptr))
+	{
+		return false;
+	}
+
+	FAssetBundleData AssetBundles;
+	FAssetBundleEntry& Entry = AssetBundles.Bundles.Add_GetRef(FAssetBundleEntry{});
+	Entry.BundleName = "TestBundle";
+	Entry.AssetPaths.Add(FTopLevelAssetPath("/Game/Characters/Steve.Steve"));
+
+	constexpr const TCHAR* NullTChar = (const TCHAR*)nullptr;
+
+	// Round trip a data with a single asset bundle with a single asset
+	FString Exported;
+	Struct->ExportText(Exported, &AssetBundles, nullptr, nullptr, PPF_None, nullptr);
+	if (TestFalse("Single-asset bundle: Export with no defaults: Exported text not empty", Exported.IsEmpty()))
+	{
+		FAssetBundleData Imported1;
+		if (TestNotEqual("Single-asset bundle: Export with no defaults: Asset bundle data re-imports successfully",
+						 Struct->ImportText(*Exported, &Imported1, nullptr, PPF_None, GLog->Get(), Struct->GetName()), NullTChar))
+		{
+			TestEqual("Single-asset bundle: Export with no defaults: Re-imported asset bundle data matches", Imported1, AssetBundles);
+		}
+	}
+
+	Exported.Reset();
+	Struct->ExportText(Exported, &AssetBundles, &AssetBundles, nullptr, PPF_None, nullptr); // Export with reference to defaults as the same pointer
+	if (TestFalse("Single-asset bundle: Export with defaults: Exported text not empty", Exported.IsEmpty()))
+	{
+		FAssetBundleData Imported2;
+		if (TestNotEqual("Single-asset bundle: Export with defaults: Asset bundle data re-imports successfully",
+						 Struct->ImportText(*Exported, &Imported2, nullptr, PPF_None, GLog->Get(), Struct->GetName()), NullTChar))
+		{
+			TestEqual("Single-asset bundle: Export with defaults: Re-imported asset bundle data matches", Imported2, AssetBundles);
+		}
+	}
+
+	// Test with multiple asset paths
+	Entry.BundleName = "TestBundle";
+	Entry.AssetPaths.Add(FTopLevelAssetPath("/Game/Characters/Geoff.Geoff"));
+	Entry.AssetPaths.Add(FTopLevelAssetPath("/Game/Characters/Mary.Mary"));
+
+	Exported.Reset();
+	Struct->ExportText(Exported, &AssetBundles, nullptr, nullptr, PPF_None, nullptr);
+	if (TestFalse("Multi-asset bundle: Export with no defaults: Exported text not empty", Exported.IsEmpty()))
+	{
+		FAssetBundleData Imported3;
+		if (TestNotEqual("Multi-asset bundle: Export with no defaults: Asset bundle data re-imports successfully",
+						 Struct->ImportText(*Exported, &Imported3, nullptr, PPF_None, GLog->Get(), Struct->GetName()), NullTChar))
+		{
+			TestEqual("Multi-asset bundle: Export with no defaults: Re-imported asset bundle data matches", Imported3, AssetBundles);
+		}
+	}
+
+	Exported.Reset();
+	Struct->ExportText(Exported, &AssetBundles, &AssetBundles, nullptr, PPF_None, nullptr); // Export with reference to defaults as the same pointer
+	if (TestFalse("Multi-asset bundle: Export with defaults: Exported text not empty", Exported.IsEmpty()))
+	{
+		FAssetBundleData Imported4;
+		if (TestNotEqual("Multi-asset bundle: Export with defaults: Asset bundle data re-imports successfully",
+						 Struct->ImportText(*Exported, &Imported4, nullptr, PPF_None, GLog->Get(), Struct->GetName()), NullTChar))
+		{
+			TestEqual("Multi-asset bundle: Export with defaults: Re-imported asset bundle data matches", Imported4, AssetBundles);
+		}
+	}
+	return true;
+}
 #endif // WITH_DEV_AUTOMATION_TESTS
