[CoreUObject] fix for not forcing PreLoad on redirected objects during a lazy load resolve

PreLoad needs to be forced otherwise the object could get used before it has been loaded

diff --git a/Engine/Source/Runtime/CoreUObject/Private/UObject/ObjectHandle.cpp b/Engine/Source/Runtime/CoreUObject/Private/UObject/ObjectHandle.cpp
--- a/Engine/Source/Runtime/CoreUObject/Private/UObject/ObjectHandle.cpp
+++ b/Engine/Source/Runtime/CoreUObject/Private/UObject/ObjectHandle.cpp
@@ -408,6 +408,10 @@ UClass* ResolveObjectRefClass(const FObjectRef& ObjectRef, uint32 LoadFlags /*=
 				if (UObjectRedirector* Redirector = dynamic_cast<UObjectRedirector*>(ClassObject))
 				{
 					ClassObject = (UClass*)Redirector->DestinationObject;
+					if (ClassObject != nullptr)
+					{
+						ClassPackage = ClassObject->GetPackage();
+					}
 				}
 				if (ClassObject->HasAnyFlags(RF_NeedLoad) && ClassPackage->GetLinker())
 				{
@@ -473,6 +477,10 @@ UObject* ResolveObjectRef(const FObjectRef& ObjectRef, uint32 LoadFlags /*= LOAD
 		if (UObjectRedirector* Redirector = dynamic_cast<UObjectRedirector*>(CurrentObject))
 		{
 			CurrentObject = Redirector->DestinationObject;
+			if (CurrentObject != nullptr)
+			{
+				TargetPackage = CurrentObject->GetPackage();
+			}
 		}
 
 		if (!CurrentObject && !TargetPackage->IsFullyLoaded() && TargetPackage->GetLinker() && TargetPackage->GetLinker()->IsLoading())
@@ -492,6 +500,10 @@ UObject* ResolveObjectRef(const FObjectRef& ObjectRef, uint32 LoadFlags /*= LOAD
 			if (UObjectRedirector* Redirector = dynamic_cast<UObjectRedirector*>(CurrentObject))
 			{
 				CurrentObject = Redirector->DestinationObject;
+				if (CurrentObject != nullptr)
+				{
+					TargetPackage = CurrentObject->GetPackage();
+				}
 			}
 		}
 
