Update comment.

diff --git a/Engine/Source/Runtime/CoreUObject/Public/UObject/SoftObjectPath.h b/Engine/Source/Runtime/CoreUObject/Public/UObject/SoftObjectPath.h
--- a/Engine/Source/Runtime/CoreUObject/Public/UObject/SoftObjectPath.h
+++ b/Engine/Source/Runtime/CoreUObject/Public/UObject/SoftObjectPath.h
@@ -34,9 +34,9 @@ struct FPropertyTag;
 struct FUObjectSerializeContext;
 
 /**
- * A struct that contains a string reference to an object, either a top level asset or a subobject.
+ * A struct that contains a string reference to an object, either a package, a top level asset or a subobject.
  * This can be used to make soft references to assets that are loaded on demand.
- * This is stored internally as an FName pointing to the top level asset (/package/path.assetname) and an option a string subobject path.
+ * This is stored internally as an FTopLevelAssetPath pointing to the top level asset (/package/path.assetname) and an optional string subobject path.
  * If the MetaClass metadata is applied to a FProperty with this the UI will restrict to that type of asset.
  */
 struct COREUOBJECT_API FSoftObjectPath
