Fix for compile errors under HoloLens/W11 builds. Need to use global namespace for type hash functions under this build environment.

diff --git a/Engine/Source/Runtime/CoreUObject/Private/AssetRegistry/AssetData.cpp b/Engine/Source/Runtime/CoreUObject/Private/AssetRegistry/AssetData.cpp
--- a/Engine/Source/Runtime/CoreUObject/Private/AssetRegistry/AssetData.cpp
+++ b/Engine/Source/Runtime/CoreUObject/Private/AssetRegistry/AssetData.cpp
@@ -127,7 +127,7 @@ public:
 		uint32 Hash = 0;
 		for (int32 ChunkID : InEntry.ChunkArray)
 		{
-			Hash = HashCombineFast(Hash, GetTypeHash(ChunkID));
+			Hash = ::HashCombineFast(Hash, ::GetTypeHash(ChunkID));
 		}
 		return Hash;
 	}
