[CoreUObject] changed the default behaviour of FArchiveUObjectFromStructuredArchiveImpl to forward FObjectPtr to the inner archive to avoid loading objects
Inner archive can choose how to handle FObjectPtr

diff --git a/Engine/Source/Runtime/CoreUObject/Private/Serialization/ArchiveUObjectFromStructuredArchive.cpp b/Engine/Source/Runtime/CoreUObject/Private/Serialization/ArchiveUObjectFromStructuredArchive.cpp
--- a/Engine/Source/Runtime/CoreUObject/Private/Serialization/ArchiveUObjectFromStructuredArchive.cpp
+++ b/Engine/Source/Runtime/CoreUObject/Private/Serialization/ArchiveUObjectFromStructuredArchive.cpp
@@ -150,8 +150,19 @@ FArchive& FArchiveUObjectFromStructuredArchiveImpl::operator<<(FWeakObjectPtr& V
 
 FArchive& FArchiveUObjectFromStructuredArchiveImpl::operator<<(FObjectPtr& Value)
 {
-	// @TODO: OBJPTR: If we want lazy loading from text archives, this will need to be modified to not route through the default raw pointer serialization path.
-	return FArchiveUObject::SerializeObjectPtr(*this, Value);
+	OpenArchive();
+
+	if (InnerArchive.IsTextFormat())
+	{
+		UObject* Obj = Value.Get();
+		*this << Obj;
+		Value = Obj;
+	}
+	else
+	{
+		InnerArchive << Value;
+	}
+	return *this;
 }
 
 void FArchiveUObjectFromStructuredArchiveImpl::PushFileRegionType(EFileRegionType Type)
