[CoreUObject] changed DuplicateDataReader/Writer to handle TObjectPtr

diff --git a/Engine/Source/Runtime/CoreUObject/Private/Serialization/DuplicateDataReader.cpp b/Engine/Source/Runtime/CoreUObject/Private/Serialization/DuplicateDataReader.cpp
--- a/Engine/Source/Runtime/CoreUObject/Private/Serialization/DuplicateDataReader.cpp
+++ b/Engine/Source/Runtime/CoreUObject/Private/Serialization/DuplicateDataReader.cpp
@@ -76,6 +76,23 @@ FArchive& FDuplicateDataReader::operator<<( UObject*& Object )
 	return *this;
 }
 
+FArchive& FDuplicateDataReader::operator<<(FObjectPtr& ObjectPtr)
+{
+	FObjectHandle& Handle  = ObjectPtr.GetHandleRef();
+	Serialize(&Handle, sizeof(FObjectHandle));
+
+	if (!IsObjectHandleNull(Handle) && IsObjectHandleResolved(Handle))
+	{
+		UObject* SourceObject = ReadObjectHandlePointerNoCheck(Handle);
+		FDuplicatedObject ObjectInfo = DuplicatedObjectAnnotation.GetAnnotation(SourceObject);
+		if (!ObjectInfo.IsDefault())
+		{
+			ObjectPtr = ObjectInfo.DuplicatedObject.GetEvenIfUnreachable();
+		}
+	}
+	return *this;
+}
+
 FArchive& FDuplicateDataReader::operator<<(FLazyObjectPtr& LazyObjectPtr)
 {
 	FArchive& Ar = *this;
diff --git a/Engine/Source/Runtime/CoreUObject/Private/Serialization/DuplicateDataWriter.cpp b/Engine/Source/Runtime/CoreUObject/Private/Serialization/DuplicateDataWriter.cpp
--- a/Engine/Source/Runtime/CoreUObject/Private/Serialization/DuplicateDataWriter.cpp
+++ b/Engine/Source/Runtime/CoreUObject/Private/Serialization/DuplicateDataWriter.cpp
@@ -78,6 +78,22 @@ FArchive& FDuplicateDataWriter::operator<<(UObject*& Object)
 	return *this;
 }
 
+
+FArchive& FDuplicateDataWriter::operator<<(FObjectPtr& ObjectPtr)
+{
+	if (ObjectPtr.IsResolved())
+	{
+		UObject* Object = ObjectPtr.Get();
+		*this << Object;
+	}
+	else
+	{
+		FObjectHandle& Handle = ObjectPtr.GetHandleRef();
+		Serialize(&Handle, sizeof(FObjectHandle));
+	}
+	return *this;
+}
+
 FArchive& FDuplicateDataWriter::operator<<(FLazyObjectPtr& LazyObjectPtr)
 {
 	if ( (GetPortFlags() & PPF_DuplicateForPIE) == 0 )
diff --git a/Engine/Source/Runtime/CoreUObject/Public/Serialization/DuplicatedDataReader.h b/Engine/Source/Runtime/CoreUObject/Public/Serialization/DuplicatedDataReader.h
--- a/Engine/Source/Runtime/CoreUObject/Public/Serialization/DuplicatedDataReader.h
+++ b/Engine/Source/Runtime/CoreUObject/Public/Serialization/DuplicatedDataReader.h
@@ -30,10 +30,11 @@ private:
 
 	//~ Begin FArchive Interface.
 
-	virtual FArchive& operator<<(FName& N);
-	virtual FArchive& operator<<(UObject*& Object);
-	virtual FArchive& operator<<(FLazyObjectPtr& LazyObjectPtr);
-	virtual FArchive& operator<<(FSoftObjectPath& SoftObjectPath);
+	virtual FArchive& operator<<(FName& N) override;
+	virtual FArchive& operator<<(UObject*& Object) override;
+	virtual FArchive& operator<<(FObjectPtr& Object) override;
+	virtual FArchive& operator<<(FLazyObjectPtr& LazyObjectPtr) override;
+	virtual FArchive& operator<<(FSoftObjectPath& SoftObjectPath) override;
 	
 	void SerializeFail();
 
diff --git a/Engine/Source/Runtime/CoreUObject/Public/Serialization/DuplicatedDataWriter.h b/Engine/Source/Runtime/CoreUObject/Public/Serialization/DuplicatedDataWriter.h
--- a/Engine/Source/Runtime/CoreUObject/Public/Serialization/DuplicatedDataWriter.h
+++ b/Engine/Source/Runtime/CoreUObject/Public/Serialization/DuplicatedDataWriter.h
@@ -46,6 +46,7 @@ private:
 	virtual FArchive& operator<<(UObject*& Object) override;
 	virtual FArchive& operator<<(FLazyObjectPtr& LazyObjectPtr) override;
 	virtual FArchive& operator<<(FField*& Field) override;
+	virtual FArchive& operator<<(FObjectPtr& Object) override;
 
 	virtual void Serialize(void* Data,int64 Num)
 	{
