Fix static analysis errors in unit test.

diff --git a/Engine/Source/Runtime/CoreUObject/Private/AssetRegistry/AssetBundleData.cpp b/Engine/Source/Runtime/CoreUObject/Private/AssetRegistry/AssetBundleData.cpp
--- a/Engine/Source/Runtime/CoreUObject/Private/AssetRegistry/AssetBundleData.cpp
+++ b/Engine/Source/Runtime/CoreUObject/Private/AssetRegistry/AssetBundleData.cpp
@@ -725,11 +725,7 @@ bool FLegacyAssetBundleEntryTest::RunTest(const FString& Parameters)
 IMPLEMENT_SIMPLE_AUTOMATION_TEST(FAssetBundlDataImportExportTextTest, "Engine.AssetRegistry.AssetBundleData.ImportExportText", EAutomationTestFlags::ApplicationContextMask | EAutomationTestFlags::EngineFilter);
 bool FAssetBundlDataImportExportTextTest::RunTest(const FString& Parameters)
 {
-	UScriptStruct* Struct = FindObject<UScriptStruct>(nullptr, TEXT("/Script/CoreUObject.AssetBundleData"));
-	if (!TestTrue("FAssetBundleData::StaticStruct", Struct != nullptr))
-	{
-		return false;
-	}
+	UScriptStruct* Struct = TBaseStructure<FAssetBundleData>::Get();
 
 	FAssetBundleData AssetBundles;
 	FAssetBundleEntry& Entry = AssetBundles.Bundles.Add_GetRef(FAssetBundleEntry{});
