Allow existing FString/FName properties to be replaced by FTopLevelAssetPath without losing serialized data

diff --git a/Engine/Source/Runtime/CoreUObject/Private/UObject/TopLevelAssetPath.cpp b/Engine/Source/Runtime/CoreUObject/Private/UObject/TopLevelAssetPath.cpp
--- a/Engine/Source/Runtime/CoreUObject/Private/UObject/TopLevelAssetPath.cpp
+++ b/Engine/Source/Runtime/CoreUObject/Private/UObject/TopLevelAssetPath.cpp
@@ -16,6 +16,14 @@
 // which are public in the temp struct and private in the actual FTopLevelAssetPath.
 // If UHT compiles with UE_IMPLEMENT_STRUCT("/Script/CoreUObject", TopLevelAssetPath) the generated code will fail to compile trying to access private struct members
 #if !HACK_HEADER_GENERATOR
+template<>
+struct TStructOpsTypeTraits<FTopLevelAssetPath> : public TStructOpsTypeTraitsBase2<FTopLevelAssetPath>
+{
+	enum
+	{
+		WithStructuredSerializeFromMismatchedTag = true,
+	};
+};
 UE_IMPLEMENT_STRUCT("/Script/CoreUObject", TopLevelAssetPath)
 #endif
 
@@ -160,6 +168,32 @@ bool FTopLevelAssetPath::TrySetPath(FAnsiStringView Path)
 	return TrySetPath(Wide);
 }
 
+bool FTopLevelAssetPath::SerializeFromMismatchedTag(const FPropertyTag& Tag, FStructuredArchive::FSlot Slot)
+{
+	if (Tag.Type == NAME_NameProperty)
+	{
+		FName Name;
+		Slot << Name;
+		
+		FNameBuilder NameBuilder(Name);
+		TrySetPath(NameBuilder.ToView());
+
+		return true;
+	}
+	
+	if (Tag.Type == NAME_StrProperty)
+	{
+		FString String;
+		Slot << String;
+
+		TrySetPath(String);
+		
+		return true;
+	}
+
+	return false;
+}
+
 
 
 #if WITH_DEV_AUTOMATION_TESTS 
diff --git a/Engine/Source/Runtime/CoreUObject/Public/UObject/TopLevelAssetPath.h b/Engine/Source/Runtime/CoreUObject/Public/UObject/TopLevelAssetPath.h
--- a/Engine/Source/Runtime/CoreUObject/Public/UObject/TopLevelAssetPath.h
+++ b/Engine/Source/Runtime/CoreUObject/Public/UObject/TopLevelAssetPath.h
@@ -16,6 +16,7 @@
 
 #include <type_traits>
 
+struct FPropertyTag;
 class CORE_API FString;
 class UObject;
 
@@ -166,6 +167,8 @@ struct COREUOBJECT_API FTopLevelAssetPath
 		return HashCombineFast(GetTypeHash(This.PackageName), GetTypeHash(This.AssetName));
 	}
 
+	bool SerializeFromMismatchedTag(const FPropertyTag& Tag, FStructuredArchive::FSlot Slot);
+
 private:
 	/** Name of the package containing the asset e.g. /Path/To/Package */
 	FName PackageName;
