Allow existing FName properties to be replaced by FSoftObjectPath/Ptr without losing serialized data

diff --git a/Engine/Source/Runtime/CoreUObject/Private/UObject/SoftObjectPath.cpp b/Engine/Source/Runtime/CoreUObject/Private/UObject/SoftObjectPath.cpp
--- a/Engine/Source/Runtime/CoreUObject/Private/UObject/SoftObjectPath.cpp
+++ b/Engine/Source/Runtime/CoreUObject/Private/UObject/SoftObjectPath.cpp
@@ -463,6 +463,15 @@ bool SerializeFromMismatchedTagTemplate(FString& Output, const FPropertyTag& Tag
 		}
 		return true;
 	}
+	else if (Tag.Type == NAME_NameProperty)
+	{
+		FName Name;
+		Slot << Name;
+
+		FNameBuilder NameBuilder(Name);
+		Output = NameBuilder.ToView();
+		return true;
+	}
 	else if (Tag.Type == NAME_StrProperty)
 	{
 		FString String;
