Defer import class verification to creation time when appropriate so that we don't emit load linker warnings for false positives.

diff --git a/Engine/Source/Runtime/CoreUObject/Private/UObject/LinkerLoad.cpp b/Engine/Source/Runtime/CoreUObject/Private/UObject/LinkerLoad.cpp
--- a/Engine/Source/Runtime/CoreUObject/Private/UObject/LinkerLoad.cpp
+++ b/Engine/Source/Runtime/CoreUObject/Private/UObject/LinkerLoad.cpp
@@ -3422,8 +3422,8 @@ bool FLinkerLoad::VerifyImportInner(const int32 ImportIndex, FString& WarningSuf
 									{
 										if ((SourceExportOuter.ClassName != OuterImport.ClassName) || (SourceExportOuter.ClassPackage != OuterImport.ClassPackage))
 										{
-											UE_ASSET_LOG(LogLinker, Warning, PackagePath, TEXT("Resolved outer import with a different class: import class '%s.%s', package class '%s.%s'. Resave to fix."), 
-												*OuterImport.ClassPackage.ToString(), *OuterImport.ClassName.ToString(), *SourceExportOuter.ClassPackage.ToString(), *SourceExportOuter.ClassName.ToString());
+											// Since we don't have an exact match, do some additional verification when we create the outer import (where we have a valid class object).
+											ImportsToVerifyOnCreate.Add(Import.OuterIndex.ToImport());
 										}
 									}
 								}
@@ -3498,11 +3498,8 @@ bool FLinkerLoad::VerifyImportInner(const int32 ImportIndex, FString& WarningSuf
 						// Found the FObjectExport for this import
 						if ((Import.ClassName != Import.SourceLinker->GetExportClassName(j)) || (Import.ClassPackage != Import.SourceLinker->GetExportClassPackage(j)))
 						{
-							UE_ASSET_LOG(LogLinker, Warning, PackagePath, TEXT("Resolved import with name '%s' from '%s' with a different class: import class '%s.%s', package class '%s.%s'. Resave to fix."), 
-								*Import.ObjectName.ToString(),
-								*Import.SourceLinker->GetPackagePath().GetPackageName(),
-								*Import.ClassPackage.ToString(), *Import.ClassName.ToString(), 
-								*Import.SourceLinker->GetExportClassPackage(j).ToString(), *Import.SourceLinker->GetExportClassName(j).ToString());
+							// Since we don't have an exact match, do some additional verification when we create the import (where we have a valid class object).
+							ImportsToVerifyOnCreate.Add(ImportIndex);
 						}
 
 						Import.SourceIndex = j;
@@ -5283,6 +5280,26 @@ UObject* FLinkerLoad::CreateImport( int32 Index )
 			const FString OuterName = Import.OuterIndex.IsNull() ? LinkerRoot->GetFullName() : GetFullImpExpName(Import.OuterIndex);
 			UE_LOG(LogLinker, Verbose, TEXT("Failed to resolve import '%d' named '%s' in '%s'"), Index, *Import.ObjectName.ToString(), *OuterName);
 		}
+		else if (ImportsToVerifyOnCreate.Contains(Index))
+		{
+			const UClass* ExpectedImportClass = nullptr;
+			if (UPackage* ImportClassPackage = FindObjectFast<UPackage>(nullptr, Import.ClassPackage, false))
+			{
+				ExpectedImportClass = FindObjectFast<UClass>(ImportClassPackage, Import.ClassName, false);
+			}
+
+			// Verify that the resolved import object's class is serialization-compatible with the expected result. Data loss will otherwise occur on load if this is not satisfied, so we warn about it. A re-save is required to fix up the import table and suppress this warning.
+			if (!ExpectedImportClass || !Import.XObject->GetClass()->IsChildOf(ExpectedImportClass))
+			{
+				UE_ASSET_LOG(LogLinker, Warning, PackagePath, TEXT("Resolved import with name '%s' from '%s' with a different class: import class '%s.%s', package class '%s.%s'. Resave to fix."), 
+								*Import.ObjectName.ToString(),
+								*Import.SourceLinker->GetPackagePath().GetPackageName(),
+								*Import.ClassPackage.ToString(), *Import.ClassName.ToString(), 
+								*Import.SourceLinker->GetExportClassPackage(Import.SourceIndex).ToString(), *Import.SourceLinker->GetExportClassName(Import.SourceIndex).ToString());
+			}
+
+			ImportsToVerifyOnCreate.Remove(Index);
+		}
 	}
 	return Import.XObject;
 }
diff --git a/Engine/Source/Runtime/CoreUObject/Public/UObject/LinkerLoad.h b/Engine/Source/Runtime/CoreUObject/Public/UObject/LinkerLoad.h
--- a/Engine/Source/Runtime/CoreUObject/Public/UObject/LinkerLoad.h
+++ b/Engine/Source/Runtime/CoreUObject/Public/UObject/LinkerLoad.h
@@ -223,6 +223,9 @@ private:
 	/** The trailer for the package */
 	TUniquePtr<UE::FPackageTrailer> PackageTrailer;
 
+	/** Set of imports that require additional verification at creation time. */
+	TSet<int32> ImportsToVerifyOnCreate;
+
 	// Helper function to access the InstancingContext IsInstanced, 
 	bool IsContextInstanced() const;
 
