Small cleanup of dead code in blueprint support

diff --git a/Engine/Source/Runtime/CoreUObject/Private/Blueprint/BlueprintSupport.cpp b/Engine/Source/Runtime/CoreUObject/Private/Blueprint/BlueprintSupport.cpp
--- a/Engine/Source/Runtime/CoreUObject/Private/Blueprint/BlueprintSupport.cpp
+++ b/Engine/Source/Runtime/CoreUObject/Private/Blueprint/BlueprintSupport.cpp
@@ -449,7 +449,7 @@ void FBlueprintSupport::ValidateNoExternalRefsToSkeletons()
  ******************************************************************************/
 
 #if WITH_EDITOR
-UClass* FScopedClassDependencyGather::BatchAuthorityClass = NULL;
+UClass* FScopedClassDependencyGather::BatchAuthorityClass = nullptr;
 TArray<UClass*> FScopedClassDependencyGather::BatchClassDependencies;
 
 FScopedClassDependencyGather::FScopedClassDependencyGather(UClass* ClassToGather, FUObjectSerializeContext* InLoadContext)
@@ -459,7 +459,7 @@ FScopedClassDependencyGather::FScopedClassDependencyGather(UClass* ClassToGather
 	// Do NOT track duplication dependencies, as these are intermediate products that we don't care about
 	if( !GIsDuplicatingClassForReinstancing )
 	{
-		if( BatchAuthorityClass == NULL )
+		if( BatchAuthorityClass == nullptr )
 		{
 			// If there is no current dependency authority, register this class as the authority, and reset the array
 			BatchAuthorityClass = ClassToGather;
@@ -480,41 +480,9 @@ FScopedClassDependencyGather::~FScopedClassDependencyGather()
 	// dependencies (unless compiling on load is explicitly disabled)
 	if( bAuthoritativeClass )
 	{
-		auto DependencyIter = BatchClassDependencies.CreateIterator();
-		// implemented as a lambda, to prevent duplicated code between 
-		// BatchAuthorityClass and BatchClassDependencies entries
-		auto RecompileClassLambda = [&DependencyIter](UClass* Class, FUObjectSerializeContext* InLoadContext)
-		{
-			Class->ConditionalRecompileClass(InLoadContext);
-
-			// because of the above call to ConditionalRecompileClass(), the 
-			// specified Class gets "cleaned and sanitized" (meaning its old 
-			// properties get moved to a TRASH class, and new ones are 
-			// constructed in their place)... the unfortunate side-effect of 
-			// this is that child classes that have already been linked are now
-			// referencing TRASH inherited properties; to resolve this issue, 
-			// here we go back through dependencies that were already recompiled
-			// and re-link any that are sub-classes
-			//
-			// @TODO: this isn't the most optimal solution to this problem; we 
-			//        should probably instead prevent CleanAndSanitizeClass()
-			//        from running for BytecodeOnly compiles (we would then need 
-			//        to block UField re-creation)... UE-14957 was created to 
-			//        track this issue
-			auto ReverseIt = DependencyIter;
-			for (--ReverseIt; ReverseIt.GetIndex() >= 0; --ReverseIt)
-			{
-				UClass* ProcessedDependency = *ReverseIt;
-				if (ProcessedDependency->IsChildOf(Class))
-				{
-					ProcessedDependency->StaticLink(/*bRelinkExistingProperties =*/true);
-				}
-			}
-		};
-
 		BatchAuthorityClass->ConditionalRecompileClass(LoadContext);
 
-		BatchAuthorityClass = NULL;
+		BatchAuthorityClass = nullptr;
 	}
 }
 
@@ -1864,7 +1832,7 @@ void FLinkerLoad::FinalizeBlueprint(UClass* LoadClass)
 
 	// Now that imports have been resolved we optionally flush the compilation
 	// queue. This is only done for level blueprints, which will have instances
-	// of actors in them that cannot reliably be reinstanced on load (see useage
+	// of actors in them that cannot reliably be reinstanced on load (see usage
 	// of Scene pointers in things like UActorComponent::ExecuteRegisterEvents)
 	// - on load the Scene may not yet be created, meaning this code cannot 
 	// correctly be run. We could address that, but avoiding reinstancings is
