ZenLoader in editor:
* Fixed issue with LOAD_PackageForPIE getting lost when calling LoadPackageAsync from LoadPackageInternal
* Fixed issue with externally created packages not being tracked in the import store and unintentionally being loaded multiple times
* Fixed issue with instancing contexts not being propagated to imported packages

diff --git a/Engine/Source/Runtime/CoreUObject/Private/Serialization/AsyncLoading2.cpp b/Engine/Source/Runtime/CoreUObject/Private/Serialization/AsyncLoading2.cpp
--- a/Engine/Source/Runtime/CoreUObject/Private/Serialization/AsyncLoading2.cpp
+++ b/Engine/Source/Runtime/CoreUObject/Private/Serialization/AsyncLoading2.cpp
@@ -400,6 +400,15 @@ struct FPackageRequest
 	TUniquePtr<FLoadPackageAsyncDelegate> PackageLoadedDelegate;
 	FPackageRequest* Next = nullptr;
 
+	FLinkerInstancingContext* GetInstancingContext()
+	{
+#if WITH_EDITORONLY_DATA
+		return &InstancingContext;
+#else
+		return nullptr;
+#endif
+	}
+
 	static FPackageRequest Create(int32 RequestId, EPackageFlags PackageFlags, int32 PIEInstanceID, int32 Priority, const FLinkerInstancingContext* InstancingContext, const FPackagePath& PackagePath, FName CustomName, TUniquePtr<FLoadPackageAsyncDelegate> PackageLoadedDelegate)
 	{
 		return FPackageRequest
@@ -434,10 +443,6 @@ struct FAsyncPackageDesc2
 #if WITH_EDITOR
 	/** PIE instance ID this package belongs to, INDEX_NONE otherwise */
 	int32 PIEInstanceID;
-#endif
-#if WITH_EDITORONLY_DATA
-	/** Instancing context, maps original package to their instanced counterpart, used to remap imports. */
-	FLinkerInstancingContext InstancingContext;
 #endif
 	// The package id of the UPackage being loaded
 	// It will be used as key when tracking active async packages
@@ -453,18 +458,6 @@ struct FAsyncPackageDesc2
 	// Packages with a a custom name can't be imported
 	bool bCanBeImported;
 
-#if WITH_EDITORONLY_DATA
-	const FLinkerInstancingContext* GetInstancingContext() const
-	{
-		return &InstancingContext;
-	}
-#else
-	const FLinkerInstancingContext* GetInstancingContext() const
-	{
-		return nullptr;
-	}
-#endif 
-
 	static FAsyncPackageDesc2 FromPackageRequest(
 		FPackageRequest& Request,
 		FName UPackageName,
@@ -478,9 +471,6 @@ struct FAsyncPackageDesc2
 			Request.PackageFlags,
 #if WITH_EDITOR
 			Request.PIEInstanceID,
-#endif
-#if WITH_EDITORONLY_DATA
-			MoveTemp(Request.InstancingContext),
 #endif
 			FPackageId::FromName(UPackageName),
 			PackageIdToLoad,
@@ -509,9 +499,6 @@ struct FAsyncPackageDesc2
 			PKG_None,
 #if WITH_EDITOR
 			INDEX_NONE,
-#endif
-#if WITH_EDITORONLY_DATA
-			ImportingPackageDesc.InstancingContext,
 #endif
 			ImportedPackageId,
 			PackageIdToLoad,
@@ -912,15 +899,34 @@ public:
 		return *PackageRef;
 	}
 
-	inline FLoadedPackageRef& AddPackageRef(FPackageId PackageId)
+	inline FLoadedPackageRef& AddPackageRef(FPackageId PackageId, FName PackageNameIfKnown)
 	{
 		LLM_SCOPE_BYNAME(TEXT("AsyncLoadPackageStore"));
 
 		FLoadedPackageRef& PackageRef = Packages.FindOrAdd(PackageId);
-		// is this the first reference to a package that has been loaded earlier?
-		if (PackageRef.RefCount == 0 && PackageRef.Package)
+		// is this the first reference to a package that already exists?
+		if (PackageRef.RefCount == 0)
 		{
-			PackageRef.PinPublicExportsForGC();
+#if WITH_EDITOR
+			if (!PackageRef.Package && !PackageNameIfKnown.IsNone())
+			{
+				UPackage* FoundPackage = FindObjectFast<UPackage>(nullptr, PackageNameIfKnown);
+				if (FoundPackage)
+				{
+					PackageRef.SetPackage(FoundPackage);
+					FoundPackage->SetCanBeImportedFlag(true);
+					FoundPackage->SetPackageId(PackageId);
+					if (FoundPackage->IsFullyLoaded())
+					{
+						PackageRef.SetAllPublicExportsLoaded();
+					}
+				}
+			}
+#endif
+			if (PackageRef.Package)
+			{
+				PackageRef.PinPublicExportsForGC();
+			}
 		}
 		++PackageRef.RefCount;
 		return PackageRef;
@@ -1210,12 +1216,7 @@ struct FAsyncPackageHeaderData
 	TArrayView<const uint8> ArcsData;
 	// Backed by allocation in FAsyncPackageData
 	TArrayView<FPackageId> ImportedPackageIds;
-	struct FImportPackageNames
-	{
-		FName UPackageName;
-		FName PackageNameToLoad;
-	};
-	TArrayView<FImportPackageNames> ImportedPackageNames;
+	TArrayView<FName> ImportedPackageNames;
 	TArrayView<FExportBundleHeader> ExportBundleHeaders;
 	TArrayView<FExportBundleEntry> ExportBundleEntries;
 
@@ -1335,19 +1336,16 @@ struct FPackageImportStore
 	}
 
 public:
-	void AddImportedPackageReferences(const TArrayView<const FPackageId>& ImportedPackageIds)
+	void AddImportedPackageReference(FPackageId ImportedPackageId, FName PackageNameIfKnown)
 	{
-		for (const FPackageId& ImportedPackageId : ImportedPackageIds)
-		{
-			LoadedPackageStore.AddPackageRef(ImportedPackageId);
-		}
+		LoadedPackageStore.AddPackageRef(ImportedPackageId, PackageNameIfKnown);
 	}
 
 	void AddPackageReference(const FAsyncPackageDesc2& Desc)
 	{
 		if (Desc.bCanBeImported)
 		{
-			FLoadedPackageRef& PackageRef = LoadedPackageStore.AddPackageRef(Desc.UPackageId);
+			FLoadedPackageRef& PackageRef = LoadedPackageStore.AddPackageRef(Desc.UPackageId, Desc.UPackageName);
 			PackageRef.ClearErrorFlags();
 		}
 	}
@@ -2271,7 +2269,8 @@ struct FAsyncPackage2
 	void StartLoading(FIoBatch& IoBatch);
 
 #if ALT2_ENABLE_LINKERLOAD_SUPPORT
-	void CreateLinker();
+	void InitializeLinkerLoadState(const FLinkerInstancingContext* InstancingContext);
+	void CreateLinker(const FLinkerInstancingContext* InstancingContext);
 	void DetachLinker();
 #endif
 
@@ -2280,7 +2279,7 @@ struct FAsyncPackage2
 #endif
 
 private:
-	void ImportPackagesRecursiveInner(FAsyncLoadingThreadState2& ThreadState, FIoBatch& IoBatch, FPackageStore& PackageStore, const TArrayView<const FPackageId>& ImportedPackageIds, const TArrayView<const FAsyncPackageHeaderData::FImportPackageNames>& ImportedPackageNames, int32& ImportedPackageIndex);
+	void ImportPackagesRecursiveInner(FAsyncLoadingThreadState2& ThreadState, FIoBatch& IoBatch, FPackageStore& PackageStore, const TArrayView<FPackageId>& ImportedPackageIds, const TArrayView<FName>& ImportedPackageNames, int32& ImportedPackageIndex);
 
 	uint8 PackageNodesMemory[EEventLoadNode2::Package_NumPhases * sizeof(FEventLoadNode2)];
 	TArrayView<FEventLoadNode2> PackageNodes;
@@ -3021,7 +3020,7 @@ private:
 		return new FAsyncPackage2(Desc, *this, GraphAllocator, EventSpecs.GetData());
 	}
 
-	void InitializeAsyncPackageFromPackageStore(FAsyncLoadingThreadState2& ThreadState, FIoBatch& IoBatch, FAsyncPackage2* AsyncPackage, const FPackageStoreEntry& PackageStoreEntry, const TArrayView<FAsyncPackageHeaderData::FImportPackageNames>& ImportedPackageNames)
+	void InitializeAsyncPackageFromPackageStore(FAsyncLoadingThreadState2& ThreadState, FIoBatch& IoBatch, FAsyncPackage2* AsyncPackage, const FPackageStoreEntry& PackageStoreEntry, const TArrayView<FName>& ImportedPackageNames)
 	{
 		TRACE_CPUPROFILER_EVENT_SCOPE(InitializeAsyncPackageFromPackageStore);
 		UE_ASYNC_PACKAGE_DEBUG(AsyncPackage->Desc);
@@ -3036,7 +3035,7 @@ private:
 		const uint64 ExportBundleEntriesMemSize = Align(sizeof(FExportBundleEntry) * ExportBundleEntriesCount, 8);
 		const int32 ImportedPackagesCount = PackageStoreEntry.ImportedPackageIds.Num();
 		const uint64 ImportedPackageIdsMemSize = Align(sizeof(FPackageId) * ImportedPackagesCount, 8);
-		const uint64 ImportedPackageNamesMemSize = Align(sizeof(FAsyncPackageHeaderData::FImportPackageNames) * ImportedPackageNames.Num(), 8);
+		const uint64 ImportedPackageNamesMemSize = Align(sizeof(FName) * ImportedPackageNames.Num(), 8);
 #if WITH_EDITOR
 		const int32 OptionalSegmentExportCount = PackageStoreEntry.OptionalSegmentExportInfo.ExportCount;
 		const int32 OptionalSegmentExportBundleCount = PackageStoreEntry.OptionalSegmentExportInfo.ExportBundleCount;
@@ -3104,8 +3103,8 @@ private:
 		DataPtr += ImportedPackageIdsMemSize;
 		if (ImportedPackageNamesMemSize)
 		{
-			HeaderData.ImportedPackageNames = MakeArrayView(reinterpret_cast<FAsyncPackageHeaderData::FImportPackageNames*>(DataPtr), ImportedPackageNames.Num());
-			FMemory::Memcpy((void*)HeaderData.ImportedPackageNames.GetData(), ImportedPackageNames.GetData(), sizeof(FAsyncPackageHeaderData::FImportPackageNames) * ImportedPackageNames.Num());
+			HeaderData.ImportedPackageNames = MakeArrayView(reinterpret_cast<FName*>(DataPtr), ImportedPackageNames.Num());
+			FMemory::Memcpy((void*)HeaderData.ImportedPackageNames.GetData(), ImportedPackageNames.GetData(), sizeof(FName) * ImportedPackageNames.Num());
 			DataPtr += ImportedPackageNamesMemSize;
 		}
 #if WITH_EDITOR
@@ -3121,10 +3120,8 @@ private:
 			OptionalSegmentHeaderData.ImportedPackageIds = MakeArrayView(reinterpret_cast<FPackageId*>(DataPtr), OptionalSegmentImportedPackagesCount);
 			FMemory::Memcpy((void*)OptionalSegmentHeaderData.ImportedPackageIds.GetData(), PackageStoreEntry.OptionalSegmentImportedPackageIds.GetData(), sizeof(FPackageId) * OptionalSegmentImportedPackagesCount);
 			DataPtr += OptionalSegmentImportedPackageIdsMemSize;
-			AsyncPackage->ImportStore.AddImportedPackageReferences(OptionalSegmentHeaderData.ImportedPackageIds);
 		}
 #endif
-		AsyncPackage->ImportStore.AddImportedPackageReferences(HeaderData.ImportedPackageIds);
 		check(DataPtr - Data.MemoryBuffer == MemoryBufferSize);
 		AsyncPackage->CreateExportBundleNodes(EventSpecs.GetData());
 
@@ -3370,7 +3367,7 @@ bool FAsyncLoadingThread2::CreateAsyncPackagesFromQueue(FAsyncLoadingThreadState
 		FPackageStoreEntry PackageEntry;
 		if (EPackageStoreEntryStatus::Ok == PackageStore.GetPackageStoreEntry(PendingPackage->Desc.PackageIdToLoad, PackageEntry))
 		{
-			InitializeAsyncPackageFromPackageStore(ThreadState, IoBatch, PendingPackage, PackageEntry, TArrayView<FAsyncPackageHeaderData::FImportPackageNames>());
+			InitializeAsyncPackageFromPackageStore(ThreadState, IoBatch, PendingPackage, PackageEntry, TArrayView<FName>());
 			PendingPackage->StartLoading(IoBatch);
 			It.RemoveCurrent();
 		}
@@ -3490,12 +3487,12 @@ bool FAsyncLoadingThread2::CreateAsyncPackagesFromQueue(FAsyncLoadingThreadState
 #if ALT2_ENABLE_LINKERLOAD_SUPPORT
 						if (!bIsZenPackage)
 						{
-							Package->LinkerLoadState.Emplace();
+							Package->InitializeLinkerLoadState(Request.GetInstancingContext());
 						}
 						else
 #endif
 						{
-							InitializeAsyncPackageFromPackageStore(ThreadState, IoBatch, Package, PackageEntry, TArrayView<FAsyncPackageHeaderData::FImportPackageNames>());
+							InitializeAsyncPackageFromPackageStore(ThreadState, IoBatch, Package, PackageEntry, TArrayView<FName>());
 						}
 						Package->StartLoading(IoBatch);
 					}
@@ -4038,7 +4035,7 @@ void FGlobalImportStore::RegistrationComplete()
 	ScriptObjects.Shrink();
 }
 
-void FAsyncPackage2::ImportPackagesRecursiveInner(FAsyncLoadingThreadState2& ThreadState, FIoBatch& IoBatch, FPackageStore& PackageStore, const TArrayView<const FPackageId>& ImportedPackageIds, const TArrayView<const FAsyncPackageHeaderData::FImportPackageNames>& ImportedPackageNames, int32& ImportedPackageIndex)
+void FAsyncPackage2::ImportPackagesRecursiveInner(FAsyncLoadingThreadState2& ThreadState, FIoBatch& IoBatch, FPackageStore& PackageStore, const TArrayView<FPackageId>& ImportedPackageIds, const TArrayView<FName>& ImportedPackageNames, int32& ImportedPackageIndex)
 {
 	const int32 ImportedPackageCount = ImportedPackageIds.Num();
 	if (!ImportedPackageCount)
@@ -4050,21 +4047,26 @@ void FAsyncPackage2::ImportPackagesRecursiveInner(FAsyncLoadingThreadState2& Thr
 	for (int32 LocalImportedPackageIndex = 0; LocalImportedPackageIndex < ImportedPackageCount; ++LocalImportedPackageIndex)
 	{
 		FPackageId ImportedPackageId = ImportedPackageIds[LocalImportedPackageIndex];
-		const FAsyncPackageHeaderData::FImportPackageNames* NamesForPackage = bHasImportedPackageNames ? &ImportedPackageNames[LocalImportedPackageIndex] : nullptr;
 		EPackageStoreEntryStatus ImportedPackageStatus = EPackageStoreEntryStatus::Missing;
 		FPackageStoreEntry ImportedPackageEntry;
-		FName ImportedPackageUPackageName;
-		FName ImportedPackageNameToLoad;
+		FName ImportedPackageUPackageName = bHasImportedPackageNames ? ImportedPackageNames[LocalImportedPackageIndex] : NAME_None;
+		FName ImportedPackageNameToLoad = ImportedPackageUPackageName;
 		FPackageId ImportedPackageIdToLoad = ImportedPackageId;
-		if (NamesForPackage)
+#if WITH_EDITORONLY_DATA && ALT2_ENABLE_LINKERLOAD_SUPPORT
+		if (bHasImportedPackageNames && LinkerLoadState.IsSet())
 		{
-			if (NamesForPackage->UPackageName != NamesForPackage->PackageNameToLoad)
+			// Use the instancing context remap from disk package name to upackage name
+			const FLinkerInstancingContext& InstancingContext = LinkerLoadState->Linker->GetInstancingContext();
+			ImportedPackageUPackageName = InstancingContext.RemapPackage(ImportedPackageNameToLoad);
+			if (ImportedPackageUPackageName != ImportedPackageNameToLoad)
 			{
-				ImportedPackageIdToLoad = FPackageId::FromName(NamesForPackage->PackageNameToLoad);
+				ImportedPackageId = FPackageId::FromName(ImportedPackageUPackageName);
+				// Rewrite the import table
+				ImportedPackageIds[LocalImportedPackageIndex] = ImportedPackageId;
+				ImportedPackageNames[LocalImportedPackageIndex] = ImportedPackageUPackageName;
 			}
-			ImportedPackageUPackageName = NamesForPackage->UPackageName;
-			ImportedPackageNameToLoad = NamesForPackage->PackageNameToLoad;
 		}
+#endif
 
 		{
 			FName SourcePackageName;
@@ -4091,6 +4093,8 @@ void FAsyncPackage2::ImportPackagesRecursiveInner(FAsyncLoadingThreadState2& Thr
 		}
 #endif
 
+		ImportStore.AddImportedPackageReference(ImportedPackageId, ImportedPackageUPackageName);
+
 		FPackagePath ImportedPackagePath;
 #if ALT2_ENABLE_LINKERLOAD_SUPPORT
 		bool bIsZenPackage = !LinkerLoadState.IsSet();
@@ -4245,12 +4249,12 @@ void FAsyncPackage2::ImportPackagesRecursiveInner(FAsyncLoadingThreadState2& Thr
 #if ALT2_ENABLE_LINKERLOAD_SUPPORT
 				if (!bIsZenPackageImport)
 				{
-					ImportedPackage->LinkerLoadState.Emplace();
+					ImportedPackage->InitializeLinkerLoadState(bIsZenPackage ? nullptr : &LinkerLoadState->Linker->GetInstancingContext());
 				}
 				else
 #endif
 				{
-					AsyncLoadingThread.InitializeAsyncPackageFromPackageStore(ThreadState, IoBatch, ImportedPackage, ImportedPackageEntry, TArrayView<FAsyncPackageHeaderData::FImportPackageNames>());
+					AsyncLoadingThread.InitializeAsyncPackageFromPackageStore(ThreadState, IoBatch, ImportedPackage, ImportedPackageEntry, TArrayView<FName>());
 				}
 				ImportedPackage->StartLoading(IoBatch);
 			}
@@ -4286,7 +4290,7 @@ void FAsyncPackage2::ImportPackagesRecursive(FAsyncLoadingThreadState2& ThreadSt
 #if WITH_EDITOR
 	if (OptionalSegmentHeaderData.IsSet())
 	{
-		ImportPackagesRecursiveInner(ThreadState, IoBatch, PackageStore, OptionalSegmentHeaderData->ImportedPackageIds, TArrayView<FAsyncPackageHeaderData::FImportPackageNames>(), ImportedPackageIndex);
+		ImportPackagesRecursiveInner(ThreadState, IoBatch, PackageStore, OptionalSegmentHeaderData->ImportedPackageIds, OptionalSegmentHeaderData->ImportedPackageNames, ImportedPackageIndex);
 	}
 #endif
 
@@ -4307,7 +4311,14 @@ void FAsyncPackage2::ImportPackagesRecursive(FAsyncLoadingThreadState2& ThreadSt
 }
 
 #if ALT2_ENABLE_LINKERLOAD_SUPPORT
-void FAsyncPackage2::CreateLinker()
+void FAsyncPackage2::InitializeLinkerLoadState(const FLinkerInstancingContext* InstancingContext)
+{
+	LinkerLoadState.Emplace();
+	CreateUPackage();
+	CreateLinker(InstancingContext);
+}
+
+void FAsyncPackage2::CreateLinker(const FLinkerInstancingContext* InstancingContext)
 {
 	FLinkerLoad* Linker = FLinkerLoad::FindExistingLinkerForPackage(LinkerRoot);
 	if (!Linker)
@@ -4317,7 +4328,6 @@ void FAsyncPackage2::CreateLinker()
 		{
 			LinkerFlags |= LOAD_PackageForPIE;
 		}
-		const FLinkerInstancingContext* InstancingContext = Desc.GetInstancingContext();
 		FUObjectSerializeContext* LoadContext = GetSerializeContext();
 #if ALT2_ENABLE_NEW_ARCHIVE_FOR_LINKERLOAD
 		Linker = new FLinkerLoad(LinkerRoot, Desc.PackagePathToLoad, LinkerFlags, InstancingContext ? *InstancingContext : FLinkerInstancingContext());
@@ -4326,7 +4336,7 @@ void FAsyncPackage2::CreateLinker()
 		FLinkerLoadArchive2* Loader = new FLinkerLoadArchive2(Desc.PackagePathToLoad);
 		Linker->SetLoader(Loader, Loader->NeedsEngineVersionChecks());
 #else
-		Linker = FLinkerLoad::CreateLinkerAsync(LoadContext, LinkerRoot, Desc.PackagePathToLoad, LinkerFlags, Desc.GetInstancingContext(), TFunction<void()>([]() {}));
+		Linker = FLinkerLoad::CreateLinkerAsync(LoadContext, LinkerRoot, Desc.PackagePathToLoad, LinkerFlags, InstancingContext, TFunction<void()>([]() {}));
 #endif
 	}
 	check(Linker);
@@ -4334,9 +4344,6 @@ void FAsyncPackage2::CreateLinker()
 	LinkerLoadState->Linker = Linker;
 #if ALT2_ENABLE_NEW_ARCHIVE_FOR_LINKERLOAD
 	Linker->ResetStatusInfo();
-	static_cast<FLinkerLoadArchive2*>(Linker->GetLoader())->BeginRead(&GetPackageNode(EEventLoadNode2::Package_ProcessSummary));
-#else
-	GetPackageNode(EEventLoadNode2::Package_ProcessSummary).ReleaseBarrier();
 #endif
 }
 
@@ -4362,8 +4369,11 @@ void FAsyncPackage2::StartLoading(FIoBatch& IoBatch)
 #if ALT2_ENABLE_LINKERLOAD_SUPPORT
 	if (LinkerLoadState.IsSet())
 	{
-		CreateUPackage();
-		CreateLinker();
+#if ALT2_ENABLE_NEW_ARCHIVE_FOR_LINKERLOAD
+		static_cast<FLinkerLoadArchive2*>(LinkerLoadState->Linker->GetLoader())->BeginRead(&GetPackageNode(EEventLoadNode2::Package_ProcessSummary));
+#else
+		GetPackageNode(EEventLoadNode2::Package_ProcessSummary).ReleaseBarrier();
+#endif
 		return;
 	}
 #endif
@@ -4501,10 +4511,10 @@ EEventLoadNodeExecutionResult FAsyncPackage2::ProcessLinkerLoadPackageSummary(FA
 	TRACE_CPUPROFILER_EVENT_SCOPE(ProcessLinkerLoadPackageSummary);
 
 #if ALT2_ENABLE_NEW_ARCHIVE_FOR_LINKERLOAD
-	if (Package->LinkerLoadState->Linker->GetLoader()->IsError())
+	if (LinkerLoadState->Linker->GetLoader()->IsError())
 	{
-		Package->bLoadHasFailed = true;
-		return 	return EEventLoadNodeExecutionResult::Complete;
+		bLoadHasFailed = true;
+		return EEventLoadNodeExecutionResult::Complete;
 	}
 	LinkerLoadState->Linker->bUseTimeLimit = false;
 	FLinkerLoad::ELinkerStatus LinkerResult = LinkerLoadState->Linker->ProcessPackageSummary(nullptr);
@@ -4520,7 +4530,7 @@ EEventLoadNodeExecutionResult FAsyncPackage2::ProcessLinkerLoadPackageSummary(FA
 	check(LinkerLoadState->Linker->HasFinishedInitialization());
 
 	LinkerLoadState->LinkerLoadHeaderData.ImportMap.SetNum(LinkerLoadState->Linker->ImportMap.Num());
-	TArray<FAsyncPackageHeaderData::FImportPackageNames, TInlineAllocator<128>> ImportedPackageNames;
+	TArray<FName, TInlineAllocator<128>> ImportedPackageNames;
 	TArray<FPackageId, TInlineAllocator<128>> ImportedPackageIds;
 	for (int32 ImportIndex = 0; ImportIndex < LinkerLoadState->Linker->ImportMap.Num(); ++ImportIndex)
 	{
@@ -4567,7 +4577,7 @@ EEventLoadNodeExecutionResult FAsyncPackage2::ProcessLinkerLoadPackageSummary(FA
 			{
 				ImportedPackageNames.AddDefaulted();
 			}
-			ImportedPackageNames[ImportedPackageIndex].PackageNameToLoad = ImportPackageName;
+			ImportedPackageNames[ImportedPackageIndex] = ImportPackageName;
 			bool bIsPackageImport = ImportIndex == PackageImportIndex;
 			if (!bIsPackageImport || bImportHasPackageName)
 			{
@@ -4593,17 +4603,6 @@ EEventLoadNodeExecutionResult FAsyncPackage2::ProcessLinkerLoadPackageSummary(FA
 		}
 	}
 
-	const FLinkerInstancingContext& InstancingContext = LinkerLoadState->Linker->GetInstancingContext();
-	for (int32 PackageImportIndex = 0; PackageImportIndex < ImportedPackageNames.Num(); ++PackageImportIndex)
-	{
-		FAsyncPackageHeaderData::FImportPackageNames& ImportedPackage = ImportedPackageNames[PackageImportIndex];
-		ImportedPackage.UPackageName = InstancingContext.RemapPackage(ImportedPackage.PackageNameToLoad);
-		if (ImportedPackage.UPackageName != ImportedPackage.PackageNameToLoad)
-		{
-			ImportedPackageIds[PackageImportIndex] = FPackageId::FromName(ImportedPackage.UPackageName);
-		}
-	}
-
 	LinkerLoadState->LinkerLoadHeaderData.PublicExportHashes.SetNum(LinkerLoadState->Linker->ExportMap.Num());
 	for (int32 ExportIndex = 0; ExportIndex < LinkerLoadState->Linker->ExportMap.Num(); ++ExportIndex)
 	{
diff --git a/Engine/Source/Runtime/CoreUObject/Private/UObject/UObjectGlobals.cpp b/Engine/Source/Runtime/CoreUObject/Private/UObject/UObjectGlobals.cpp
--- a/Engine/Source/Runtime/CoreUObject/Private/UObject/UObjectGlobals.cpp
+++ b/Engine/Source/Runtime/CoreUObject/Private/UObject/UObjectGlobals.cpp
@@ -1457,7 +1457,11 @@ UPackage* LoadPackageInternal(UPackage* InOuter, const FPackagePath& PackagePath
 		}
 
 		GSyncLoadUsingAsyncLoaderCount++;
-		constexpr EPackageFlags PackageFlags = PKG_None;
+		EPackageFlags PackageFlags = PKG_None;
+		if ((!FApp::IsGame() || GIsEditor) && (LoadFlags & LOAD_PackageForPIE) != 0)
+		{
+			PackageFlags |= PKG_PlayInEditor;
+		}
 		constexpr int32 PIEInstanceID = INDEX_NONE;
 		constexpr int32 Priority = INT32_MAX;
 		int32 RequestID = LoadPackageAsync(PackagePath, PackageName, FLoadPackageAsyncDelegate(), PackageFlags, PIEInstanceID, Priority, InstancingContext);
