ZenLoader: Let the async loading thread progress on other packages if there's no more work to do for the current sync load

diff --git a/Engine/Source/Runtime/CoreUObject/Private/Serialization/AsyncLoading2.cpp b/Engine/Source/Runtime/CoreUObject/Private/Serialization/AsyncLoading2.cpp
--- a/Engine/Source/Runtime/CoreUObject/Private/Serialization/AsyncLoading2.cpp
+++ b/Engine/Source/Runtime/CoreUObject/Private/Serialization/AsyncLoading2.cpp
@@ -3856,6 +3856,10 @@ bool FAsyncLoadEventQueue2::ExecuteSyncLoadEvents(FAsyncLoadingThreadState2& Thr
 		}
 		Node = Next;
 	}
+	if (!bDidSomething && ThreadState.bIsAsyncLoadingThread)
+	{
+		return PopAndExecute(ThreadState);
+	}
 	return bDidSomething;
 }
 
