Core - Changes the base CPPTYPE of FObjectProperty from UObject* to TObjectPtr<UObject>, this ensures that no matter how people access the value be it as a FObjectProperty or not, calling GetPropertyValue won't return a UObject* directly - which effectively would reinterpret cast a TObjectPtr, if it is one, and then if lazy loading is enabled would potentially result in a crash from trying to treat a non-pointer as a TObjectPtr.
Class - When iterating the classes sparse values, don't accidentally load deferred objects, or crash because you're reinterpret casting the value of something that isn't a pointer.

diff --git a/Engine/Plugins/Experimental/StructUtils/Source/StructUtils/Private/ChunkedStructBuffer.cpp b/Engine/Plugins/Experimental/StructUtils/Source/StructUtils/Private/ChunkedStructBuffer.cpp
--- a/Engine/Plugins/Experimental/StructUtils/Source/StructUtils/Private/ChunkedStructBuffer.cpp
+++ b/Engine/Plugins/Experimental/StructUtils/Source/StructUtils/Private/ChunkedStructBuffer.cpp
@@ -124,8 +124,7 @@ void FChunkedStructBuffer::AddStructReferencedObjects(class FReferenceCollector&
 				// The iterator will recursively loop through all structs in structs too.
 				for (TPropertyValueIterator<const FObjectProperty> It(ScriptStruct, StructMemory); It; ++It)
 				{
-					UObject** ObjectPtr = static_cast<UObject**>(const_cast<void*>(It.Value()));
-					Collector.AddReferencedObject(*ObjectPtr);
+					Collector.AddReferencedObject(It.Key()->GetObjectPtrPropertyValueRef(It.Value()));
 				}
 			}
 		}
diff --git a/Engine/Plugins/Experimental/StructUtils/Source/StructUtils/Private/InstancedStructStream.cpp b/Engine/Plugins/Experimental/StructUtils/Source/StructUtils/Private/InstancedStructStream.cpp
--- a/Engine/Plugins/Experimental/StructUtils/Source/StructUtils/Private/InstancedStructStream.cpp
+++ b/Engine/Plugins/Experimental/StructUtils/Source/StructUtils/Private/InstancedStructStream.cpp
@@ -124,8 +124,7 @@ void FInstancedStructStream::AddStructReferencedObjects(class FReferenceCollecto
 				// The iterator will recursively loop through all structs in structs too.
 				for (TPropertyValueIterator<const FObjectProperty> It(ScriptStruct, StructMemory); It; ++It)
 				{
-					UObject** ObjectPtr = static_cast<UObject**>(const_cast<void*>(It.Value()));
-					Collector.AddReferencedObject(*ObjectPtr);
+					Collector.AddReferencedObject(It.Key()->GetObjectPtrPropertyValueRef(It.Value()));
 				}
 			}
 		}
diff --git a/Engine/Source/Runtime/CoreUObject/Private/UObject/Class.cpp b/Engine/Source/Runtime/CoreUObject/Private/UObject/Class.cpp
--- a/Engine/Source/Runtime/CoreUObject/Private/UObject/Class.cpp
+++ b/Engine/Source/Runtime/CoreUObject/Private/UObject/Class.cpp
@@ -3974,8 +3974,8 @@ void UClass::AddReferencedObjects(UObject* InThis, FReferenceCollector& Collecto
 			// The iterator will recursively loop through all structs in structs too.
 			for (TPropertyValueIterator<FObjectProperty> It(This->SparseClassDataStruct, This->SparseClassData); It; ++It)
 			{
-				UObject** ObjectPtr = static_cast<UObject**>(const_cast<void*>(It.Value()));
-				Collector.AddReferencedObject(*ObjectPtr);
+				// @TODO: OBJPTR: 
+				Collector.AddReferencedObject(It.Key()->GetObjectPtrPropertyValueRef(It.Value()));
 			}
 		}
 	}
diff --git a/Engine/Source/Runtime/CoreUObject/Private/UObject/PropertyBaseObject.cpp b/Engine/Source/Runtime/CoreUObject/Private/UObject/PropertyBaseObject.cpp
--- a/Engine/Source/Runtime/CoreUObject/Private/UObject/PropertyBaseObject.cpp
+++ b/Engine/Source/Runtime/CoreUObject/Private/UObject/PropertyBaseObject.cpp
@@ -184,7 +184,7 @@ FString FObjectPropertyBase::GetExportPath(FTopLevelAssetPath ClassPathName, con
 	return FString::Printf(TEXT("%s.%s'%s'"), *ClassPathName.GetPackageName().ToString(), *ClassPathName.GetAssetName().ToString(), *ObjectPathName);
 }
 
-FString FObjectPropertyBase::GetExportPath(const UObject* Object, const UObject* Parent /*= nullptr*/, const UObject* ExportRootScope /*= nullptr*/, const uint32 PortFlags /*= PPF_None*/)
+FString FObjectPropertyBase::GetExportPath(const TObjectPtr<const UObject>& Object, const UObject* Parent /*= nullptr*/, const UObject* ExportRootScope /*= nullptr*/, const uint32 PortFlags /*= PPF_None*/)
 {
 	bool bExportFullyQualified = true;
 
@@ -220,19 +220,19 @@ FString FObjectPropertyBase::GetExportPath(const UObject* Object, const UObject*
 
 	// Take the path name relative to the stopping point outermost ptr.
 	// This is so that cases like a component referencing a component in another actor work correctly when pasted
-	FString PathName = Object->GetPathName(StopOuter);
+	FString PathName = StopOuter ? Object->GetPathName(StopOuter) : Object.GetPathName();
 	// Object names that contain invalid characters and paths that contain spaces must be put into quotes to be handled correctly
 	if (PortFlags & PPF_Delimited)
 	{
 		PathName = FString::Printf(TEXT("\"%s\""), *PathName.ReplaceQuotesWithEscapedQuotes());
 	}
-	FTopLevelAssetPath ClassPathName = Object->GetClass()->GetClassPathName();
+	const FTopLevelAssetPath ClassPathName = Object.GetClass()->GetClassPathName();
 	return GetExportPath(ClassPathName, *PathName);
 }
 
 void FObjectPropertyBase::ExportText_Internal( FString& ValueStr, const void* PropertyValueOrContainer, EPropertyPointerType PropertyPointerType, const void* DefaultValue, UObject* Parent, int32 PortFlags, UObject* ExportRootScope ) const
 {
-	UObject* Temp = nullptr;
+	TObjectPtr<UObject> Temp;
 
 	if (PropertyPointerType == EPropertyPointerType::Container && HasGetter())
 	{
@@ -240,7 +240,7 @@ void FObjectPropertyBase::ExportText_Internal( FString& ValueStr, const void* Pr
 	}
 	else
 	{
-		Temp = GetObjectPropertyValue(PointerToValuePtr(PropertyValueOrContainer, PropertyPointerType));
+		Temp = GetObjectPtrPropertyValue(PointerToValuePtr(PropertyValueOrContainer, PropertyPointerType));
 	}
 
 	if (0 != (PortFlags & PPF_ExportCpp))
@@ -254,21 +254,32 @@ void FObjectPropertyBase::ExportText_Internal( FString& ValueStr, const void* Pr
 		return;
 	}
 
-	if( Temp != nullptr )
+	if (Temp.IsNullNoResolve())
+	{
+		ValueStr += TEXT("None");
+	}
+	else
 	{
 		if (PortFlags & PPF_DebugDump)
 		{
-			ValueStr += Temp->GetFullName();
+			ValueStr += Temp ? Temp->GetFullName() : TEXT("None");
 		}
-		else if (Parent && !Parent->HasAnyFlags(RF_ClassDefaultObject) && Temp->IsDefaultSubobject())
+		else if (
+		         Parent && !Parent->HasAnyFlags(RF_ClassDefaultObject)
+				 // @NOTE: OBJPTR: In the event that we're trying to handle a default subobject, the requirement would 
+				 // be that it's inside the package we are currently in, which means the Temp pointer should be resolved 
+				 // already.  So don't move forward with the check unless it's resolved, we don't want to force a deferred
+				 // loaded asset if we don't have to with this check.
+				 && Temp.IsResolved() && Temp && Temp->IsDefaultSubobject()
+				)
 		{
 			if (PortFlags & PPF_Delimited)
 			{
-				ValueStr += FString::Printf(TEXT("\"%s\""), *Temp->GetName().ReplaceQuotesWithEscapedQuotes());
+				ValueStr += Temp ? FString::Printf(TEXT("\"%s\""), *Temp->GetName().ReplaceQuotesWithEscapedQuotes()) : TEXT("None");
 			}
 			else
 			{
-				ValueStr += Temp->GetName();
+				ValueStr += Temp.GetName();
 			}
 		}
 		else
@@ -276,10 +287,6 @@ void FObjectPropertyBase::ExportText_Internal( FString& ValueStr, const void* Pr
 			ValueStr += GetExportPath(Temp, Parent, ExportRootScope, PortFlags);
 		}
 	}
-	else
-	{
-		ValueStr += TEXT("None");
-	}
 }
 
 /**
@@ -573,6 +580,12 @@ FName FObjectPropertyBase::GetID() const
 	return NAME_ObjectProperty;
 }
 
+TObjectPtr<UObject> FObjectPropertyBase::GetObjectPtrPropertyValue(const void* PropertyValueAddress) const
+{
+	checkf(false, TEXT("%s is missing implementation of GetObjectPtrPropertyValue"), *GetFullName());
+	return TObjectPtr<UObject>();
+}
+
 UObject* FObjectPropertyBase::GetObjectPropertyValue(const void* PropertyValueAddress) const
 {
 	checkf(false, TEXT("%s is missing implementation of GetObjectPropertyValue"), *GetFullName());
diff --git a/Engine/Source/Runtime/CoreUObject/Private/UObject/PropertyLazyObjectPtr.cpp b/Engine/Source/Runtime/CoreUObject/Private/UObject/PropertyLazyObjectPtr.cpp
--- a/Engine/Source/Runtime/CoreUObject/Private/UObject/PropertyLazyObjectPtr.cpp
+++ b/Engine/Source/Runtime/CoreUObject/Private/UObject/PropertyLazyObjectPtr.cpp
@@ -105,6 +105,11 @@ bool FLazyObjectProperty::Identical( const void* A, const void* B, uint32 PortFl
 	return bResult;
 }
 
+TObjectPtr<UObject> FLazyObjectProperty::GetObjectPtrPropertyValue(const void* PropertyValueAddress) const
+{
+	return TObjectPtr<UObject>(GetPropertyValue(PropertyValueAddress).Get());
+}
+
 UObject* FLazyObjectProperty::GetObjectPropertyValue(const void* PropertyValueAddress) const
 {
 	return GetPropertyValue(PropertyValueAddress).Get();
diff --git a/Engine/Source/Runtime/CoreUObject/Private/UObject/PropertyObject.cpp b/Engine/Source/Runtime/CoreUObject/Private/UObject/PropertyObject.cpp
--- a/Engine/Source/Runtime/CoreUObject/Private/UObject/PropertyObject.cpp
+++ b/Engine/Source/Runtime/CoreUObject/Private/UObject/PropertyObject.cpp
@@ -98,7 +98,7 @@ void FObjectProperty::SerializeItem( FStructuredArchive::FSlot Slot, void* Value
 	if (UnderlyingArchive.IsObjectReferenceCollector())
 	{
 		// Serialize in place
-		UObject** ObjectPtr = GetPropertyValuePtr(Value);
+		TObjectPtr<UObject>* ObjectPtr = GetPropertyValuePtr(Value);
 		Slot << (*ObjectPtr);
 
 		if(!UnderlyingArchive.IsSaving())
@@ -145,7 +145,7 @@ void FObjectProperty::SerializeItem( FStructuredArchive::FSlot Slot, void* Value
 
 const TCHAR* FObjectProperty::ImportText_Internal(const TCHAR* Buffer, void* ContainerOrPropertyPtr, EPropertyPointerType PropertyPointerType, UObject* OwnerObject, int32 PortFlags, FOutputDevice* ErrorText) const
 {
-	const TCHAR* Result = TFObjectPropertyBase<UObject*>::ImportText_Internal(Buffer, ContainerOrPropertyPtr, PropertyPointerType, OwnerObject, PortFlags, ErrorText);
+	const TCHAR* Result = TFObjectPropertyBase<TObjectPtr<UObject>>::ImportText_Internal(Buffer, ContainerOrPropertyPtr, PropertyPointerType, OwnerObject, PortFlags, ErrorText);
 	if (Result)
 	{
 		void* Data = PointerToValuePtr(ContainerOrPropertyPtr, PropertyPointerType);
@@ -178,11 +178,16 @@ uint32 FObjectProperty::GetValueTypeHashInternal(const void* Src) const
 	return GetTypeHash(GetPropertyValue(Src));
 }
 
-UObject* FObjectProperty::GetObjectPropertyValue(const void* PropertyValueAddress) const
+TObjectPtr<UObject> FObjectProperty::GetObjectPtrPropertyValue(const void* PropertyValueAddress) const
 {
 	return GetPropertyValue(PropertyValueAddress);
 }
 
+UObject* FObjectProperty::GetObjectPropertyValue(const void* PropertyValueAddress) const
+{
+	return GetPropertyValue(PropertyValueAddress).Get();
+}
+
 UObject* FObjectProperty::GetObjectPropertyValue_InContainer(const void* ContainerAddress, int32 ArrayIndex) const
 {
 	// Dummy struct to satisfy GetWrappedObjectPropertyValue_InContainer API requirements (Get() function)
diff --git a/Engine/Source/Runtime/CoreUObject/Private/UObject/PropertyObjectPtr.cpp b/Engine/Source/Runtime/CoreUObject/Private/UObject/PropertyObjectPtr.cpp
--- a/Engine/Source/Runtime/CoreUObject/Private/UObject/PropertyObjectPtr.cpp
+++ b/Engine/Source/Runtime/CoreUObject/Private/UObject/PropertyObjectPtr.cpp
@@ -55,7 +55,8 @@ void FObjectPtrProperty::StaticSerializeItem(const FObjectPropertyBase* ObjectPr
 {
 	FArchive& UnderlyingArchive = Slot.GetUnderlyingArchive();
 	FObjectPtr* ObjectPtr = (FObjectPtr*)GetPropertyValuePtr(Value);
-	UObject* CurrentValue = IsObjectHandleResolved(ObjectPtr->GetHandle()) ? ObjectPtr->Get() : nullptr;
+	const bool IsHandleResolved = IsObjectHandleResolved(ObjectPtr->GetHandle());
+	UObject* CurrentValue = IsHandleResolved ? ObjectPtr->Get() : nullptr;
 
 	if (UnderlyingArchive.IsObjectReferenceCollector())
 	{
@@ -143,9 +144,14 @@ FObjectPtr& FObjectPtrProperty::GetObjectPropertyValueAsPtr(const void* Property
 	return (FObjectPtr&)GetPropertyValue(PropertyValueAddress);
 }
 
+TObjectPtr<UObject> FObjectPtrProperty::GetObjectPtrPropertyValue(const void* PropertyValueAddress) const
+{
+	return GetPropertyValue(PropertyValueAddress);
+}
+
 UObject* FObjectPtrProperty::GetObjectPropertyValue(const void* PropertyValueAddress) const
 {
-	return ((FObjectPtr&)GetPropertyValue(PropertyValueAddress)).Get();
+	return GetPropertyValue(PropertyValueAddress).Get();
 }
 
 UObject* FObjectPtrProperty::GetObjectPropertyValue_InContainer(const void* ContainerAddress, int32 ArrayIndex) const
diff --git a/Engine/Source/Runtime/CoreUObject/Private/UObject/PropertySoftObjectPtr.cpp b/Engine/Source/Runtime/CoreUObject/Private/UObject/PropertySoftObjectPtr.cpp
--- a/Engine/Source/Runtime/CoreUObject/Private/UObject/PropertySoftObjectPtr.cpp
+++ b/Engine/Source/Runtime/CoreUObject/Private/UObject/PropertySoftObjectPtr.cpp
@@ -237,6 +237,11 @@ UObject* FSoftObjectProperty::LoadObjectPropertyValue(const void* PropertyValueA
 	return GetPropertyValue(PropertyValueAddress).LoadSynchronous();
 }
 
+TObjectPtr<UObject> FSoftObjectProperty::GetObjectPtrPropertyValue(const void* PropertyValueAddress) const
+{
+	return TObjectPtr<UObject>(GetPropertyValue(PropertyValueAddress).Get());
+}
+
 UObject* FSoftObjectProperty::GetObjectPropertyValue(const void* PropertyValueAddress) const
 {
 	return GetPropertyValue(PropertyValueAddress).Get();
diff --git a/Engine/Source/Runtime/CoreUObject/Private/UObject/PropertyWeakObjectPtr.cpp b/Engine/Source/Runtime/CoreUObject/Private/UObject/PropertyWeakObjectPtr.cpp
--- a/Engine/Source/Runtime/CoreUObject/Private/UObject/PropertyWeakObjectPtr.cpp
+++ b/Engine/Source/Runtime/CoreUObject/Private/UObject/PropertyWeakObjectPtr.cpp
@@ -64,6 +64,11 @@ void FWeakObjectProperty::SerializeItem( FStructuredArchive::FSlot Slot, void* V
 	}
 }
 
+TObjectPtr<UObject> FWeakObjectProperty::GetObjectPtrPropertyValue(const void* PropertyValueAddress) const
+{
+	return TObjectPtr<UObject>(GetPropertyValue(PropertyValueAddress).Get());
+}
+
 UObject* FWeakObjectProperty::GetObjectPropertyValue(const void* PropertyValueAddress) const
 {
 	return GetPropertyValue(PropertyValueAddress).Get();
diff --git a/Engine/Source/Runtime/CoreUObject/Private/UObject/UObjectGlobals.cpp b/Engine/Source/Runtime/CoreUObject/Private/UObject/UObjectGlobals.cpp
--- a/Engine/Source/Runtime/CoreUObject/Private/UObject/UObjectGlobals.cpp
+++ b/Engine/Source/Runtime/CoreUObject/Private/UObject/UObjectGlobals.cpp
@@ -4484,8 +4484,8 @@ void FReferenceCollector::AddReferencedObjects(const UScriptStruct*& ScriptStruc
 	// Iterate through all object properties within the struct (will also search through structs within the struct)
 	for (TPropertyValueIterator<const FObjectProperty> ObjectPropertyIter(ScriptStruct, StructMemory); ObjectPropertyIter; ++ObjectPropertyIter)
 	{
-		UObject** ObjectPtr = static_cast<UObject**>(const_cast<void*>(ObjectPropertyIter.Value()));
-		AddReferencedObject(*ObjectPtr, ReferencingObject, ReferencingProperty);
+		TObjectPtr<UObject>& ObjectPtrRef = ObjectPropertyIter.Key()->GetObjectPtrPropertyValueRef(ObjectPropertyIter.Value());
+		AddReferencedObject(ObjectPtrRef, ReferencingObject, ReferencingProperty);
 	}
 }
 
diff --git a/Engine/Source/Runtime/CoreUObject/Public/UObject/UnrealType.h b/Engine/Source/Runtime/CoreUObject/Public/UObject/UnrealType.h
--- a/Engine/Source/Runtime/CoreUObject/Public/UObject/UnrealType.h
+++ b/Engine/Source/Runtime/CoreUObject/Public/UObject/UnrealType.h
@@ -2709,7 +2709,7 @@ public:
 	 * @param PortFlags Property port flags
 	 * @return A string representing the export path of an object, usually in the form of /ClassPackage.ClassName'/Package/Path.Object'
 	 */
-	static FString GetExportPath(const UObject* Object, const UObject* Parent = nullptr, const UObject* ExportRootScope = nullptr, const uint32 PortFlags = PPF_None);
+	static FString GetExportPath(const TObjectPtr<const UObject>& Object, const UObject* Parent = nullptr, const UObject* ExportRootScope = nullptr, const uint32 PortFlags = PPF_None);
 
 	/**
 	 * Returns the qualified export path given a class path name and object path name
@@ -2730,6 +2730,8 @@ public:
 	{
 		return LoadObjectPropertyValue(ContainerPtrToValuePtr<void>(PropertyValueAddress, ArrayIndex));
 	}
+
+	virtual TObjectPtr<UObject> GetObjectPtrPropertyValue(const void* PropertyValueAddress) const;
 	virtual UObject* GetObjectPropertyValue(const void* PropertyValueAddress) const;
 	virtual UObject* GetObjectPropertyValue_InContainer(const void* ContainerAddress, int32 ArrayIndex = 0) const;
 	virtual void SetObjectPropertyValue(void* PropertyValueAddress, UObject* Value) const;
@@ -2896,9 +2898,9 @@ PRAGMA_ENABLE_DEPRECATION_WARNINGS
 //
 // Describes a reference variable to another object which may be nil.
 //
-class COREUOBJECT_API FObjectProperty : public TFObjectPropertyBase<UObject*>
+class COREUOBJECT_API FObjectProperty : public TFObjectPropertyBase<TObjectPtr<UObject>>
 {
-	DECLARE_FIELD(FObjectProperty, TFObjectPropertyBase<UObject*>, CASTCLASS_FObjectProperty)
+	DECLARE_FIELD(FObjectProperty, TFObjectPropertyBase<TObjectPtr<UObject>>, CASTCLASS_FObjectProperty)
 
 	FObjectProperty(FFieldVariant InOwner, const FName& InName, EObjectFlags InObjectFlags)
 		: TFObjectPropertyBase(InOwner, InName, InObjectFlags)
@@ -2944,12 +2946,23 @@ public:
 	// End of FProperty interface
 
 	// FObjectPropertyBase interface
+	virtual TObjectPtr<UObject> GetObjectPtrPropertyValue(const void* PropertyValueAddress) const override;
 	virtual UObject* GetObjectPropertyValue(const void* PropertyValueAddress) const override;
 	virtual UObject* GetObjectPropertyValue_InContainer(const void* ContainerAddress, int32 ArrayIndex = 0) const override;
 	virtual void SetObjectPropertyValue(void* PropertyValueAddress, UObject* Value) const override;
 	virtual void SetObjectPropertyValue_InContainer(void* ContainerAddress, UObject* Value, int32 ArrayIndex = 0) const override;
 	virtual FString GetCPPTypeCustom(FString* ExtendedTypeText, uint32 CPPExportFlags, const FString& InnerNativeTypeName)  const override;
 	// End of FObjectPropertyBase interface
+	
+	inline TObjectPtr<UObject>* GetObjectPtrPropertyValuePtr(const void* PropertyValueAddress) const
+	{
+		return reinterpret_cast<TObjectPtr<UObject>*>(const_cast<void*>(PropertyValueAddress));
+	}
+
+	inline TObjectPtr<UObject>& GetObjectPtrPropertyValueRef(const void* PropertyValueAddress) const
+	{
+		return *reinterpret_cast<TObjectPtr<UObject>*>(const_cast<void*>(PropertyValueAddress));
+	}
 };
 
 //
@@ -2988,6 +3001,7 @@ class COREUOBJECT_API FObjectPtrProperty : public FObjectProperty
 	FObjectPtr& GetObjectPropertyValueAsPtr(const void* PropertyValueAddress) const;
 
 	// FObjectProperty interface
+	virtual TObjectPtr<UObject> GetObjectPtrPropertyValue(const void* PropertyValueAddress) const override;
 	virtual UObject* GetObjectPropertyValue(const void* PropertyValueAddress) const override;
 	virtual UObject* GetObjectPropertyValue_InContainer(const void* ContainerAddress, int32 ArrayIndex = 0) const override;
 	virtual void SetObjectPropertyValue(void* PropertyValueAddress, UObject* Value) const override;
@@ -3052,6 +3066,7 @@ public:
 	// End of FProperty interface
 
 	// FObjectProperty interface
+	virtual TObjectPtr<UObject> GetObjectPtrPropertyValue(const void* PropertyValueAddress) const override;
 	virtual UObject* GetObjectPropertyValue(const void* PropertyValueAddress) const override;
 	virtual UObject* GetObjectPropertyValue_InContainer(const void* ContainerAddress, int32 ArrayIndex = 0) const override;
 	virtual void SetObjectPropertyValue(void* PropertyValueAddress, UObject* Value) const override;
@@ -3109,6 +3124,7 @@ PRAGMA_ENABLE_DEPRECATION_WARNINGS
 	// End of FProperty interface
 
 	// FObjectProperty interface
+	virtual TObjectPtr<UObject> GetObjectPtrPropertyValue(const void* PropertyValueAddress) const override;
 	virtual UObject* GetObjectPropertyValue(const void* PropertyValueAddress) const override;
 	virtual UObject* GetObjectPropertyValue_InContainer(const void* ContainerAddress, int32 ArrayIndex = 0) const override;
 	virtual void SetObjectPropertyValue(void* PropertyValueAddress, UObject* Value) const override;
@@ -3183,6 +3199,7 @@ public:
 
 	// FObjectProperty interface
 	virtual UObject* LoadObjectPropertyValue(const void* PropertyValueAddress) const override;
+	virtual TObjectPtr<UObject> GetObjectPtrPropertyValue(const void* PropertyValueAddress) const override;
 	virtual UObject* GetObjectPropertyValue(const void* PropertyValueAddress) const override;
 	virtual UObject* GetObjectPropertyValue_InContainer(const void* ContainerAddress, int32 ArrayIndex = 0) const override;
 	virtual void SetObjectPropertyValue(void* PropertyValueAddress, UObject* Value) const override;
diff --git a/Engine/Source/Runtime/MovieScene/Public/EntitySystem/MovieSceneComponentTypeInfo.h b/Engine/Source/Runtime/MovieScene/Public/EntitySystem/MovieSceneComponentTypeInfo.h
--- a/Engine/Source/Runtime/MovieScene/Public/EntitySystem/MovieSceneComponentTypeInfo.h
+++ b/Engine/Source/Runtime/MovieScene/Public/EntitySystem/MovieSceneComponentTypeInfo.h
@@ -36,8 +36,7 @@ typename TEnableIf< TIsSame<decltype(T::StaticStruct()), decltype(T::StaticStruc
 {
 	for (TPropertyValueIterator<const FObjectProperty> It(T::StaticStruct(), Component); It; ++It)
 	{
-		UObject** ObjectPtr = static_cast<UObject**>(const_cast<void*>(It.Value()));
-		ReferenceCollector->AddReferencedObject(*ObjectPtr);
+		ReferenceCollector->AddReferencedObject(It.Key()->GetObjectPtrPropertyValueRef(It.Value()));
 	}
 }
 
