Core - Adding support for getting the PathName from a ObjectPtr without resolving the Object*.
Core - Refactoring the Object Handle event callbacks so that we can hook and unhook them with delegates/delegatehandles.
Core - Adding support for detecting loads via object handles as a new event distinct from resolution callback.
Core - Expanding the utility functions on FObjectRef so that we can resolve it in a myraid of fashions to get the path, the class path, the full path and the export path.

diff --git a/Engine/Source/Editor/UnrealEd/Private/Cooker/PackageBuildDependencyTracker.cpp b/Engine/Source/Editor/UnrealEd/Private/Cooker/PackageBuildDependencyTracker.cpp
--- a/Engine/Source/Editor/UnrealEd/Private/Cooker/PackageBuildDependencyTracker.cpp
+++ b/Engine/Source/Editor/UnrealEd/Private/Cooker/PackageBuildDependencyTracker.cpp
@@ -56,20 +56,15 @@ TArray<FBuildDependencyAccessData> FPackageBuildDependencyTracker::GetAccessData
 
 FPackageBuildDependencyTracker::FPackageBuildDependencyTracker()
 {
-	PreviousObjectHandleReadFunction = SetObjectHandleReadCallback(StaticOnObjectHandleRead);
+	ObjectHandleReadHandle = AddObjectHandleReadCallback(FObjectHandleReadDelegate::CreateStatic(StaticOnObjectHandleRead));
 }
 
-void FPackageBuildDependencyTracker::StaticOnObjectHandleRead(UObject* ReadObject)
+FPackageBuildDependencyTracker::~FPackageBuildDependencyTracker()
 {
-	StaticOnObjectHandleReadInternal(ReadObject);
-
-	if (Singleton.PreviousObjectHandleReadFunction)
-	{
-		Singleton.PreviousObjectHandleReadFunction(ReadObject);
-	}
+	RemoveObjectHandleReadCallback(ObjectHandleReadHandle);
 }
 
-void FPackageBuildDependencyTracker::StaticOnObjectHandleReadInternal(UObject* ReadObject)
+void FPackageBuildDependencyTracker::StaticOnObjectHandleRead(UObject* ReadObject)
 {
 	if (!ReadObject)
 	{
diff --git a/Engine/Source/Editor/UnrealEd/Private/Cooker/PackageBuildDependencyTracker.h b/Engine/Source/Editor/UnrealEd/Private/Cooker/PackageBuildDependencyTracker.h
--- a/Engine/Source/Editor/UnrealEd/Private/Cooker/PackageBuildDependencyTracker.h
+++ b/Engine/Source/Editor/UnrealEd/Private/Cooker/PackageBuildDependencyTracker.h
@@ -37,17 +37,17 @@ public:
 
 private:
 	FPackageBuildDependencyTracker();
+	virtual ~FPackageBuildDependencyTracker();
 
 	/** Track object reference reads */
 	static void StaticOnObjectHandleRead(UObject* ReadObject);
-	static void StaticOnObjectHandleReadInternal(UObject* ReadObject);
-	ObjectHandleReadFunction* PreviousObjectHandleReadFunction = nullptr;
 
 	mutable FCriticalSection RecordsLock;
 	TMap<FName, TSet<FBuildDependencyAccessData>> Records;
 	FName LastReferencer = NAME_None;
 	FBuildDependencyAccessData LastAccessData{ NAME_None, nullptr };
 	TSet<FBuildDependencyAccessData>* LastReferencerSet = nullptr;
+	FDelegateHandle ObjectHandleReadHandle;
 	static FPackageBuildDependencyTracker Singleton;
 };
 
diff --git a/Engine/Source/Runtime/CoreUObject/Private/UObject/ObjectHandle.cpp b/Engine/Source/Runtime/CoreUObject/Private/UObject/ObjectHandle.cpp
--- a/Engine/Source/Runtime/CoreUObject/Private/UObject/ObjectHandle.cpp
+++ b/Engine/Source/Runtime/CoreUObject/Private/UObject/ObjectHandle.cpp
@@ -16,6 +16,67 @@
 
 DEFINE_LOG_CATEGORY(LogObjectHandle);
 
+FName FObjectRef::GetFName() const
+{
+	if (PackageName == NAME_None)
+	{
+		return NAME_None;
+	}
+
+	FObjectPathId::ResolvedNameContainerType ResolvedNames;
+	ObjectPath.Resolve(ResolvedNames);
+
+	if (ResolvedNames.Num() > 0)
+	{
+		const FName& Name = ResolvedNames[ResolvedNames.Num() - 1];
+		return Name;
+	}
+
+	return NAME_None;
+}
+
+void FObjectRef::AppendClassPathName(FStringBuilderBase& OutClassPathNameBuilder, EObjectFullNameFlags Flags) const
+{
+	if (EnumHasAllFlags(Flags, EObjectFullNameFlags::IncludeClassPackage))
+	{
+		ClassPackageName.AppendString(OutClassPathNameBuilder);
+		OutClassPathNameBuilder << TEXT(".");
+	}
+
+	ClassName.AppendString(OutClassPathNameBuilder);
+}
+
+void FObjectRef::AppendPathName(FStringBuilderBase& OutPathNameBuilder) const
+{
+	if (PackageName == NAME_None)
+	{
+		OutPathNameBuilder.Append(TEXT("None"));
+		return;
+	}
+
+	PackageName.ToString(OutPathNameBuilder);
+
+	FObjectPathId::ResolvedNameContainerType ResolvedNames;
+	ObjectPath.Resolve(ResolvedNames);
+
+	for (int Index = 0; Index < ResolvedNames.Num(); ++Index)
+	{
+		if (Index == 1)
+		{
+			OutPathNameBuilder << SUBOBJECT_DELIMITER_CHAR;
+		}
+		else
+		{
+			OutPathNameBuilder << TEXT(".");
+		}
+
+		const FName& Name = ResolvedNames[Index];
+		Name.AppendString(OutPathNameBuilder);
+	}
+}
+
+//===========================================================================
+
 bool operator==(const FObjectHandleDataClassDescriptor& Lhs, const FObjectHandleDataClassDescriptor& Rhs)
 {
 	return (Lhs.PackageName == Rhs.PackageName) && (Lhs.ClassName == Rhs.ClassName);
@@ -365,8 +426,10 @@ FPackedObjectRef MakePackedObjectRef(const FObjectRef& ObjectRef)
 	return ObjectHandle_Private::MakePackedObjectRef(ObjectRef.PackageName, ObjectRef.ClassPackageName, ObjectRef.ClassName, ObjectRef.ObjectPath);
 }
 
-static inline UPackage* FindOrLoadPackage(FName PackageName, int32 LoadFlags)
+static inline UPackage* FindOrLoadPackage(FName PackageName, int32 LoadFlags, OUT bool& bWasPackageLoaded)
 {
+	bWasPackageLoaded = false;
+
 	// @TODO: OBJPTR: Want to replicate the functional path of an import here.  See things like FindImportFast in BlueprintSupport.cpp
 	// 		 for additional behavior that we're not handling here yet.
 	FName* ScriptPackageName = FPackageName::FindScriptPackageName(PackageName);
@@ -387,6 +450,7 @@ static inline UPackage* FindOrLoadPackage(FName PackageName, int32 LoadFlags)
 		}
 		LoadFlags |= LOAD_NoWarn | LOAD_NoVerify; //This does nothing? | LOAD_DisableDependencyPreloading;
 		TargetPackage = LoadPackage(nullptr, *PackageName.ToString(), LoadFlags);
+		bWasPackageLoaded = true;
 	}
 	return TargetPackage;
 }
@@ -398,7 +462,8 @@ UClass* ResolveObjectRefClass(const FObjectRef& ObjectRef, uint32 LoadFlags /*=
 	UPackage* ClassPackage = nullptr;
 	if (!ObjectRef.ClassPackageName.IsNone())
 	{
-		ClassPackage = FindOrLoadPackage(ObjectRef.ClassPackageName, LoadFlags);
+		bool bWasPackageLoaded = false;
+		ClassPackage = FindOrLoadPackage(ObjectRef.ClassPackageName, LoadFlags, bWasPackageLoaded);
 
 		if (!ObjectRef.ClassName.IsNone())
 		{
@@ -458,7 +523,9 @@ UObject* ResolveObjectRef(const FObjectRef& ObjectRef, uint32 LoadFlags /*= LOAD
 
 	ResolveObjectRefClass(ObjectRef, LoadFlags);
 
-	UPackage* TargetPackage = FindOrLoadPackage(ObjectRef.PackageName, LoadFlags);
+	bool bWasObjectOrPackageLoaded = false;
+
+	UPackage* TargetPackage = FindOrLoadPackage(ObjectRef.PackageName, LoadFlags, bWasObjectOrPackageLoaded);
 
 	if (!TargetPackage)
 	{
@@ -488,12 +555,14 @@ UObject* ResolveObjectRef(const FObjectRef& ObjectRef, uint32 LoadFlags /*= LOAD
 			if (IsInAsyncLoadingThread() || IsInGameThread())
 			{
 				TargetPackage->GetLinker()->LoadAllObjects(true);
+				bWasObjectOrPackageLoaded = true;
 			}
 			else
 			{
 				// Shunt the load request to happen on the game thread and block on its completion.  This is a deadlock risk!  The game thread may be blocked waiting on this thread.
 				UE_LOG(LogObjectHandle, Warning, TEXT("Resolve of object in package '%s' from a non-game thread was shunted to the game thread."), *ObjectRef.PackageName.ToString());
 				TGraphTask<FFullyLoadPackageOnHandleResolveTask>::CreateTask().ConstructAndDispatchWhenReady(TargetPackage)->Wait();
+				bWasObjectOrPackageLoaded = true;
 			}
 
 			CurrentObject = StaticFindObjectFastInternal(nullptr, PreviousOuter, ResolvedNames[ObjectPathIndex]);
@@ -519,16 +588,23 @@ UObject* ResolveObjectRef(const FObjectRef& ObjectRef, uint32 LoadFlags /*= LOAD
 		if (IsInAsyncLoadingThread() || IsInGameThread())
 		{
 			TargetPackage->GetLinker()->LoadAllObjects(true);
+			bWasObjectOrPackageLoaded = true;
 		}
 		else
 		{
 			// Shunt the load request to happen on the game thread and block on its completion.  This is a deadlock risk!  The game thread may be blocked waiting on this thread.
 			UE_LOG(LogObjectHandle, Warning, TEXT("Resolve of object in package '%s' from a non-game thread was shunted to the game thread."), *ObjectRef.PackageName.ToString());
 			TGraphTask<FFullyLoadPackageOnHandleResolveTask>::CreateTask().ConstructAndDispatchWhenReady(TargetPackage)->Wait();
+			bWasObjectOrPackageLoaded = true;
 		}
 	}
 	ObjectHandle_Private::OnReferenceResolved(ObjectRef, TargetPackage, CurrentObject);
 
+	if (bWasObjectOrPackageLoaded)
+	{
+		ObjectHandle_Private::OnReferenceLoaded(ObjectRef, TargetPackage, CurrentObject);
+	}
+
 	return CurrentObject;
 }
 
@@ -543,22 +619,91 @@ UObject* ResolvePackedObjectRef(FPackedObjectRef PackedObjectRef, uint32 LoadFla
 }
 
 #if UE_WITH_OBJECT_HANDLE_TRACKING
-COREUOBJECT_API ObjectHandleReadFunction* ObjectHandle_Private::ObjectHandleReadCallback = nullptr;
-COREUOBJECT_API ObjectHandleClassResolvedFunction* ObjectHandle_Private::ObjectHandleClassResolvedCallback = nullptr;
-COREUOBJECT_API ObjectHandleReferenceResolvedFunction* ObjectHandle_Private::ObjectHandleReferenceResolvedCallback = nullptr;
 
-ObjectHandleReadFunction* SetObjectHandleReadCallback(ObjectHandleReadFunction* Function)
+COREUOBJECT_API std::atomic<int32> ObjectHandle_Private::ObjectHandleEventIndex;
+COREUOBJECT_API ObjectHandle_Private::FObjectHandleEvents ObjectHandle_Private::ObjectHandleEvents[2];
+
+ObjectHandle_Private::FObjectHandleEvents& BeginWritingEvents()
+{
+	// We only allow adding/removing the callbacks from a single thread.
+	check(IsInGameThread());
+
+	// Get what the last and new index will be.
+	const int32 LastIndex = ObjectHandle_Private::ObjectHandleEventIndex;
+	const int32 NewIndex = LastIndex == 0 ? 1 : 0;
+
+	// Make sure we're not still using the events, this should definitely not be true, because
+	// we spin lock at the end of any of these functions to make sure all using has stopped.
+	check(!ObjectHandle_Private::ObjectHandleEvents[NewIndex].IsUsing());
+
+	// We need to take whatever delegates are setup on the previous index and copy them over to new set.
+	ObjectHandle_Private::ObjectHandleEvents[NewIndex].ObjectHandleReadEvent = ObjectHandle_Private::ObjectHandleEvents[LastIndex].ObjectHandleReadEvent;
+	ObjectHandle_Private::ObjectHandleEvents[NewIndex].ClassReferenceResolvedEvent = ObjectHandle_Private::ObjectHandleEvents[LastIndex].ClassReferenceResolvedEvent;
+	ObjectHandle_Private::ObjectHandleEvents[NewIndex].ObjectHandleReferenceResolvedEvent = ObjectHandle_Private::ObjectHandleEvents[LastIndex].ObjectHandleReferenceResolvedEvent;
+	ObjectHandle_Private::ObjectHandleEvents[NewIndex].ObjectHandleReferenceLoadedEvent = ObjectHandle_Private::ObjectHandleEvents[LastIndex].ObjectHandleReferenceLoadedEvent;
+
+	// Now toggle which buffer we're using.
+	ObjectHandle_Private::ObjectHandleEventIndex = NewIndex;
+
+	// Ok, now we need to spin lock until nobody is using the LastIndex, this will prevent us from
+	// calling this function right after this one and assuming the previous set isn't being used,
+	// similarly while we were doing this, the BeginReadingEvents() function spun lock after marked
+	// that they were using the event set, so now that we've set the new event index, if we spin lock
+	// here, that should ensure that any existing readers finish up what they were doing.
+	while (ObjectHandle_Private::ObjectHandleEvents[LastIndex].IsUsing())
+	{
+		// Spin!
+	}
+
+	return ObjectHandle_Private::ObjectHandleEvents[NewIndex];
+}
+
+FDelegateHandle AddObjectHandleReadCallback(FObjectHandleReadDelegate Callback)
 {
-	return (ObjectHandleReadFunction*)FPlatformAtomics::InterlockedExchangePtr((void**)&ObjectHandle_Private::ObjectHandleReadCallback, (void*)Function);
+	ObjectHandle_Private::FObjectHandleEvents& Events = BeginWritingEvents();
+	return Events.ObjectHandleReadEvent.Add(Callback);
 }
 
-ObjectHandleClassResolvedFunction* SetObjectHandleClassResolvedCallback(ObjectHandleClassResolvedFunction* Function)
+void RemoveObjectHandleReadCallback(FDelegateHandle DelegateHandle)
 {
-	return (ObjectHandleClassResolvedFunction*)FPlatformAtomics::InterlockedExchangePtr((void**)&ObjectHandle_Private::ObjectHandleClassResolvedCallback, (void*)Function);
+	ObjectHandle_Private::FObjectHandleEvents& Events = BeginWritingEvents();
+	Events.ObjectHandleReadEvent.Remove(DelegateHandle);
 }
 
-ObjectHandleReferenceResolvedFunction* SetObjectHandleReferenceResolvedCallback(ObjectHandleReferenceResolvedFunction* Function)
+FDelegateHandle AddObjectHandleClassResolvedCallback(FObjectHandleClassResolvedDelegate Callback)
 {
-	return (ObjectHandleReferenceResolvedFunction*)FPlatformAtomics::InterlockedExchangePtr((void**)&ObjectHandle_Private::ObjectHandleReferenceResolvedCallback, (void*)Function);
+	ObjectHandle_Private::FObjectHandleEvents& Events = BeginWritingEvents();
+	return Events.ClassReferenceResolvedEvent.Add(Callback);
 }
+
+void RemoveObjectHandleClassResolvedCallback(FDelegateHandle DelegateHandle)
+{
+	ObjectHandle_Private::FObjectHandleEvents& Events = BeginWritingEvents();
+	Events.ClassReferenceResolvedEvent.Remove(DelegateHandle);
+}
+
+FDelegateHandle AddObjectHandleReferenceResolvedCallback(FObjectHandleReferenceResolvedDelegate Callback)
+{
+	ObjectHandle_Private::FObjectHandleEvents& Events = BeginWritingEvents();
+	return Events.ObjectHandleReferenceResolvedEvent.Add(Callback);
+}
+
+void RemoveObjectHandleReferenceResolvedCallback(FDelegateHandle DelegateHandle)
+{
+	ObjectHandle_Private::FObjectHandleEvents& Events = BeginWritingEvents();
+	Events.ObjectHandleReferenceResolvedEvent.Remove(DelegateHandle);
+}
+
+FDelegateHandle AddObjectHandleReferenceLoadedCallback(FObjectHandleReferenceLoadedDelegate Callback)
+{
+	ObjectHandle_Private::FObjectHandleEvents& Events = BeginWritingEvents();
+	return Events.ObjectHandleReferenceLoadedEvent.Add(Callback);
+}
+
+void RemoveObjectHandleReferenceLoadedCallback(FDelegateHandle DelegateHandle)
+{
+	ObjectHandle_Private::FObjectHandleEvents& Events = BeginWritingEvents();
+	Events.ObjectHandleReferenceLoadedEvent.Remove(DelegateHandle);
+}
+
 #endif // UE_WITH_OBJECT_HANDLE_TRACKING
diff --git a/Engine/Source/Runtime/CoreUObject/Private/UObject/ObjectPtr.cpp b/Engine/Source/Runtime/CoreUObject/Private/UObject/ObjectPtr.cpp
new file mode 100644
--- /dev/null
+++ b/Engine/Source/Runtime/CoreUObject/Private/UObject/ObjectPtr.cpp
@@ -0,0 +1,48 @@
+// Copyright Epic Games, Inc. All Rights Reserved.
+
+#include "UObject/ObjectPtr.h"
+
+#if UE_WITH_OBJECT_HANDLE_LATE_RESOLVE
+
+FString FObjectPtr::GetPathName() const
+{
+	if (IsResolved())
+	{
+		return Get()->GetPathName();
+	}
+	else
+	{
+		FObjectRef ObjectRef = MakeObjectRef(Handle);
+		return ObjectRef.GetPathName();
+	}
+}
+
+COREUOBJECT_API FName FObjectPtr::GetFName() const
+{
+	if (IsResolved())
+	{
+		UObject* ResolvedObject = Get();
+		return ResolvedObject ? ResolvedObject->GetFName() : NAME_None;
+	}
+	else
+	{
+		FObjectRef ObjectRef = MakeObjectRef(Handle);
+		return ObjectRef.GetFName();
+	}
+}
+
+COREUOBJECT_API FString FObjectPtr::GetFullName(EObjectFullNameFlags Flags) const
+{
+	if (IsResolved())
+	{
+		// UObjectBaseUtility::GetFullName is safe to call on null objects.
+		return Get()->GetFullName(nullptr, Flags);
+	}
+	else
+	{
+		FObjectRef ObjectRef = MakeObjectRef(Handle);
+		return ObjectRef.GetFullName(Flags);
+	}
+}
+
+#endif
\ No newline at end of file
diff --git a/Engine/Source/Runtime/CoreUObject/Private/UObject/SoftObjectPath.cpp b/Engine/Source/Runtime/CoreUObject/Private/UObject/SoftObjectPath.cpp
--- a/Engine/Source/Runtime/CoreUObject/Private/UObject/SoftObjectPath.cpp
+++ b/Engine/Source/Runtime/CoreUObject/Private/UObject/SoftObjectPath.cpp
@@ -14,14 +14,6 @@
 #include "Misc/AutomationTest.h"
 #include "String/Find.h"
 
-FSoftObjectPath::FSoftObjectPath(const UObject* InObject)
-{
-	if (InObject)
-	{
-		SetPath(InObject->GetPathName());
-	}
-}
-
 // Deprecated constructor
 FSoftObjectPath::FSoftObjectPath(FName InAssetPathName, FString InSubPathString)
 {
diff --git a/Engine/Source/Runtime/CoreUObject/Public/UObject/ObjectHandle.h b/Engine/Source/Runtime/CoreUObject/Public/UObject/ObjectHandle.h
--- a/Engine/Source/Runtime/CoreUObject/Public/UObject/ObjectHandle.h
+++ b/Engine/Source/Runtime/CoreUObject/Public/UObject/ObjectHandle.h
@@ -40,6 +40,66 @@ struct FObjectRef
 	{
 		return (PackageName == Other.PackageName) && (ObjectPath == Other.ObjectPath) && (ClassPackageName == Other.ClassPackageName) && (ClassName == Other.ClassName);
 	}
+
+	/** Returns the path to the class for this object. */
+	FString GetClassPathName(EObjectFullNameFlags Flags = EObjectFullNameFlags::IncludeClassPackage) const
+	{
+		TStringBuilder<FName::StringBufferSize> Result;
+		AppendClassPathName(Result, Flags);
+		return FString(Result);
+	}
+
+	/** Appends the path to the class for this object to the builder, does not reset builder. */
+	COREUOBJECT_API void AppendClassPathName(FStringBuilderBase& OutClassPathNameBuilder, EObjectFullNameFlags Flags = EObjectFullNameFlags::IncludeClassPackage) const;
+
+	/** Returns the name of the object in the form: ObjectName */
+	COREUOBJECT_API FName GetFName() const;
+
+	/** Returns the full path for the object in the form: ObjectPath */
+	FString GetPathName() const
+	{
+		TStringBuilder<FName::StringBufferSize> Result;
+		AppendPathName(Result);
+		return FString(Result);
+	}
+
+	/** Appends the path to the object to the builder, does not reset builder. */
+	COREUOBJECT_API void AppendPathName(FStringBuilderBase& OutPathNameBuilder) const;
+
+	/** Returns the full name for the object in the form: Class ObjectPath */
+	FString GetFullName(EObjectFullNameFlags Flags = EObjectFullNameFlags::None) const
+	{
+		TStringBuilder<256> FullName;
+		GetFullName(FullName, Flags);
+		return FString(FullName);
+	}
+
+	/** Populates OutFullNameBuilder with the full name for the object in the form: Class ObjectPath */
+	void GetFullName(FStringBuilderBase& OutFullNameBuilder, EObjectFullNameFlags Flags = EObjectFullNameFlags::None) const
+	{
+		OutFullNameBuilder.Reset();
+		AppendClassPathName(OutFullNameBuilder, Flags);
+		OutFullNameBuilder.AppendChar(TEXT(' '));
+		AppendPathName(OutFullNameBuilder);
+	}
+
+	/** Returns the name for the asset in the form: Class'ObjectPath' */
+	FString GetExportTextName() const
+	{
+		TStringBuilder<256> ExportTextName;
+		GetExportTextName(ExportTextName);
+		return FString(ExportTextName);
+	}
+
+	/** Populates OutExportTextNameBuilder with the name for the object in the form: Class'ObjectPath' */
+	void GetExportTextName(FStringBuilderBase& OutExportTextNameBuilder) const
+	{
+		OutExportTextNameBuilder.Reset();
+		AppendClassPathName(OutExportTextNameBuilder);
+		OutExportTextNameBuilder.AppendChar(TEXT('\''));
+		AppendPathName(OutExportTextNameBuilder);
+		OutExportTextNameBuilder.AppendChar(TEXT('\''));
+	}
 };
 
 inline bool IsObjectRefNull(const FObjectRef& ObjectRef) { return ObjectRef.PackageName.IsNone() && ObjectRef.ObjectPath.IsNone(); }
@@ -131,13 +191,14 @@ inline FPackedObjectRef MakePackedObjectRef(FObjectHandle Handle);
  * 4) high performance (ie: will be called many times)
  */
 #if UE_WITH_OBJECT_HANDLE_TRACKING
+
  /**
   * Callback notifying when an object value is read from a handle.  Fired regardless of whether the handle
   * was resolved as part of the read operation or not and whether the object being read is null or not.
   *
   * @param ReadObject	The object that was read from a handle.
   */
-using ObjectHandleReadFunction = void(UObject* ReadObject);
+DECLARE_DELEGATE_OneParam(FObjectHandleReadDelegate, UObject* ReadObject);
 
 /**
  * Callback notifying when a class is resolved from an object handle or object reference.
@@ -147,42 +208,63 @@ using ObjectHandleReadFunction = void(UObject* ReadObject);
  * @param ClassPackage	The package containing the resolved class.
  * @param Class			The resolved class.
  */
-using ObjectHandleClassResolvedFunction = void(const FObjectRef& SourceRef, UPackage* ClassPackage, UClass* Class);
+DECLARE_DELEGATE_ThreeParams(FObjectHandleClassResolvedDelegate, const FObjectRef& SourceRef, UPackage* ObjectPackage, UClass* Class);
+
+ /**
+  * Callback notifying when a object handle is resolved.
+  *
+  * @param SourceRef		The object reference (either user provided or internally computed from a handle) from which the class was resolved.
+  * @param ClassPackage	The package containing the resolved class.
+  * @param Class			The resolved class.
+  */
+DECLARE_DELEGATE_ThreeParams(FObjectHandleReferenceResolvedDelegate, const FObjectRef& SourceRef, UPackage* ObjectPackage, UObject* Object);
 
 /**
- * Callback notifying when an object is resolved from an object handle or object reference.
+ * Callback notifying when an object was loaded through an object handle.  Will not notify you about global object loads, just ones that occur
+ * as the byproduct of resolving an ObjectHandle.
  *
- * @param SourceRef		The object reference (either user provided or internally computed from a handle) from which the object was resolved.
- * @param ObjectPackage	The package containing the resolved object.
- * @param Object		The resolved object.
+ * @param SourceRef		The object reference (either user provided or internally computed from a handle) from which the class was resolved.
+ * @param ClassPackage	The package containing the resolved class.
+ * @param Class			The resolved class.
  */
-using ObjectHandleReferenceResolvedFunction = void(const FObjectRef& SourceRef, UPackage* ObjectPackage, UObject* Object);
+DECLARE_DELEGATE_ThreeParams(FObjectHandleReferenceLoadedDelegate, const FObjectRef& SourceRef, UPackage* ObjectPackage, UObject* Object);
 
 /**
  * Installs a new callback for notifications that an object value has been read from a handle.
  *
  * @param Function		The new handle read callback to install.
- * @return				The previous handle read callback (or nullptr).  The caller is expected to store this and call in their own handle read callback.
+ * @return				The DelegateHandle so that you can remove the callback at a later date.
  */
-COREUOBJECT_API /*UE_NODISCARD*/ ObjectHandleReadFunction* SetObjectHandleReadCallback(ObjectHandleReadFunction* Function);
+COREUOBJECT_API FDelegateHandle AddObjectHandleReadCallback(FObjectHandleReadDelegate Delegate);
+COREUOBJECT_API void RemoveObjectHandleReadCallback(FDelegateHandle DelegateHandle);
 
 /**
  * Installs a new callback for notifications that a class has been resolved from an object handle or object reference.
  *
  * @param Function		The new class resolved callback to install.
- * @return				The previous object resolved callback (or nullptr).  The caller is expected to store this and call in their own class resolved callback.
+ * @return				The DelegateHandle so that you can remove the callback at a later date.
  */
-COREUOBJECT_API /*UE_NODISCARD*/ ObjectHandleClassResolvedFunction* SetObjectHandleClassResolvedCallback(ObjectHandleClassResolvedFunction* Function);
+COREUOBJECT_API FDelegateHandle AddObjectHandleClassResolvedCallback(FObjectHandleClassResolvedDelegate Callback);
+COREUOBJECT_API void RemoveObjectHandleClassResolvedCallback(FDelegateHandle DelegateHandle);
 
 /**
  * Installs a new callback for notifications that an object has been resolved from an object handle or object reference.
  *
  * @param Function		The new object resolved callback to install.
- * @return				The previous object resolved callback (or nullptr).  The caller is expected to store this and call in their own object resolved callback.
+ * @return				The DelegateHandle so that you can remove the callback at a later date.
  */
-COREUOBJECT_API /*UE_NODISCARD*/ ObjectHandleReferenceResolvedFunction* SetObjectHandleReferenceResolvedCallback(ObjectHandleReferenceResolvedFunction* Function);
-#endif
+COREUOBJECT_API FDelegateHandle AddObjectHandleReferenceResolvedCallback(FObjectHandleReferenceResolvedDelegate Callback);
+COREUOBJECT_API void RemoveObjectHandleReferenceResolvedCallback(FDelegateHandle DelegateHandle);
 
+/**
+ * Installs a new callback for notifications that an object has been loaded from an object handle or object reference.
+ *
+ * @param Function		The new object resolved callback to install.
+ * @return				The DelegateHandle so that you can remove the callback at a later date.
+ */
+COREUOBJECT_API FDelegateHandle AddObjectHandleReferenceLoadedCallback(FObjectHandleReferenceLoadedDelegate Callback);
+COREUOBJECT_API void RemoveObjectHandleReferenceLoadedCallback(FDelegateHandle DelegateHandle);
+#endif
 
 
 ///////////////////////////////////////////////////////////////////////////////////////////////////////////////////////
@@ -213,37 +295,83 @@ namespace ObjectHandle_Private
 	constexpr uint32 PackageIdMask = 0x7FFF'FFFF;
 
 #if UE_WITH_OBJECT_HANDLE_TRACKING
-	extern COREUOBJECT_API ObjectHandleReadFunction* ObjectHandleReadCallback;
-	extern COREUOBJECT_API ObjectHandleClassResolvedFunction* ObjectHandleClassResolvedCallback;
-	extern COREUOBJECT_API ObjectHandleReferenceResolvedFunction* ObjectHandleReferenceResolvedCallback;
+	DECLARE_MULTICAST_DELEGATE_OneParam(FObjectHandleReadEvent, UObject* Object);
+	DECLARE_MULTICAST_DELEGATE_ThreeParams(FClassReferenceResolvedEvent, const FObjectRef& ObjectRef, UPackage* Package, UClass* Class);
+	DECLARE_MULTICAST_DELEGATE_ThreeParams(FObjectHandleReferenceResolvedEvent, const FObjectRef& SourceRef, UPackage* ObjectPackage, UObject* Object);
+	DECLARE_MULTICAST_DELEGATE_ThreeParams(FObjectHandleReferenceLoadedEvent, const FObjectRef& SourceRef, UPackage* ObjectPackage, UObject* Object);
 
-	inline void OnHandleRead(UObject* Object)
+	COREUOBJECT_API struct FObjectHandleEvents
+	{
+		FObjectHandleReadEvent ObjectHandleReadEvent;
+		FClassReferenceResolvedEvent ClassReferenceResolvedEvent;
+		FObjectHandleReferenceResolvedEvent ObjectHandleReferenceResolvedEvent;
+		FObjectHandleReferenceLoadedEvent ObjectHandleReferenceLoadedEvent;
+
+		inline void BeginUsing() { UsingCount++; }
+		inline void EndUsing() { --UsingCount; }
+		inline bool IsUsing() const { return UsingCount > 0; }
+
+	private:
+		std::atomic<int32> UsingCount;
+	};
+
+	extern COREUOBJECT_API std::atomic<int32> ObjectHandleEventIndex;
+	extern COREUOBJECT_API FObjectHandleEvents ObjectHandleEvents[2];
+
+	inline FObjectHandleEvents& BeginReadingEvents()
 	{
-		if (ObjectHandleReadFunction* Callback = ObjectHandleReadCallback)
+		// Quick spin lock to try and begin using the events, normally this wont actually spin, in like 99.999% of the time.
+		while (true)
 		{
-			Callback(Object);
+			// Start by getting the current event index from the double buffer.
+			const int32 InitialEventIndex = ObjectHandleEventIndex;
+			// Grab the event set at that index.
+			FObjectHandleEvents& Events = ObjectHandleEvents[InitialEventIndex];
+			// Begin using the events, this will signal we're using them.
+			Events.BeginUsing();
+			// Ok - now we're going to check that we're *still* using the same index.  If we are
+			// then we know that it didn't change out from under us in between getting Events, and calling BeginUsing().
+			if (InitialEventIndex == ObjectHandleEventIndex)
+			{
+				return Events;
+			}
+			// If the check above failed, then we need to cease using the events, because we're about to try this again.
+			Events.EndUsing();
 		}
 	}
+
+	inline void OnHandleRead(UObject* Object)
+	{
+		FObjectHandleEvents& Events = BeginReadingEvents();
+		Events.ObjectHandleReadEvent.Broadcast(Object);
+		Events.EndUsing();
+	}
 	
 	inline void OnClassReferenceResolved(const FObjectRef& ObjectRef, UPackage* Package, UClass* Class)
 	{
-		if (ObjectHandleClassResolvedFunction* Callback = ObjectHandleClassResolvedCallback)
-		{
-			Callback(ObjectRef, Package, Class);
-		}
+		FObjectHandleEvents& Events = BeginReadingEvents();
+		Events.ClassReferenceResolvedEvent.Broadcast(ObjectRef, Package, Class);
+		Events.EndUsing();
 	}
 
 	inline void OnReferenceResolved(const FObjectRef& ObjectRef, UPackage* Package, UObject* Object)
 	{
-		if (ObjectHandleReferenceResolvedFunction* Callback = ObjectHandleReferenceResolvedCallback)
-		{
-			Callback(ObjectRef, Package, Object);
-		}
+		FObjectHandleEvents& Events = BeginReadingEvents();
+		Events.ObjectHandleReferenceResolvedEvent.Broadcast(ObjectRef, Package, Object);
+		Events.EndUsing();
+	}
+
+	inline void OnReferenceLoaded(const FObjectRef& ObjectRef, UPackage* Package, UObject* Object)
+	{
+		FObjectHandleEvents& Events = BeginReadingEvents();
+		Events.ObjectHandleReferenceLoadedEvent.Broadcast(ObjectRef, Package, Object);
+		Events.EndUsing();
 	}
 #else
 	inline void OnHandleRead(UObject* Object) {}
 	inline void OnClassReferenceResolved(const FObjectRef& ObjectRef, UPackage* Package, UClass* Class) {}
 	inline void OnReferenceResolved(const FObjectRef& ObjectRef, UPackage* Package, UObject* Object) {}
+	inline void OnReferenceLoaded(const FObjectRef& ObjectRef, UPackage* Package, UObject* Object) {}
 #endif
 }
 
diff --git a/Engine/Source/Runtime/CoreUObject/Public/UObject/ObjectPtr.h b/Engine/Source/Runtime/CoreUObject/Public/UObject/ObjectPtr.h
--- a/Engine/Source/Runtime/CoreUObject/Public/UObject/ObjectPtr.h
+++ b/Engine/Source/Runtime/CoreUObject/Public/UObject/ObjectPtr.h
@@ -142,7 +142,48 @@ public:
 
 	FORCEINLINE bool IsResolved() const { return IsObjectHandleResolved(Handle); }
 
-	FORCEINLINE FString GetPath() const { return Get()->GetPathName(); }
+#if UE_WITH_OBJECT_HANDLE_LATE_RESOLVE
+	// Gets the PathName of the object without resolving the object reference.
+	COREUOBJECT_API FString GetPathName() const;
+
+	// Gets the FName of the object without resolving the object reference.
+	COREUOBJECT_API FName GetFName() const;
+
+	// Gets the string name of the object without resolving the object reference.
+	FString GetName() const
+	{
+		return GetFName().ToString();
+	}
+
+	/** Returns the full name for the object in the form: Class ObjectPath */
+	COREUOBJECT_API FString GetFullName(EObjectFullNameFlags Flags = EObjectFullNameFlags::None) const;
+#else
+	FString GetPathName() const
+	{
+		return Get()->GetPathName();
+	}
+
+	FName GetFName() const
+	{
+		const UObject* ResolvedObject = Get();
+		return ResolvedObject ? ResolvedObject->GetFName() : NAME_None;
+	}
+
+	FString GetName() const
+	{
+		return GetFName().ToString();
+	}
+
+	/**
+	 * Returns the fully qualified pathname for this object as well as the name of the class, in the format:
+	 * 'ClassName Outermost.[Outer:]Name'.
+	 */
+	FString GetFullName(EObjectFullNameFlags Flags = EObjectFullNameFlags::None) const
+	{
+		// UObjectBaseUtility::GetFullName is safe to call on null objects.
+		return Get()->GetFullName(nullptr, Flags);
+	}
+#endif
 
 	FORCEINLINE FObjectHandle GetHandle() const { return Handle; }
 	FORCEINLINE FObjectHandle& GetHandleRef() const { return Handle; }
@@ -151,10 +192,12 @@ public:
 	{
 		checkfSlow(SomeBase, TEXT("IsA(NULL) cannot yield meaningful results"));
 
-		const UClass* ThisClass = GetClass();
-		checkfSlow(ThisClass, TEXT("Object pointer class did not resolve to a non-null value"));
+		if (const UClass* ThisClass = GetClass())
+		{
+			return ThisClass->IsChildOf(SomeBase);
+		}
 
-		return ThisClass->IsChildOf(SomeBase);
+		return false;
 	}
 
 	template <typename T>
@@ -333,7 +376,10 @@ public:
 	using FObjectPtr::operator!;
 	using FObjectPtr::operator bool;
 	using FObjectPtr::IsResolved;
-	using FObjectPtr::GetPath;
+	using FObjectPtr::GetFName;
+	using FObjectPtr::GetName;
+	using FObjectPtr::GetPathName;
+	using FObjectPtr::GetFullName;
 	using FObjectPtr::GetHandle;
 
 	using FObjectPtr::IsA;
diff --git a/Engine/Source/Runtime/CoreUObject/Public/UObject/SoftObjectPath.h b/Engine/Source/Runtime/CoreUObject/Public/UObject/SoftObjectPath.h
--- a/Engine/Source/Runtime/CoreUObject/Public/UObject/SoftObjectPath.h
+++ b/Engine/Source/Runtime/CoreUObject/Public/UObject/SoftObjectPath.h
@@ -87,17 +87,25 @@ PRAGMA_ENABLE_DEPRECATION_WARNINGS;
 	// UE_DEPRECATED(5.1, "Asset path FNames have been deprecated. This constructor should be used only temporarily to fix up old codepaths that produce an FName.")
 	FSoftObjectPath(FName InAssetPathName, FString InSubPathString);
 	
-	/** Construct from an existing object in memory */
-	FSoftObjectPath(const UObject* InObject);
-
-	FSoftObjectPath(FObjectPtr InObject)
-		: FSoftObjectPath(InObject.Get())
-	{}
-
 	template <typename T>
 	FSoftObjectPath(const TObjectPtr<T>& InObject)
-		: FSoftObjectPath(InObject.Get())
-	{}
+	{
+		SetPath(InObject.GetPathName());
+	}
+
+	FSoftObjectPath(const FObjectPtr& InObject)
+	{
+		SetPath(InObject.GetPathName());
+	}
+
+	/** Construct from an existing object in memory */
+	FSoftObjectPath(const UObject* InObject)
+	{
+		if (InObject)
+		{
+			SetPath(InObject->GetPathName());
+		}
+	}
 
 	~FSoftObjectPath() {}
 PRAGMA_DISABLE_DEPRECATION_WARNINGS	
diff --git a/Engine/Source/Runtime/CoreUObject/Tests/ObjectRefTrackingTestBase.h b/Engine/Source/Runtime/CoreUObject/Tests/ObjectRefTrackingTestBase.h
--- a/Engine/Source/Runtime/CoreUObject/Tests/ObjectRefTrackingTestBase.h
+++ b/Engine/Source/Runtime/CoreUObject/Tests/ObjectRefTrackingTestBase.h
@@ -86,19 +86,10 @@ private:
 		{
 			NumFailedResolves++;
 		}
-
-		if (PrevResolvedFunc)
-		{
-			PrevResolvedFunc(ObjectRef, Pkg, Obj);
-		}
 	}
 	static void OnRefRead(UObject* Obj)
 	{
 		NumReads++;
-		if (PrevReadFunc)
-		{
-			PrevReadFunc(Obj);
-		}
 	}
 #endif
 	
@@ -116,16 +107,18 @@ private:
 		if (!bCallbacksInstalled)
 		{
 #if UE_WITH_OBJECT_HANDLE_TRACKING
-			PrevResolvedFunc = SetObjectHandleReferenceResolvedCallback(OnRefResolved);
-			PrevReadFunc = SetObjectHandleReadCallback(OnRefRead);
+			ResolvedCallbackHandle = AddObjectHandleReferenceResolvedCallback(FObjectHandleReferenceResolvedDelegate::CreateStatic(OnRefResolved));
+			HandleReadCallbackHandle = AddObjectHandleReadCallback(FObjectHandleReadDelegate::CreateStatic(OnRefRead));
+			// TODO We should unhook these handles somewhere, but i don't want to refactor the test, it's not as if they were
+			// being unhooked before.  So...
 #endif
 			bCallbacksInstalled = true;
 		}
 	}
 
 #if UE_WITH_OBJECT_HANDLE_TRACKING
-	static ObjectHandleReferenceResolvedFunction* PrevResolvedFunc;
-	static ObjectHandleReadFunction* PrevReadFunc;
+	static FDelegateHandle ResolvedCallbackHandle;
+	static FDelegateHandle HandleReadCallbackHandle;
 #endif
 	static thread_local uint32 NumResolves;
 	static thread_local uint32 NumFailedResolves;
