Removing some unused stuff from ObjectRefTrackingTestBase for builds with LOW_LEVEL_TESTS in them.

diff --git a/Engine/Source/Runtime/CoreUObject/Tests/ObjectRefTrackingTestBase.cpp b/Engine/Source/Runtime/CoreUObject/Tests/ObjectRefTrackingTestBase.cpp
--- a/Engine/Source/Runtime/CoreUObject/Tests/ObjectRefTrackingTestBase.cpp
+++ b/Engine/Source/Runtime/CoreUObject/Tests/ObjectRefTrackingTestBase.cpp
@@ -4,10 +4,6 @@
 
 #include "ObjectRefTrackingTestBase.h"
 
-#if UE_WITH_OBJECT_HANDLE_TRACKING
-ObjectHandleReferenceResolvedFunction* FObjectRefTrackingTestBase::PrevResolvedFunc = nullptr;
-ObjectHandleReadFunction* FObjectRefTrackingTestBase::PrevReadFunc = nullptr;
-#endif
 thread_local uint32 FObjectRefTrackingTestBase::NumResolves = 0;
 thread_local uint32 FObjectRefTrackingTestBase::NumFailedResolves = 0;
 thread_local uint32 FObjectRefTrackingTestBase::NumReads = 0;
