Fixing a COREUOBJECT_API being in the wrong spot on a struct for linux.

diff --git a/Engine/Source/Runtime/CoreUObject/Public/UObject/ObjectHandle.h b/Engine/Source/Runtime/CoreUObject/Public/UObject/ObjectHandle.h
--- a/Engine/Source/Runtime/CoreUObject/Public/UObject/ObjectHandle.h
+++ b/Engine/Source/Runtime/CoreUObject/Public/UObject/ObjectHandle.h
@@ -300,7 +300,7 @@ namespace ObjectHandle_Private
 	DECLARE_MULTICAST_DELEGATE_ThreeParams(FObjectHandleReferenceResolvedEvent, const FObjectRef& SourceRef, UPackage* ObjectPackage, UObject* Object);
 	DECLARE_MULTICAST_DELEGATE_ThreeParams(FObjectHandleReferenceLoadedEvent, const FObjectRef& SourceRef, UPackage* ObjectPackage, UObject* Object);
 
-	COREUOBJECT_API struct FObjectHandleEvents
+	struct FObjectHandleEvents
 	{
 		FObjectHandleReadEvent ObjectHandleReadEvent;
 		FClassReferenceResolvedEvent ClassReferenceResolvedEvent;
