Making GetPath public for TObjectPtr to satisfy existing uses.

diff --git a/Engine/Source/Runtime/CoreUObject/Public/UObject/ObjectPtr.h b/Engine/Source/Runtime/CoreUObject/Public/UObject/ObjectPtr.h
--- a/Engine/Source/Runtime/CoreUObject/Public/UObject/ObjectPtr.h
+++ b/Engine/Source/Runtime/CoreUObject/Public/UObject/ObjectPtr.h
@@ -382,6 +382,7 @@ public:
 	using FObjectPtr::IsResolved;
 	using FObjectPtr::GetFName;
 	using FObjectPtr::GetName;
+	using FObjectPtr::GetPath;
 	using FObjectPtr::GetPathName;
 	using FObjectPtr::GetFullName;
 	using FObjectPtr::GetHandle;
