EditorPackageLoader: Fixed issue with exports from already loaded uncooked packages not being registered in the import store

diff --git a/Engine/Source/Runtime/CoreUObject/Private/Serialization/AsyncLoading2.cpp b/Engine/Source/Runtime/CoreUObject/Private/Serialization/AsyncLoading2.cpp
--- a/Engine/Source/Runtime/CoreUObject/Private/Serialization/AsyncLoading2.cpp
+++ b/Engine/Source/Runtime/CoreUObject/Private/Serialization/AsyncLoading2.cpp
@@ -4125,7 +4125,7 @@ void FAsyncPackage2::ImportPackagesRecursiveInner(FAsyncLoadingThreadState2& Thr
 #if WITH_EDITOR
 		else if (AsyncLoadingThread.UncookedPackageLoader && !ImportedPackageEntry.UncookedPackageName.IsNone())
 		{
-			UPackage* UncookedPackage = nullptr;
+			UPackage* UncookedPackage = ImportedPackageRef.GetPackage();
 			if (!ImportedPackageRef.AreAllPublicExportsLoaded())
 			{
 				UE_ASYNC_PACKAGE_LOG(Verbose, Desc, TEXT("ImportPackages: LoadUncookedImport"), TEXT("Loading imported uncooked package '%s' '0x%llX'"), *ImportedPackageEntry.UncookedPackageName.ToString(), ImportedPackageId.ValueForDebugging());
@@ -4136,53 +4136,49 @@ void FAsyncPackage2::ImportPackagesRecursiveInner(FAsyncLoadingThreadState2& Thr
 				int32 ImportRequestId = AsyncLoadingThread.UncookedPackageLoader->LoadPackage(ImportedPackagePath, NAME_None, FLoadPackageAsyncDelegate(), PKG_None, INDEX_NONE, 0, nullptr);
 				AsyncLoadingThread.UncookedPackageLoader->FlushLoading(ImportRequestId);
 				UncookedPackage = FindObjectFast<UPackage>(nullptr, ImportedPackagePath.GetPackageFName());
+				ImportedPackageRef.SetPackage(UncookedPackage);
 				if (UncookedPackage)
 				{
 					UncookedPackage->SetCanBeImportedFlag(true);
 					UncookedPackage->SetPackageId(ImportedPackageId);
-					checkf(!UncookedPackage->HasAnyInternalFlags(EInternalObjectFlags::LoaderImport), TEXT("%s"), *UncookedPackage->GetFullName());
 					UncookedPackage->SetInternalFlags(EInternalObjectFlags::LoaderImport);
-
-					ForEachObjectWithOuter(UncookedPackage, [this, ImportedPackageId](UObject* Object)
+					ImportedPackageRef.SetAllPublicExportsLoaded();
+				}
+			}
+			if (UncookedPackage)
+			{
+				ForEachObjectWithOuter(UncookedPackage, [this, ImportedPackageId](UObject* Object)
+				{
+					if (Object->HasAllFlags(RF_Public))
 					{
-						if (Object->HasAllFlags(RF_Public))
-						{
-							checkf(!Object->HasAnyInternalFlags(EInternalObjectFlags::LoaderImport), TEXT("%s"), *Object->GetFullName());
-							Object->SetInternalFlags(EInternalObjectFlags::LoaderImport);
+						Object->SetInternalFlags(EInternalObjectFlags::LoaderImport);
 
-							TArray<FName, TInlineAllocator<64>> FullPath;
-							FullPath.Add(Object->GetFName());
-							UObject* Outer = Object->GetOuter();
-							while (Outer)
-							{
-								FullPath.Add(Outer->GetFName());
-								Outer = Outer->GetOuter();
-							}
-							TStringBuilder<256> PackageRelativeExportPath;
-							for (int32 PathIndex = FullPath.Num() - 2; PathIndex >= 0; --PathIndex)
+						TArray<FName, TInlineAllocator<64>> FullPath;
+						FullPath.Add(Object->GetFName());
+						UObject* Outer = Object->GetOuter();
+						while (Outer)
+						{
+							FullPath.Add(Outer->GetFName());
+							Outer = Outer->GetOuter();
+						}
+						TStringBuilder<256> PackageRelativeExportPath;
+						for (int32 PathIndex = FullPath.Num() - 2; PathIndex >= 0; --PathIndex)
+						{
+							TCHAR NameStr[FName::StringBufferSize];
+							uint32 NameLen = FullPath[PathIndex].ToString(NameStr);
+							for (uint32 I = 0; I < NameLen; ++I)
 							{
-								TCHAR NameStr[FName::StringBufferSize];
-								uint32 NameLen = FullPath[PathIndex].ToString(NameStr);
-								for (uint32 I = 0; I < NameLen; ++I)
-								{
-									NameStr[I] = TChar<TCHAR>::ToLower(NameStr[I]);
-								}
-								PackageRelativeExportPath.AppendChar('/');
-								PackageRelativeExportPath.Append(FStringView(NameStr, NameLen));
+								NameStr[I] = TChar<TCHAR>::ToLower(NameStr[I]);
 							}
-							uint64 ExportHash = CityHash64(reinterpret_cast<const char*>(PackageRelativeExportPath.GetData() + 1), (PackageRelativeExportPath.Len() - 1) * sizeof(TCHAR));
-							ImportStore.StoreGlobalObject(ImportedPackageId, ExportHash, Object);
+							PackageRelativeExportPath.AppendChar('/');
+							PackageRelativeExportPath.Append(FStringView(NameStr, NameLen));
 						}
-					}, /* bIncludeNestedObjects*/ true);
-				}
-				ImportedPackageRef.SetPackage(UncookedPackage);
-				ImportedPackageRef.SetAllPublicExportsLoaded();
+						uint64 ExportHash = CityHash64(reinterpret_cast<const char*>(PackageRelativeExportPath.GetData() + 1), (PackageRelativeExportPath.Len() - 1) * sizeof(TCHAR));
+						ImportStore.StoreGlobalObject(ImportedPackageId, ExportHash, Object);
+					}
+				}, /* bIncludeNestedObjects*/ true);
 			}
 			else
-			{
-				UncookedPackage = ImportedPackageRef.GetPackage();
-			}
-			if (!UncookedPackage)
 			{
 				ImportedPackageRef.SetHasFailed();
 				UE_ASYNC_PACKAGE_LOG(Warning, Desc, TEXT("ImportPackages: SkipPackage"),
