Deprecate FNames containing full object paths.

diff --git a/Engine/Plugins/Editor/AssetManagerEditor/Source/AssetManagerEditor/Public/AssetManagerEditorModule.h b/Engine/Plugins/Editor/AssetManagerEditor/Source/AssetManagerEditor/Public/AssetManagerEditorModule.h
--- a/Engine/Plugins/Editor/AssetManagerEditor/Source/AssetManagerEditor/Public/AssetManagerEditorModule.h
+++ b/Engine/Plugins/Editor/AssetManagerEditor/Source/AssetManagerEditor/Public/AssetManagerEditorModule.h
@@ -217,7 +217,7 @@ public:
 		}
 	}
 
-	// UE_DEPRECATED(5.1, "Asset path FNames have been deprecated, use FSoftObjectPath instead.")
+	UE_DEPRECATED(5.1, "Asset path FNames have been deprecated, use FSoftObjectPath instead.")
 	FAssetData GetAssetByObjectPath(FName ObjectPath) const
 	{
 PRAGMA_DISABLE_DEPRECATION_WARNINGS
diff --git a/Engine/Source/Runtime/AssetRegistry/Private/AssetRegistry.h b/Engine/Source/Runtime/AssetRegistry/Private/AssetRegistry.h
--- a/Engine/Source/Runtime/AssetRegistry/Private/AssetRegistry.h
+++ b/Engine/Source/Runtime/AssetRegistry/Private/AssetRegistry.h
@@ -42,7 +42,7 @@ public:
 	virtual bool GetAssets(const FARFilter& Filter, TArray<FAssetData>& OutAssetData, bool bSkipARFilteredAssets = true) const override;
 	virtual bool EnumerateAssets(const FARFilter& Filter, TFunctionRef<bool(const FAssetData&)> Callback, bool bSkipARFilteredAssets=true) const override;
 	virtual bool EnumerateAssets(const FARCompiledFilter& Filter, TFunctionRef<bool(const FAssetData&)> Callback, bool bSkipARFilteredAssets = true) const override;
-	// UE_DEPRECATED(5.1, "Asset path FNames have been deprecated, use FSoftObjectPath instead.")
+	UE_DEPRECATED(5.1, "Asset path FNames have been deprecated, use FSoftObjectPath instead.")
 	virtual FAssetData GetAssetByObjectPath( const FName ObjectPath, bool bIncludeOnlyOnDiskAssets = false ) const override;
 	virtual FAssetData GetAssetByObjectPath(const FSoftObjectPath& ObjectPath, bool bIncludeOnlyOnDiskAssets = false) const override;
 	virtual bool GetAllAssets(TArray<FAssetData>& OutAssetData, bool bIncludeOnlyOnDiskAssets = false) const override;
@@ -68,7 +68,7 @@ public:
 	virtual void EnumerateAllPackages(TFunctionRef<void(FName PackageName, const FAssetPackageData& PackageData)> Callback) const override;
 	virtual bool DoesPackageExistOnDisk(FName PackageName, FString* OutCorrectCasePackageName = nullptr, FString* OutExtension = nullptr) const override;
 	virtual FSoftObjectPath GetRedirectedObjectPath(const FSoftObjectPath& ObjectPath) const override;
-	// UE_DEPRECATED(5.1, "Asset path FNames have been deprecated, use FSoftObjectPath instead.")
+	UE_DEPRECATED(5.1, "Asset path FNames have been deprecated, use FSoftObjectPath instead.")
 	virtual FName GetRedirectedObjectPath(const FName ObjectPath) const override;
 	UE_DEPRECATED(5.1, "Class names are now represented by path names. Please use a version of this function that uses FTopLevelAssetPath.")
 	virtual bool GetAncestorClassNames(FName ClassName, TArray<FName>& OutAncestorClassNames) const override;
diff --git a/Engine/Source/Runtime/AssetRegistry/Private/AssetRegistryImpl.h b/Engine/Source/Runtime/AssetRegistry/Private/AssetRegistryImpl.h
--- a/Engine/Source/Runtime/AssetRegistry/Private/AssetRegistryImpl.h
+++ b/Engine/Source/Runtime/AssetRegistry/Private/AssetRegistryImpl.h
@@ -125,7 +125,7 @@ public:
 
 	bool HasAssets(const FName PackagePath, const bool bRecursive) const;
 	FSoftObjectPath GetRedirectedObjectPath(const FSoftObjectPath& ObjectPath) const;
-	// UE_DEPRECATED(5.1, "Asset path FNames have been deprecated, use FSoftObjectPath instead.")
+	UE_DEPRECATED(5.1, "Asset path FNames have been deprecated, use FSoftObjectPath instead.")
 	FName GetRedirectedObjectPath(const FName ObjectPath) const;
 	bool GetAncestorClassNames(Impl::FClassInheritanceContext& InheritanceContext, FTopLevelAssetPath ClassName,
 		TArray<FTopLevelAssetPath>& OutAncestorClassNames) const;
diff --git a/Engine/Source/Runtime/AssetRegistry/Public/AssetRegistry/IAssetRegistry.h b/Engine/Source/Runtime/AssetRegistry/Public/AssetRegistry/IAssetRegistry.h
--- a/Engine/Source/Runtime/AssetRegistry/Public/AssetRegistry/IAssetRegistry.h
+++ b/Engine/Source/Runtime/AssetRegistry/Public/AssetRegistry/IAssetRegistry.h
@@ -216,7 +216,7 @@ public:
 	 * @param bIncludeOnlyOnDiskAssets if true, in-memory objects will be ignored. The call will be faster.
 	 * @return the assets data;Will be invalid if object could not be found
 	 */
-	// UE_DEPRECATED(5.1, "Asset path FNames have been deprecated, use Soft Object Path instead.")
+	UE_DEPRECATED(5.1, "Asset path FNames have been deprecated, use Soft Object Path instead.")
 	UFUNCTION(BlueprintCallable, BlueprintPure=false, Category="AssetRegistry")
 	virtual FAssetData GetAssetByObjectPath( const FName ObjectPath, bool bIncludeOnlyOnDiskAssets = false ) const = 0;
 
@@ -364,7 +364,7 @@ public:
 	/** Uses the asset registry to look for ObjectRedirectors. This will follow the chain of redirectors. It will return the original path if no redirectors are found */
 	virtual FSoftObjectPath GetRedirectedObjectPath(const FSoftObjectPath& ObjectPath) const = 0;
 
-	// UE_DEPRECATED(5.1, "Asset path FNames have been deprecated, use FSoftObjectPath instead.")
+	UE_DEPRECATED(5.1, "Asset path FNames have been deprecated, use FSoftObjectPath instead.")
 	virtual FName GetRedirectedObjectPath(const FName ObjectPath) const = 0;
 
 	UE_DEPRECATED(4.27, "Loading then discarding tags is no longer allowed as it can "
@@ -706,7 +706,7 @@ protected:
 	virtual void SetManageReferences(const TMultiMap<FAssetIdentifier, FAssetIdentifier>& ManagerMap, bool bClearExisting, UE::AssetRegistry::EDependencyCategory RecurseType, TSet<FDependsNode*>& ExistingManagedNodes, ShouldSetManagerPredicate ShouldSetManager = nullptr) = 0;
 
 	/** Sets the PrimaryAssetId for a specific asset. This should only be called by the AssetManager, and is needed when the AssetManager is more up to date than the on disk Registry */
-	// UE_DEPRECATED(5.1, "Asset path FNames have been deprecated, use FSoftObjectPath instead.")
+	UE_DEPRECATED(5.1, "Asset path FNames have been deprecated, use FSoftObjectPath instead.")
 	virtual bool SetPrimaryAssetIdForObjectPath(const FName ObjectPath, FPrimaryAssetId PrimaryAssetId) = 0;
 	virtual bool SetPrimaryAssetIdForObjectPath(const FSoftObjectPath& ObjectPath, FPrimaryAssetId PrimaryAssetId) = 0;
 };
diff --git a/Engine/Source/Runtime/CoreUObject/Public/AssetRegistry/ARFilter.h b/Engine/Source/Runtime/CoreUObject/Public/AssetRegistry/ARFilter.h
--- a/Engine/Source/Runtime/CoreUObject/Public/AssetRegistry/ARFilter.h
+++ b/Engine/Source/Runtime/CoreUObject/Public/AssetRegistry/ARFilter.h
@@ -32,7 +32,7 @@ PRAGMA_ENABLE_DEPRECATION_WARNINGS
 	TArray<FName> PackagePaths;
 
 	/** The filter component containing specific object paths */
-	// UE_DEPRECATED(5.1, "Asset path FNames have been deprecated, use FSoftObjectPath instead.")
+	UE_DEPRECATED(5.1, "Asset path FNames have been deprecated, use FSoftObjectPath instead.")
 	TArray<FName> ObjectPaths;
 
 	/** 
@@ -161,7 +161,7 @@ PRAGMA_ENABLE_DEPRECATION_WARNINGS
 	TSet<FName> PackagePaths;
 
 	/** The filter component containing specific object paths */
-	// UE_DEPRECATED(5.1, "Object path FNames have been deprecated, use FSoftObjectPath instead.")
+	UE_DEPRECATED(5.1, "Object path FNames have been deprecated, use FSoftObjectPath instead.")
 	TSet<FName> ObjectPaths;
 
 	/** The filter component containing specific object paths */
diff --git a/Engine/Source/Runtime/CoreUObject/Public/AssetRegistry/AssetBundleData.h b/Engine/Source/Runtime/CoreUObject/Public/AssetRegistry/AssetBundleData.h
--- a/Engine/Source/Runtime/CoreUObject/Public/AssetRegistry/AssetBundleData.h
+++ b/Engine/Source/Runtime/CoreUObject/Public/AssetRegistry/AssetBundleData.h
@@ -36,7 +36,7 @@ PRAGMA_ENABLE_DEPRECATION_WARNINGS
 	FName BundleName;
 
 	/** List of string assets contained in this bundle */
-	// UE_DEPRECATED(5.1, "BundleAssets has been replaced with AssetPaths.")
+	UE_DEPRECATED(5.1, "BundleAssets has been replaced with AssetPaths.")
 	TArray<FSoftObjectPath> BundleAssets;
 
 	/** List of references to top-level assets contained in this bundle */
@@ -110,22 +110,22 @@ struct COREUOBJECT_API FAssetBundleData
 	}
 
 	/** Adds or updates an entry with the given BundleName -> Path. Scope is empty and will be filled in later */
-	// UE_DEPRECATED(5.1, "The contents of an asset bundle now use FTopLevelAssetPath instead of FSoftObjectPath.")
+	UE_DEPRECATED(5.1, "The contents of an asset bundle now use FTopLevelAssetPath instead of FSoftObjectPath.")
 	void AddBundleAsset(FName BundleName, const FSoftObjectPath& AssetPath);
 
 	template< typename T >
-	// UE_DEPRECATED(5.1, "The contents of an asset bundle now use FTopLevelAssetPath instead of FSoftObjectPath.")
+	UE_DEPRECATED(5.1, "The contents of an asset bundle now use FTopLevelAssetPath instead of FSoftObjectPath.")
 	void AddBundleAsset(FName BundleName, const TSoftObjectPtr<T>& SoftObjectPtr)
 	{
 		AddBundleAsset(BundleName, SoftObjectPtr.ToSoftObjectPath().GetAssetPath());
 	}
 
 	/** Adds multiple assets at once */
-	// UE_DEPRECATED(5.1, "The contents of an asset bundle now use FTopLevelAssetPath instead of FSoftObjectPath.")
+	UE_DEPRECATED(5.1, "The contents of an asset bundle now use FTopLevelAssetPath instead of FSoftObjectPath.")
 	void AddBundleAssets(FName BundleName, const TArray<FSoftObjectPath>& AssetPaths);
 
 	/** A fast set of asset bundle assets, will destroy copied in path list */
-	// UE_DEPRECATED(5.1, "The contents of an asset bundle now use FTopLevelAssetPath instead of FSoftObjectPath.")
+	UE_DEPRECATED(5.1, "The contents of an asset bundle now use FTopLevelAssetPath instead of FSoftObjectPath.")
 	void SetBundleAssets(FName BundleName, TArray<FSoftObjectPath>&& AssetPaths);
 
 	/** Adds or updates an entry with the given BundleName -> Path.*/
diff --git a/Engine/Source/Runtime/CoreUObject/Public/AssetRegistry/AssetData.h b/Engine/Source/Runtime/CoreUObject/Public/AssetRegistry/AssetData.h
--- a/Engine/Source/Runtime/CoreUObject/Public/AssetRegistry/AssetData.h
+++ b/Engine/Source/Runtime/CoreUObject/Public/AssetRegistry/AssetData.h
@@ -155,7 +155,7 @@ public:
 
 public:
 	/** The object path for the asset in the form PackageName.ObjectName, or PackageName.ObjectName:SubObjectName. */
-	// UE_DEPRECATED(5.1, "FName asset paths have been deprecated. Use GetSoftObjectPath to get the path this asset will use in memory when loaded.")
+	UE_DEPRECATED(5.1, "FName asset paths have been deprecated. Use GetSoftObjectPath to get the path this asset will use in memory when loaded.")
 	FName ObjectPath;
 	/** The name of the package in which the asset is found, this is the full long package name such as /Game/Path/Package */
 	FName PackageName;
@@ -164,7 +164,7 @@ public:
 	/** The name of the asset without the package */
 	FName AssetName;
 	/** The name of the asset's class */
-	// UE_DEPRECATED(5.1, "Class names are now represented by path names. Please use AssetClassPath.")
+	UE_DEPRECATED(5.1, "Class names are now represented by path names. Please use AssetClassPath.")
 	FName AssetClass;
 	/** The path of the asset's class, e.g. /Script/Engine.StaticMesh */
 	FTopLevelAssetPath AssetClassPath;
diff --git a/Engine/Source/Runtime/CoreUObject/Public/Misc/RedirectCollector.h b/Engine/Source/Runtime/CoreUObject/Public/Misc/RedirectCollector.h
--- a/Engine/Source/Runtime/CoreUObject/Public/Misc/RedirectCollector.h
+++ b/Engine/Source/Runtime/CoreUObject/Public/Misc/RedirectCollector.h
@@ -106,17 +106,17 @@ public:
 	void ProcessSoftObjectPathPackageList(FName FilterPackage, bool bGetEditorOnly, TSet<FName>& OutReferencedPackages);
 
 	/** Adds a new mapping for redirector path to destination path, this is called from the Asset Registry to register all redirects it knows about */
-	// UE_DEPRECATED(5.1, "Asset path FNames have been deprecated, use FSoftObjectPath instead.")
+	UE_DEPRECATED(5.1, "Asset path FNames have been deprecated, use FSoftObjectPath instead.")
 	void AddAssetPathRedirection(FName OriginalPath, FName RedirectedPath);
 	void AddAssetPathRedirection(const FSoftObjectPath& OriginalPath, const FSoftObjectPath& RedirectedPath);
 
 	/** Removes an asset path redirection, call this when deleting redirectors */
-	// UE_DEPRECATED(5.1, "Asset path FNames have been deprecated, use FSoftObjectPath instead.")
+	UE_DEPRECATED(5.1, "Asset path FNames have been deprecated, use FSoftObjectPath instead.")
 	void RemoveAssetPathRedirection(FName OriginalPath);
 	void RemoveAssetPathRedirection(const FSoftObjectPath& OriginalPath);
 
 	/** Returns a remapped asset path, if it returns null there is no relevant redirector */
-	// UE_DEPRECATED(5.1, "Asset path FNames have been deprecated, use FSoftObjectPath instead.")
+	UE_DEPRECATED(5.1, "Asset path FNames have been deprecated, use FSoftObjectPath instead.")
 	FName GetAssetPathRedirection(FName OriginalPath);
 	/** Returns a remapped asset path, if there is no relevant redirector, the return value reports true from IsNull() */
 	FSoftObjectPath GetAssetPathRedirection(const FSoftObjectPath& OriginalPath);
diff --git a/Engine/Source/Runtime/CoreUObject/Public/UObject/SoftObjectPath.h b/Engine/Source/Runtime/CoreUObject/Public/UObject/SoftObjectPath.h
--- a/Engine/Source/Runtime/CoreUObject/Public/UObject/SoftObjectPath.h
+++ b/Engine/Source/Runtime/CoreUObject/Public/UObject/SoftObjectPath.h
@@ -75,7 +75,7 @@ struct COREUOBJECT_API FSoftObjectPath
 	explicit FSoftObjectPath(const ANSICHAR* Path)				{ SetPath(FAnsiStringView(Path)); }
 	explicit FSoftObjectPath(TYPE_OF_NULLPTR)					{}
 
-	// UE_DEPRECATED(5.1, "Asset path FNames have been deprecated. This constructor should be used only temporarily to fix up old codepaths that produce an FName.")
+	UE_DEPRECATED(5.1, "Asset path FNames have been deprecated. This constructor should be used only temporarily to fix up old codepaths that produce an FName.")
 	explicit FSoftObjectPath(FName Path) 
 	{
 PRAGMA_DISABLE_DEPRECATION_WARNINGS;
@@ -84,7 +84,7 @@ PRAGMA_ENABLE_DEPRECATION_WARNINGS;
 	}
 
 	/** Construct from an asset FName and subobject pair */
-	// UE_DEPRECATED(5.1, "Asset path FNames have been deprecated. This constructor should be used only temporarily to fix up old codepaths that produce an FName.")
+	UE_DEPRECATED(5.1, "Asset path FNames have been deprecated. This constructor should be used only temporarily to fix up old codepaths that produce an FName.")
 	FSoftObjectPath(FName InAssetPathName, FString InSubPathString);
 	
 	template <typename T>
@@ -120,7 +120,7 @@ PRAGMA_ENABLE_DEPRECATION_WARNINGS
 	FSoftObjectPath& operator=(const ANSICHAR* Path)					{ SetPath(FAnsiStringView(Path)); return *this; }
 	FSoftObjectPath& operator=(TYPE_OF_NULLPTR)							{ Reset(); return *this; }
 
-	// UE_DEPRECATED(5.1, "Asset path FNames have been deprecated. This assignment operator should be used only temporarily to fix up old codepaths that produce an FName.")
+	UE_DEPRECATED(5.1, "Asset path FNames have been deprecated. This assignment operator should be used only temporarily to fix up old codepaths that produce an FName.")
 	FSoftObjectPath& operator=(FName Path) 
 	{
 PRAGMA_DISABLE_DEPRECATION_WARNINGS;
@@ -132,7 +132,7 @@ PRAGMA_ENABLE_DEPRECATION_WARNINGS;
 	/** Returns string representation of reference, in form /package/path.assetname[:subpath] */
 	FString ToString() const;
 
-	// UE_DEPRECATED(5.1, "Asset path FNames have been deprecated. This function should be used only temporarily to interface with old APIs that require an FName.")
+	UE_DEPRECATED(5.1, "Asset path FNames have been deprecated. This function should be used only temporarily to interface with old APIs that require an FName.")
 	FName ToFName() const { return *ToString(); }
 
 	/** Append string representation of reference, in form /package/path.assetname[:subpath] */
@@ -155,10 +155,10 @@ PRAGMA_ENABLE_DEPRECATION_WARNINGS;
 	}
 	
 	/** Returns the entire asset path as an FName, including both package and asset but not sub object */
-	// UE_DEPRECATED(5.1, "Asset path FNames have been deprecated. Use GetAssetPath instead.")
+	UE_DEPRECATED(5.1, "Asset path FNames have been deprecated. Use GetAssetPath instead.")
 	FName GetAssetPathName() const;
 	
-	// UE_DEPRECATED(5.1, "Asset path FNames have been deprecated. Use SetAssetPath instead.")
+	UE_DEPRECATED(5.1, "Asset path FNames have been deprecated. Use SetAssetPath instead.")
 	void SetAssetPathName(FName InAssetPathName);
 
 	/** Returns string version of asset path, including both package and asset but not sub object */
@@ -214,7 +214,7 @@ PRAGMA_ENABLE_DEPRECATION_WARNINGS;
 	void SetPath(FWideStringView Path);
 	void SetPath(FAnsiStringView Path);
 	void SetPath(FUtf8StringView Path);
-	// UE_DEPRECATED(5.1, "Asset path FNames have been deprecated. This function should be used only temporarily to fix up old codepaths that produce an FName.")
+	UE_DEPRECATED(5.1, "Asset path FNames have been deprecated. This function should be used only temporarily to fix up old codepaths that produce an FName.")
 	void SetPath(FName Path);
 	void SetPath(const WIDECHAR* Path)			{ SetPath(FWideStringView(Path)); }
 	void SetPath(const ANSICHAR* Path)			{ SetPath(FAnsiStringView(Path)); }
@@ -587,7 +587,7 @@ struct FSoftObjectPathFixupArchive : public FArchiveUObject
 
 namespace UE::SoftObjectPath::Private
 {
-	// UE_DEPRECATED(5.1, "This function is only for use in fixing up deprecated APIs.")
+	UE_DEPRECATED(5.1, "This function is only for use in fixing up deprecated APIs.")
 	inline TArray<FName> ConvertSoftObjectPaths(TConstArrayView<FSoftObjectPath> InPaths)
 	{
 		TArray<FName> Out;
@@ -598,7 +598,7 @@ PRAGMA_ENABLE_DEPRECATION_WARNINGS
 		return Out;
 	}
 
-	// UE_DEPRECATED(5.1, "This function is only for use in fixing up deprecated APIs.")
+	UE_DEPRECATED(5.1, "This function is only for use in fixing up deprecated APIs.")
 	inline TArray<FSoftObjectPath> ConvertObjectPathNames(TConstArrayView<FName> InPaths)
 	{
 		TArray<FSoftObjectPath> Out;
