Bulk Data - crash fix for bulk data with wrong flags getting uploaded to DDC

#rnx

diff --git a/Engine/Source/Runtime/CoreUObject/Private/Serialization/BulkData.cpp b/Engine/Source/Runtime/CoreUObject/Private/Serialization/BulkData.cpp
--- a/Engine/Source/Runtime/CoreUObject/Private/Serialization/BulkData.cpp
+++ b/Engine/Source/Runtime/CoreUObject/Private/Serialization/BulkData.cpp
@@ -1335,6 +1335,13 @@ void FBulkData::Serialize(FArchive& Ar, UObject* Owner, bool bAttemptFileMapping
 
 			FBulkMetaResource MetaResource;
 			Ar << MetaResource;
+#if WITH_EDITOR
+			if (Owner == nullptr)
+			{
+				// Temp fix for accidentially uploading bulk data with wrong flags to DDC
+				ClearBulkDataFlagsOn(MetaResource.Flags, static_cast<EBulkDataFlags>(BULKDATA_PayloadAtEndOfFile | BULKDATA_PayloadInSeperateFile | BULKDATA_WorkspaceDomainPayload));
+			}
+#endif
 			BulkMeta = FBulkMetaData::FromSerialized(MetaResource, ElementSize);
 
 			bool bIsUsingIoDispatcher = false;
