Core - Reverting the base type of FObjectProperty's FObjectPropertyBase to UObject* instead of TObjectPtr<UObject>, there's some range garbage collection issue this causes where things arent cleaned up in the expected fashion.  This may be limited to script, but unsure.  Working on it.  Reverting for now to fix the EngineTest tests.

diff --git a/Engine/Source/Runtime/CoreUObject/Private/UObject/PropertyObject.cpp b/Engine/Source/Runtime/CoreUObject/Private/UObject/PropertyObject.cpp
--- a/Engine/Source/Runtime/CoreUObject/Private/UObject/PropertyObject.cpp
+++ b/Engine/Source/Runtime/CoreUObject/Private/UObject/PropertyObject.cpp
@@ -98,7 +98,7 @@ void FObjectProperty::SerializeItem( FStructuredArchive::FSlot Slot, void* Value
 	if (UnderlyingArchive.IsObjectReferenceCollector())
 	{
 		// Serialize in place
-		TObjectPtr<UObject>* ObjectPtr = GetPropertyValuePtr(Value);
+		UObject** ObjectPtr = GetPropertyValuePtr(Value);
 		Slot << (*ObjectPtr);
 
 		if(!UnderlyingArchive.IsSaving())
@@ -145,7 +145,7 @@ void FObjectProperty::SerializeItem( FStructuredArchive::FSlot Slot, void* Value
 
 const TCHAR* FObjectProperty::ImportText_Internal(const TCHAR* Buffer, void* ContainerOrPropertyPtr, EPropertyPointerType PropertyPointerType, UObject* OwnerObject, int32 PortFlags, FOutputDevice* ErrorText) const
 {
-	const TCHAR* Result = TFObjectPropertyBase<TObjectPtr<UObject>>::ImportText_Internal(Buffer, ContainerOrPropertyPtr, PropertyPointerType, OwnerObject, PortFlags, ErrorText);
+	const TCHAR* Result = TFObjectPropertyBase<UObject*>::ImportText_Internal(Buffer, ContainerOrPropertyPtr, PropertyPointerType, OwnerObject, PortFlags, ErrorText);
 	if (Result)
 	{
 		void* Data = PointerToValuePtr(ContainerOrPropertyPtr, PropertyPointerType);
@@ -185,7 +185,7 @@ TObjectPtr<UObject> FObjectProperty::GetObjectPtrPropertyValue(const void* Prope
 
 UObject* FObjectProperty::GetObjectPropertyValue(const void* PropertyValueAddress) const
 {
-	return GetPropertyValue(PropertyValueAddress).Get();
+	return GetPropertyValue(PropertyValueAddress);
 }
 
 UObject* FObjectProperty::GetObjectPropertyValue_InContainer(const void* ContainerAddress, int32 ArrayIndex) const
diff --git a/Engine/Source/Runtime/CoreUObject/Private/UObject/PropertyObjectPtr.cpp b/Engine/Source/Runtime/CoreUObject/Private/UObject/PropertyObjectPtr.cpp
--- a/Engine/Source/Runtime/CoreUObject/Private/UObject/PropertyObjectPtr.cpp
+++ b/Engine/Source/Runtime/CoreUObject/Private/UObject/PropertyObjectPtr.cpp
@@ -151,7 +151,7 @@ TObjectPtr<UObject> FObjectPtrProperty::GetObjectPtrPropertyValue(const void* Pr
 
 UObject* FObjectPtrProperty::GetObjectPropertyValue(const void* PropertyValueAddress) const
 {
-	return GetPropertyValue(PropertyValueAddress).Get();
+	return GetPropertyValue(PropertyValueAddress);
 }
 
 UObject* FObjectPtrProperty::GetObjectPropertyValue_InContainer(const void* ContainerAddress, int32 ArrayIndex) const
diff --git a/Engine/Source/Runtime/CoreUObject/Public/UObject/UnrealType.h b/Engine/Source/Runtime/CoreUObject/Public/UObject/UnrealType.h
--- a/Engine/Source/Runtime/CoreUObject/Public/UObject/UnrealType.h
+++ b/Engine/Source/Runtime/CoreUObject/Public/UObject/UnrealType.h
@@ -2898,9 +2898,12 @@ PRAGMA_ENABLE_DEPRECATION_WARNINGS
 //
 // Describes a reference variable to another object which may be nil.
 //
-class COREUOBJECT_API FObjectProperty : public TFObjectPropertyBase<TObjectPtr<UObject>>
+//class COREUOBJECT_API FObjectProperty : public TFObjectPropertyBase<TObjectPtr<UObject>>
+//{
+//	DECLARE_FIELD(FObjectProperty, TFObjectPropertyBase<TObjectPtr<UObject>>, CASTCLASS_FObjectProperty)
+class COREUOBJECT_API FObjectProperty : public TFObjectPropertyBase<UObject*>
 {
-	DECLARE_FIELD(FObjectProperty, TFObjectPropertyBase<TObjectPtr<UObject>>, CASTCLASS_FObjectProperty)
+	DECLARE_FIELD(FObjectProperty, TFObjectPropertyBase<UObject*>, CASTCLASS_FObjectProperty)
 
 	FObjectProperty(FFieldVariant InOwner, const FName& InName, EObjectFlags InObjectFlags)
 		: TFObjectPropertyBase(InOwner, InName, InObjectFlags)
