Allow replicated FSoftObjectPtrs to resolve unstably-named objects on clients.
These now serialize a NetGUID tracked by the PackageMap, as other replicated references to objects do. Effectively acts as a weak pointer in the replication case now as well. Soft pointers to stably-named objects continue to replicate their string path.
Added EngineTest cases for replicated soft pointers.

diff --git a/Engine/Source/Runtime/Core/Public/Misc/NetworkVersion.h b/Engine/Source/Runtime/Core/Public/Misc/NetworkVersion.h
--- a/Engine/Source/Runtime/Core/Public/Misc/NetworkVersion.h
+++ b/Engine/Source/Runtime/Core/Public/Misc/NetworkVersion.h
@@ -67,9 +67,10 @@ enum EEngineNetworkVersionHistory
 	HISTORY_PACKED_VECTOR_LWC_SUPPORT = 23,			// Bump version to support quantized LWC FVector net serialization
 	HISTORY_PAWN_REMOTEVIEWPITCH = 24,				// Bump version to support serialization changes to RemoteViewPitch
 	HISTORY_REPMOVE_SERVERFRAME_AND_HANDLE = 25,	// Bump version to support serialization changes to RepMove so we can get the serverframe and physics handle associated with the object
-	HISTORY_21_AND_VIEWPITCH_ONLY_DO_NOT_USE = 26,  // Bump version to support up to history 21 + HISTORY_PAWN_REMOTEVIEWPITCH.  DO NOT USE!!!
-	HISTORY_PLACEHOLDER = 27,                       // Bump version to a placeholder.  This version is the same as HISTORY_REPMOVE_SERVERFRAME_AND_HANDLE
+	HISTORY_21_AND_VIEWPITCH_ONLY_DO_NOT_USE = 26,	// Bump version to support up to history 21 + HISTORY_PAWN_REMOTEVIEWPITCH.  DO NOT USE!!!
+	HISTORY_PLACEHOLDER = 27,						// Bump version to a placeholder.  This version is the same as HISTORY_REPMOVE_SERVERFRAME_AND_HANDLE
 	HISTORY_RUNTIME_FEATURES_COMPATIBILITY = 28,	// Bump version to add network runtime feature compatibility test to handshake (hello/upgrade) control messages
+	HISTORY_SOFTOBJECTPTR_NETGUIDS = 29,			// Bump version to support replicating SoftObjectPtrs by NetGuid instead of raw strings.
 	// New history items go above here.
 
 	HISTORY_ENGINENETVERSION_PLUS_ONE,
diff --git a/Engine/Source/Runtime/CoreUObject/Private/UObject/CoreNet.cpp b/Engine/Source/Runtime/CoreUObject/Private/UObject/CoreNet.cpp
--- a/Engine/Source/Runtime/CoreUObject/Private/UObject/CoreNet.cpp
+++ b/Engine/Source/Runtime/CoreUObject/Private/UObject/CoreNet.cpp
@@ -444,7 +444,37 @@ FArchive& FNetBitWriter::operator<<(FSoftObjectPath& Value)
 
 FArchive& FNetBitWriter::operator<<(FSoftObjectPtr& Value)
 {
-	return FArchiveUObject::SerializeSoftObjectPtr(*this, Value);
+	// Keep in sync with FNetBitReader::operator<<(FSoftObjectPtr&)
+
+	if (EngineNetVer() >= EEngineNetworkVersionHistory::HISTORY_SOFTOBJECTPTR_NETGUIDS)
+	{
+		// Treat soft pointers to dynamic (non-stably named) objects as weak pointers.
+		// This allows these objects to be resolved on clients and to only replicate a NetGUID.
+		//
+		// In theory stably named objects could be replicated as NetGUIDs also via the package map,
+		// only exporting and replicating the path string once, but this would require a refactor
+		// of the package map to handle "soft" references to packages and not automatically load them
+		// when they're referenced.
+		UObject* Object = Value.Get();
+		const bool bUsePath = !Object || Object->IsFullNameStableForNetworking();
+
+		WriteBit(bUsePath);
+
+		if (bUsePath)
+		{
+			FArchiveUObject::SerializeSoftObjectPtr(*this, Value);
+		}
+		else
+		{
+			PackageMap->SerializeObject(*this, UObject::StaticClass(), Object, nullptr);
+		}
+	}
+	else
+	{
+		FArchiveUObject::SerializeSoftObjectPtr(*this, Value);
+	}
+
+	return *this;
 }
 
 FArchive& FNetBitWriter::operator<<(FObjectPtr& Value)
@@ -505,7 +535,42 @@ FArchive& FNetBitReader::operator<<(FSoftObjectPath& Value)
 
 FArchive& FNetBitReader::operator<<(FSoftObjectPtr& Value)
 {
-	return FArchiveUObject::SerializeSoftObjectPtr(*this, Value);
+	// Keep in sync with FNetBitWriter::operator<<(FSoftObjectPtr&)
+
+	if (EngineNetVer() >= EEngineNetworkVersionHistory::HISTORY_SOFTOBJECTPTR_NETGUIDS)
+	{
+		// Treat soft pointers to dynamic (non-stably named) objects as weak pointers.
+		// This allows these objects to be resolved on clients and to only replicate a NetGUID.
+		//
+		// In theory stably named objects could be replicated as NetGUIDs also via the package map,
+		// only exporting and replicating the path string once, but this would require a refactor
+		// of the package map to handle "soft" references to packages and not automatically load them
+		// when they're referenced.
+		const bool bUsePath = ReadBit() != 0;
+
+		UObject* Object = nullptr;
+
+		if (bUsePath)
+		{
+			FArchiveUObject::SerializeSoftObjectPtr(*this, Value);
+		}
+		else
+		{
+			PackageMap->SerializeObject(*this, UObject::StaticClass(), Object, nullptr);
+		}
+
+		if (!bUsePath)
+		{
+			Value = Object;
+		}
+		Value.GetUniqueID().FixupForPIE();
+	}
+	else
+	{
+		FArchiveUObject::SerializeSoftObjectPtr(*this, Value);
+	}
+
+	return *this;
 }
 
 FArchive& FNetBitReader::operator<<(FObjectPtr& Value)
