Ensure that FSimpleObjectReferenceCollectorArchive allows subclasses to override the UObject* overload of operator<< and that FObjectPtr overloads of operator<< redirect to the appropriate method in the derived type.

diff --git a/Engine/Source/Runtime/CoreUObject/Private/UObject/UObjectGlobals.cpp b/Engine/Source/Runtime/CoreUObject/Private/UObject/UObjectGlobals.cpp
--- a/Engine/Source/Runtime/CoreUObject/Private/UObject/UObjectGlobals.cpp
+++ b/Engine/Source/Runtime/CoreUObject/Private/UObject/UObjectGlobals.cpp
@@ -4339,11 +4339,10 @@ protected:
 	{
 		if (IsObjectHandleResolved(Object.GetHandle()) && Object)
 		{
-			FReferenceCollector& CurrentCollector = GetCollector();
-			FProperty* OldCollectorSerializedProperty = CurrentCollector.GetSerializedProperty();
-			CurrentCollector.SetSerializedProperty(GetSerializedProperty());
-			CurrentCollector.AddReferencedObject(Object.ToTObjectPtr(), GetSerializingObject(), GetSerializedProperty());
-			CurrentCollector.SetSerializedProperty(OldCollectorSerializedProperty);
+			// NOTE: This is deliberately not triggering access tracking as that is an undesirable overhead
+			//		 during garbage collect and GC reference collection is not meant to be trackable.
+			UObject*& RawObjectPointer = *(UObject**)&Object.GetHandleRef();
+			*this << RawObjectPointer;
 		}
 		return *this;
 	}
