Simpler Cast<> implementation, using std:: traits.

diff --git a/Engine/Source/Runtime/CoreUObject/Public/Templates/Casts.h b/Engine/Source/Runtime/CoreUObject/Public/Templates/Casts.h
--- a/Engine/Source/Runtime/CoreUObject/Public/Templates/Casts.h
+++ b/Engine/Source/Runtime/CoreUObject/Public/Templates/Casts.h
@@ -5,21 +5,17 @@
 #include "Containers/UnrealString.h"
 #include "CoreMinimal.h"
 #include "CoreTypes.h"
-#include "Templates/AndOrNot.h"
-#include "Templates/ChooseClass.h"
-#include "Templates/EnableIf.h"
 #include "Templates/LosesQualifiersFromTo.h"
-#include "Templates/PointerIsConvertibleFromTo.h"
-#include "Templates/RemoveReference.h"
 #include "Templates/UnrealTemplate.h"
-#include "Templates/UnrealTypeTraits.h"
-#include "Traits/IsVoidType.h"
 #include "UObject/Class.h"
 #include "UObject/Object.h"
 #include "UObject/ObjectHandle.h"
 #include "UObject/ObjectMacros.h"
 #include "UObject/ObjectPtr.h"
 #include "UObject/WeakObjectPtrTemplates.h"
+#include <type_traits>
+
+#define UE_USE_CAST_FLAGS (USTRUCT_FAST_ISCHILDOF_IMPL != USTRUCT_ISCHILDOF_STRUCTARRAY)
 
 class AActor;
 class APawn;
@@ -35,8 +31,6 @@ class USkinnedMeshComponent;
 class UStaticMeshComponent;
 /// @cond DOXYGEN_WARNINGS
 template<class TClass> class TSubclassOf;
-template <typename ...Types> struct TAnd;
-template <typename T> struct TIsPointer;
 template <typename Type> struct TCastFlags;
 /// @endcond
 
@@ -51,7 +45,7 @@ UE_NORETURN COREUOBJECT_API void CastLogError(const TCHAR* FromType, const TCHAR
  *
  * Otherwise, assume it's an IInterface.
  */
-template <typename T, bool bIsAUObject_IMPL = TPointerIsConvertibleFromTo<T, const volatile UObject>::Value>
+template <typename T, bool bIsAUObject_IMPL = std::is_convertible_v<T*, const volatile UObject*>>
 struct TIsIInterface
 {
 	enum { Value = false };
@@ -67,239 +61,68 @@ struct TIsIInterface<T, false>
 };
 
 template <typename T>
-struct TIsCastable
-{
-	// It's from-castable if it's an interface or a UObject-derived type
-	enum { Value = TIsIInterface<T>::Value || TPointerIsConvertibleFromTo<T, const volatile UObject>::Value };
-};
-
-template <typename T>
-struct TIsCastableToPointer : TOr<TIsVoidType<T>, TIsCastable<T>>
-{
-	// It's to-pointer-castable if it's from-castable or void
-};
-
-template <typename T>
-FORCEINLINE typename TEnableIf<TIsIInterface<T>::Value, FString>::Type GetTypeName()
-{
-	return T::UClassType::StaticClass()->GetName();
-}
-
-template <typename T>
-FORCEINLINE typename TEnableIf<!TIsIInterface<T>::Value, FString>::Type GetTypeName()
+FORCEINLINE FString GetTypeName()
 {
-	return T::StaticClass()->GetName();
+	if constexpr (TIsIInterface<T>::Value)
+	{
+		return T::UClassType::StaticClass()->GetName();
+	}
+	else
+	{
+		return T::StaticClass()->GetName();
+	}
 }
 
-enum class ECastType
-{
-	UObjectToUObject,
-	InterfaceToUObject,
-	UObjectToInterface,
-	InterfaceToInterface,
-	FromCastFlags
-};
-
 template <typename Type>
 struct TCastFlags
 {
 	static const EClassCastFlags Value = CASTCLASS_None;
 };
 
-template <typename From, typename To, bool bFromInterface = TIsIInterface<From>::Value, bool bToInterface = TIsIInterface<To>::Value, EClassCastFlags CastClass = TCastFlags<To>::Value>
-struct TGetCastType
-{
-#if USTRUCT_FAST_ISCHILDOF_IMPL == USTRUCT_ISCHILDOF_STRUCTARRAY
-	static const ECastType Value = ECastType::UObjectToUObject;
-#else
-	static const ECastType Value = ECastType::FromCastFlags;
-#endif
-};
-
-template <typename From, typename To                           > struct TGetCastType<From, To, false, false, CASTCLASS_None> { static const ECastType Value = ECastType::UObjectToUObject;     };
-template <typename From, typename To                           > struct TGetCastType<From, To, false, true , CASTCLASS_None> { static const ECastType Value = ECastType::UObjectToInterface;   };
-template <typename From, typename To, EClassCastFlags CastClass> struct TGetCastType<From, To, true,  false, CastClass     > { static const ECastType Value = ECastType::InterfaceToUObject;   };
-template <typename From, typename To, EClassCastFlags CastClass> struct TGetCastType<From, To, true,  true , CastClass     > { static const ECastType Value = ECastType::InterfaceToInterface; };
-
+// Dynamically cast an object type-safely.
 template <typename To, typename From>
-To* Cast(From* Src);
-
-template <typename From, typename To, ECastType CastType = TGetCastType<From, To>::Value>
-struct TCastImpl
-{
-	// This is the cast flags implementation
-	FORCEINLINE static To* DoCast( UObject* Src )
-	{
-		return Src && Src->GetClass()->HasAnyCastFlag(TCastFlags<To>::Value) ? (To*)Src : nullptr;
-	}
-
-	FORCEINLINE static To* DoCastCheckedWithoutTypeCheck( UObject* Src )
-	{
-		return (To*)Src;
-	}
-
-	FORCEINLINE static To* DoCast( const FObjectPtr& Src )
-	{
-		UObject* SrcObj = ResolveObjectHandleNoRead(Src.GetHandleRef());
-		if (SrcObj && Src->GetClass()->HasAnyCastFlag(TCastFlags<To>::Value))
-		{
-			ObjectHandle_Private::OnHandleRead(SrcObj);
-			return (To*)SrcObj;
-		}
-		return nullptr;
-	}
-
-	FORCEINLINE static To* DoCastCheckedWithoutTypeCheck( const FObjectPtr& Src )
-	{
-		return (To*)Src.Get();
-	}
-
-	UE_DEPRECATED(4.25, "Cast<>() and CastChecked<>() should not be used with FProperties. Use CastField<>() or CastFieldChecked<>() instead.")
-	FORCEINLINE static To* DoCast( FField* Src )
-	{
-		return Src && Src->IsA<To>() ? (To*)Src : nullptr;
-	}
-
-	UE_DEPRECATED(4.25, "Cast<>() and CastChecked<>() should not be used with FProperties. Use CastField<>() or CastFieldChecked<>() instead.")
-	FORCEINLINE static To* DoCastCheckedWithoutTypeCheck( FField* Src )
-	{
-		return (To*)Src;
-	}
-};
-
-template <typename From, typename To>
-struct TCastImpl<From, To, ECastType::UObjectToUObject>
-{
-	FORCEINLINE static To* DoCast( UObject* Src )
-	{
-		return Src && Src->IsA<To>() ? (To*)Src : nullptr;
-	}
-
-	FORCEINLINE static To* DoCastCheckedWithoutTypeCheck( UObject* Src )
-	{
-		return (To*)Src;
-	}
-
-	FORCEINLINE static To* DoCast( const FObjectPtr& Src )
-	{
-		UObject* SrcObj = ResolveObjectHandleNoRead(Src.GetHandleRef());
-		if (SrcObj && SrcObj->IsA<To>())
-		{
-			ObjectHandle_Private::OnHandleRead(SrcObj);
-			return (To*)SrcObj;
-		}
-		return nullptr;
-	}
-
-	FORCEINLINE static To* DoCastCheckedWithoutTypeCheck( const FObjectPtr& Src )
-	{
-		return (To*)Src.Get();
-	}
-
-	UE_DEPRECATED(4.25, "Cast<>() and CastChecked<>() should not be used with FProperties. Use CastField<>() or CastFieldChecked<>() instead.")
-	FORCEINLINE static To* DoCast( FField* Src )
-	{
-		return Src && Src->IsA<To>() ? (To*)Src : nullptr;
-	}
-
-	UE_DEPRECATED(4.25, "Cast<>() and CastChecked<>() should not be used with FProperties. Use CastField<>() or CastFieldChecked<>() instead.")
-	FORCEINLINE static To* DoCastCheckedWithoutTypeCheck( FField* Src )
-	{
-		return (To*)Src;
-	}
-};
-
-template <typename From, typename To>
-struct TCastImpl<From, To, ECastType::InterfaceToUObject>
+FORCEINLINE To* Cast(From* Src)
 {
-	FORCEINLINE static To* DoCast( From* Src )
+	if (Src)
 	{
-		if (Src)
+		if constexpr (TIsIInterface<From>::Value)
 		{
-			UObject* Obj = Src->_getUObject();
-			if (Obj && Obj->IsA<To>())
+			if (UObject* Obj = Src->_getUObject())
 			{
-				return (To*)Obj;
+				if constexpr (TIsIInterface<To>::Value)
+				{
+					return (To*)Obj->GetInterfaceAddress(To::UClassType::StaticClass());
+				}
+				else
+				{
+					if (Obj->IsA<To>())
+					{
+						return (To*)Obj;
+					}
+				}
 			}
 		}
-		return nullptr;
-	}
-
-	FORCEINLINE static To* DoCastCheckedWithoutTypeCheck( From* Src )
-	{
-		return Src ? (To*)Src->_getUObject() : nullptr;
-	}
-};
-
-template <typename From, typename To>
-struct TCastImpl<From, To, ECastType::UObjectToInterface>
-{
-	FORCEINLINE static To* DoCast( UObject* Src )
-	{
-		return Src ? (To*)Src->GetInterfaceAddress(To::UClassType::StaticClass()) : nullptr;
-	}
-
-	FORCEINLINE static To* DoCastCheckedWithoutTypeCheck( UObject* Src )
-	{
-		return Src ? (To*)Src->GetInterfaceAddress(To::UClassType::StaticClass()) : nullptr;
-	}
-
-	FORCEINLINE static To* DoCast( const FObjectPtr& Src )
-	{
-		UObject* SrcObj = ResolveObjectHandleNoRead(Src.GetHandleRef());
-		if (SrcObj)
-		{
-			ObjectHandle_Private::OnHandleRead(SrcObj);
-			return (To*)Src.Get()->GetInterfaceAddress(To::UClassType::StaticClass());
-		}
-		return nullptr;
-	}
-
-	FORCEINLINE static To* DoCastCheckedWithoutTypeCheck( const FObjectPtr& Src )
-	{
-		UObject* SrcObj = ResolveObjectHandleNoRead(Src.GetHandleRef());
-		if (SrcObj)
+		else if constexpr (UE_USE_CAST_FLAGS && TCastFlags<To>::Value != CASTCLASS_None)
 		{
-			ObjectHandle_Private::OnHandleRead(SrcObj);
-			return (To*)Src.Get()->GetInterfaceAddress(To::UClassType::StaticClass());
-		}
-		return nullptr;
-	}
-};
-
-template <typename From, typename To>
-struct TCastImpl<From, To, ECastType::InterfaceToInterface>
-{
-	FORCEINLINE static To* DoCast( From* Src )
-	{
-		if (Src)
-		{
-			if (UObject* Obj = Src->_getUObject())
+			if (((const UObject*)Src)->GetClass()->HasAnyCastFlag(TCastFlags<To>::Value))
 			{
-				return (To*)Obj->GetInterfaceAddress(To::UClassType::StaticClass());
+				return (To*)Src;
 			}
 		}
-		return nullptr;
-	}
-
-	FORCEINLINE static To* DoCastCheckedWithoutTypeCheck( From* Src )
-	{
-		if (Src)
+		else if constexpr (TIsIInterface<To>::Value)
 		{
-			if (UObject* Obj = Src->_getUObject())
+			return (To*)((UObject*)Src)->GetInterfaceAddress(To::UClassType::StaticClass());
+		}
+		else
+		{
+			if (((const UObject*)Src)->IsA<To>())
 			{
-				return (To*)Obj->GetInterfaceAddress(To::UClassType::StaticClass());
+				return (To*)Src;
 			}
 		}
-		return nullptr;
 	}
-};
 
-// Dynamically cast an object type-safely.
-template <typename To, typename From>
-FORCEINLINE To* Cast(From* Src)
-{
-	return TCastImpl<From, To>::DoCast(Src);
+	return nullptr;
 }
 
 template< class T >
@@ -311,22 +134,18 @@ FORCEINLINE T* ExactCast( UObject* Src )
 #if DO_CHECK
 
 	// Helper function to get the full name for UObjects and UInterfaces
-	template <
-		typename T,
-		typename TEnableIf<!TIsDerivedFrom<T, FField>::Value>::Type* = nullptr
-	>
+	template <typename T>
 	FString GetFullNameForCastLogError(T* InObjectOrInterface)
 	{
-		return Cast<UObject>(InObjectOrInterface)->GetFullName();
-	}
-	// And a special version for FFields
-	template <
-		typename T,
-		typename TEnableIf<TIsDerivedFrom<T, FField>::Value>::Type* = nullptr
-	>
-	FString GetFullNameForCastLogError(T* InField)
-	{
-		return GetFullNameSafe(InField);
+		// A special version for FFields
+		if constexpr (std::is_base_of_v<FField, T>)
+		{
+			return GetFullNameSafe(InObjectOrInterface);
+		}
+		else
+		{
+			return Cast<UObject>(InObjectOrInterface)->GetFullName();
+		}
 	}
 
 	template <typename To, typename From>
@@ -377,13 +196,32 @@ FORCEINLINE T* ExactCast( UObject* Src )
 		FORCEINLINE To* CastChecked(From* Src)
 	FUNCTION_NON_NULL_RETURN_END
 	{
-		return TCastImpl<From, To>::DoCastCheckedWithoutTypeCheck(Src);
+		if constexpr (TIsIInterface<From>::Value)
+		{
+			UObject* Obj = Src->_getUObject();
+			if constexpr (TIsIInterface<To>::Value)
+			{
+				return (To*)Obj->GetInterfaceAddress(To::UClassType::StaticClass());
+			}
+			else
+			{
+				return (To*)Obj;
+			}
+		}
+		else if constexpr (TIsIInterface<To>::Value)
+		{
+			return (To*)((UObject*)Src)->GetInterfaceAddress(To::UClassType::StaticClass());
+		}
+		else
+		{
+			return (To*)Src;
+		}
 	}
 
 	template <typename To, typename From>
 	FORCEINLINE To* CastChecked(From* Src, ECastCheckedType::Type CheckType)
 	{
-		return TCastImpl<From, To>::DoCastCheckedWithoutTypeCheck(Src);
+		return CastChecked<To>(Src);
 	}
 
 #endif
@@ -394,26 +232,61 @@ template< class T, class U > FORCEINLINE T* ExactCast  ( const TWeakObjectPtr<U>
 template< class T, class U > FORCEINLINE T* CastChecked( const TWeakObjectPtr<U>& Src, ECastCheckedType::Type CheckType = ECastCheckedType::NullChecked ) { return CastChecked<T>(Src.Get(), CheckType); }
 
 // object ptr versions
-template <class T, class U> FORCEINLINE T* Cast       ( const TObjectPtr<U>& Src                                                                   ) { return TCastImpl<U, T>::DoCast((const FObjectPtr&)Src); }
-template <class T, class U> FORCEINLINE T* ExactCast  ( const TObjectPtr<U>& Src                                                                   )
+template <typename To, typename From>
+FORCEINLINE To* Cast(const TObjectPtr<From>& InSrc)
+{
+	const FObjectPtr& Src = (const FObjectPtr&)InSrc;
+
+	if (UObject* SrcObj = ResolveObjectHandleNoRead(Src.GetHandleRef()))
+	{
+		if constexpr (UE_USE_CAST_FLAGS && TCastFlags<To>::Value != CASTCLASS_None)
+		{
+			if (Src->GetClass()->HasAnyCastFlag(TCastFlags<To>::Value))
+			{
+				ObjectHandle_Private::OnHandleRead(SrcObj);
+				return (To*)SrcObj;
+			}
+		}
+		else if constexpr (TIsIInterface<To>::Value)
+		{
+			ObjectHandle_Private::OnHandleRead(SrcObj);
+			return (To*)Src.Get()->GetInterfaceAddress(To::UClassType::StaticClass());
+		}
+		else
+		{
+			if (SrcObj->IsA<To>())
+			{
+				ObjectHandle_Private::OnHandleRead(SrcObj);
+				return (To*)SrcObj;
+			}
+		}
+	}
+
+	return nullptr;
+}
+
+template <typename To, typename From>
+FORCEINLINE To* ExactCast(const TObjectPtr<From>& Src)
 {
 	UObject* SrcObj = ResolveObjectHandleNoRead(((const FObjectPtr&)Src).GetHandleRef());
-	if (SrcObj && (SrcObj->GetClass() == T::StaticClass()))
+	if (SrcObj && (SrcObj->GetClass() == To::StaticClass()))
 	{
 		ObjectHandle_Private::OnHandleRead(SrcObj);
-		return (T*)SrcObj;
+		return (To*)SrcObj;
 	}
 	return nullptr;
 }
-template <class T, class U> FORCEINLINE T* CastChecked( const TObjectPtr<U>& Src, ECastCheckedType::Type CheckType = ECastCheckedType::NullChecked )
+
+template <typename To, typename From>
+FORCEINLINE To* CastChecked(const TObjectPtr<From>& Src, ECastCheckedType::Type CheckType = ECastCheckedType::NullChecked)
 {
 #if DO_CHECK
 	if (!Src.IsNull())
 	{
-		T* Result = Cast<T>(Src);
+		To* Result = Cast<To>(Src);
 		if (!Result)
 		{
-			CastLogError(*GetFullNameForCastLogError(Src.Get()), *GetTypeName<T>());
+			CastLogError(*GetFullNameForCastLogError(Src.Get()), *GetTypeName<To>());
 		}
 
 		return Result;
@@ -421,12 +294,25 @@ template <class T, class U> FORCEINLINE T* CastChecked( const TObjectPtr<U>& Src
 
 	if (CheckType == ECastCheckedType::NullChecked)
 	{
-		CastLogError(TEXT("nullptr"), *GetTypeName<T>());
+		CastLogError(TEXT("nullptr"), *GetTypeName<To>());
 	}
 
 	return nullptr;
 #else
-	return TCastImpl<U, T>::DoCastCheckedWithoutTypeCheck((const FObjectPtr&)Src);
+	if constexpr (UE_USE_CAST_FLAGS && TCastFlags<To>::Value != CASTCLASS_None)
+	{
+		return (To*)((const FObjectPtr&)Src).Get();
+	}
+	else if constexpr (TIsIInterface<To>::Value)
+	{
+		UObject* SrcObj = ResolveObjectHandleNoRead(((const FObjectPtr&)Src).GetHandleRef());
+		ObjectHandle_Private::OnHandleRead(SrcObj);
+		return (To*)((const FObjectPtr&)Src).Get()->GetInterfaceAddress(To::UClassType::StaticClass());
+	}
+	else
+	{
+		return (To*)((const FObjectPtr&)Src).Get();
+	}
 #endif
 }
 
@@ -544,46 +430,61 @@ private:
 
 namespace UECasts_Private
 {
-	template <typename To, typename From>
-	FORCEINLINE typename TEnableIf<TAnd<TIsPointer<To>, TAnd<TIsCastableToPointer<typename TRemovePointer<To>::Type>, TIsCastable<From>>>::Value, To>::Type DynamicCast(From* Arg)
+	template <typename T>
+	struct TIsCastable
 	{
-		typedef typename TRemovePointer<To  >::Type ToValueType;
-		typedef typename TRemovePointer<From>::Type FromValueType;
-
-		// Casting away const/volatile
-		static_assert(!TLosesQualifiersFromTo<FromValueType, ToValueType>::Value, "Conversion loses qualifiers");
-
-		// If we're casting to void, cast to UObject instead and let it implicitly cast to void
-		return Cast<typename TChooseClass<TIsVoidType<ToValueType>::Value, UObject, ToValueType>::Result>(Arg);
-	}
+		// It's from-castable if it's an interface or a UObject-derived type
+		enum { Value = TIsIInterface<T>::Value || std::is_convertible_v<T*, const volatile UObject*> };
+	};
 
 	template <typename To, typename From>
-	FORCEINLINE typename TEnableIf<!TAnd<TIsPointer<To>, TAnd<TIsCastableToPointer<typename TRemovePointer<To>::Type>, TIsCastable<From>>>::Value, To>::Type DynamicCast(From* Arg)
+	FORCEINLINE To DynamicCast(From* Arg)
 	{
-		return dynamic_cast<To>(Arg);
+		using ToValueType = std::remove_pointer_t<To>;
+
+		if constexpr (!std::is_pointer_v<To> || !TIsCastable<From>::Value || !TIsCastable<ToValueType>::Value)
+		{
+			return dynamic_cast<To>(Arg);
+		}
+		else
+		{
+			// Casting away const/volatile
+			static_assert(!TLosesQualifiersFromTo<From, ToValueType>::Value, "Conversion loses qualifiers");
+
+			if constexpr (std::is_void_v<ToValueType>)
+			{
+				// When casting to void, cast to UObject instead and let it implicitly cast to void
+				return Cast<UObject>(Arg);
+			}
+			else
+			{
+				return Cast<ToValueType>(Arg);
+			}
+		}
 	}
 
 	template <typename To, typename From>
-	FORCEINLINE typename TEnableIf<TAnd<TIsCastable<typename TRemoveReference<To>::Type>, TIsCastable<typename TRemoveReference<From>::Type>>::Value, To>::Type DynamicCast(From&& Arg)
+	FORCEINLINE To DynamicCast(From&& Arg)
 	{
-		typedef typename TRemoveReference<From>::Type FromValueType;
-		typedef typename TRemoveReference<To  >::Type ToValueType;
-
-		// Casting away const/volatile
-		static_assert(!TLosesQualifiersFromTo<FromValueType, ToValueType>::Value, "Conversion loses qualifiers");
+		using FromValueType = std::remove_reference_t<From>;
+		using ToValueType   = std::remove_reference_t<To>;
 
-		// T&& can only be cast to U&&
-		// http://en.cppreference.com/w/cpp/language/dynamic_cast
-		static_assert(TOr<TIsLValueReferenceType<From>, TIsRValueReferenceType<To>>::Value, "Cannot dynamic_cast from an rvalue to a non-rvalue reference");
+		if constexpr (!TIsCastable<FromValueType>::Value || !TIsCastable<ToValueType>::Value)
+		{
+			// This may fail when dynamic_casting rvalue references due to patchy compiler support
+			return dynamic_cast<To>(Arg);
+		}
+		else
+		{
+			// Casting away const/volatile
+			static_assert(!TLosesQualifiersFromTo<FromValueType, ToValueType>::Value, "Conversion loses qualifiers");
 
-		return Forward<To>(*CastChecked<typename TRemoveReference<To>::Type>(&Arg));
-	}
+			// T&& can only be cast to U&&
+			// http://en.cppreference.com/w/cpp/language/dynamic_cast
+			static_assert(std::is_lvalue_reference_v<From> || std::is_rvalue_reference_v<To>, "Cannot dynamic_cast from an rvalue to a non-rvalue reference");
 
-	template <typename To, typename From>
-	FORCEINLINE typename TEnableIf<!TAnd<TIsCastable<typename TRemoveReference<To>::Type>, TIsCastable<typename TRemoveReference<From>::Type>>::Value, To>::Type DynamicCast(From&& Arg)
-	{
-		// This may fail when dynamic_casting rvalue references due to patchy compiler support
-		return dynamic_cast<To>(Arg);
+			return Forward<To>(*CastChecked<ToValueType>(&Arg));
+		}
 	}
 }
 
