Added FCoreUObjectDelegates::OnPostInitSparseClassData, to allow extra initialization of sparse class data outside of the struct itself

diff --git a/Engine/Source/Runtime/CoreUObject/Private/UObject/Class.cpp b/Engine/Source/Runtime/CoreUObject/Private/UObject/Class.cpp
--- a/Engine/Source/Runtime/CoreUObject/Private/UObject/Class.cpp
+++ b/Engine/Source/Runtime/CoreUObject/Private/UObject/Class.cpp
@@ -5567,6 +5567,8 @@ void* UClass::CreateSparseClassData()
 		{
 			SparseClassDataArchetypeStruct->CopyScriptStruct(SparseClassData, SparseArchetypeData);
 		}
+
+		FCoreUObjectDelegates::OnPostInitSparseClassData.Broadcast(this, SparseClassDataStruct, SparseClassData);
 	}
 
 	return SparseClassData;
diff --git a/Engine/Source/Runtime/CoreUObject/Private/UObject/UObjectGlobals.cpp b/Engine/Source/Runtime/CoreUObject/Private/UObject/UObjectGlobals.cpp
--- a/Engine/Source/Runtime/CoreUObject/Private/UObject/UObjectGlobals.cpp
+++ b/Engine/Source/Runtime/CoreUObject/Private/UObject/UObjectGlobals.cpp
@@ -141,6 +141,7 @@ FCoreUObjectDelegates::FReloadCompleteDelegate FCoreUObjectDelegates::ReloadComp
 FCoreUObjectDelegates::FReloadAddedClassesDelegate FCoreUObjectDelegates::ReloadAddedClassesDelegate;
 FCoreUObjectDelegates::FCompiledInUObjectsRegisteredDelegate FCoreUObjectDelegates::CompiledInUObjectsRegisteredDelegate;
 FCoreUObjectDelegates::FIsPackageOKToSaveDelegate FCoreUObjectDelegates::IsPackageOKToSaveDelegate;
+FCoreUObjectDelegates::FOnPostInitSparseClassData FCoreUObjectDelegates::OnPostInitSparseClassData;
 FCoreUObjectDelegates::FOnPackageReloaded FCoreUObjectDelegates::OnPackageReloaded;
 FCoreUObjectDelegates::FNetworkFileRequestPackageReload FCoreUObjectDelegates::NetworkFileRequestPackageReload;
 #if WITH_EDITOR
diff --git a/Engine/Source/Runtime/CoreUObject/Public/UObject/UObjectGlobals.h b/Engine/Source/Runtime/CoreUObject/Public/UObject/UObjectGlobals.h
--- a/Engine/Source/Runtime/CoreUObject/Public/UObject/UObjectGlobals.h
+++ b/Engine/Source/Runtime/CoreUObject/Public/UObject/UObjectGlobals.h
@@ -2728,6 +2728,10 @@ struct COREUOBJECT_API FCoreUObjectDelegates
 	static FAutoPackageBackupDelegate AutoPackageBackupDelegate;
 #endif // WITH_EDITOR
 
+	/** Called when new sparse class data has been created (and the base data initialized) for the given class */
+	DECLARE_MULTICAST_DELEGATE_ThreeParams(FOnPostInitSparseClassData, UClass*, UScriptStruct*, void*);
+	static FOnPostInitSparseClassData OnPostInitSparseClassData;
+
 	/** Called by ReloadPackage during package reloading. It will be called several times for different phases of fix-up to allow custom code to handle updating objects as needed */
 	DECLARE_MULTICAST_DELEGATE_TwoParams(FOnPackageReloaded, EPackageReloadPhase, FPackageReloadedEvent*);
 	static FOnPackageReloaded OnPackageReloaded;
