Bulk Data - unit tests

#rnx

diff --git a/Engine/Source/Runtime/CoreUObject/Tests/Serialization/BulkDataTests.cpp b/Engine/Source/Runtime/CoreUObject/Tests/Serialization/BulkDataTests.cpp
--- a/Engine/Source/Runtime/CoreUObject/Tests/Serialization/BulkDataTests.cpp
+++ b/Engine/Source/Runtime/CoreUObject/Tests/Serialization/BulkDataTests.cpp
@@ -3,73 +3,354 @@
 #if WITH_LOW_LEVEL_TESTS
 
 #include "TestHarness.h"
+#include "TestMacros/Assertions.h"
+#include "Memory/SharedBuffer.h"
 #include "Serialization/BulkData.h"
+#include "Serialization/LargeMemoryWriter.h"
+#include "Serialization/MemoryReader.h"
+#include <catch2/generators/catch_generators.hpp>
 
-namespace BulkDataTest
+namespace UE::Serialization::BulkDataTest
 {
-	// Test code paths for BulkData objects that do not reference a file on disk.
-	TEST_CASE("CoreUObject::Serialization::FByteBulkData::Transient", "[CoreUObject][Serialization]")
+
+static FUniqueBuffer CreatePayload(uint64 Size, uint64 Seed = 0)
+{
+	FUniqueBuffer Buffer = FUniqueBuffer::Alloc(Size);
+	TArrayView<uint64> Values(static_cast<uint64*>(Buffer.GetData()), Size / sizeof(uint64));
+
+	uint64 Index = Seed;
+	for (uint64& Value : Values)
+	{
+		Value = ++Index;
+	}
+
+	return Buffer;
+};
+
+static void CopyPayload(FBulkData& BulkData, FMemoryView Src)
+{
+	BulkData.Lock(LOCK_READ_WRITE);
+	FMutableMemoryView Dst(BulkData.Realloc(Src.GetSize(), 1), Src.GetSize());
+	Dst.CopyFrom(Src);
+	BulkData.Unlock();
+}
+
+template<typename T>
+bool TestBulkDataFlags(FBulkData& BulkData, uint32 Flags, T&& IsSet)
+{
+	bool bOk = false;
+
+	if (IsSet() == false)
 	{
-		FByteBulkData BulkData;
+		BulkData.SetBulkDataFlags(Flags);
+		bOk = IsSet();
+	}
+	BulkData.ResetBulkDataFlags(0);
+
+	return bOk;
+}
 
-		// We should be able to lock for read access but there should be no valid data
-		const void* ReadOnlyDataPtr = BulkData.LockReadOnly();
-		TEST_NULL(TEXT("Locking an empty BulkData object for reading should return nullptr!"), ReadOnlyDataPtr);
+TEST_CASE("CoreUObject::Serialization::FBulkData::Basic", "[CoreUObject][Serialization]")
+{
+	SECTION("Default construction")
+	{
+		FBulkData BulkData;
+		CHECK(BulkData.GetBulkDataFlags() == 0);
+		CHECK(BulkData.GetBulkDataSize() == 0);
+		CHECK(BulkData.GetBulkDataOffsetInFile() == -1);
+		CHECK(BulkData.CanLoadFromDisk() == false);
+		CHECK(BulkData.DoesExist() == false);
+		CHECK(BulkData.IsBulkDataLoaded() == false);
+	}
+	
+	SECTION("Lock empty")
+	{
+		FBulkData BulkData;
+		CHECK(nullptr == BulkData.Lock(LOCK_READ_ONLY));
+		BulkData.Unlock();
+		CHECK(nullptr == BulkData.Lock(LOCK_READ_WRITE));
 		BulkData.Unlock();
+		CHECK(nullptr == BulkData.LockReadOnly());
+		BulkData.Unlock();
+	}
+	
+	SECTION("Create payload")
+	{
+		const int64 ExpectedPayloadSize = 1024;
 
-		void* DataPtr = BulkData.Lock(LOCK_READ_WRITE);
-		TEST_NULL(TEXT("Locking an empty BulkData object for writing should return nullptr!"), DataPtr);
+		FBulkData BulkData;
+		void* Payload = BulkData.Lock(LOCK_READ_WRITE);
+		CHECK(Payload == nullptr);
+		
+		Payload = BulkData.Realloc(ExpectedPayloadSize, 1);
+		CHECK(Payload != nullptr);
 		BulkData.Unlock();
+		
+		CHECK(BulkData.GetBulkDataFlags() == 0);
+		CHECK(BulkData.GetBulkDataSize() == ExpectedPayloadSize);
+		CHECK(BulkData.GetBulkDataOffsetInFile() == -1);
+	}
+	
+	SECTION("Get copy")
+	{
+		// Empty
+		{
+			void* Dst = nullptr;
+			bool bDiscardInternalCopy = true;
+			
+			FBulkData BulkData;
+			BulkData.GetCopy(&Dst, bDiscardInternalCopy); 
 
-		void* CopyEmptyPtr = nullptr;
-		BulkData.GetCopy(&CopyEmptyPtr, true);
-		TEST_NULL(TEXT("Getting a copy of an empty BulkData object for writing should return nullptr!"), CopyEmptyPtr);
+			CHECK(Dst == nullptr);
+		}
+		
+		// Non empty
+		{
+			const int64 ExpectedPayloadSize = 64;
 
-		BulkData.Lock(LOCK_READ_WRITE);
-		DataPtr = BulkData.Realloc(32 * 32 * 4);
-		TEST_NOT_NULL(TEXT("Reallocating an empty BulkData object should return a valid pointer!"), DataPtr);
-		BulkData.Unlock();
+			FUniqueBuffer Payload = CreatePayload(ExpectedPayloadSize);
+			FBulkData BulkData;
+			CopyPayload(BulkData, Payload.GetView());
+			CHECK(BulkData.GetBulkDataSize() == ExpectedPayloadSize);
 
-		TEST_TRUE(TEXT("BulkData should be loaded now that it has been reallocated"), BulkData.IsBulkDataLoaded());
+			void* Dst = nullptr;
+			bool bDiscardInternalCopy = false;
+			BulkData.GetCopy(&Dst, bDiscardInternalCopy);
+			CHECK(BulkData.GetBulkDataSize() == ExpectedPayloadSize);
 
-		void* CopyWithDiscardPtr = nullptr;
-		BulkData.GetCopy(&CopyWithDiscardPtr, true); // The second parameter should be ignored because the bulkdata cannot be reloaded from disk!
-		TEST_NOT_NULL(TEXT("GetCopy should return a valid pointer!"), CopyWithDiscardPtr);
-		TEST_NOT_EQUAL(TEXT("GetCopy should return a copy of the data so the pointers should be different!"), DataPtr, CopyWithDiscardPtr);
-		TEST_TRUE(TEXT("BulkData should still loaded after taking a copy"), BulkData.IsBulkDataLoaded());
+			FMemoryView DstView(Dst, ExpectedPayloadSize);
+			CHECK(DstView.EqualBytes(Payload.GetView()));
+		}
 
-		// Now try GetCopy again but this time without the discard request
-		void* CopyNoDiscardPtr = nullptr;
-		BulkData.GetCopy(&CopyNoDiscardPtr, false);
-		TEST_NOT_NULL(TEXT("GetCopy should return a valid pointer!"), CopyNoDiscardPtr);
-		TEST_NOT_EQUAL(TEXT("GetCopy should return a copy of the data so the pointers should be different!"), DataPtr, CopyNoDiscardPtr);
-		TEST_NOT_EQUAL(TEXT("GetCopy should return a copy of the data so the pointers should be different!"), CopyWithDiscardPtr, CopyNoDiscardPtr);
-		TEST_TRUE(TEXT("BulkData should still loaded after taking a copy"), BulkData.IsBulkDataLoaded());
+		// Get copy and discard with non-allocatated destination buffer 
+		{
+			const int64 ExpectedPayloadSize = 2048;
 
-		// Clean up allocations from GetCopy
-		FMemory::Free(CopyWithDiscardPtr);
-		FMemory::Free(CopyNoDiscardPtr);
+			FUniqueBuffer Payload = CreatePayload(ExpectedPayloadSize);
+			FBulkData BulkData;
+			CopyPayload(BulkData, Payload.GetView());
+			CHECK(BulkData.GetBulkDataSize() == ExpectedPayloadSize);
 
-		// Now do one last lock test after GetCopy
-		DataPtr = BulkData.Lock(LOCK_READ_WRITE);
-		BulkData.Unlock();
+			void* Dst = nullptr;
+			bool bDiscardInternalCopy = true;
+			BulkData.GetCopy(&Dst, bDiscardInternalCopy);
+
+			CHECK(BulkData.IsBulkDataLoaded()); // Still loaded (not discardable)
+			CHECK(BulkData.GetBulkDataSize() == ExpectedPayloadSize);
 
-		TEST_TRUE(TEXT("BulkData should still loaded after locking for write"), BulkData.IsBulkDataLoaded());
-		TEST_NOT_NULL(TEXT("Locking for write should return a valid pointer!"), DataPtr);
+			FMemoryView DstView(Dst, ExpectedPayloadSize);
+			CHECK(DstView.EqualBytes(Payload.GetView()));
+			
+			bDiscardInternalCopy = true;
+			BulkData.SetBulkDataFlags(BULKDATA_AlwaysAllowDiscard);
+			BulkData.GetCopy(&Dst, bDiscardInternalCopy);
+			
+			CHECK(BulkData.IsBulkDataLoaded() == false);
+			CHECK(BulkData.GetBulkDataSize() == ExpectedPayloadSize);
+
+			DstView = FMemoryView(Dst, ExpectedPayloadSize);
+			CHECK(DstView.EqualBytes(Payload.GetView()));
+			FMemory::Free(Dst);
+		}
+		
+		// Get copy and discard with allocatated destination buffer 
+		{
+			const int64 ExpectedPayloadSize = 512;
+			
+			FBulkData BulkData;
+			{
+				FUniqueBuffer Payload = CreatePayload(ExpectedPayloadSize);
+				CopyPayload(BulkData, Payload.GetView());
+			}
+
+			FUniqueBuffer Payload = FUniqueBuffer::Alloc(ExpectedPayloadSize);
+			
+			bool bDiscardInternalCopy = true;
+			void* Dst = Payload.GetData();
+			BulkData.GetCopy(&Dst, bDiscardInternalCopy);
+			
+			CHECK(BulkData.IsBulkDataLoaded());
+			CHECK(BulkData.GetBulkDataSize() == ExpectedPayloadSize);
+
+			FMemoryView BulkDataView(BulkData.LockReadOnly(), BulkData.GetBulkDataSize());
+			CHECK(BulkDataView.EqualBytes(Payload.GetView()));
+			BulkData.Unlock();
+		
+			bDiscardInternalCopy = false;
+			BulkData.SetBulkDataFlags(BULKDATA_AlwaysAllowDiscard);
+			BulkData.GetCopy(&Dst, bDiscardInternalCopy);
+			
+			CHECK(BulkData.IsBulkDataLoaded());
+			
+			bDiscardInternalCopy = true;
+			BulkData.GetCopy(&Dst, bDiscardInternalCopy);
+			
+			CHECK(BulkData.IsBulkDataLoaded() == false);
+		}
+	}
+	
+	SECTION("Remove payload")
+	{
+		const int64 ExpectedPayloadSize = 256;
+
+		FBulkData BulkData;
+		FUniqueBuffer Payload = CreatePayload(ExpectedPayloadSize);
+		CopyPayload(BulkData, Payload);
+		CHECK(BulkData.IsBulkDataLoaded());
 
-		// Now remove the bulkdata and make sure that we cannot access the old data anymore
 		BulkData.RemoveBulkData();
-		TEST_FALSE(TEXT("RemoveBulkData should've discarded the BulkData"), BulkData.IsBulkDataLoaded());
+		CHECK(BulkData.IsBulkDataLoaded() == false);
+	}
+	
+	SECTION("Flags")
+	{
+		FBulkData BulkData;
+		
+		CHECK(TestBulkDataFlags(BulkData, BULKDATA_PayloadAtEndOfFile, [&BulkData]() { return !BulkData.IsInlined(); }));
+		CHECK(TestBulkDataFlags(BulkData, BULKDATA_PayloadInSeperateFile, [&BulkData]() { return BulkData.IsInSeparateFile(); }));
+		CHECK(TestBulkDataFlags(BulkData, BULKDATA_OptionalPayload, [&BulkData]() { return BulkData.IsOptional(); }));
+		CHECK(TestBulkDataFlags(BulkData, BULKDATA_DuplicateNonOptionalPayload, [&BulkData]() { return BulkData.IsDuplicateNonOptional(); }));
+		CHECK(TestBulkDataFlags(BulkData, BULKDATA_DataIsMemoryMapped, [&BulkData]() { return BulkData.IsDataMemoryMapped(); }));
+		CHECK(TestBulkDataFlags(BulkData, BULKDATA_SingleUse, [&BulkData]() { return BulkData.IsSingleUse(); }));
+		CHECK(TestBulkDataFlags(BulkData, BULKDATA_UsesIoDispatcher, [&BulkData]() { return BulkData.IsUsingIODispatcher(); }));
+	}
+}
 
-		DataPtr = BulkData.Lock(LOCK_READ_WRITE);
-		BulkData.Unlock();
+TEST_CASE("CoreUObject::Serialization::FBulkData::Serialize", "[CoreUObject][Serialization]")
+{
+	SECTION("Serialize to memory archive")
+	{
+		const int64 ExpectedPayloadSize = 128;
+		
+		FUniqueBuffer Payload = CreatePayload(ExpectedPayloadSize);
+		FLargeMemoryWriter Ar;
+		Ar.SetIsPersistent(true);
+
+		{
+			FBulkData BulkData;
+			CopyPayload(BulkData, Payload);
+			BulkData.Serialize(Ar, nullptr, false, 1, EFileRegionType::None); 
+			CHECK(Ar.TotalSize() > BulkData.GetBulkDataSize()); // Bulk meta data + payload
+		}
+
+		{
+			FMemoryReaderView ReaderAr(Ar.GetView());
+			ReaderAr.SetIsPersistent(true);
+
+			FBulkData BulkData;
+			BulkData.Serialize(ReaderAr, nullptr, false, 1, EFileRegionType::None); 
+			FMemoryView PayloadView(BulkData.LockReadOnly(), BulkData.GetBulkDataSize());
 
-		TEST_NULL(TEXT("Locking for write after calling ::RemoveBulkData should return a nullptr!"), DataPtr);
+			CHECK(BulkData.GetBulkDataSize() == ExpectedPayloadSize);
+			CHECK(BulkData.GetBulkDataOffsetInFile() == ReaderAr.Tell() - ExpectedPayloadSize);
+			CHECK(BulkData.GetBulkDataFlags() == 0);
+			CHECK(BulkData.IsInlined());
+			CHECK(PayloadView.EqualBytes(Payload.GetView()));
 
-		CopyEmptyPtr = nullptr;
-		BulkData.GetCopy(&CopyNoDiscardPtr, true);
-		TEST_NULL(TEXT("Getting a copy of BulkData object after calling ::RemoveBulkData should return nullptr!"), DataPtr);
+			BulkData.Unlock();
+		}
+	}
+	
+	SECTION("Serialize to memory archive does not serialize invalid flags")
+	{
+		const int64 ExpectedPayloadSize = 1024;
+		const uint32 InvalidFlags = uint32(BULKDATA_PayloadAtEndOfFile | BULKDATA_PayloadInSeperateFile | BULKDATA_WorkspaceDomainPayload);
+
+		FUniqueBuffer Payload = CreatePayload(ExpectedPayloadSize);
+		FLargeMemoryWriter Ar;
+		Ar.SetIsPersistent(true);
+		
+		{
+			FBulkData BulkData;
+			BulkData.SetBulkDataFlags(InvalidFlags);
+			BulkData.Serialize(Ar, nullptr, false, 1, EFileRegionType::None); 
+			CHECK(BulkData.GetBulkDataFlags() == InvalidFlags);
+			CHECK(BulkData.GetBulkDataOffsetInFile() == -1);
+		}
+
+		{
+			FMemoryReaderView ReaderAr(Ar.GetView());
+			ReaderAr.SetIsPersistent(true);
+
+			FBulkData BulkData;
+			BulkData.Serialize(Ar, nullptr, false, 1, EFileRegionType::None); 
+			CHECK(BulkData.GetBulkDataFlags() == 0);
+			CHECK(BulkData.IsInlined());
+		}
+	}
+	
+	SECTION("Serialize many to memory archive")
+	{
+		const int64 ExpectedPayloadSizes[] = {8, 16, 32, 64, 128, 256, 512, 1024, 2048};
+		TArrayView<const int64> PayloadSizes = MakeArrayView<const int64>(ExpectedPayloadSizes, 9); 
+		TArray<FUniqueBuffer> ExpectedPayloads;
+
+		for (int32 Idx = 0; Idx < PayloadSizes.Num(); ++Idx)
+		{
+			ExpectedPayloads.Emplace(CreatePayload(PayloadSizes[Idx], Idx)); 
+		}
+		
+		FLargeMemoryWriter Ar;
+		Ar.SetIsPersistent(true);
+
+		for (FUniqueBuffer& Payload : ExpectedPayloads)
+		{
+			FBulkData BulkData;
+			CopyPayload(BulkData, Payload);
+			BulkData.Serialize(Ar, nullptr, false, 1, EFileRegionType::None); 
+			CHECK(BulkData.GetBulkDataOffsetInFile() == -1);
+		}
+
+		{
+			FMemoryReaderView ReaderAr(Ar.GetView());
+			ReaderAr.SetIsPersistent(true);
+
+			for (const FUniqueBuffer& ExpectedPayload : ExpectedPayloads)
+			{
+				FBulkData BulkData;
+				BulkData.Serialize(ReaderAr, nullptr, false, 1, EFileRegionType::None); 
+				FMemoryView PayloadView(BulkData.LockReadOnly(), BulkData.GetBulkDataSize());
+
+				CHECK(BulkData.GetBulkDataOffsetInFile() == (ReaderAr.Tell() - BulkData.GetBulkDataSize()));
+				CHECK(BulkData.GetBulkDataFlags() == 0);
+				CHECK(BulkData.IsInlined());
+				CHECK(PayloadView.EqualBytes(ExpectedPayload.GetView()));
+
+				BulkData.Unlock();
+			}
+		}
+	}
+	
+	SECTION("Serialize compressed")
+	{
+		const int64 UncompressedPayloadSize = 4 << 20;
+		FUniqueBuffer Payload = CreatePayload(UncompressedPayloadSize);
+		
+		FLargeMemoryWriter Ar;
+		Ar.SetIsPersistent(true);
+
+		{
+			FBulkData BulkData;
+			CopyPayload(BulkData, Payload);
+			BulkData.SetBulkDataFlags(BULKDATA_SerializeCompressedZLIB);
+			BulkData.Serialize(Ar, nullptr, false, 1, EFileRegionType::None);
+		}
+
+		{
+			FMemoryReaderView ReaderAr(Ar.GetView());
+			ReaderAr.SetIsPersistent(true);
+			
+			FBulkData BulkData;
+			BulkData.Serialize(ReaderAr, nullptr, false, 1, EFileRegionType::None);
+
+			CHECK(BulkData.IsBulkDataLoaded());
+			CHECK(BulkData.IsStoredCompressedOnDisk());
+			CHECK(BulkData.GetBulkDataFlags() == BULKDATA_SerializeCompressedZLIB);
+			CHECK(BulkData.GetBulkDataSize() == UncompressedPayloadSize );
+		}
 	}
 }
 
+} // namespace UE::Serialization::BulkDataTest
+
 #endif // WITH_LOW_LEVEL_TESTS
