ZenLoader in editor: Remove garbage objects from linker export tables when notifying unreachable objects so that we don't end up with garbage objects in the import store

diff --git a/Engine/Source/Runtime/CoreUObject/Private/Serialization/AsyncLoading2.cpp b/Engine/Source/Runtime/CoreUObject/Private/Serialization/AsyncLoading2.cpp
--- a/Engine/Source/Runtime/CoreUObject/Private/Serialization/AsyncLoading2.cpp
+++ b/Engine/Source/Runtime/CoreUObject/Private/Serialization/AsyncLoading2.cpp
@@ -1052,19 +1052,17 @@ public:
 	{
 		TRACE_CPUPROFILER_EVENT_SCOPE(RemovePublicExports);
 
-		TArray<FPublicExportKey> PublicExportKeys;
-		PublicExportKeys.Reserve(ObjectsToRemove.Num());
+		FPackageId LastPackageId;
+		FLoadedPackageRef* PackageRef = nullptr;
 
 		for (const FUnreachableObject& Item : ObjectsToRemove)
 		{
 			int32 ObjectIndex = Item.ObjectIndex;
-			check(ObjectIndex >= 0)
+			check(ObjectIndex >= 0);
 
 			FPublicExportKey PublicExportKey;
 			if (ObjectIndexToPublicExport.RemoveAndCopyValue(ObjectIndex, PublicExportKey))
 			{
-				PublicExportKeys.Emplace(PublicExportKey);
-
 #if ALT2_VERIFY_UNREACHABLE_OBJECTS
 				if (GGRemoveUnreachableObjectsFromGCNotifyOnGT)
 				{
@@ -1091,21 +1089,17 @@ public:
 						*ExistingObject->GetFullName());
 				}
 #endif
-			}
-		}
 
-		FPackageId LastPackageId;
-		FLoadedPackageRef* PackageRef = nullptr;
-		for (const FPublicExportKey& PublicExportKey : PublicExportKeys)
-		{
-			FPackageId PackageId = PublicExportKey.GetPackageId();
-			if (PackageId != LastPackageId)
-			{
-				LastPackageId = PackageId;
-				PackageRef = LoadedPackageStore.FindPackageRef(PackageId);
+				FPackageId PackageId = PublicExportKey.GetPackageId();
+				if (PackageId != LastPackageId)
+				{
+					LastPackageId = PackageId;
+					PackageRef = LoadedPackageStore.FindPackageRef(PackageId);
+				}
+
+				check(PackageRef);
+				PackageRef->RemovePublicExport(PublicExportKey.GetExportHash());
 			}
-			check(PackageRef);
-			PackageRef->RemovePublicExport(PublicExportKey.GetExportHash());
 		}
 	}
 
@@ -4800,6 +4794,7 @@ bool FAsyncPackage2::ProcessLinkerLoadPackageExports(FAsyncLoadingThreadState2&
 		FExportObject& ExportObject = Data.Exports[ExportIndex];
 		if (UObject* Object = LinkerLoadState->Linker->CreateExport(ExportIndex))
 		{
+			checkf(!Object->HasAnyInternalFlags(EInternalObjectFlags::Unreachable), TEXT("Trying to store an unreachable object '%s' in the import store"), *Object->GetFullName());
 			ExportObject.Object = Object;
 			ExportObject.bWasFoundInMemory = true; // Make sure that the async flags are cleared in ClearConstructedObjects
 			EInternalObjectFlags FlagsToSet = EInternalObjectFlags::Async;
@@ -6995,15 +6990,15 @@ static void VerifyLoadFlagsWhenFinishedLoading()
 }
 #endif
 
-FORCENOINLINE static void FilterUnreachableObjects(
+FORCENOINLINE static void CollectUnreachableObjects(
 	TArrayView<FUObjectItem*> UnreachableObjectItems,
 	FUnreachableObjects& OutUnreachableObjects)
 {
-	TRACE_CPUPROFILER_EVENT_SCOPE(FilterUnreachableObjects);
+	TRACE_CPUPROFILER_EVENT_SCOPE(CollectUnreachableObjects);
 
 	OutUnreachableObjects.SetNum(UnreachableObjectItems.Num());
 
-	ParallelFor(TEXT("FilterUnreachableObjects"), UnreachableObjectItems.Num(), 2048, [&UnreachableObjectItems, &OutUnreachableObjects](int32 Index)
+	ParallelFor(TEXT("CollectUnreachableObjects"), UnreachableObjectItems.Num(), 2048, [&UnreachableObjectItems, &OutUnreachableObjects](int32 Index)
 	{
 		UObject* Object = static_cast<UObject*>(UnreachableObjectItems[Index]->Object);
 
@@ -7024,6 +7019,33 @@ FORCENOINLINE static void FilterUnreachableObjects(
 				Item.PackageId = Package->GetPackageId();
 			}
 		}
+
+#if ALT2_ENABLE_LINKERLOAD_SUPPORT
+		// Clear garbage objects from linker export tables
+		// Normally done from UObject::BeginDestroy but we need to do it already here
+		if (FLinkerLoad* ObjectLinker = Object->GetLinker())
+		{
+#if WITH_EDITORONLY_DATA
+			// Make sure the linker entry stays as 'bExportLoadFailed' if the entry was marked as such, 
+			// Not doing this will cause crashes when resolving circular imports!
+			const int32 CachedLinkerIndex = Object->GetLinkerIndex();
+			bool bLinkerEntryWasInvalid = false;
+			if (ObjectLinker->ExportMap.IsValidIndex(CachedLinkerIndex))
+			{
+				FObjectExport& ObjExport = ObjectLinker->ExportMap[CachedLinkerIndex];
+				bLinkerEntryWasInvalid = ObjExport.bExportLoadFailed;
+			}
+#endif // WITH_EDITORONLY_DATA
+			Object->SetLinker(NULL, INDEX_NONE);
+#if WITH_EDITORONLY_DATA
+			if (bLinkerEntryWasInvalid)
+			{
+				FObjectExport& ObjExport = ObjectLinker->ExportMap[CachedLinkerIndex];
+				ObjExport.bExportLoadFailed = true;
+			}
+#endif // WITH_EDITORONLY_DATA
+		}
+#endif // ALT2_ENABLE_LINKERLOAD_SUPPORT
 	});
 }
 
@@ -7067,7 +7089,7 @@ void FAsyncLoadingThread2::OnLeakedPackageRename(UPackage* Package)
 	});
 
 	FUnreachableObjects LeakedUnreachableObjects;
-	FilterUnreachableObjects(LeakedObjectItems, LeakedUnreachableObjects);
+	CollectUnreachableObjects(LeakedObjectItems, LeakedUnreachableObjects);
 	RemoveUnreachableObjects(LeakedUnreachableObjects);
 
 	// Clear the CanBeImportedFlag so that this package is only removed once,
@@ -7126,7 +7148,7 @@ void FAsyncLoadingThread2::NotifyUnreachableObjects(const TArrayView<FUObjectIte
 	// if not handle them here before adding new ones
 	RemoveUnreachableObjects(UnreachableObjects);
 
-	FilterUnreachableObjects(UnreachableObjectItems, UnreachableObjects);
+	CollectUnreachableObjects(UnreachableObjectItems, UnreachableObjects);
 
 #if ALT2_VERIFY_ASYNC_FLAGS
 	if (!IsAsyncLoading())
