Adding Optional Specialization for TNonNullSubclassOf

diff --git a/Engine/Source/Runtime/CoreUObject/Public/Templates/NonNullSubclassOf.h b/Engine/Source/Runtime/CoreUObject/Public/Templates/NonNullSubclassOf.h
--- a/Engine/Source/Runtime/CoreUObject/Public/Templates/NonNullSubclassOf.h
+++ b/Engine/Source/Runtime/CoreUObject/Public/Templates/NonNullSubclassOf.h
@@ -57,3 +57,89 @@ public:
 		return *this;
 	}
 };
+
+
+/**
+ * Specialization of TOptional for TNonNullSubclassOf value types
+ */
+template<typename TClass>
+struct TOptional<TNonNullSubclassOf<TClass>>
+{
+public:
+
+	using TClassType = typename TSubclassOf<TClass>::TClassType;
+
+	/** Construct an OptionaType with a valid value. */
+	TOptional(const TNonNullSubclassOf<TClass>& InPointer)
+		: Pointer(InPointer)
+	{
+	}
+
+	/** Construct an TClass with no value; i.e. unset */
+	TOptional()
+		: Pointer(nullptr)
+	{
+	}
+
+	/** Construct an TClass with an invalid value. */
+	TOptional(FNullOpt)
+		: TOptional()
+	{
+	}
+
+	TOptional& operator=(const TOptional& Other)
+	{
+		Pointer = Other.Pointer;
+		return *this;
+	}
+
+	TOptional& operator=(TClass* InPointer)
+	{
+		Pointer = InPointer;
+		return *this;
+	}
+
+	void Reset()
+	{
+		Pointer = nullptr;
+	}
+
+	TClass* Emplace(TClass* InPointer)
+	{
+		Pointer = InPointer;
+		return InPointer;
+	}
+
+	friend bool operator==(const TOptional& lhs, const TOptional& rhs)
+	{
+		return lhs.Pointer == rhs.Pointer;
+	}
+
+	friend bool operator!=(const TOptional& lhs, const TOptional& rhs)
+	{
+		return !(lhs == rhs);
+	}
+
+	friend FArchive& operator<<(FArchive& Ar, TOptional& Optional)
+	{
+		Ar << Optional.Pointer;
+		return Ar;
+	}
+
+	/** @return true when the value is meaningful; false if calling GetValue() is undefined. */
+	bool IsSet() const { return Pointer != nullptr; }
+	FORCEINLINE explicit operator bool() const { return Pointer != nullptr; }
+
+	/** @return The optional value; undefined when IsSet() returns false. */
+	TClassType* GetValue() const { checkf(IsSet(), TEXT("It is an error to call GetValue() on an unset TOptional. Please either check IsSet() or use Get(DefaultValue) instead.")); return Pointer; }
+
+	TClassType* operator->() const { return Pointer; }
+
+	TClassType& operator*() const { return *Pointer; }
+
+	/** @return The optional value when set; DefaultValue otherwise. */
+	TClassType* Get(TClassType* DefaultPointer) const { return IsSet() ? Pointer : DefaultPointer; }
+
+private:
+	TClassType* Pointer;
+};
