Core - Now that the autoweak bug is fixed with FObjectPtr reference checker, im restoring the change to the FObjectProperty base CPPType expectations.

diff --git a/Engine/Source/Runtime/CoreUObject/Private/UObject/PropertyObject.cpp b/Engine/Source/Runtime/CoreUObject/Private/UObject/PropertyObject.cpp
--- a/Engine/Source/Runtime/CoreUObject/Private/UObject/PropertyObject.cpp
+++ b/Engine/Source/Runtime/CoreUObject/Private/UObject/PropertyObject.cpp
@@ -98,7 +98,7 @@ void FObjectProperty::SerializeItem( FStructuredArchive::FSlot Slot, void* Value
 	if (UnderlyingArchive.IsObjectReferenceCollector())
 	{
 		// Serialize in place
-		UObject** ObjectPtr = GetPropertyValuePtr(Value);
+		TObjectPtr<UObject>* ObjectPtr = GetPropertyValuePtr(Value);
 		Slot << (*ObjectPtr);
 
 		if(!UnderlyingArchive.IsSaving())
@@ -145,7 +145,7 @@ void FObjectProperty::SerializeItem( FStructuredArchive::FSlot Slot, void* Value
 
 const TCHAR* FObjectProperty::ImportText_Internal(const TCHAR* Buffer, void* ContainerOrPropertyPtr, EPropertyPointerType PropertyPointerType, UObject* OwnerObject, int32 PortFlags, FOutputDevice* ErrorText) const
 {
-	const TCHAR* Result = TFObjectPropertyBase<UObject*>::ImportText_Internal(Buffer, ContainerOrPropertyPtr, PropertyPointerType, OwnerObject, PortFlags, ErrorText);
+	const TCHAR* Result = TFObjectPropertyBase<TObjectPtr<UObject>>::ImportText_Internal(Buffer, ContainerOrPropertyPtr, PropertyPointerType, OwnerObject, PortFlags, ErrorText);
 	if (Result)
 	{
 		void* Data = PointerToValuePtr(ContainerOrPropertyPtr, PropertyPointerType);
@@ -185,7 +185,7 @@ TObjectPtr<UObject> FObjectProperty::GetObjectPtrPropertyValue(const void* Prope
 
 UObject* FObjectProperty::GetObjectPropertyValue(const void* PropertyValueAddress) const
 {
-	return GetPropertyValue(PropertyValueAddress);
+	return GetPropertyValue(PropertyValueAddress).Get();
 }
 
 UObject* FObjectProperty::GetObjectPropertyValue_InContainer(const void* ContainerAddress, int32 ArrayIndex) const
diff --git a/Engine/Source/Runtime/CoreUObject/Private/UObject/PropertyObjectPtr.cpp b/Engine/Source/Runtime/CoreUObject/Private/UObject/PropertyObjectPtr.cpp
--- a/Engine/Source/Runtime/CoreUObject/Private/UObject/PropertyObjectPtr.cpp
+++ b/Engine/Source/Runtime/CoreUObject/Private/UObject/PropertyObjectPtr.cpp
@@ -151,7 +151,7 @@ TObjectPtr<UObject> FObjectPtrProperty::GetObjectPtrPropertyValue(const void* Pr
 
 UObject* FObjectPtrProperty::GetObjectPropertyValue(const void* PropertyValueAddress) const
 {
-	return GetPropertyValue(PropertyValueAddress);
+	return GetPropertyValue(PropertyValueAddress).Get();
 }
 
 UObject* FObjectPtrProperty::GetObjectPropertyValue_InContainer(const void* ContainerAddress, int32 ArrayIndex) const
diff --git a/Engine/Source/Runtime/CoreUObject/Public/UObject/UnrealType.h b/Engine/Source/Runtime/CoreUObject/Public/UObject/UnrealType.h
--- a/Engine/Source/Runtime/CoreUObject/Public/UObject/UnrealType.h
+++ b/Engine/Source/Runtime/CoreUObject/Public/UObject/UnrealType.h
@@ -2898,12 +2898,9 @@ PRAGMA_ENABLE_DEPRECATION_WARNINGS
 //
 // Describes a reference variable to another object which may be nil.
 //
-//class COREUOBJECT_API FObjectProperty : public TFObjectPropertyBase<TObjectPtr<UObject>>
-//{
-//	DECLARE_FIELD(FObjectProperty, TFObjectPropertyBase<TObjectPtr<UObject>>, CASTCLASS_FObjectProperty)
-class COREUOBJECT_API FObjectProperty : public TFObjectPropertyBase<UObject*>
+class COREUOBJECT_API FObjectProperty : public TFObjectPropertyBase<TObjectPtr<UObject>>
 {
-	DECLARE_FIELD(FObjectProperty, TFObjectPropertyBase<UObject*>, CASTCLASS_FObjectProperty)
+	DECLARE_FIELD(FObjectProperty, TFObjectPropertyBase<TObjectPtr<UObject>>, CASTCLASS_FObjectProperty)
 
 	FObjectProperty(FFieldVariant InOwner, const FName& InName, EObjectFlags InObjectFlags)
 		: TFObjectPropertyBase(InOwner, InName, InObjectFlags)
