Core: propagate the PIE package flag to external packages containing objects outered to a PIE package.

#rnx

diff --git a/Engine/Source/Runtime/CoreUObject/Private/UObject/LinkerLoad.cpp b/Engine/Source/Runtime/CoreUObject/Private/UObject/LinkerLoad.cpp
--- a/Engine/Source/Runtime/CoreUObject/Private/UObject/LinkerLoad.cpp
+++ b/Engine/Source/Runtime/CoreUObject/Private/UObject/LinkerLoad.cpp
@@ -4964,14 +4964,21 @@ UObject* FLinkerLoad::CreateExport( int32 Index )
 		Params.Template = Template;
 		// if our outer is actually an import, then the package we are an export of is not in our outer chain, set our package in that case
 		Params.ExternalPackage = Export.OuterIndex.IsImport() ? LinkerRoot : nullptr;
+
+		// Propagate relevant flags from the outer package to the external package
+		if (Params.ExternalPackage)
+		{
+			Params.ExternalPackage->SetPackageFlags(ThisParent->GetPackage()->GetPackageFlags() & PKG_PlayInEditor);
+		}
+
 		Export.Object = StaticConstructObject_Internal(Params);
 
 		if (FPlatformProperties::RequiresCookedData())
 		{
 			if (GIsInitialLoad || GUObjectArray.IsOpenForDisregardForGC())
-		{
-			Export.Object->AddToRoot();
-		}
+			{
+				Export.Object->AddToRoot();
+			}
 		}
 		
 		LoadClass = Export.Object->GetClass(); // this may have changed if we are overwriting a CDO component
