ZenLoader in editor: Fixed handling of failed packages

diff --git a/Engine/Source/Runtime/CoreUObject/Private/Serialization/AsyncLoading2.cpp b/Engine/Source/Runtime/CoreUObject/Private/Serialization/AsyncLoading2.cpp
--- a/Engine/Source/Runtime/CoreUObject/Private/Serialization/AsyncLoading2.cpp
+++ b/Engine/Source/Runtime/CoreUObject/Private/Serialization/AsyncLoading2.cpp
@@ -4505,13 +4505,12 @@ EEventLoadNodeExecutionResult FAsyncPackage2::ProcessLinkerLoadPackageSummary(FA
 	TRACE_CPUPROFILER_EVENT_SCOPE(ProcessLinkerLoadPackageSummary);
 
 #if ALT2_ENABLE_NEW_ARCHIVE_FOR_LINKERLOAD
-	if (LinkerLoadState->Linker->GetLoader()->IsError())
+	FLinkerLoad::ELinkerStatus LinkerResult = FLinkerLoad::LINKER_Failed;
+	if (!LinkerLoadState->Linker->GetLoader()->IsError())
 	{
-		bLoadHasFailed = true;
-		return EEventLoadNodeExecutionResult::Complete;
+		LinkerLoadState->Linker->bUseTimeLimit = false;
+		LinkerResult = LinkerLoadState->Linker->ProcessPackageSummary(nullptr);
 	}
-	LinkerLoadState->Linker->bUseTimeLimit = false;
-	FLinkerLoad::ELinkerStatus LinkerResult = LinkerLoadState->Linker->ProcessPackageSummary(nullptr);
 #else
 	FLinkerLoad::ELinkerStatus LinkerResult = LinkerLoadState->Linker->Tick(/* RemainingTimeLimit */ 0.0, /* bUseTimeLimit */ false, /* bUseFullTimeLimit */ false, nullptr);
 #endif
@@ -4519,9 +4518,8 @@ EEventLoadNodeExecutionResult FAsyncPackage2::ProcessLinkerLoadPackageSummary(FA
 	if (LinkerResult == FLinkerLoad::LINKER_Failed)
 	{
 		bLoadHasFailed = true;
-		return EEventLoadNodeExecutionResult::Complete;
 	}
-	check(LinkerLoadState->Linker->HasFinishedInitialization());
+	check(LinkerLoadState->Linker->HasFinishedInitialization() || bLoadHasFailed);
 
 	LinkerLoadState->LinkerLoadHeaderData.ImportMap.SetNum(LinkerLoadState->Linker->ImportMap.Num());
 	TArray<FName, TInlineAllocator<128>> ImportedPackageNames;
@@ -4643,6 +4641,12 @@ EEventLoadNodeExecutionResult FAsyncPackage2::ProcessLinkerLoadPackageSummary(FA
 	HeaderData.ImportedPublicExportHashes = LinkerLoadState->LinkerLoadHeaderData.ImportedPublicExportHashes;
 	HeaderData.ImportMap = LinkerLoadState->LinkerLoadHeaderData.ImportMap;
 
+	if (bLoadHasFailed && Desc.bCanBeImported)
+	{
+		FLoadedPackageRef& PackageRef = ImportStore.LoadedPackageStore.FindPackageRefChecked(Desc.UPackageId);
+		PackageRef.SetHasFailed();
+	}
+
 	AsyncPackageLoadingState = EAsyncPackageLoadingState2::SetupDependencies;
 	GetPackageNode(Package_SetupDependencies).ReleaseBarrier();
 	return EEventLoadNodeExecutionResult::Complete;
@@ -4880,11 +4884,7 @@ EEventLoadNodeExecutionResult FAsyncPackage2::Event_ProcessPackageSummary(FAsync
 #if ALT2_ENABLE_LINKERLOAD_SUPPORT
 	if (Package->LinkerLoadState.IsSet())
 	{
-		EEventLoadNodeExecutionResult Result = Package->ProcessLinkerLoadPackageSummary(ThreadState);
-		if (!Package->bLoadHasFailed)
-		{
-			return Result;
-		}
+		return Package->ProcessLinkerLoadPackageSummary(ThreadState);
 	}
 #endif
 	if (Package->bLoadHasFailed)
@@ -7569,24 +7569,11 @@ EAsyncPackageState::Type FAsyncPackage2::CreateClusters(FAsyncLoadingThreadState
 
 void FAsyncPackage2::FinishUPackage()
 {
-	if (LinkerRoot)
+	if (LinkerRoot && !bLoadHasFailed)
 	{
-		if (!bLoadHasFailed)
-		{
-			// Mark package as having been fully loaded and update load time.
-			LinkerRoot->MarkAsFullyLoaded();
-			LinkerRoot->SetLoadTime(FPlatformTime::Seconds() - LoadStartTime);
-		}
-		else
-		{
-			// Clean up UPackage so it can't be found later
-			if (bCreatedLinkerRoot && !LinkerRoot->IsRooted())
-			{
-				LinkerRoot->ClearFlags(RF_NeedPostLoad | RF_NeedLoad | RF_NeedPostLoadSubobjects);
-				LinkerRoot->MarkAsGarbage();
-				LinkerRoot->Rename(*MakeUniqueObjectName(GetTransientPackage(), UPackage::StaticClass()).ToString(), nullptr, REN_DontCreateRedirectors | REN_DoNotDirty | REN_ForceNoResetLoaders | REN_NonTransactional);
-			}
-		}
+		// Mark package as having been fully loaded and update load time.
+		LinkerRoot->MarkAsFullyLoaded();
+		LinkerRoot->SetLoadTime(FPlatformTime::Seconds() - LoadStartTime);
 	}
 }
 
