[CoreUObject] fix for static analysis nullptr

diff --git a/Engine/Source/Runtime/CoreUObject/Private/UObject/ObjectHandle.cpp b/Engine/Source/Runtime/CoreUObject/Private/UObject/ObjectHandle.cpp
--- a/Engine/Source/Runtime/CoreUObject/Private/UObject/ObjectHandle.cpp
+++ b/Engine/Source/Runtime/CoreUObject/Private/UObject/ObjectHandle.cpp
@@ -468,16 +468,17 @@ UClass* ResolveObjectRefClass(const FObjectRef& ObjectRef, uint32 LoadFlags /*=
 		if (!ObjectRef.ClassName.IsNone())
 		{
 			ClassObject = (UClass*)StaticFindObjectFastInternal(UClass::StaticClass(), ClassPackage, ObjectRef.ClassName);
-			if (ClassObject)
+			if (UObjectRedirector* Redirector = dynamic_cast<UObjectRedirector*>(ClassObject))
 			{
-				if (UObjectRedirector* Redirector = dynamic_cast<UObjectRedirector*>(ClassObject))
+				ClassObject = (UClass*)Redirector->DestinationObject;
+				if (ClassObject != nullptr)
 				{
-					ClassObject = (UClass*)Redirector->DestinationObject;
-					if (ClassObject != nullptr)
-					{
-						ClassPackage = ClassObject->GetPackage();
-					}
+					ClassPackage = ClassObject->GetPackage();
 				}
+			}
+
+			if (ClassObject)
+			{
 				if (ClassObject->HasAnyFlags(RF_NeedLoad) && ClassPackage->GetLinker())
 				{
 					ClassPackage->GetLinker()->Preload(ClassObject);
