Bulk Data - fix failing tests on consoles

#rnx

diff --git a/Engine/Source/Runtime/CoreUObject/Tests/Serialization/BulkDataTests.cpp b/Engine/Source/Runtime/CoreUObject/Tests/Serialization/BulkDataTests.cpp
--- a/Engine/Source/Runtime/CoreUObject/Tests/Serialization/BulkDataTests.cpp
+++ b/Engine/Source/Runtime/CoreUObject/Tests/Serialization/BulkDataTests.cpp
@@ -216,6 +216,8 @@ TEST_CASE("CoreUObject::Serialization::FBulkData::Basic", "[CoreUObject][Seriali
 	}
 }
 
+#if WITH_EDITORONLY_DATA
+
 TEST_CASE("CoreUObject::Serialization::FBulkData::Serialize", "[CoreUObject][Serialization]")
 {
 	SECTION("Serialize to memory archive")
@@ -351,6 +353,8 @@ TEST_CASE("CoreUObject::Serialization::FBulkData::Serialize", "[CoreUObject][Ser
 	}
 }
 
+#endif // WITH_EDITORONLY_DATA
+
 } // namespace UE::Serialization::BulkDataTest
 
 #endif // WITH_LOW_LEVEL_TESTS
