//UE5/Main - Cook EngineTest CookedEditor Win64 - CustomAttributesModifier.uasset warning: The node is from an Editor Only module
#fix separated out behaviour for testing whether or not an object is _editor only_ from checking whether or not an object will be stripped or not for (non) editor cooks/builds - as the API naming gave a false expecation of its behaviour

diff --git a/Engine/Source/Runtime/CoreUObject/Private/UObject/SavePackage/PackageHarvester.cpp b/Engine/Source/Runtime/CoreUObject/Private/UObject/SavePackage/PackageHarvester.cpp
--- a/Engine/Source/Runtime/CoreUObject/Private/UObject/SavePackage/PackageHarvester.cpp
+++ b/Engine/Source/Runtime/CoreUObject/Private/UObject/SavePackage/PackageHarvester.cpp
@@ -9,9 +9,6 @@
 #include "UObject/UObjectHash.h"
 #include "Interfaces/ITargetPlatform.h"
 
-// bring the UObectGlobal declaration visible to non editor
-bool IsEditorOnlyObject(const UObject* InObject, bool bCheckRecursive, bool bCheckMarks);
-
 EObjectMark GenerateMarksForObject(const UObject* InObject, const ITargetPlatform* TargetPlatform)
 {
 	EObjectMark Marks = OBJECTMARK_NOMARKS;
@@ -49,7 +46,7 @@ EObjectMark GenerateMarksForObject(const UObject* InObject, const ITargetPlatfor
 #endif
 	
 	// CDOs must be included if their class is so only inherit marks, for everything else we check the native overrides as well
-	if (IsEditorOnlyObject(InObject, false, false))
+	if (SavePackageUtilities::IsStrippedEditorOnlyObject(InObject, false, false))
 	{
 		Marks = (EObjectMark)(Marks | OBJECTMARK_EditorOnly);
 	}
@@ -203,7 +200,7 @@ void FPackageHarvester::ProcessExport(const FExportWithContext& InProcessContext
 	UObject* Export = InProcessContext.Export;
 
 	// No need to check marks since we do not set them on objects anymore
-	bool bReferencerIsEditorOnly = IsEditorOnlyObject(Export, true /* bCheckRecursive */, false /* bCheckMarks */) && !Export->HasNonEditorOnlyReferences();
+	bool bReferencerIsEditorOnly = SavePackageUtilities::IsStrippedEditorOnlyObject(Export, true /* bCheckRecursive */, false /* bCheckMarks */) && !Export->HasNonEditorOnlyReferences();
 	FExportScope HarvesterScope(*this, InProcessContext, bReferencerIsEditorOnly);
 	
 	// The export scope set the current harvesting context
diff --git a/Engine/Source/Runtime/CoreUObject/Private/UObject/SavePackage/SavePackageUtilities.cpp b/Engine/Source/Runtime/CoreUObject/Private/UObject/SavePackage/SavePackageUtilities.cpp
--- a/Engine/Source/Runtime/CoreUObject/Private/UObject/SavePackage/SavePackageUtilities.cpp
+++ b/Engine/Source/Runtime/CoreUObject/Private/UObject/SavePackage/SavePackageUtilities.cpp
@@ -596,6 +596,32 @@ void GetCDOSubobjects(UObject* CDO, TArray<UObject*>& Subobjects)
 		}
 	}
 }
+	
+bool IsStrippedEditorOnlyObject(const UObject* InObject, bool bCheckRecursive, bool bCheckMarks)
+{
+#if WITH_EDITOR
+	// Configurable via ini setting
+	static struct FCanStripEditorOnlyExportsAndImports
+	{
+		bool bCanStripEditorOnlyObjects;
+		FCanStripEditorOnlyExportsAndImports()
+			: bCanStripEditorOnlyObjects(true)
+		{
+			GConfig->GetBool(TEXT("Core.System"), TEXT("CanStripEditorOnlyExportsAndImports"), bCanStripEditorOnlyObjects, GEngineIni);
+		}
+		FORCEINLINE operator bool() const { return bCanStripEditorOnlyObjects; }
+	} CanStripEditorOnlyExportsAndImports;
+	if (!CanStripEditorOnlyExportsAndImports)
+	{
+		return false;
+	}
+	
+	return IsEditorOnlyObject(InObject, bCheckRecursive, bCheckMarks);
+#else
+	return true;
+#endif
+}
+	
 } // end namespace SavePackageUtilities
 
 namespace UE::SavePackageUtilities
@@ -744,22 +770,6 @@ void FObjectSaveContextData::Set(UPackage* Package, const ITargetPlatform* InTar
 bool IsEditorOnlyObject(const UObject* InObject, bool bCheckRecursive, bool bCheckMarks)
 {
 	DECLARE_SCOPE_CYCLE_COUNTER(TEXT("IsEditorOnlyObject"), STAT_IsEditorOnlyObject, STATGROUP_LoadTime);
-
-	// Configurable via ini setting
-	static struct FCanStripEditorOnlyExportsAndImports
-	{
-		bool bCanStripEditorOnlyObjects;
-		FCanStripEditorOnlyExportsAndImports()
-			: bCanStripEditorOnlyObjects(true)
-		{
-			GConfig->GetBool(TEXT("Core.System"), TEXT("CanStripEditorOnlyExportsAndImports"), bCanStripEditorOnlyObjects, GEngineIni);
-		}
-		FORCEINLINE operator bool() const { return bCanStripEditorOnlyObjects; }
-	} CanStripEditorOnlyExportsAndImports;
-	if (!CanStripEditorOnlyExportsAndImports)
-	{
-		return false;
-	}
 	check(InObject);
 
 	if ((bCheckMarks && InObject->HasAnyMarks(OBJECTMARK_EditorOnly)) || InObject->IsEditorOnly())
diff --git a/Engine/Source/Runtime/CoreUObject/Private/UObject/SavePackage/SavePackageUtilities.h b/Engine/Source/Runtime/CoreUObject/Private/UObject/SavePackage/SavePackageUtilities.h
--- a/Engine/Source/Runtime/CoreUObject/Private/UObject/SavePackage/SavePackageUtilities.h
+++ b/Engine/Source/Runtime/CoreUObject/Private/UObject/SavePackage/SavePackageUtilities.h
@@ -442,6 +442,9 @@ namespace SavePackageUtilities
 	void AsyncWriteFile(EAsyncWriteOptions Options, FSavePackageOutputFile& File);
 
 	void GetCDOSubobjects(UObject* CDO, TArray<UObject*>& Subobjects);
+
+	/** Returns result of IsEditorOnlyObject if Engine:[Core.System]:CanStripEditorOnlyExportsAndImports (ini) is set to true */
+	bool IsStrippedEditorOnlyObject(const UObject* InObject, bool bCheckRecursive = true, bool bCheckMarks = true);
 }
 
 #if ENABLE_COOK_STATS
diff --git a/Engine/Source/Runtime/CoreUObject/Private/UObject/SavePackage2.cpp b/Engine/Source/Runtime/CoreUObject/Private/UObject/SavePackage2.cpp
--- a/Engine/Source/Runtime/CoreUObject/Private/UObject/SavePackage2.cpp
+++ b/Engine/Source/Runtime/CoreUObject/Private/UObject/SavePackage2.cpp
@@ -50,9 +50,6 @@
 // defined in UObjectGlobals.cpp
 COREUOBJECT_API extern bool GOutputCookingWarnings;
 
-// bring the UObectGlobal declaration visible to non editor
-bool IsEditorOnlyObject(const UObject* InObject, bool bCheckRecursive, bool bCheckMarks);
-
 namespace
 {
 
diff --git a/Engine/Source/Runtime/CoreUObject/Public/UObject/UObjectGlobals.h b/Engine/Source/Runtime/CoreUObject/Public/UObject/UObjectGlobals.h
--- a/Engine/Source/Runtime/CoreUObject/Public/UObject/UObjectGlobals.h
+++ b/Engine/Source/Runtime/CoreUObject/Public/UObject/UObjectGlobals.h
@@ -2960,7 +2960,7 @@ struct FAssetMsg
  * or
  * - if bCheckMarks is true, if it has the EditorOnly object mark
  * or 
- * - if bCheckRecursive is true, if it's class, outer, or archetypes are editor only
+ * - if bCheckRecursive is true, if its class, super struct, outer, or archetypes are editor only
  */
 COREUOBJECT_API bool IsEditorOnlyObject(const UObject* InObject, bool bCheckRecursive = true, bool bCheckMarks = true);
 #endif //WITH_EDITOR
