Improve debugability of FObjectPtr/TObjectPtr by supplying a "DebugPtr" member on both that can be used to dereference resolved object references in debugger watch windows.  This required changing TObjectPtr to not privately inherit from FObjectPtr, but to instead contain an FObjectPtr.

diff --git a/Engine/Source/Runtime/CoreUObject/Public/UObject/ObjectPtr.h b/Engine/Source/Runtime/CoreUObject/Public/UObject/ObjectPtr.h
--- a/Engine/Source/Runtime/CoreUObject/Public/UObject/ObjectPtr.h
+++ b/Engine/Source/Runtime/CoreUObject/Public/UObject/ObjectPtr.h
@@ -210,12 +210,6 @@ public:
 		return IsA(T::StaticClass());
 	}
 
-protected:
-	FORCEINLINE UObject* GetNoRead() const
-	{
-		return ResolveObjectHandleNoRead(Handle);
-	}
-
 private:
 	friend FORCEINLINE uint32 GetTypeHash(const FObjectPtr& Object)
 	{
@@ -223,7 +217,13 @@ private:
 		return GetTypeHash(Object.Handle);
 	}
 
-	mutable FObjectHandle Handle;
+	union
+	{
+		mutable FObjectHandle Handle;
+		// DebugPtr allows for easier dereferencing of a resolved FObjectPtr in watch windows of debuggers.  If the address in the pointer
+		// is an odd/uneven number, that means the object reference is unresolved and you will not be able to dereference it successfully.
+		UObject* DebugPtr;
+	};
 };
 
 template <typename T>
@@ -254,22 +254,26 @@ namespace ObjectPtr_Private
  * with FObjectPtr/TObjectPtr members instead.
  */
 template <typename T>
-struct TObjectPtr : private FObjectPtr
+struct TObjectPtr
 {
 public:
 	using ElementType = T;
 
-	TObjectPtr() = default;
-	TObjectPtr(TObjectPtr<T>&&) = default;
-	TObjectPtr(const TObjectPtr<T>&) = default;
+	TObjectPtr()
+		: ObjectPtr()
+	{
+	}
+
+	TObjectPtr(TObjectPtr<T>&& Other) = default;
+	TObjectPtr(const TObjectPtr<T>& Other) = default;
 
 	explicit FORCEINLINE TObjectPtr(ENoInit)
-		: FObjectPtr(NoInit)
+		: ObjectPtr(NoInit)
 	{
 	}
 
 	FORCEINLINE TObjectPtr(TYPE_OF_NULLPTR)
-		: FObjectPtr(nullptr)
+		: ObjectPtr(nullptr)
 	{
 	}
 
@@ -278,7 +282,7 @@ public:
 		decltype(ImplicitConv<T*>(std::declval<U*>()))* = nullptr
 	>
 	FORCEINLINE TObjectPtr(const TObjectPtr<U>& Other)
-		: FObjectPtr(static_cast<const FObjectPtr&>(Other))
+		: ObjectPtr(Other.ObjectPtr)
 	{
 	}
 
@@ -290,12 +294,12 @@ public:
 		>* = nullptr
 	>
 	FORCEINLINE TObjectPtr(U&& Object)
-		: FObjectPtr(const_cast<std::remove_const_t<T>*>(ImplicitConv<T*>(Object)))
+		: ObjectPtr(const_cast<std::remove_const_t<T>*>(ImplicitConv<T*>(Object)))
 	{
 	}
 
 	explicit FORCEINLINE TObjectPtr(TPrivateObjectPtr<T>&& PrivatePtr)
-		: FObjectPtr(const_cast<UObject*>(PrivatePtr.Pointer))
+		: ObjectPtr(const_cast<UObject*>(PrivatePtr.Pointer))
 	{
 	}
 
@@ -304,7 +308,7 @@ public:
 
 	FORCEINLINE TObjectPtr<T>& operator=(TYPE_OF_NULLPTR)
 	{
-		FObjectPtr::operator=(nullptr);
+		ObjectPtr = nullptr;
 		return *this;
 	}
 
@@ -314,7 +318,7 @@ public:
 	>
 	FORCEINLINE TObjectPtr<T>& operator=(const TObjectPtr<U>& Other)
 	{
-		FObjectPtr::operator=(static_cast<const FObjectPtr&>(Other));
+		ObjectPtr = Other.ObjectPtr;
 		return *this;
 	}
 
@@ -327,13 +331,13 @@ public:
 	>
 	FORCEINLINE TObjectPtr<T>& operator=(U&& Object)
 	{
-		FObjectPtr::operator=(const_cast<std::remove_const_t<T>*>(ImplicitConv<T*>(Object)));
+		ObjectPtr = const_cast<std::remove_const_t<T>*>(ImplicitConv<T*>(Object));
 		return *this;
 	}
 
 	FORCEINLINE TObjectPtr<T>& operator=(TPrivateObjectPtr<T>&& PrivatePtr)
 	{
-		FObjectPtr::operator=(const_cast<UObject*>(PrivatePtr.Pointer));
+		ObjectPtr = const_cast<UObject*>(PrivatePtr.Pointer);
 		return *this;
 	}
 
@@ -343,7 +347,7 @@ public:
 	>
 	FORCEINLINE bool operator==(const TObjectPtr<U>& Other) const
 	{
-		return FObjectPtr::operator==(static_cast<const FObjectPtr&>(Other));
+		return ObjectPtr == Other.ObjectPtr;
 	}
 
 	template <
@@ -352,7 +356,7 @@ public:
 	>
 	FORCEINLINE bool operator!=(const TObjectPtr<U>& Other) const
 	{
-		return FObjectPtr::operator!=(static_cast<const FObjectPtr&>(Other));
+		return ObjectPtr != Other.ObjectPtr;
 	}
 
 	// @TODO: OBJPTR: There is a risk that the FObjectPtr is storing a reference to the wrong type.  This could
@@ -361,8 +365,8 @@ public:
 	//			a reference to the wrong type of object which we'll just send back static_casted as the wrong type.  Doing
 	//			a check or checkSlow here could catch this, but it would be better if the check could happen elsewhere that
 	//			isn't called as frequently.
-	FORCEINLINE T* Get() const { return (T*)(FObjectPtr::Get()); }
-	using FObjectPtr::GetClass;
+	FORCEINLINE T* Get() const { return (T*)(ObjectPtr.Get()); }
+	FORCEINLINE UClass* GetClass() const { return ObjectPtr.GetClass(); }
 	
 	FORCEINLINE operator T* () const { return Get(); }
 	template <typename U>
@@ -375,51 +379,59 @@ public:
 	UE_OBJPTR_DEPRECATED(5.0, "Conversion to a mutable pointer is deprecated.  Please pass a TObjectPtr<T>& instead so that assignment can be tracked accurately.")
 	explicit FORCEINLINE operator T*& () { return GetInternalRef(); }
 
-	using FObjectPtr::IsNull;
-	using FObjectPtr::IsNullNoResolve;
-	using FObjectPtr::operator!;
-	using FObjectPtr::operator bool;
-	using FObjectPtr::IsResolved;
-	using FObjectPtr::GetFName;
-	using FObjectPtr::GetName;
-	using FObjectPtr::GetPath;
-	using FObjectPtr::GetPathName;
-	using FObjectPtr::GetFullName;
-	using FObjectPtr::GetHandle;
-
-	using FObjectPtr::IsA;
+	FORCEINLINE bool IsNull() const { return ObjectPtr.IsNull(); }
+	FORCEINLINE bool IsNullNoResolve() const { return ObjectPtr.IsNullNoResolve(); }
+	FORCEINLINE bool operator!() const { return ObjectPtr.operator!(); }
+	explicit FORCEINLINE operator bool() const { return ObjectPtr.operator bool(); }
+	FORCEINLINE bool IsResolved() const { return ObjectPtr.IsResolved(); }
+	FORCEINLINE FString GetPath() const { return ObjectPtr.GetPath(); }
+	FORCEINLINE FString GetPathName() const { return ObjectPtr.GetPathName(); }
+	FORCEINLINE FName GetFName() const { return ObjectPtr.GetFName(); }
+	FORCEINLINE FString GetName() const { return ObjectPtr.GetName(); }
+	FORCEINLINE FString GetFullName(EObjectFullNameFlags Flags = EObjectFullNameFlags::None) const { return ObjectPtr.GetFullName(Flags); }
+	FORCEINLINE FObjectHandle GetHandle() const { return ObjectPtr.GetHandle(); }
+	FORCEINLINE bool IsA(const UClass* SomeBase) const { return ObjectPtr.IsA(SomeBase); }
+	template <typename U> FORCEINLINE bool IsA() const { return ObjectPtr.IsA<U>(); }
 
-	friend FORCEINLINE uint32 GetTypeHash(const TObjectPtr<T>& ObjectPtr)
+	friend FORCEINLINE uint32 GetTypeHash(const TObjectPtr<T>& InObjectPtr)
 	{
-		return GetTypeHash(static_cast<const FObjectPtr&>(ObjectPtr));
+		return GetTypeHash(InObjectPtr.ObjectPtr);
 	}
 
-	friend FORCEINLINE FArchive& operator<<(FArchive& Ar, TObjectPtr<T>& ObjectPtr)
+	friend FORCEINLINE FArchive& operator<<(FArchive& Ar, TObjectPtr<T>& InObjectPtr)
 	{
-		Ar << static_cast<FObjectPtr&>(ObjectPtr);
+		Ar << InObjectPtr.ObjectPtr;
 		return Ar;
 	}
 
-	friend FORCEINLINE void operator<<(FStructuredArchiveSlot Slot, TObjectPtr<T>& ObjectPtr)
+	friend FORCEINLINE void operator<<(FStructuredArchiveSlot Slot, TObjectPtr<T>& InObjectPtr)
 	{
-		Slot << static_cast<FObjectPtr&>(ObjectPtr);
+		Slot << InObjectPtr.ObjectPtr;
 	}
 
 	friend struct FObjectPtr;
 	template <typename U> friend struct TObjectPtr;
 	template <typename U, typename V> friend bool ObjectPtr_Private::IsObjectPtrEqualToRawPtrOfRelatedType(const TObjectPtr<U>& Ptr, const V* Other);
-private:
 
-	FORCEINLINE T* GetNoRead() const { return (T*)(FObjectPtr::GetNoRead()); }
+private:
+	FORCEINLINE T* GetNoRead() const { return (T*)(ResolveObjectHandleNoRead(ObjectPtr.GetHandleRef())); }
 
 	// @TODO: OBJPTR: There is a risk of a gap in access tracking here.  The caller may get a mutable pointer, write to it, then
 	//			read from it.  That last read would happen without an access being recorded.  Not sure if there is a good way
 	//			to handle this case without forcing the calling code to be modified.
 	FORCEINLINE T*& GetInternalRef()
 	{
-		FObjectPtr::Get();
-		return (T*&)FObjectPtr::GetHandleRef();
+		ObjectPtr.Get();
+		return (T*&)ObjectPtr.GetHandleRef();
 	}
+
+	union
+	{
+		FObjectPtr ObjectPtr;
+		// DebugPtr allows for easier dereferencing of a resolved TObjectPtr in watch windows of debuggers.  If the address in the pointer
+		// is an odd/uneven number, that means the object reference is unresolved and you will not be able to dereference it successfully.
+		T* DebugPtr;
+	};
 };
 
 template <typename T> struct TIsTObjectPtr<               TObjectPtr<T>> { enum { Value = true }; };
@@ -734,12 +746,12 @@ FORCEINLINE TWeakObjectPtr<T> MakeWeakObjectPtr(TObjectPtr<T> Ptr)
 
 FORCEINLINE TObjectPtr<UObject>& FObjectPtr::ToTObjectPtr()
 {
-	return *static_cast<TObjectPtr<UObject>*>(this);
+	return *reinterpret_cast<TObjectPtr<UObject>*>(this);
 }
 
 FORCEINLINE const TObjectPtr<UObject>& FObjectPtr::ToTObjectPtr() const
 {
-	return *static_cast<const TObjectPtr<UObject>*>(this);
+	return *reinterpret_cast<const TObjectPtr<UObject>*>(this);
 }
 
 /** Swap variants between TObjectPtr<T> and raw pointer to T */
