LinkerLoad: Ensure that objects are associated with the correct async package for NotifyConstructedDuringAsyncLoading

diff --git a/Engine/Source/Runtime/CoreUObject/Private/Serialization/AsyncLoading.cpp b/Engine/Source/Runtime/CoreUObject/Private/Serialization/AsyncLoading.cpp
--- a/Engine/Source/Runtime/CoreUObject/Private/Serialization/AsyncLoading.cpp
+++ b/Engine/Source/Runtime/CoreUObject/Private/Serialization/AsyncLoading.cpp
@@ -542,9 +542,7 @@ struct FAsyncPackageScope
 {
 	/** Outer scope package */
 	void* PreviousPackage;
-#if WITH_IOSTORE_IN_EDITOR
 	IAsyncPackageLoader* PreviousAsyncPackageLoader;
-#endif
 	/** Cached ThreadContext so we don't have to access it again */
 	FUObjectThreadContext& ThreadContext;
 
@@ -553,17 +551,13 @@ struct FAsyncPackageScope
 	{
 		PreviousPackage = ThreadContext.AsyncPackage;
 		ThreadContext.AsyncPackage = InPackage;
-#if WITH_IOSTORE_IN_EDITOR
 		PreviousAsyncPackageLoader = ThreadContext.AsyncPackageLoader;
 		ThreadContext.AsyncPackageLoader = &InPackage->AsyncLoadingThread;
-#endif
 	}
 	~FAsyncPackageScope()
 	{
 		ThreadContext.AsyncPackage = PreviousPackage;
-#if WITH_IOSTORE_IN_EDITOR
 		ThreadContext.AsyncPackageLoader = PreviousAsyncPackageLoader;
-#endif
 	}
 };
 
@@ -1730,10 +1724,8 @@ FScopedAsyncPackageEvent::FScopedAsyncPackageEvent(FAsyncPackage* InPackage)
 	FUObjectThreadContext& ThreadContext = FUObjectThreadContext::Get();
 	PreviousPackage = static_cast<FAsyncPackage*>(ThreadContext.AsyncPackage);
 	ThreadContext.AsyncPackage = Package;
-#if WITH_IOSTORE_IN_EDITOR
 	PreviousAsyncPackageLoader = ThreadContext.AsyncPackageLoader;
 	ThreadContext.AsyncPackageLoader = &InPackage->AsyncLoadingThread;
-#endif
 	Package->BeginAsyncLoad();
 	FExclusiveLoadPackageTimeTracker::PushLoadPackage(Package->Desc.PackagePath.GetPackageFName());
 }
@@ -1748,9 +1740,7 @@ FScopedAsyncPackageEvent::~FScopedAsyncPackageEvent()
 	// Restore the package from the outer scope
 	FUObjectThreadContext& ThreadContext = FUObjectThreadContext::Get();
 	ThreadContext.AsyncPackage = PreviousPackage;
-#if WITH_IOSTORE_IN_EDITOR
 	ThreadContext.AsyncPackageLoader = PreviousAsyncPackageLoader;
-#endif
 }
 
 FORCENOINLINE static bool CheckForFileOpenLogCommandLine()
diff --git a/Engine/Source/Runtime/CoreUObject/Private/Serialization/AsyncLoading.h b/Engine/Source/Runtime/CoreUObject/Private/Serialization/AsyncLoading.h
--- a/Engine/Source/Runtime/CoreUObject/Private/Serialization/AsyncLoading.h
+++ b/Engine/Source/Runtime/CoreUObject/Private/Serialization/AsyncLoading.h
@@ -980,9 +980,7 @@ struct FScopedAsyncPackageEvent
 	FAsyncPackage* Package;
 	/** Outer scope package */
 	FAsyncPackage* PreviousPackage;
-#if WITH_IOSTORE_IN_EDITOR
 	IAsyncPackageLoader* PreviousAsyncPackageLoader;
-#endif
 	
 	FScopedAsyncPackageEvent(FAsyncPackage* InPackage);
 	~FScopedAsyncPackageEvent();
diff --git a/Engine/Source/Runtime/CoreUObject/Private/Serialization/AsyncLoading2.cpp b/Engine/Source/Runtime/CoreUObject/Private/Serialization/AsyncLoading2.cpp
--- a/Engine/Source/Runtime/CoreUObject/Private/Serialization/AsyncLoading2.cpp
+++ b/Engine/Source/Runtime/CoreUObject/Private/Serialization/AsyncLoading2.cpp
@@ -2184,7 +2184,6 @@ private:
 
 struct FAsyncPackage2
 {
-	friend struct FScopedAsyncPackageEvent2;
 	friend struct FAsyncPackageScope2;
 	friend class FAsyncLoadingThread2;
 	friend class FAsyncLoadEventQueue2;
@@ -2543,20 +2542,6 @@ public:
 	FUObjectSerializeContext* GetSerializeContext();
 };
 
-struct FScopedAsyncPackageEvent2
-{
-	/** Current scope package */
-	FAsyncPackage2* Package;
-	/** Outer scope package */
-	FAsyncPackage2* PreviousPackage;
-#if WITH_EDITOR
-	IAsyncPackageLoader* PreviousAsyncPackageLoader;
-#endif
-
-	FScopedAsyncPackageEvent2(FAsyncPackage2* InPackage);
-	~FScopedAsyncPackageEvent2();
-};
-
 class FAsyncLoadingThread2 final
 	: public FRunnable
 	, public IAsyncPackageLoader
@@ -3857,34 +3842,6 @@ bool FAsyncLoadEventQueue2::ExecuteSyncLoadEvents(FAsyncLoadingThreadState2& Thr
 	return bDidSomething;
 }
 
-FScopedAsyncPackageEvent2::FScopedAsyncPackageEvent2(FAsyncPackage2* InPackage)
-	:Package(InPackage)
-{
-	check(Package);
-
-	// Update the thread context with the current package. This is used by NotifyConstructedDuringAsyncLoading.
-	FUObjectThreadContext& ThreadContext = FUObjectThreadContext::Get();
-	PreviousPackage = static_cast<FAsyncPackage2*>(ThreadContext.AsyncPackage);
-	ThreadContext.AsyncPackage = Package;
-#if WITH_IOSTORE_IN_EDITOR
-	PreviousAsyncPackageLoader = ThreadContext.AsyncPackageLoader;
-	ThreadContext.AsyncPackageLoader = &InPackage->AsyncLoadingThread;
-#endif
-	Package->BeginAsyncLoad();
-}
-
-FScopedAsyncPackageEvent2::~FScopedAsyncPackageEvent2()
-{
-	Package->EndAsyncLoad();
-
-	// Restore the package from the outer scope
-	FUObjectThreadContext& ThreadContext = FUObjectThreadContext::Get();
-	ThreadContext.AsyncPackage = PreviousPackage;
-#if WITH_IOSTORE_IN_EDITOR
-	ThreadContext.AsyncPackageLoader = PreviousAsyncPackageLoader;
-#endif
-}
-
 FUObjectSerializeContext* FAsyncPackage2::GetSerializeContext()
 {
 	return FUObjectThreadContext::Get().GetSerializeContext();
@@ -4335,6 +4292,8 @@ void FAsyncPackage2::CreateLinker(const FLinkerInstancingContext* InstancingCont
 	}
 	check(Linker);
 	check(Linker->LinkerRoot == LinkerRoot);
+	check(!Linker->AsyncRoot);
+	Linker->AsyncRoot = this;
 	LinkerLoadState->Linker = Linker;
 #if ALT2_ENABLE_NEW_ARCHIVE_FOR_LINKERLOAD
 	Linker->ResetStatusInfo();
@@ -4345,6 +4304,8 @@ void FAsyncPackage2::DetachLinker()
 {
 	if (LinkerLoadState.IsSet() && LinkerLoadState->Linker)
 	{
+		check(LinkerLoadState->Linker->AsyncRoot == this);
+		LinkerLoadState->Linker->AsyncRoot = nullptr;
 		LinkerLoadState->Linker = nullptr;
 	}
 }
@@ -4879,7 +4840,7 @@ EEventLoadNodeExecutionResult FAsyncPackage2::Event_ProcessPackageSummary(FAsync
 	check(Package->AsyncPackageLoadingState == EAsyncPackageLoadingState2::WaitingForIo);
 	Package->AsyncPackageLoadingState = EAsyncPackageLoadingState2::ProcessPackageSummary;
 
-	FScopedAsyncPackageEvent2 Scope(Package);
+	FAsyncPackageScope2 Scope(Package);
 
 #if ALT2_ENABLE_LINKERLOAD_SUPPORT
 	if (Package->LinkerLoadState.IsSet())
@@ -5045,7 +5006,7 @@ EEventLoadNodeExecutionResult FAsyncPackage2::Event_ProcessExportBundle(FAsyncLo
 	UE_ASYNC_PACKAGE_DEBUG(Package->Desc);
 	check(Package->AsyncPackageLoadingState == EAsyncPackageLoadingState2::ProcessExportBundles);
 
-	FScopedAsyncPackageEvent2 Scope(Package);
+	FAsyncPackageScope2 Scope(Package);
 
 #if ALT2_ENABLE_LINKERLOAD_SUPPORT
 	if (Package->LinkerLoadState.IsSet())
diff --git a/Engine/Source/Runtime/CoreUObject/Private/Serialization/AsyncPackageLoader.cpp b/Engine/Source/Runtime/CoreUObject/Private/Serialization/AsyncPackageLoader.cpp
--- a/Engine/Source/Runtime/CoreUObject/Private/Serialization/AsyncPackageLoader.cpp
+++ b/Engine/Source/Runtime/CoreUObject/Private/Serialization/AsyncPackageLoader.cpp
@@ -676,12 +676,6 @@ void NotifyRegistrationComplete()
 	GPackageLoader->StartThread();
 }
 
-void NotifyConstructedDuringAsyncLoading(UObject* Object, bool bSubObject)
-{
-	LLM_SCOPE(ELLMTag::AsyncLoading);
-	GetAsyncPackageLoader().NotifyConstructedDuringAsyncLoading(Object, bSubObject);
-}
-
 void NotifyUnreachableObjects(const TArrayView<FUObjectItem*>& UnreachableObjects)
 {
 	LLM_SCOPE(ELLMTag::AsyncLoading);
diff --git a/Engine/Source/Runtime/CoreUObject/Private/Serialization/EditorPackageLoader.cpp b/Engine/Source/Runtime/CoreUObject/Private/Serialization/EditorPackageLoader.cpp
--- a/Engine/Source/Runtime/CoreUObject/Private/Serialization/EditorPackageLoader.cpp
+++ b/Engine/Source/Runtime/CoreUObject/Private/Serialization/EditorPackageLoader.cpp
@@ -179,12 +179,7 @@ public:
 
 	virtual void NotifyConstructedDuringAsyncLoading(UObject* Object, bool bSubObject) override
 	{
-		FUObjectThreadContext& ThreadContext = FUObjectThreadContext::Get();
-
-		if (ThreadContext.AsyncPackageLoader)
-		{
-			ThreadContext.AsyncPackageLoader->NotifyConstructedDuringAsyncLoading(Object, bSubObject);
-		}
+		checkf(false, TEXT("This is never called"));
 	}
 
 	virtual void NotifyUnreachableObjects(const TArrayView<FUObjectItem*>& UnreachableObjects) override
diff --git a/Engine/Source/Runtime/CoreUObject/Private/UObject/LinkerLoad.cpp b/Engine/Source/Runtime/CoreUObject/Private/UObject/LinkerLoad.cpp
--- a/Engine/Source/Runtime/CoreUObject/Private/UObject/LinkerLoad.cpp
+++ b/Engine/Source/Runtime/CoreUObject/Private/UObject/LinkerLoad.cpp
@@ -4497,6 +4497,7 @@ UObject* FLinkerLoad::CreateExport( int32 Index )
 	// Check whether we already loaded the object and if not whether the context flags allow loading it.
 	if( !Export.Object && !FilterExport(Export) ) // for some acceptable position, it was not "not for" 
 	{
+		TGuardValue<void*> GuardThreadContextAsyncPackage(FUObjectThreadContext::Get().AsyncPackage, AsyncRoot);
 		FUObjectSerializeContext* CurrentLoadContext = GetSerializeContext();
 		check(!GEventDrivenLoaderEnabled || !bLockoutLegacyOperations || !EVENT_DRIVEN_ASYNC_LOAD_ACTIVE_AT_RUNTIME);
 		check(Export.ObjectName!=NAME_None || !(Export.ObjectFlags&RF_Public));
diff --git a/Engine/Source/Runtime/CoreUObject/Private/UObject/UObjectGlobals.cpp b/Engine/Source/Runtime/CoreUObject/Private/UObject/UObjectGlobals.cpp
--- a/Engine/Source/Runtime/CoreUObject/Private/UObject/UObjectGlobals.cpp
+++ b/Engine/Source/Runtime/CoreUObject/Private/UObject/UObjectGlobals.cpp
@@ -66,6 +66,7 @@
 #include "UObject/PropertyWithSetterAndGetter.h"
 #include "UObject/AnyPackagePrivate.h"
 #include "UObject/UObjectGlobalsInternal.h"
+#include "Serialization/AsyncPackageLoader.h"
 
 DEFINE_LOG_CATEGORY(LogUObjectGlobals);
 
@@ -3037,11 +3038,6 @@ bool StaticAllocateObjectErrorTests( const UClass* InClass, UObject* InOuter, FN
 	return false;
 }
 
-/**
-* Call back into the async loading code to inform of the creation of a new object
-*/
-void NotifyConstructedDuringAsyncLoading(UObject* Object, bool bSubObject);
-
 /**
 * For object overwrites, the class may want to persist some info over the re-initialize
 * this is only used for classes in the script compiler
@@ -3295,7 +3291,12 @@ UObject* StaticAllocateObject
 
 	if (IsInAsyncLoadingThread())
 	{
-		NotifyConstructedDuringAsyncLoading(Obj, bSubObject);
+		FUObjectThreadContext& ThreadContext = FUObjectThreadContext::Get();
+		if (ThreadContext.AsyncPackageLoader)
+		{
+			LLM_SCOPE(ELLMTag::AsyncLoading);
+			ThreadContext.AsyncPackageLoader->NotifyConstructedDuringAsyncLoading(Obj, bSubObject);
+		}
 	}
 	else
 	{
diff --git a/Engine/Source/Runtime/CoreUObject/Private/UObject/UObjectThreadContext.cpp b/Engine/Source/Runtime/CoreUObject/Private/UObject/UObjectThreadContext.cpp
--- a/Engine/Source/Runtime/CoreUObject/Private/UObject/UObjectThreadContext.cpp
+++ b/Engine/Source/Runtime/CoreUObject/Private/UObject/UObjectThreadContext.cpp
@@ -17,9 +17,7 @@ FUObjectThreadContext::FUObjectThreadContext()
 , IsInConstructor(0)
 , ConstructedObject(nullptr)
 , AsyncPackage(nullptr)
-#if WITH_IOSTORE_IN_EDITOR
 , AsyncPackageLoader(nullptr)
-#endif
 , SerializeContext(new FUObjectSerializeContext())
 {}
 
diff --git a/Engine/Source/Runtime/CoreUObject/Public/UObject/LinkerLoad.h b/Engine/Source/Runtime/CoreUObject/Public/UObject/LinkerLoad.h
--- a/Engine/Source/Runtime/CoreUObject/Public/UObject/LinkerLoad.h
+++ b/Engine/Source/Runtime/CoreUObject/Public/UObject/LinkerLoad.h
@@ -264,7 +264,7 @@ public:
 	}
 
 	/** The async package associated with this linker */
-	struct FAsyncPackage* AsyncRoot;
+	void* AsyncRoot;
 #if WITH_EDITOR
 	/** Bulk data that does not need to be loaded when the linker is loaded.												*/
 	TArray<FBulkData*> BulkDataLoaders;
diff --git a/Engine/Source/Runtime/CoreUObject/Public/UObject/UObjectThreadContext.h b/Engine/Source/Runtime/CoreUObject/Public/UObject/UObjectThreadContext.h
--- a/Engine/Source/Runtime/CoreUObject/Public/UObject/UObjectThreadContext.h
+++ b/Engine/Source/Runtime/CoreUObject/Public/UObject/UObjectThreadContext.h
@@ -92,10 +92,8 @@ public:
 	UObject* ConstructedObject;
 	/** Async Package currently processing objects */
 	void* AsyncPackage;
-#if WITH_IOSTORE_IN_EDITOR
 	/** Async package loader currently processing objects */
 	IAsyncPackageLoader* AsyncPackageLoader;
-#endif
 
 #if !(UE_BUILD_SHIPPING || UE_BUILD_TEST)
 	/** Stack to ensure that PostInitProperties is routed through Super:: calls. **/
