ZenLoader in editor: PostLoad all constructed objects for a package instead of only the exports

diff --git a/Engine/Source/Runtime/CoreUObject/Private/Serialization/AsyncLoading2.cpp b/Engine/Source/Runtime/CoreUObject/Private/Serialization/AsyncLoading2.cpp
--- a/Engine/Source/Runtime/CoreUObject/Private/Serialization/AsyncLoading2.cpp
+++ b/Engine/Source/Runtime/CoreUObject/Private/Serialization/AsyncLoading2.cpp
@@ -4811,16 +4811,11 @@ EEventLoadNodeExecutionResult FAsyncPackage2::PostLoadLinkerLoadPackageExports(F
 {
 	FAsyncLoadingTickScope2 InAsyncLoadingTick(AsyncLoadingThread);
 
-	const int32 ExportCount = LinkerLoadState->Linker->ExportMap.Num();
-	while (LinkerLoadState->PostLoadExportIndex < ExportCount)
+	const int32 ObjectCount = ConstructedObjects.Num();
+	while (LinkerLoadState->PostLoadExportIndex < ObjectCount)
 	{
-		const int32 ExportIndex = LinkerLoadState->PostLoadExportIndex++;
-		const FObjectExport& Export = LinkerLoadState->Linker->ExportMap[ExportIndex];
-		UObject* Object = Export.Object;
-		if (Object)
-		{
-			Object->ConditionalPostLoad();
-		}
+		const int32 ObjectIndex = LinkerLoadState->PostLoadExportIndex++;
+		ConstructedObjects[ObjectIndex]->ConditionalPostLoad();
 	}
 
 	FUObjectSerializeContext* LoadContext = GetSerializeContext();
