ZenLoader in editor: Only call CreateImports/Exports when in the correct state, fixes issues with trying to use the linker after it has been detached.

diff --git a/Engine/Source/Runtime/CoreUObject/Private/Serialization/AsyncLoading2.cpp b/Engine/Source/Runtime/CoreUObject/Private/Serialization/AsyncLoading2.cpp
--- a/Engine/Source/Runtime/CoreUObject/Private/Serialization/AsyncLoading2.cpp
+++ b/Engine/Source/Runtime/CoreUObject/Private/Serialization/AsyncLoading2.cpp
@@ -4646,6 +4646,7 @@ EEventLoadNodeExecutionResult FAsyncPackage2::ProcessLinkerLoadPackageImportsAnd
 
 bool FAsyncPackage2::ProcessLinkerLoadPackageImports(FAsyncLoadingThreadState2& ThreadState)
 {
+	check(AsyncPackageLoadingState == EAsyncPackageLoadingState2::ProcessExportBundles);
 	if (LinkerLoadState->bIsCurrentlyProcessingImports)
 	{
 		return true;
@@ -4655,15 +4656,19 @@ bool FAsyncPackage2::ProcessLinkerLoadPackageImports(FAsyncLoadingThreadState2&
 	while (LinkerLoadState->ProcessingImportedPackageIndex < ImportedPackagesCount)
 	{
 		FAsyncPackage2* ImportedPackage = Data.ImportedAsyncPackages[LinkerLoadState->ProcessingImportedPackageIndex];
-		if (ImportedPackage && ImportedPackage->LinkerLoadState.IsSet() && ImportedPackage->LinkerLoadState->Linker)
+		if (ImportedPackage && ImportedPackage->LinkerLoadState.IsSet())
 		{
-			if (!ImportedPackage->ProcessLinkerLoadPackageImports(ThreadState))
+			check(ImportedPackage->AsyncPackageLoadingState >= EAsyncPackageLoadingState2::ProcessExportBundles);
+			if (ImportedPackage->AsyncPackageLoadingState == EAsyncPackageLoadingState2::ProcessExportBundles)
 			{
-				return false;
-			}
-			if (!ImportedPackage->ProcessLinkerLoadPackageExports(ThreadState))
-			{
-				return false;
+				if (!ImportedPackage->ProcessLinkerLoadPackageImports(ThreadState))
+				{
+					return false;
+				}
+				if (!ImportedPackage->ProcessLinkerLoadPackageExports(ThreadState))
+				{
+					return false;
+				}
 			}
 		}
 		++LinkerLoadState->ProcessingImportedPackageIndex;
@@ -4716,6 +4721,7 @@ bool FAsyncPackage2::ProcessLinkerLoadPackageImports(FAsyncLoadingThreadState2&
 
 bool FAsyncPackage2::ProcessLinkerLoadPackageExports(FAsyncLoadingThreadState2& ThreadState)
 {
+	check(AsyncPackageLoadingState == EAsyncPackageLoadingState2::ProcessExportBundles);
 #if WITH_EDITORONLY_DATA
 	// Create metadata object
 	int32 MetaDataIndex = LinkerLoadState->MetaDataIndex.IsSet() ? LinkerLoadState->MetaDataIndex.GetValue() : -1;
