Add missing support for resolving weak references to placeholder linker objects created in deferred Blueprint object serialization paths on editor load.

diff --git a/Engine/Source/Runtime/CoreUObject/Private/UObject/LinkerPlaceholderBase.cpp b/Engine/Source/Runtime/CoreUObject/Private/UObject/LinkerPlaceholderBase.cpp
--- a/Engine/Source/Runtime/CoreUObject/Private/UObject/LinkerPlaceholderBase.cpp
+++ b/Engine/Source/Runtime/CoreUObject/Private/UObject/LinkerPlaceholderBase.cpp
@@ -75,10 +75,10 @@ int32 FLinkerPlaceholderObjectImpl::ResolvePlaceholderValues(const TArray<FField
 		if (PropertyIndex == 0)
 		{
 #if USE_DEFERRED_DEPENDENCY_CHECK_VERIFICATION_TESTS
-			check(Property.IsA<FObjectProperty>() || Property.IsA<UObjectProperty>());
+			check(Property.IsA<FObjectPropertyBase>() || Property.IsA<UObjectPropertyBase>());
 #endif // USE_DEFERRED_DEPENDENCY_CHECK_VERIFICATION_TESTS
 
-			const FObjectProperty* ReferencingProperty = Property.Get<FObjectProperty>();
+			const FObjectPropertyBase* ReferencingProperty = Property.Get<FObjectPropertyBase>();
 			check(ReferencingProperty);
 
 			UObject* CurrentValue = ReferencingProperty->GetObjectPropertyValue(ValueAddress);
@@ -392,7 +392,7 @@ FLinkerPlaceholderBase::FPlaceholderValuePropertyPath::FPlaceholderValueProperty
 bool FLinkerPlaceholderBase::FPlaceholderValuePropertyPath::IsValid() const
 {
 	return (PropertyChain.Num() > 0) && 
-		(PropertyChain[0].IsA<FObjectProperty>() || PropertyChain[0].IsA<UObjectProperty>()) &&
+		(PropertyChain[0].IsA<FObjectPropertyBase>() || PropertyChain[0].IsA<UObjectPropertyBase>()) &&
 		PropertyChain.Last().GetOwnerClass();
 }
 
@@ -463,13 +463,13 @@ FLinkerPlaceholderBase::~FLinkerPlaceholderBase()
 //------------------------------------------------------------------------------
 bool FLinkerPlaceholderBase::AddReferencingPropertyValue(FFieldVariant ReferencingProperty, void* DataPtr)
 {
-	check(ReferencingProperty.IsA<FObjectProperty>() || ReferencingProperty.IsA<UObjectProperty>())
+	check(ReferencingProperty.IsA<FObjectPropertyBase>() || ReferencingProperty.IsA<UObjectPropertyBase>())
 	FPlaceholderValuePropertyPath PropertyChain(ReferencingProperty);
 	UObject* ReferencingContainer = FLinkerPlaceholderObjectImpl::FindPlaceholderContainer(PropertyChain);
 	if (ReferencingContainer != nullptr)
 	{
 #if USE_DEFERRED_DEPENDENCY_CHECK_VERIFICATION_TESTS
-		check(ReferencingProperty.Get<FObjectProperty>()->GetObjectPropertyValue(DataPtr) == GetPlaceholderAsUObject());
+		check(ReferencingProperty.Get<FObjectPropertyBase>()->GetObjectPropertyValue(DataPtr) == GetPlaceholderAsUObject());
 		check(PropertyChain.IsValid());
 #endif // USE_DEFERRED_DEPENDENCY_CHECK_VERIFICATION_TESTS
 
diff --git a/Engine/Source/Runtime/CoreUObject/Private/UObject/PropertyWeakObjectPtr.cpp b/Engine/Source/Runtime/CoreUObject/Private/UObject/PropertyWeakObjectPtr.cpp
--- a/Engine/Source/Runtime/CoreUObject/Private/UObject/PropertyWeakObjectPtr.cpp
+++ b/Engine/Source/Runtime/CoreUObject/Private/UObject/PropertyWeakObjectPtr.cpp
@@ -3,6 +3,8 @@
 #include "CoreMinimal.h"
 #include "UObject/ObjectMacros.h"
 #include "UObject/UnrealType.h"
+#include "UObject/LinkerPlaceholderExportObject.h"
+#include "UObject/LinkerPlaceholderClass.h"
 
 /*-----------------------------------------------------------------------------
 	FWeakObjectProperty.
@@ -56,11 +58,31 @@ void FWeakObjectProperty::SerializeItem( FStructuredArchive::FSlot Slot, void* V
 {
 	FArchive& UnderlyingArchive = Slot.GetUnderlyingArchive();
 
-	UObject* ObjectValue = GetObjectPropertyValue(Value);
+	UObject* OldObjectValue = GetObjectPropertyValue(Value);
 	Slot << *(FWeakObjectPtr*)Value;
-	if ((UnderlyingArchive.IsLoading() || UnderlyingArchive.IsModifyingWeakAndStrongReferences()) && ObjectValue != GetObjectPropertyValue(Value))
+
+	if (UnderlyingArchive.IsLoading() || UnderlyingArchive.IsModifyingWeakAndStrongReferences())
 	{
-		CheckValidObject(Value, nullptr); // FWeakObjectProperty is never non-nullable at this point so it's ok to pass null as the current value
+		UObject* NewObjectValue = GetObjectPropertyValue(Value);
+
+		if (OldObjectValue != NewObjectValue)
+		{
+#if USE_CIRCULAR_DEPENDENCY_LOAD_DEFERRING
+			if (UnderlyingArchive.IsLoading() && !UnderlyingArchive.IsObjectReferenceCollector())
+			{
+				if (ULinkerPlaceholderExportObject* PlaceholderVal = Cast<ULinkerPlaceholderExportObject>(NewObjectValue))
+				{
+					PlaceholderVal->AddReferencingPropertyValue(this, Value);
+				}
+				else if (ULinkerPlaceholderClass* PlaceholderClass = Cast<ULinkerPlaceholderClass>(NewObjectValue))
+				{
+					PlaceholderClass->AddReferencingPropertyValue(this, Value);
+				}
+			}
+#endif // USE_CIRCULAR_DEPENDENCY_LOAD_DEFERRING
+
+			CheckValidObject(Value, nullptr); // FWeakObjectProperty is never non-nullable at this point so it's ok to pass null as the current value
+		}
 	}
 }
 
