Change FTopLevelAssetPath to import/export as a string.
Implement placeholder import for structured format in case there is any stale data in that form.
FAssetBundleEntry::ImportTextItem now doesn't depend on UScriptStruct for FTopLevelAssetPath so it asset registry can load bundles very early during engine boot.

diff --git a/Engine/Source/Runtime/CoreUObject/Private/AssetRegistry/AssetBundleData.cpp b/Engine/Source/Runtime/CoreUObject/Private/AssetRegistry/AssetBundleData.cpp
--- a/Engine/Source/Runtime/CoreUObject/Private/AssetRegistry/AssetBundleData.cpp
+++ b/Engine/Source/Runtime/CoreUObject/Private/AssetRegistry/AssetBundleData.cpp
@@ -106,8 +106,6 @@ bool FAssetBundleEntry::ImportTextItem(const TCHAR*& Buffer, int32 PortFlags, UO
 		return false;
 	}
 
-	UScriptStruct* TopLevelAssetPathStruct = TBaseStructure<FTopLevelAssetPath>::Get();
-	static const FString TopLevelPathStructName = TopLevelAssetPathStruct->GetName(); 
 	FOutputDeviceNull SilenceErrors;
 	bool bBundleAssetsSeen = false;
 	bool bAssetPathsSeen = false;
@@ -173,13 +171,14 @@ bool FAssetBundleEntry::ImportTextItem(const TCHAR*& Buffer, int32 PortFlags, UO
 
 			// Try multiple formats to parse our list of FTopLevelAssetPaths
 			//     - the native export converts FTopLevelAssetPaths to FSoftObjectPaths and exports the FSoftObjectPaths as strings
-			//     - the default export path (activated when defaults is not null) exports FTopLevelAssetpaths as formatted structs
+			//     - the default export path (activated when defaults is not null) delegates to FTopLevelAssetPath
 			FSoftObjectPath TempObjectPath;
 			FTopLevelAssetPath TempTopLevelPath;
 			while (true)
 			{
 				TempObjectPath.Reset();
-				// Import string e.g. /Path/To/Package.AssetName
+				TempTopLevelPath.Reset();
+				// FSoftObjectPath format 
 				if (TempObjectPath.ImportTextItem(BufferIt, PortFlags, Parent, &SilenceErrors))
 				{
 					if (!TempObjectPath.GetSubPathString().IsEmpty())
@@ -191,19 +190,16 @@ bool FAssetBundleEntry::ImportTextItem(const TCHAR*& Buffer, int32 PortFlags, UO
 					}
 					LocalObjectPaths.Emplace(MoveTemp(TempObjectPath));
 				}
-				else
+				// FTopLevelAssetPath's native ExportText will export a string that can also be parsed by FSoftObjectPath so this won't be called unless some stale data is in structured text form
+				else if(TempTopLevelPath.ImportTextItem(BufferIt, PortFlags, Parent, &SilenceErrors))
 				{
-					// Import formatted struct e.g. (PackageName="/Path/To/Package", AssetName="AssetName")
-					TempTopLevelPath.Reset();
-					const TCHAR* NewBuffer = TopLevelAssetPathStruct->ImportText(BufferIt, &TempTopLevelPath, Parent, PortFlags, &SilenceErrors, TopLevelPathStructName);;
-					if (!NewBuffer)
-					{
-						ErrorText->Logf(ELogVerbosity::Error, TEXT("FAssetBundleEntry::ImportTextItem - unterminated or ill formed AssetPaths list while importing asset bundle entry: '%s'. When parsing import text '%s'."), BufferIt, Buffer);
-						break;
-					}
-					BufferIt = NewBuffer;
 					LocalObjectPaths.Emplace(TempTopLevelPath);
 				}
+				else
+				{					
+					ErrorText->Logf(ELogVerbosity::Error, TEXT("FAssetBundleEntry::ImportTextItem - unterminated or ill formed AssetPaths list while importing asset bundle entry: '%s'. When parsing import text '%s'."), BufferIt, Buffer);
+					return false;
+				}
 
 				if (*BufferIt == ',')
 				{
diff --git a/Engine/Source/Runtime/CoreUObject/Private/UObject/TopLevelAssetPath.cpp b/Engine/Source/Runtime/CoreUObject/Private/UObject/TopLevelAssetPath.cpp
--- a/Engine/Source/Runtime/CoreUObject/Private/UObject/TopLevelAssetPath.cpp
+++ b/Engine/Source/Runtime/CoreUObject/Private/UObject/TopLevelAssetPath.cpp
@@ -5,6 +5,7 @@
 #include "Containers/UnrealString.h"
 
 #include "Misc/AsciiSet.h"
+#include "Misc/Optional.h"
 #include "Misc/RedirectCollector.h"
 #include "UObject/CoreRedirects.h"
 #include "UObject/ObjectRedirector.h"
@@ -22,6 +23,8 @@ struct TStructOpsTypeTraits<FTopLevelAssetPath> : public TStructOpsTypeTraitsBas
 	enum
 	{
 		WithStructuredSerializeFromMismatchedTag = true,
+		WithExportTextItem = true,
+		WithImportTextItem = true,
 	};
 };
 UE_IMPLEMENT_STRUCT("/Script/CoreUObject", TopLevelAssetPath)
@@ -168,6 +171,332 @@ bool FTopLevelAssetPath::TrySetPath(FAnsiStringView Path)
 	return TrySetPath(Wide);
 }
 
+bool FTopLevelAssetPath::ExportTextItem(FString& ValueStr, FTopLevelAssetPath const& DefaultValue, UObject* Parent, int32 PortFlags, UObject* ExportRootScope) const
+{
+	if (0 != (PortFlags & EPropertyPortFlags::PPF_ExportCpp))
+	{
+		return false;
+	}
+
+	if (!IsNull())
+	{
+		if (PortFlags & PPF_Delimited)
+		{
+			ValueStr += TEXT("\"");
+			ValueStr += ToString().ReplaceQuotesWithEscapedQuotes();
+			ValueStr += TEXT("\"");
+		}
+		else
+		{
+			ValueStr += ToString();
+		}
+	}
+	else
+	{
+		ValueStr += TEXT("None");
+	}
+	return true;
+}
+
+namespace UE::TopLevelAssetPath::Private
+{
+	void SkipWhitespace(const TCHAR*& Str)
+	{
+		while(FChar::IsWhitespace(*Str))
+		{
+			Str++;
+		}
+	}
+
+	// Copied from FNameProperty::ImportText_Internal
+	const TCHAR* ImportName(
+		const TCHAR* Buffer,
+		FName& OutName,
+		int32 PortFlags
+	) 
+	{
+		FName ImportedName;
+		if (!(PortFlags & PPF_Delimited))
+		{
+			ImportedName = FName(Buffer);		
+
+			// in order to indicate that the value was successfully imported, advance the buffer past the last character that was imported
+			Buffer += FCString::Strlen(Buffer);
+		}
+		else
+		{
+			TStringBuilder<256> Token;
+			Buffer = FPropertyHelpers::ReadToken(Buffer, /* out */ Token, true);
+			if (!Buffer)
+			{
+				return nullptr;
+			}
+
+			ImportedName = FName(Token);
+		}
+
+		OutName = ImportedName;
+		return Buffer;
+	}
+
+	// Stripped down version of FProperty::ImportSingleProperty that only supports loading the two FName fields we need 
+	const TCHAR* ImportSingleProperty(
+		const TCHAR* Str,
+		TOptional<FName>& OutPackageName, 
+		TOptional<FName>& OutAssetName, 
+		UObject* SubobjectOuter,
+		int32 PortFlags,
+		FOutputDevice* Warn
+	)
+	{
+		constexpr FAsciiSet Whitespaces(" \t");
+		constexpr FAsciiSet Delimiters("=([.");
+	
+		static const FName NAME_PackageName("PackageName");
+		static const FName NAME_AssetName("AssetName");
+
+		// strip leading whitespace
+		const TCHAR* Start = FAsciiSet::Skip(Str, Whitespaces);
+		// find first delimiter
+		Str = FAsciiSet::FindFirstOrEnd(Start, Delimiters);
+		// check if delimiter was found...
+		if (*Str)
+		{
+			// strip trailing whitespace
+			int32 Len = UE_PTRDIFF_TO_INT32(Str - Start);
+			while (Len > 0 && Whitespaces.Contains(Start[Len - 1]))
+			{
+				--Len;
+			}
+
+			const FName PropertyName(Len, Start);
+			TOptional<FName>* Destination = nullptr;
+			if (PropertyName == NAME_PackageName)
+			{
+				Destination = &OutPackageName;
+			}
+			else if (PropertyName == NAME_AssetName)
+			{
+				Destination = &OutAssetName;
+			}
+			
+			// check to see if this property has already imported data
+			if (Destination && Destination->IsSet())
+			{
+				Warn->Logf(ELogVerbosity::Warning, TEXT("redundant data: %s"), Start);
+				return Str;
+			}
+			else if (Destination == nullptr)
+			{
+				UE_SUPPRESS(LogExec, Verbose, Warn->Logf(TEXT("Unknown property in TopLevelAssetPath: %s "), Start));
+				return Str;
+			}
+
+			// strip whitespace before =
+			SkipWhitespace(Str);
+			if (*Str++ != '=')
+			{
+				Warn->Logf(ELogVerbosity::Warning, TEXT("Missing '=' in default properties assignment: %s"), Start );
+				return Str;
+			}
+			// strip whitespace after =
+			SkipWhitespace(Str);
+
+			if (*Str == TCHAR('\0') || *Str == TCHAR(',') || *Str == TCHAR(')'))
+			{
+				// Empty value
+				return Str;
+			}
+
+			const TCHAR* Result = ImportName(Str, Destination->Emplace(NAME_None), PortFlags);
+			if (Result == NULL || Result == Str)
+			{
+				UE_SUPPRESS(LogExec, Verbose, Warn->Logf(TEXT("Unknown property in TopLevelAssetPath: %s "), Start));
+			}
+
+			// in the failure case, don't return NULL so the caller can potentially skip less and get values further in the string
+			if (Result != NULL)
+			{
+				Str = Result;
+			}
+		}
+		return Str;
+	};
+
+	 // Structure copied from UScriptStruct::ImportText to avoid dependency on UScriptStruct as we load asset registry data early during engine initialization 
+	 // Support for this format is deprecated and this code should be removed in 5.2 or 5.3.
+	static const TCHAR* ImportFormattedStruct(
+		const TCHAR* InBuffer,
+		TOptional<FName>& OutPackageName, 
+		TOptional<FName>& OutAssetName,
+		UObject* OwnerObject, 
+		int32 PortFlags,
+		FOutputDevice* ErrorText
+	)
+	{	
+		// this keeps track of the number of errors we've logged, so that we can add new lines when logging more than one error
+		int32 ErrorCount = 0;
+		const TCHAR* Buffer = InBuffer;
+		if (*Buffer++ == TCHAR('('))
+		{
+			// Parse all properties.
+			while (*Buffer != TCHAR(')'))
+			{
+				// parse and import the value
+				Buffer = ImportSingleProperty(Buffer, OutPackageName, OutAssetName, OwnerObject, PortFlags | PPF_Delimited, ErrorText);
+
+				// skip any remaining text before the next property value
+				SkipWhitespace(Buffer);
+				int32 SubCount = 0;
+				while (*Buffer && *Buffer != TCHAR('\r') && *Buffer != TCHAR('\n') &&
+					(SubCount > 0 || *Buffer != TCHAR(')')) && (SubCount > 0 || *Buffer != TCHAR(',')))
+				{
+					SkipWhitespace(Buffer);
+					if (*Buffer == TCHAR('\"'))
+					{
+						do
+						{
+							Buffer++;
+						} while (*Buffer && *Buffer != TCHAR('\"') && *Buffer != TCHAR('\n') && *Buffer != TCHAR('\r'));
+
+						if (*Buffer != TCHAR('\"'))
+						{
+							ErrorText->Logf(TEXT("%sImportText (TopLevelAssetPath): Bad quoted string at: %s"), 
+								ErrorCount++ > 0 ? LINE_TERMINATOR : TEXT(""), Buffer);
+							return nullptr;
+						}
+					}
+					else if (*Buffer == TCHAR('('))
+					{
+						SubCount++;
+					}
+					else if (*Buffer == TCHAR(')'))
+					{
+						SubCount--;
+						if (SubCount < 0)
+						{
+							ErrorText->Logf(TEXT("%sImportText (TopLevelAssetPath): Too many closing parenthesis in: %s"),
+								ErrorCount++ > 0 ? LINE_TERMINATOR : TEXT(""), InBuffer);
+							return nullptr;
+						}
+					}
+					Buffer++;
+				}
+				if (SubCount > 0)
+				{
+					ErrorText->Logf(TEXT("%sImportText(TopLevelAssetPath): Not enough closing parenthesis in: %s"),
+						ErrorCount++ > 0 ? LINE_TERMINATOR : TEXT(""), InBuffer);
+					return nullptr;
+				}
+
+				// Skip comma.
+				if (*Buffer == TCHAR(','))
+				{
+					// Skip comma.
+					Buffer++;
+				}
+				else if (*Buffer != TCHAR(')'))
+				{
+					ErrorText->Logf(TEXT("%sImportText (TopLevelAssetPath): Missing closing parenthesis: %s"), 
+						ErrorCount++ > 0 ? LINE_TERMINATOR : TEXT(""), InBuffer);
+					return nullptr;
+				}
+
+				SkipWhitespace(Buffer);
+			}
+
+			// Skip trailing ')'.
+			Buffer++;
+		}
+		else
+		{
+			ErrorText->Logf(TEXT("%sImportText (TopLevelAssetPath): Missing opening parenthesis: %s"),
+				ErrorCount++ > 0 ? LINE_TERMINATOR : TEXT(""), InBuffer); //-V547
+			return nullptr;
+		}
+		return Buffer;
+	}
+}
+
+bool FTopLevelAssetPath::ImportTextItem(const TCHAR*& Buffer, int32 PortFlags, UObject* Parent, FOutputDevice* ErrorText, FArchive* InSerializingArchive)
+{
+	const TCHAR* OriginalBuffer = Buffer;
+	TStringBuilder<256> ImportedPath;
+
+	if (const TCHAR* NewBuffer = FPropertyHelpers::ReadToken(Buffer, /* out */ ImportedPath, /* dotted names */ true))
+	{
+		Buffer = NewBuffer;
+	}
+	else
+	{
+		return false;
+	}
+
+	if (ImportedPath == TEXTVIEW("None"))
+	{
+		Reset();
+		return true;
+	}
+	else if (*Buffer == TCHAR('('))
+	{
+		// Blueprints and other utilities may pass in () as a hardcoded value for an empty struct, so treat that like an empty string
+		Buffer++;
+
+		if (*Buffer == TCHAR(')'))
+		{
+			Buffer++;
+			Reset();
+			return true;
+		}
+		else
+		{
+			Buffer--;
+			// Attempt to parse classic struct format in case there is some data in that format
+			TOptional<FName> NewPackageName, NewAssetName;
+			if (const TCHAR* NewBuffer = UE::TopLevelAssetPath::Private::ImportFormattedStruct(Buffer, NewPackageName, NewAssetName, Parent, PortFlags, ErrorText))
+			{
+				// Emit a warning because this format is deprecated and shouldn't be accepted forever.
+				ErrorText->Logf(ELogVerbosity::Warning, TEXT("Struct format for FTopLevelAssetPath is deprecated. Imported struct: %.*s"), 
+					(NewBuffer - Buffer), Buffer);
+				
+				Buffer = NewBuffer;
+				TrySetPath(NewPackageName.Get({}), NewAssetName.Get({}));
+				return true;
+			}
+			else
+			{
+				Buffer = OriginalBuffer;
+				return false;
+			}
+		}
+	}
+
+	else if (*Buffer == TCHAR('\''))
+	{
+		// A ' token likely means we're looking at a path string in the form "Texture2d'/Game/UI/HUD/Actions/Barrel'" and we need to read and append the path part
+		// We have to skip over the first ' as FPropertyHelpers::ReadToken doesn't read single-quoted strings correctly, but does read a path correctly
+		Buffer++; // Skip the leading '
+		ImportedPath.Reset();
+		if (const TCHAR* NewBuffer = FPropertyHelpers::ReadToken(Buffer, /* out */ ImportedPath, /* dotted names */ true))
+		{
+			Buffer = NewBuffer;
+		}
+		else
+		{
+			Buffer = OriginalBuffer;
+			return false;
+		}
+		if (*Buffer++ != TCHAR('\''))
+		{
+			Buffer = OriginalBuffer;
+			return false;
+		}
+	}
+
+	TrySetPath(ImportedPath);
+	return true;
+}
+
 bool FTopLevelAssetPath::SerializeFromMismatchedTag(const FPropertyTag& Tag, FStructuredArchive::FSlot Slot)
 {
 	if (Tag.Type == NAME_NameProperty)
@@ -247,6 +576,24 @@ bool FTopLevelAssetPathTest::RunTest(const FString& Parameters)
 	TestFalse(TEXT("TrySetPath with malformed path string fails"), MalformedPath.TrySetPath("/Path/To/Package.Asset.Malformed"));
 	TestEqual(TEXT("Failed set to string is empty string"), MalformedPath.ToString(), FString());
 
+	FTopLevelAssetPath FromStructuredText;
+	const TCHAR* StructuredTextBuffer = TEXT("(PackageName=\"/Path/To/Package\", AssetName=\"Asset\")");
+	TestTrue(TEXT("FromStructuredText ImportTextItem succeeds"), FromStructuredText.ImportTextItem(StructuredTextBuffer, PPF_None, nullptr, GLog->Get(), nullptr));
+	TestEqual(TEXT("FromStructuredText ImportTextItem advances buffer"), *StructuredTextBuffer, '\0');
+	TestEqual(TEXT("FromStructuredText imports correct path"), FromStructuredText.ToString(), TEXT("/Path/To/Package.Asset"));
+
+	FTopLevelAssetPath FromExportTextPath;
+	const TCHAR* ExportTextBuffer = TEXT("/Script/Module.ClassName'/Path/To/Package.Asset'");
+	TestTrue(TEXT("FromExportTextPath ImportTextItem succeeds"), FromExportTextPath.ImportTextItem(ExportTextBuffer, PPF_None, nullptr, GLog->Get()));
+	TestEqual(TEXT("FromExportTextPath ImportTextItem advances buffer"), *ExportTextBuffer, '\0');
+	TestEqual(TEXT("FromExportTextPath imports correct path"), FromExportTextPath.ToString(), TEXT("/Path/To/Package.Asset"));
+
+	FTopLevelAssetPath FromShortExportTextPath;
+	const TCHAR* ShortExportTextBuffer = TEXT("ClassName'/Path/To/Package.Asset'");
+	TestTrue(TEXT("FromShortExportTextPath ImportTextItem succeeds"), FromShortExportTextPath.ImportTextItem(ShortExportTextBuffer, PPF_None, nullptr, GLog->Get()));
+	TestEqual(TEXT("FromShortExportTextPath ImportTextItem advances buffer"), *ShortExportTextBuffer , '\0');
+	TestEqual(TEXT("FromShortExportTextPath  imports correct path"), FromShortExportTextPath.ToString(), TEXT("/Path/To/Package.Asset"));
+
 	// Test ExportText path with and without root 
 	// Test starting with . 
 
diff --git a/Engine/Source/Runtime/CoreUObject/Public/UObject/TopLevelAssetPath.h b/Engine/Source/Runtime/CoreUObject/Public/UObject/TopLevelAssetPath.h
--- a/Engine/Source/Runtime/CoreUObject/Public/UObject/TopLevelAssetPath.h
+++ b/Engine/Source/Runtime/CoreUObject/Public/UObject/TopLevelAssetPath.h
@@ -167,6 +167,8 @@ struct COREUOBJECT_API FTopLevelAssetPath
 		return HashCombineFast(GetTypeHash(This.PackageName), GetTypeHash(This.AssetName));
 	}
 
+	bool ExportTextItem(FString& ValueStr, FTopLevelAssetPath const& DefaultValue, UObject* Parent, int32 PortFlags, UObject* ExportRootScope) const;
+	bool ImportTextItem( const TCHAR*& Buffer, int32 PortFlags, UObject* Parent, FOutputDevice* ErrorText, FArchive* InSerializingArchive = nullptr );
 	bool SerializeFromMismatchedTag(const FPropertyTag& Tag, FStructuredArchive::FSlot Slot);
 
 private:
