DerivedData: Enabled the derived data reference cache tests on Mac

#rnx

diff --git a/Engine/Source/Runtime/CoreUObject/Tests/Serialization/DerivedDataTest.cpp b/Engine/Source/Runtime/CoreUObject/Tests/Serialization/DerivedDataTest.cpp
--- a/Engine/Source/Runtime/CoreUObject/Tests/Serialization/DerivedDataTest.cpp
+++ b/Engine/Source/Runtime/CoreUObject/Tests/Serialization/DerivedDataTest.cpp
@@ -287,7 +287,6 @@ TEST_CASE("CoreUObject::Serialization::DerivedData", "[CoreUObject][Serializatio
 	#endif
 	}
 
-#if !PLATFORM_MAC
 	SECTION("CacheRecord")
 	{
 		constexpr uint64 HalfSize = 8 * 1024;
@@ -487,7 +486,6 @@ TEST_CASE("CoreUObject::Serialization::DerivedData", "[CoreUObject][Serializatio
 		}
 	#endif
 	}
-#endif // !PLATFORM_MAC
 #endif // WITH_EDITORONLY_DATA
 }
 
