ZenLoader: Removed editor ifdef that was supposed to have been removed in 22152724

diff --git a/Engine/Source/Runtime/CoreUObject/Private/Serialization/AsyncLoading2.cpp b/Engine/Source/Runtime/CoreUObject/Private/Serialization/AsyncLoading2.cpp
--- a/Engine/Source/Runtime/CoreUObject/Private/Serialization/AsyncLoading2.cpp
+++ b/Engine/Source/Runtime/CoreUObject/Private/Serialization/AsyncLoading2.cpp
@@ -3141,9 +3141,8 @@ struct FAsyncPackageScope2
 {
 	/** Outer scope package */
 	void* PreviousPackage;
-#if WITH_IOSTORE_IN_EDITOR
 	IAsyncPackageLoader* PreviousAsyncPackageLoader;
-#endif
+
 	/** Cached ThreadContext so we don't have to access it again */
 	FUObjectThreadContext& ThreadContext;
 
@@ -3152,17 +3151,13 @@ struct FAsyncPackageScope2
 	{
 		PreviousPackage = ThreadContext.AsyncPackage;
 		ThreadContext.AsyncPackage = InPackage;
-#if WITH_IOSTORE_IN_EDITOR
 		PreviousAsyncPackageLoader = ThreadContext.AsyncPackageLoader;
 		ThreadContext.AsyncPackageLoader = &InPackage->AsyncLoadingThread;
-#endif
 	}
 	~FAsyncPackageScope2()
 	{
 		ThreadContext.AsyncPackage = PreviousPackage;
-#if WITH_IOSTORE_IN_EDITOR
 		ThreadContext.AsyncPackageLoader = PreviousAsyncPackageLoader;
-#endif
 	}
 };
 
