Remove deprecated fields containing full object path names from when stripping editoronly content to enable memory savings.

diff --git a/Engine/Source/Runtime/AssetRegistry/Private/AssetRegistry.cpp b/Engine/Source/Runtime/AssetRegistry/Private/AssetRegistry.cpp
--- a/Engine/Source/Runtime/AssetRegistry/Private/AssetRegistry.cpp
+++ b/Engine/Source/Runtime/AssetRegistry/Private/AssetRegistry.cpp
@@ -2994,9 +2994,13 @@ void FAssetRegistryImpl::CompileFilter(Impl::FClassInheritanceContext& Inheritan
 		OutCompiledFilter.PackagePaths.Add(FPathTree::NormalizePackagePath(PackagePath));
 	}
 	OutCompiledFilter.SoftObjectPaths.Append(InFilter.SoftObjectPaths);
+
+#if WITH_EDITORONLY_DATA
 PRAGMA_DISABLE_DEPRECATION_WARNINGS
 	OutCompiledFilter.SoftObjectPaths.Append(UE::SoftObjectPath::Private::ConvertObjectPathNames(InFilter.ObjectPaths));
 PRAGMA_ENABLE_DEPRECATION_WARNINGS
+#endif
+
 PRAGMA_DISABLE_DEPRECATION_WARNINGS
 	if (!ensureAlwaysMsgf(InFilter.ClassNames.Num() == 0, TEXT("Asset Registry Filter using ClassNames instead of ClassPaths. First class name: \"%s\""), *InFilter.ClassNames[0].ToString()))
 	{
diff --git a/Engine/Source/Runtime/AssetRegistry/Private/AssetRegistryArchive.cpp b/Engine/Source/Runtime/AssetRegistry/Private/AssetRegistryArchive.cpp
--- a/Engine/Source/Runtime/AssetRegistry/Private/AssetRegistryArchive.cpp
+++ b/Engine/Source/Runtime/AssetRegistry/Private/AssetRegistryArchive.cpp
@@ -54,6 +54,7 @@ static void LoadBundleEntries(FArchive& Ar, TArray<FAssetBundleEntry>& Entries)
 			Path = TmpPath.GetAssetPath();
 		}
 
+#if WITH_EDITORONLY_DATA
 PRAGMA_DISABLE_DEPRECATION_WARNINGS
 		Entry.BundleAssets.Reserve(Entry.AssetPaths.Num());
 		for (const FTopLevelAssetPath& Path : Entry.AssetPaths)
@@ -61,6 +62,7 @@ PRAGMA_DISABLE_DEPRECATION_WARNINGS
 			Entry.BundleAssets.Add(FSoftObjectPath(Path, {}));
 		}
 PRAGMA_ENABLE_DEPRECATION_WARNINGS
+#endif
 	}
 }
 
@@ -100,6 +102,7 @@ PRAGMA_ENABLE_DEPRECATION_WARNINGS
 			}
 		}
 		
+#if WITH_EDITORONLY_DATA
 PRAGMA_DISABLE_DEPRECATION_WARNINGS
 		Entry.BundleAssets.Reserve(Entry.AssetPaths.Num());
 		for (const FTopLevelAssetPath& Path : Entry.AssetPaths)
@@ -107,6 +110,7 @@ PRAGMA_DISABLE_DEPRECATION_WARNINGS
 			Entry.BundleAssets.Add(FSoftObjectPath(Path, {}));
 		}
 PRAGMA_ENABLE_DEPRECATION_WARNINGS
+#endif
 	}
 }
 
diff --git a/Engine/Source/Runtime/CoreUObject/Private/AssetRegistry/AssetBundleData.cpp b/Engine/Source/Runtime/CoreUObject/Private/AssetRegistry/AssetBundleData.cpp
--- a/Engine/Source/Runtime/CoreUObject/Private/AssetRegistry/AssetBundleData.cpp
+++ b/Engine/Source/Runtime/CoreUObject/Private/AssetRegistry/AssetBundleData.cpp
@@ -243,6 +243,7 @@ bool FAssetBundleEntry::ImportTextItem(const TCHAR*& Buffer, int32 PortFlags, UO
 		AssetPaths.Add(Path.GetAssetPath());
 	}
 
+#if WITH_EDITORONLY_DATA
 PRAGMA_DISABLE_DEPRECATION_WARNINGS
 	BundleAssets.Reset(AssetPaths.Num());
 	for (const FTopLevelAssetPath& Path : AssetPaths)
@@ -250,6 +251,7 @@ PRAGMA_DISABLE_DEPRECATION_WARNINGS
 		BundleAssets.Add(FSoftObjectPath(Path, {}));
 	}
 PRAGMA_ENABLE_DEPRECATION_WARNINGS
+#endif
 
 	// Success, update how much we read for the caller
 	Buffer = BufferIt;
@@ -283,9 +285,11 @@ void FAssetBundleData::AddBundleAsset(FName BundleName, const FTopLevelAssetPath
 		FoundEntry = new(Bundles) FAssetBundleEntry(BundleName);
 	}
 
+#if WITH_EDITORONLY_DATA
 PRAGMA_DISABLE_DEPRECATION_WARNINGS
 	FoundEntry->BundleAssets.AddUnique(FSoftObjectPath(AssetPath, {}));
 PRAGMA_ENABLE_DEPRECATION_WARNINGS
+#endif
 	FoundEntry->AssetPaths.AddUnique(AssetPath);
 }
 
@@ -303,9 +307,11 @@ void FAssetBundleData::AddBundleAssets(FName BundleName, const TArray<FTopLevelA
 				FoundEntry = new(Bundles) FAssetBundleEntry(BundleName);
 			}
 
+#if WITH_EDITORONLY_DATA
 PRAGMA_DISABLE_DEPRECATION_WARNINGS
 			FoundEntry->BundleAssets.AddUnique(FSoftObjectPath(Path, {}));
 PRAGMA_ENABLE_DEPRECATION_WARNINGS
+#endif
 			FoundEntry->AssetPaths.AddUnique(Path);
 		}
 	}
@@ -321,10 +327,12 @@ void FAssetBundleData::SetBundleAssets(FName BundleName, TArray<FTopLevelAssetPa
 	}
 
 	FoundEntry->AssetPaths = MoveTemp(AssetPaths);
+#if WITH_EDITORONLY_DATA
 PRAGMA_DISABLE_DEPRECATION_WARNINGS
 	FoundEntry->BundleAssets.Reset();
 	Algo::Transform(FoundEntry->AssetPaths, FoundEntry->BundleAssets, [](FTopLevelAssetPath Path) { return FSoftObjectPath(Path, {}); });
 PRAGMA_ENABLE_DEPRECATION_WARNINGS
+#endif
 }
 
 void FAssetBundleData::AddBundleAsset(FName BundleName, const FSoftObjectPath& AssetPath)
@@ -341,9 +349,11 @@ void FAssetBundleData::AddBundleAsset(FName BundleName, const FSoftObjectPath& A
 		FoundEntry = new(Bundles) FAssetBundleEntry(BundleName);
 	}
 
+#if WITH_EDITORONLY_DATA
 PRAGMA_DISABLE_DEPRECATION_WARNINGS
 	FoundEntry->BundleAssets.AddUnique(AssetPath);
 PRAGMA_ENABLE_DEPRECATION_WARNINGS
+#endif
 	FoundEntry->AssetPaths.AddUnique(AssetPath.GetAssetPath());
 }
 
@@ -361,9 +371,11 @@ void FAssetBundleData::AddBundleAssets(FName BundleName, const TArray<FSoftObjec
 				FoundEntry = new(Bundles) FAssetBundleEntry(BundleName);
 			}
 
+#if WITH_EDITORONLY_DATA
 PRAGMA_DISABLE_DEPRECATION_WARNINGS
 			FoundEntry->BundleAssets.AddUnique(Path);
 PRAGMA_ENABLE_DEPRECATION_WARNINGS
+#endif
 			FoundEntry->AssetPaths.AddUnique(Path.GetAssetPath());
 		}
 	}
@@ -384,10 +396,12 @@ void FAssetBundleData::SetBundleAssets(FName BundleName, TArray<FSoftObjectPath>
 		FoundEntry->AssetPaths.Add(Path.GetAssetPath());
 	}
 
+#if WITH_EDITORONLY_DATA
 PRAGMA_DISABLE_DEPRECATION_WARNINGS
 	FoundEntry->BundleAssets.Reset();
 	Algo::Transform(FoundEntry->AssetPaths, FoundEntry->BundleAssets, [](FTopLevelAssetPath Path) { return FSoftObjectPath(Path, {}); });
 PRAGMA_ENABLE_DEPRECATION_WARNINGS
+#endif
 }
 
 void FAssetBundleData::AddBundleAssetTruncated(FName BundleName, const FSoftObjectPath& AssetPath)
@@ -404,9 +418,11 @@ void FAssetBundleData::AddBundleAssetTruncated(FName BundleName, const FSoftObje
 		FoundEntry = new(Bundles) FAssetBundleEntry(BundleName);
 	}
 
+#if WITH_EDITORONLY_DATA
 PRAGMA_DISABLE_DEPRECATION_WARNINGS
 	FoundEntry->BundleAssets.AddUnique(AssetPath.GetWithoutSubPath());
 PRAGMA_ENABLE_DEPRECATION_WARNINGS
+#endif
 	FoundEntry->AssetPaths.AddUnique(AssetPath.GetAssetPath());
 }
 
@@ -424,9 +440,11 @@ void FAssetBundleData::AddBundleAssetsTruncated(FName BundleName, const TArray<F
 				FoundEntry = new(Bundles) FAssetBundleEntry(BundleName);
 			}
 
+#if WITH_EDITORONLY_DATA
 PRAGMA_DISABLE_DEPRECATION_WARNINGS
 			FoundEntry->BundleAssets.AddUnique(Path.GetWithoutSubPath());
 PRAGMA_ENABLE_DEPRECATION_WARNINGS
+#endif
 			FoundEntry->AssetPaths.AddUnique(Path.GetAssetPath());
 		}
 	}
@@ -447,10 +465,12 @@ void FAssetBundleData::SetBundleAssetsTruncated(FName BundleName, const TArray<F
 		FoundEntry->AssetPaths.Add(Path.GetAssetPath());
 	}
 
+#if WITH_EDITORONLY_DATA
 PRAGMA_DISABLE_DEPRECATION_WARNINGS
 	FoundEntry->BundleAssets.Reset();
 	Algo::Transform(FoundEntry->AssetPaths, FoundEntry->BundleAssets, [](FTopLevelAssetPath Path) { return FSoftObjectPath(Path); });
 PRAGMA_ENABLE_DEPRECATION_WARNINGS
+#endif 
 }
 
 void FAssetBundleData::Reset()
diff --git a/Engine/Source/Runtime/CoreUObject/Private/AssetRegistry/AssetData.cpp b/Engine/Source/Runtime/CoreUObject/Private/AssetRegistry/AssetData.cpp
--- a/Engine/Source/Runtime/CoreUObject/Private/AssetRegistry/AssetData.cpp
+++ b/Engine/Source/Runtime/CoreUObject/Private/AssetRegistry/AssetData.cpp
@@ -242,9 +242,11 @@ FAssetData::FAssetData(FName InPackageName, FName InPackagePath, FName InAssetNa
 	FNameBuilder ObjectPathStr(PackageName);
 	ObjectPathStr << TEXT('.');
 	AssetName.AppendString(ObjectPathStr);
+#if WITH_EDITORONLY_DATA
 PRAGMA_DISABLE_DEPRECATION_WARNINGS
 	ObjectPath = FName(FStringView(ObjectPathStr));
 PRAGMA_ENABLE_DEPRECATION_WARNINGS
+#endif
 
 	SetChunkIDs(InChunkIDs);
 }
@@ -255,9 +257,12 @@ FAssetData::FAssetData(const FString& InLongPackageName, const FString& InObject
 	, PackageFlags(InPackageFlags)
 {
 	using namespace UE::AssetRegistry::Private;
+
+#if WITH_EDITORONLY_DATA
 PRAGMA_DISABLE_DEPRECATION_WARNINGS
 	ObjectPath = *InObjectPath;
 PRAGMA_ENABLE_DEPRECATION_WARNINGS
+#endif
 
 	SetTagsAndAssetBundles(MoveTemp(InTags));
 
@@ -297,10 +302,10 @@ FAssetData::FAssetData(const UObject* InAsset, FAssetData::ECreationFlags InCrea
 		PackagePath = FName(*FPackageName::GetLongPackagePath(Package->GetName()));
 		AssetName = InAsset->GetFName();
 		AssetClassPath = InAsset->GetClass()->GetPathName();
+#if WITH_EDITORONLY_DATA
 PRAGMA_DISABLE_DEPRECATION_WARNINGS
 		ObjectPath = FName(*InAsset->GetPathName());
 PRAGMA_ENABLE_DEPRECATION_WARNINGS
-#if WITH_EDITORONLY_DATA
 		if (InAsset->GetOuter() != Package)
 		{
 			OptionalOuterPath = *InAsset->GetOuter()->GetPathName();
@@ -564,12 +569,14 @@ PRAGMA_ENABLE_DEPRECATION_WARNINGS
 
 	Ar << PackageFlags;
 
+#if WITH_EDITORONLY_DATA
 	// Rebuild deprecated ObjectPath field 
 	TStringBuilder<FName::StringBufferSize> Builder;
 	AppendObjectPath(Builder);
 PRAGMA_DISABLE_DEPRECATION_WARNINGS
 	ObjectPath = *Builder;
 PRAGMA_ENABLE_DEPRECATION_WARNINGS
+#endif
 }
 
 void FAssetData::NetworkWrite(FCbWriter& Writer, bool bWritePackageName) const
@@ -683,12 +690,14 @@ bool FAssetData::TryNetworkRead(FCbFieldView Field, bool bReadPackageName, FName
 	check(Algo::IsSorted(SerializedChunkIDs));
 	ChunkArrayRegistryHandle = UE::AssetRegistry::GChunkArrayRegistry.FindOrAddSorted(MoveTemp(SerializedChunkIDs));
 
+#if WITH_EDITORONLY_DATA
 	// Rebuild deprecated ObjectPath field 
 	TStringBuilder<FName::StringBufferSize> Builder;
 	AppendObjectPath(Builder);
 PRAGMA_DISABLE_DEPRECATION_WARNINGS
 	ObjectPath = *Builder;
 PRAGMA_ENABLE_DEPRECATION_WARNINGS
+#endif
 
 	return bOk;
 }
diff --git a/Engine/Source/Runtime/CoreUObject/Public/AssetRegistry/ARFilter.h b/Engine/Source/Runtime/CoreUObject/Public/AssetRegistry/ARFilter.h
--- a/Engine/Source/Runtime/CoreUObject/Public/AssetRegistry/ARFilter.h
+++ b/Engine/Source/Runtime/CoreUObject/Public/AssetRegistry/ARFilter.h
@@ -31,9 +31,11 @@ PRAGMA_ENABLE_DEPRECATION_WARNINGS
 	/** The filter component for package paths */
 	TArray<FName> PackagePaths;
 
+#if WITH_EDITORONLY_DATA
 	/** The filter component containing specific object paths */
 	UE_DEPRECATED(5.1, "Asset path FNames have been deprecated, use FSoftObjectPath instead.")
 	TArray<FName> ObjectPaths;
+#endif
 
 	/** 
 	 * The filter component containing the paths of specific assets to match. 
diff --git a/Engine/Source/Runtime/CoreUObject/Public/AssetRegistry/AssetBundleData.h b/Engine/Source/Runtime/CoreUObject/Public/AssetRegistry/AssetBundleData.h
--- a/Engine/Source/Runtime/CoreUObject/Public/AssetRegistry/AssetBundleData.h
+++ b/Engine/Source/Runtime/CoreUObject/Public/AssetRegistry/AssetBundleData.h
@@ -35,9 +35,11 @@ PRAGMA_ENABLE_DEPRECATION_WARNINGS
 	/** Specific name of this bundle, should be unique for a given scope */
 	FName BundleName;
 
+#if WITH_EDITORONLY_DATA
 	/** List of string assets contained in this bundle */
 	UE_DEPRECATED(5.1, "BundleAssets has been replaced with AssetPaths.")
 	TArray<FSoftObjectPath> BundleAssets;
+#endif
 
 	/** List of references to top-level assets contained in this bundle */
 	TArray<FTopLevelAssetPath> AssetPaths;
diff --git a/Engine/Source/Runtime/CoreUObject/Public/AssetRegistry/AssetData.h b/Engine/Source/Runtime/CoreUObject/Public/AssetRegistry/AssetData.h
--- a/Engine/Source/Runtime/CoreUObject/Public/AssetRegistry/AssetData.h
+++ b/Engine/Source/Runtime/CoreUObject/Public/AssetRegistry/AssetData.h
@@ -154,9 +154,11 @@ public:
 	};
 
 public:
+#if WITH_EDITORONLY_DATA
 	/** The object path for the asset in the form PackageName.ObjectName, or PackageName.ObjectName:SubObjectName. */
 	UE_DEPRECATED(5.1, "FName asset paths have been deprecated. Use GetSoftObjectPath to get the path this asset will use in memory when loaded.")
 	FName ObjectPath;
+#endif
 	/** The name of the package in which the asset is found, this is the full long package name such as /Game/Path/Package */
 	FName PackageName;
 	/** The path to the package in which the asset is found, this is /Game/Path with the Package stripped off */
diff --git a/Engine/Source/Runtime/CoreUObject/Public/UObject/NoExportTypes.h b/Engine/Source/Runtime/CoreUObject/Public/UObject/NoExportTypes.h
--- a/Engine/Source/Runtime/CoreUObject/Public/UObject/NoExportTypes.h
+++ b/Engine/Source/Runtime/CoreUObject/Public/UObject/NoExportTypes.h
@@ -2474,9 +2474,11 @@ struct FARFilter
 	UPROPERTY(EditDefaultsOnly, BlueprintReadWrite, Category = AssetRegistry)
 	TArray<FName> PackagePaths;
 
+#if WITH_EDITORONLY_DATA
 	/** The filter component containing specific object paths. Deprecated. */
-	UPROPERTY(EditDefaultsOnly, BlueprintReadWrite, Category = AssetRegistry, meta=(DeprecatedProperty, DeprecationMessage="Names containing full asset paths are deprecated. Use SoftObjectPaths instead."))
+	UPROPERTY()
 	TArray<FName> ObjectPaths;
+#endif
 
 	/** The filter component containing specific object paths */
 	UPROPERTY(EditDefaultsOnly, BlueprintReadWrite, Category = AssetRegistry)
@@ -2527,9 +2529,11 @@ struct COREUOBJECT_API FAssetBundleEntry
 	UPROPERTY()
 	FName BundleName;
 
+#if WITH_EDITORONLY_DATA
 	/** List of string assets contained in this bundle */
 	UPROPERTY(meta=(DeprecatedProperty, DeprecationMessage="Asset bundles may only contain top level asset paths which are referenced through the AssetPaths property."))
 	TArray<FSoftObjectPath> BundleAssets;
+#endif
 
 	UPROPERTY()
 	TArray<FTopLevelAssetPath> AssetPaths;
@@ -2551,8 +2555,10 @@ struct COREUOBJECT_API FAssetBundleData
 USTRUCT(noexport, BlueprintType, IsAlwaysAccessible, HasDefaults)
 struct FAssetData
 {
-	UPROPERTY(BlueprintReadOnly, Category = AssetData, transient, meta=(DeprecatedProperty, DeprecationMessage="Names containing full asset paths are deprecated. Use 'To Soft Object Path' instead."))
+#if WITH_EDITORONLY_DATA
+	UPROPERTY()
 	FName ObjectPath;
+#endif
 	/** The name of the package in which the asset is found, this is the full long package name such as /Game/Path/Package */
 	UPROPERTY(BlueprintReadOnly, Category = AssetData, transient)
 	FName PackageName;
