AsyncLoading: InstanceGraph: Make sure to never reconstruct found existing destination subobjects in cooked packages, they should always have been created from the correct up-to-date template already (at least after cl 19308998).
We are never updating archetypes when loading cooked packages, and we can't safely run the reconstruct logic with UObject destruction from the async loading thread.

#rnx

diff --git a/Engine/Source/Runtime/CoreUObject/Private/UObject/CoreNative.cpp b/Engine/Source/Runtime/CoreUObject/Private/UObject/CoreNative.cpp
--- a/Engine/Source/Runtime/CoreUObject/Private/UObject/CoreNative.cpp
+++ b/Engine/Source/Runtime/CoreUObject/Private/UObject/CoreNative.cpp
@@ -115,6 +115,14 @@ void FObjectInstancingGraph::SetDestinationRoot(UObject* DestinationSubobjectRoo
 	SourceToDestinationMap.Add(SourceRoot, DestinationRoot);
 
 	bCreatingArchetype = DestinationSubobjectRoot->HasAnyFlags(RF_ArchetypeObject);
+	if (DestinationSubobjectRoot->GetPackage()->HasAnyPackageFlags(PKG_Cooked))
+	{
+		//We are never updating archetypes when loading cooked packages,
+		//and we can't safely run the reconstruct logic with UObject destruction from the async loading thread.
+		//Make sure to never reconstruct found existing destination subobjects in cooked packages,
+		//they should always have been created from the correct up-to-date template already.
+		bCreatingArchetype = false;
+	}
 }
 
 UObject* FObjectInstancingGraph::GetDestinationObject(UObject* SourceObject)
@@ -196,13 +204,9 @@ UObject* FObjectInstancingGraph::GetInstancedSubobject( UObject* SourceSubobject
 
 							FName SubobjectName = SourceSubobject->GetFName();
 
-							// final archetype archetype will be the archetype of the template
-							UObject* FinalSubobjectArchetype = CurrentValue->GetArchetype();
-
 							// Don't search for the existing subobjects on Blueprint-generated classes. What we'll find is a subobject
 							// created by the constructor which may not have all of its fields initialized to the correct value (which
 							// should be coming from a blueprint).
-							// NOTE: Since this function is called ONLY for Blueprint-generated classes, we may as well delete this 'if'.
 							if (!SubobjectOuter->GetClass()->HasAnyClassFlags(CLASS_CompiledFromBlueprint))
 							{
 								InstancedSubobject = StaticFindObjectFast(nullptr, SubobjectOuter, SubobjectName);
