ZenLoader: Optimized management of sync load context packages

diff --git a/Engine/Source/Runtime/CoreUObject/Private/Serialization/AsyncLoading2.cpp b/Engine/Source/Runtime/CoreUObject/Private/Serialization/AsyncLoading2.cpp
--- a/Engine/Source/Runtime/CoreUObject/Private/Serialization/AsyncLoading2.cpp
+++ b/Engine/Source/Runtime/CoreUObject/Private/Serialization/AsyncLoading2.cpp
@@ -316,6 +316,14 @@ static FAutoConsoleVariableRef CVarGRemoveUnreachableObjectsFromGCNotifyOnGT(
 	ECVF_Default
 	);
 
+static bool GOnlyProcessRequiredPackagesWhenSyncLoading = true;
+static FAutoConsoleVariableRef CVarGOnlyProcessRequiredPackagesWhenSyncLoading(
+	TEXT("s.OnlyProcessRequiredPackagesWhenSyncLoading"),
+	GOnlyProcessRequiredPackagesWhenSyncLoading,
+	TEXT("When sync loading a package process only that package and its imports"),
+	ECVF_Default
+);
+
 CSV_DECLARE_CATEGORY_MODULE_EXTERN(CORE_API, Basic);
 CSV_DECLARE_CATEGORY_MODULE_EXTERN(CORE_API, FileIO);
 
@@ -1779,12 +1787,15 @@ public:
 	FAsyncLoadingSyncLoadContext(int32 InRequestId)
 		: RequestId(InRequestId)
 	{
-
+		ContextId = NextContextId++;
+		if (NextContextId == 0)
+		{
+			NextContextId = 1;
+		}
 	}
 
 	~FAsyncLoadingSyncLoadContext()
 	{
-		check(Packages.IsEmpty());
 	}
 
 	void AddRef()
@@ -1802,17 +1813,20 @@ public:
 		}
 	}
 
+	uint64 ContextId;
 	int32 RequestId;
 	FAsyncPackage2* RequestedPackageDebug = nullptr;
 	FAsyncPackage2* RequestingPackageDebug = nullptr;
-	FCriticalSection CriticalSection;
-	TSet<FAsyncPackage2*> Packages;
 	bool bHasFoundRequestedPackage = false;
 
 private:
 	std::atomic<int32> RefCount = 1;
+
+	static uint64 NextContextId;
 };
 
+uint64 FAsyncLoadingSyncLoadContext::NextContextId = 1;
+
 struct FAsyncLoadingThreadState2
 {
 	static void Set(FAsyncLoadingThreadState2* State)
@@ -2310,7 +2324,7 @@ private:
 	/** Package which is going to have its exports and imports loaded */
 	UPackage* LinkerRoot = nullptr;
 	// The sync load context associated with this package
-	FAsyncLoadingSyncLoadContext* SyncLoadContext = nullptr;
+	std::atomic<uint64>			SyncLoadContextId = 0;
 	/** Time load begun. This is NOT the time the load was requested in the case of pending requests.	*/
 	double						LoadStartTime = 0.0;
 	TAtomic<int32> RefCount{ 0 };
@@ -2985,9 +2999,7 @@ private:
 	*/
 	EAsyncPackageState::Type ProcessLoadedPackagesFromGameThread(FAsyncLoadingThreadState2& ThreadState, bool& bDidSomething, int32 FlushRequestID = INDEX_NONE);
 
-	void IncludePackageInSyncLoadContext(FAsyncLoadingSyncLoadContext& SyncContext, FAsyncPackage2* Package);
-	void DetachPackageFromSyncLoadContext(FAsyncPackage2* Package);
-	void DetachAllPackagesFromSyncLoadContext(FAsyncLoadingSyncLoadContext& SyncContext);
+	void IncludePackageInSyncLoadContext(uint64 ContextId, FAsyncPackage2* Package);
 	FAsyncLoadingSyncLoadContext* UpdateSyncLoadContext(FAsyncLoadingThreadState2& ThreadState);
 
 	bool CreateAsyncPackagesFromQueue(FAsyncLoadingThreadState2& ThreadState);
@@ -3121,10 +3133,6 @@ private:
 		{
 			RequestIdToPackageMap.Remove(RequestId);
 		}
-		if (Package->SyncLoadContext)
-		{
-			DetachPackageFromSyncLoadContext(Package);
-		}
 		delete Package;
 		--PackagesWithRemainingWorkCounter;
 	}
@@ -3263,18 +3271,11 @@ FAsyncPackage2* FAsyncLoadingThread2::FindOrInsertPackage(FAsyncPackageDesc2& De
 	return Package;
 }
 
-void FAsyncLoadingThread2::IncludePackageInSyncLoadContext(FAsyncLoadingSyncLoadContext& SyncLoadContext, FAsyncPackage2* Package)
+void FAsyncLoadingThread2::IncludePackageInSyncLoadContext(uint64 ContextId, FAsyncPackage2* Package)
 {
-	if (Package->SyncLoadContext)
+	if (Package->SyncLoadContextId == ContextId)
 	{
-		if (Package->SyncLoadContext != &SyncLoadContext)
-		{
-			DetachPackageFromSyncLoadContext(Package);
-		}
-		else
-		{
-			return;
-		}
+		return;
 	}
 
 	if (Package->AsyncPackageLoadingState >= EAsyncPackageLoadingState2::DeferredPostLoadDone)
@@ -3282,48 +3283,16 @@ void FAsyncLoadingThread2::IncludePackageInSyncLoadContext(FAsyncLoadingSyncLoad
 		return;
 	}
 
-	// Assumes that SyncLoadContext.CriticalSection is already locked by the caller
-	Package->SyncLoadContext = &SyncLoadContext;
-	SyncLoadContext.Packages.Add(Package);
+	Package->SyncLoadContextId = ContextId;
 	for (FAsyncPackage2* ImportedPackage : Package->Data.ImportedAsyncPackages)
 	{
-		if (ImportedPackage && ImportedPackage->SyncLoadContext != &SyncLoadContext)
+		if (ImportedPackage && ImportedPackage->SyncLoadContextId != ContextId)
 		{
-			IncludePackageInSyncLoadContext(SyncLoadContext, ImportedPackage);
+			IncludePackageInSyncLoadContext(ContextId, ImportedPackage);
 		}
 	}
 }
 
-void FAsyncLoadingThread2::DetachPackageFromSyncLoadContext(FAsyncPackage2* Package)
-{
-	FAsyncLoadingSyncLoadContext* SyncLoadContext = Package->SyncLoadContext;
-	check(SyncLoadContext);
-	FScopeLock Lock(&SyncLoadContext->CriticalSection);
-	SyncLoadContext->Packages.Remove(Package);
-	if (SyncLoadContext->RequestedPackageDebug == Package)
-	{
-		SyncLoadContext->RequestedPackageDebug = nullptr;
-	}
-	if (SyncLoadContext->RequestingPackageDebug == Package)
-	{
-		SyncLoadContext->RequestingPackageDebug = nullptr;
-	}
-	Package->SyncLoadContext = nullptr;
-}
-
-void FAsyncLoadingThread2::DetachAllPackagesFromSyncLoadContext(FAsyncLoadingSyncLoadContext& SyncLoadContext)
-{
-	FScopeLock Lock(&SyncLoadContext.CriticalSection);
-	for (FAsyncPackage2* Package : SyncLoadContext.Packages)
-	{
-		check(Package->SyncLoadContext == &SyncLoadContext);
-		Package->SyncLoadContext = nullptr;
-	}
-	SyncLoadContext.Packages.Empty();
-	SyncLoadContext.RequestedPackageDebug = nullptr;
-	SyncLoadContext.RequestingPackageDebug = nullptr;
-}
-
 bool FAsyncLoadingThread2::CreateAsyncPackagesFromQueue(FAsyncLoadingThreadState2& ThreadState)
 {
 	TRACE_CPUPROFILER_EVENT_SCOPE(CreateAsyncPackagesFromQueue);
@@ -3767,8 +3736,7 @@ bool FAsyncLoadEventQueue2::ExecuteSyncLoadEvents(FAsyncLoadingThreadState2& Thr
 
 	auto ShouldExecuteNode = [&SyncLoadContext](FEventLoadNode2* Node) -> bool
 	{
-		FScopeLock Lock(&SyncLoadContext.CriticalSection);
-		return SyncLoadContext.Packages.Contains(Node->Package);
+		return Node->Package->SyncLoadContextId == SyncLoadContext.ContextId;
 	};
 
 	bool bDidSomething = false;
@@ -4240,14 +4208,13 @@ void FAsyncPackage2::ImportPackagesRecursive(FAsyncLoadingThreadState2& ThreadSt
 	}
 #endif
 
-	if (SyncLoadContext)
+	if (SyncLoadContextId)
 	{
-		FScopeLock Lock(&SyncLoadContext->CriticalSection);
 		for (FAsyncPackage2* ImportedPackage : Data.ImportedAsyncPackages)
 		{
 			if (ImportedPackage)
 			{
-				AsyncLoadingThread.IncludePackageInSyncLoadContext(*SyncLoadContext, ImportedPackage);
+				AsyncLoadingThread.IncludePackageInSyncLoadContext(SyncLoadContextId, ImportedPackage);
 			}
 		}
 	}
@@ -6077,7 +6044,6 @@ FAsyncLoadingSyncLoadContext* FAsyncLoadingThread2::UpdateSyncLoadContext(FAsync
 	FAsyncLoadingSyncLoadContext* SyncLoadContext = ThreadState.SyncLoadContextStack.Top();
 	if (ThreadState.bIsAsyncLoadingThread && !ContainsRequestID(SyncLoadContext->RequestId))
 	{
-		DetachAllPackagesFromSyncLoadContext(*SyncLoadContext);
 		SyncLoadContext->ReleaseRef();
 		ThreadState.SyncLoadContextStack.Pop();
 		if (ThreadState.SyncLoadContextStack.IsEmpty())
@@ -6090,10 +6056,9 @@ FAsyncLoadingSyncLoadContext* FAsyncLoadingThread2::UpdateSyncLoadContext(FAsync
 	{
 		if (FAsyncPackage2* RequestedPackage = RequestIdToPackageMap.FindRef(SyncLoadContext->RequestId))
 		{
-			FScopeLock Lock(&SyncLoadContext->CriticalSection);
 			SyncLoadContext->bHasFoundRequestedPackage = true;
 			SyncLoadContext->RequestedPackageDebug = RequestedPackage;
-			IncludePackageInSyncLoadContext(*SyncLoadContext, RequestedPackage);
+			IncludePackageInSyncLoadContext(SyncLoadContext->ContextId, RequestedPackage);
 		}
 	}
 	return SyncLoadContext;
@@ -7687,7 +7652,7 @@ void FAsyncLoadingThread2::FlushLoading(int32 RequestId)
 		}
 
 		FAsyncLoadingSyncLoadContext* SyncLoadContext = nullptr;
-		if (RequestId != INDEX_NONE)
+		if (RequestId != INDEX_NONE && GOnlyProcessRequiredPackagesWhenSyncLoading)
 		{
 			SyncLoadContext = new FAsyncLoadingSyncLoadContext(RequestId);
 			SyncLoadContext->RequestingPackageDebug = CurrentlyExecutingPackage;
@@ -7743,10 +7708,6 @@ void FAsyncLoadingThread2::FlushLoading(int32 RequestId)
 		if (SyncLoadContext)
 		{
 			check(GameThreadState->SyncLoadContextStack.Top() == SyncLoadContext);
-			if (GameThreadState->bCanAccessAsyncLoadingThreadData)
-			{
-				DetachAllPackagesFromSyncLoadContext(*SyncLoadContext);
-			}
 			SyncLoadContext->ReleaseRef();
 			GameThreadState->SyncLoadContextStack.Pop();
 		}
