Make UClass::FindFunctionByName thread-safe

diff --git a/Engine/Source/Runtime/CoreUObject/Private/UObject/Class.cpp b/Engine/Source/Runtime/CoreUObject/Private/UObject/Class.cpp
--- a/Engine/Source/Runtime/CoreUObject/Private/UObject/Class.cpp
+++ b/Engine/Source/Runtime/CoreUObject/Private/UObject/Class.cpp
@@ -4914,7 +4914,10 @@ void UClass::Serialize( FArchive& Ar )
 
 	// serialize the function map
 	//@TODO: UCREMOVAL: Should we just regenerate the FuncMap post load, instead of serializing it?
-	Ar << FuncMap;
+	{
+		FWriteScopeLock ScopeLock(FuncMapLock);
+		Ar << FuncMap;
+	}
 
 	// Class flags first.
 	if (Ar.IsSaving())
@@ -5331,7 +5334,10 @@ void UClass::PurgeClass(bool bRecompilingOnLoad)
 	ScriptAndPropertyObjectReferences.Empty();
 	DeleteUnresolvedScriptProperties();
 
-	FuncMap.Empty();
+	{
+		FWriteScopeLock ScopeLock(FuncMapLock);
+		FuncMap.Empty();
+	}
 	ClearFunctionMapsCaches();
 	PropertyLink = nullptr;
 
@@ -5898,14 +5904,20 @@ void UClass::CreateLinkAndAddChildFunctionsToMap(const FClassFunctionLinkInfo* F
 
 void UClass::ClearFunctionMapsCaches()
 {
-	FRWScopeLock ScopeLock(SuperFuncMapLock, FRWScopeLockType::SLT_Write);
+	FWriteScopeLock ScopeLock(SuperFuncMapLock);
 	SuperFuncMap.Empty();
 }
 
 UFunction* UClass::FindFunctionByName(FName InName, EIncludeSuperFlag::Type IncludeSuper) const
 {
 	LLM_SCOPE(ELLMTag::UObject);
-	UFunction* Result = FuncMap.FindRef(InName);
+
+	UFunction* Result = nullptr;
+	{
+		FReadScopeLock ScopeLock(FuncMapLock);
+		Result = FuncMap.FindRef(InName);
+	}
+
 	if (Result == nullptr && IncludeSuper == EIncludeSuperFlag::IncludeSuper)
 	{
 		UClass* SuperClass = GetSuperClass();
@@ -5913,7 +5925,7 @@ UFunction* UClass::FindFunctionByName(FName InName, EIncludeSuperFlag::Type Incl
 		{
 			bool bFoundInSuperFuncMap = false;
 			{
-				FRWScopeLock ScopeLock(SuperFuncMapLock, FRWScopeLockType::SLT_ReadOnly);
+				FReadScopeLock ScopeLock(SuperFuncMapLock);
 				if (UFunction** SuperResult = SuperFuncMap.Find(InName))
 				{
 					Result = *SuperResult;
@@ -5937,7 +5949,7 @@ UFunction* UClass::FindFunctionByName(FName InName, EIncludeSuperFlag::Type Incl
 					Result = SuperClass->FindFunctionByName(InName);
 				}
 
-				FRWScopeLock ScopeLock(SuperFuncMapLock, FRWScopeLockType::SLT_Write);
+				FWriteScopeLock ScopeLock(SuperFuncMapLock);
 				SuperFuncMap.Add(InName, Result);
 			}
 		}
diff --git a/Engine/Source/Runtime/CoreUObject/Public/UObject/Class.h b/Engine/Source/Runtime/CoreUObject/Public/UObject/Class.h
--- a/Engine/Source/Runtime/CoreUObject/Public/UObject/Class.h
+++ b/Engine/Source/Runtime/CoreUObject/Public/UObject/Class.h
@@ -2679,6 +2679,7 @@ public:
 #if WITH_EDITOR
 	void GenerateFunctionList(TArray<FName>& OutArray) const 
 	{ 
+		FReadScopeLock ScopeLock(FuncMapLock);
 		FuncMap.GenerateKeyArray(OutArray); 
 	}
 #endif // WITH_EDITOR
@@ -2697,6 +2698,9 @@ private:
 	/** Map of all functions by name contained in this class */
 	TMap<FName, UFunction*> FuncMap;
 
+	/** Scope lock to avoid the FuncMap being read and written to simultaneously on multiple threads. */
+	mutable FRWLock FuncMapLock;
+
 	/** A cache of all functions by name that exist in a parent (superclass or interface) context */
 	mutable TMap<FName, UFunction*> SuperFuncMap;
 
@@ -2789,6 +2793,7 @@ public:
 	/** Add a function to the function map */
 	void AddFunctionToFunctionMap(UFunction* Function, FName FuncName)
 	{
+		FWriteScopeLock ScopeLock(FuncMapLock);
 		FuncMap.Add(FuncName, Function);
 	}
 
@@ -2797,6 +2802,7 @@ public:
 	/** Remove a function from the function map */
 	void RemoveFunctionFromFunctionMap(UFunction* Function)
 	{
+		FWriteScopeLock ScopeLock(FuncMapLock);
 		FuncMap.Remove(Function->GetFName());
 	}
 
diff --git a/Engine/Source/Runtime/Engine/Private/Animation/ExposedValueHandler.cpp b/Engine/Source/Runtime/Engine/Private/Animation/ExposedValueHandler.cpp
--- a/Engine/Source/Runtime/Engine/Private/Animation/ExposedValueHandler.cpp
+++ b/Engine/Source/Runtime/Engine/Private/Animation/ExposedValueHandler.cpp
@@ -33,8 +33,6 @@ void FExposedValueHandler::Initialize(UClass* InClass, const FPropertyAccessLibr
 		if (Function == nullptr)
 #endif
 		{
-			// we cant call FindFunction on anything but the game thread as it accesses a shared map in the object's class
-			check(IsInGameThread());
 			Function = InClass->FindFunctionByName(BoundFunction);
 			check(Function);
 		}
