DerivedData: Extracted private FCookedData from FDerivedData

The new FCookedData has been used to clean up the FLinkerSave API and the implementation of FDerivedData.

#rnx

diff --git a/Engine/Source/Runtime/CoreUObject/Private/Serialization/DerivedData.cpp b/Engine/Source/Runtime/CoreUObject/Private/Serialization/DerivedData.cpp
--- a/Engine/Source/Runtime/CoreUObject/Private/Serialization/DerivedData.cpp
+++ b/Engine/Source/Runtime/CoreUObject/Private/Serialization/DerivedData.cpp
@@ -27,6 +27,48 @@
 
 namespace UE::DerivedData::Private { class FIoResponseDispatcher; }
 
+namespace UE::DerivedData::Private
+{
+
+static FStringBuilderBase& operator<<(FStringBuilderBase& Builder, const FCookedData& CookedData)
+{
+	Builder << TEXTVIEW("Chunk: ID ");
+	String::BytesToHexLower(MakeArrayView(CookedData.ChunkId), Builder);
+	if (CookedData.ChunkOffset)
+	{
+		Builder << TEXTVIEW(" / Offset ") << CookedData.ChunkOffset;
+	}
+	if (CookedData.ChunkSize != MAX_uint64)
+	{
+		Builder << TEXTVIEW(" / Size ") << CookedData.ChunkSize;
+	}
+	return Builder;
+}
+
+bool FCookedData::ReferenceEquals(const FCookedData& Other) const
+{
+	return ChunkOffset == Other.ChunkOffset && ChunkSize == Other.ChunkSize && MakeArrayView(ChunkId) == Other.ChunkId;
+}
+
+uint32 FCookedData::ReferenceHash() const
+{
+	FXxHash64Builder Builder;
+	Builder.Update(&ChunkOffset, sizeof(ChunkOffset));
+	Builder.Update(&ChunkSize, sizeof(ChunkSize));
+	Builder.Update(MakeMemoryView(ChunkId));
+	return uint32(Builder.Finalize().Hash);
+}
+
+void FCookedData::Serialize(FArchive& Ar)
+{
+	Ar << ChunkOffset;
+	Ar << ChunkSize;
+	Ar.Serialize(ChunkId, sizeof(ChunkId));
+	static_assert(sizeof(ChunkId) == sizeof(FIoChunkId));
+}
+
+} // UE::DerivedData::Private
+
 #if WITH_EDITORONLY_DATA
 namespace UE::DerivedData::Private
 {
@@ -83,8 +125,6 @@ public:
 	bool ReferenceEquals(const FEditorData& Other) const;
 	uint32 ReferenceHash() const;
 
-	void Serialize(FArchive& Ar, UObject* Owner) const;
-
 	template <typename VisitorType>
 	inline void Visit(VisitorType&& Visitor) const
 	{
@@ -187,49 +227,6 @@ uint32 FEditorData::ReferenceHash() const
 	return Visitor.Hash;
 }
 
-void FEditorData::Serialize(FArchive& Ar, UObject* Owner) const
-{
-	checkf(Ar.IsSaving() && Ar.IsCooking(),
-		TEXT("FEditorData for FDerivedData only supports saving to cooked packages. %s"), *WriteToString<256>(*this));
-
-	struct FVisitor
-	{
-		inline explicit FVisitor(FArchive& Ar)
-			: Linker(Cast<FLinkerSave>(Ar.GetLinker()))
-		{
-		}
-
-		inline void operator()(const FCompositeBufferWithHash& BufferWithHash)
-		{
-			ChunkId = Linker->AddDerivedData(FValue::Compress(BufferWithHash.Buffer).GetData());
-		}
-
-		inline void operator()(const FCompressedBuffer& Buffer)
-		{
-			ChunkId = Linker->AddDerivedData(Buffer);
-		}
-
-		inline void operator()(const FCacheKeyWithId& CacheKeyWithId)
-		{
-			ChunkId = Linker->AddDerivedData(CacheKeyWithId.Key, CacheKeyWithId.Id);
-		}
-
-		FLinkerSave* Linker;
-		uint64 ChunkOffset = 0;
-		uint64 ChunkSize = 0;
-		FIoChunkId ChunkId;
-	};
-
-	static_assert(sizeof(FIoChunkId) == 12);
-	FVisitor Visitor(Ar);
-	checkf(Visitor.Linker, TEXT("Serializing FDerivedData requires a linker. %s"), *WriteToString<256>(*this));
-	Visit(Visitor);
-
-	Ar << Visitor.ChunkOffset;
-	Ar << Visitor.ChunkSize;
-	Ar.Serialize(&Visitor.ChunkId, sizeof(Visitor.ChunkId));
-}
-
 } // UE::DerivedData::Private
 #endif // WITH_EDITORONLY_DATA
 
@@ -246,22 +243,12 @@ FStringBuilderBase& operator<<(FStringBuilderBase& Builder, const FDerivedData&
 	}
 #endif
 
-	if (Data)
+	if (Data.CookedData)
 	{
-		Builder << TEXTVIEW("Chunk: ID ");
-		String::BytesToHexLower(MakeArrayView(Data.ChunkId), Builder);
-		if (Data.ChunkOffset)
-		{
-			Builder << TEXTVIEW(" / Offset ") << Data.ChunkOffset;
-		}
-		if (Data.ChunkSize != MAX_uint64)
-		{
-			Builder << TEXTVIEW(" / Size ") << Data.ChunkSize;
-		}
-		return Builder;
+		return Builder << Data.CookedData;
 	}
 
-	return Builder << TEXT("Null");
+	return Builder << TEXTVIEW("Null");
 }
 
 bool FDerivedData::ReferenceEquals(const FDerivedData& Other) const
@@ -272,13 +259,7 @@ bool FDerivedData::ReferenceEquals(const FDerivedData& Other) const
 		return EditorData->ReferenceEquals(*Other.EditorData);
 	}
 #endif
-
-	if (ChunkOffset == Other.ChunkOffset && ChunkSize == Other.ChunkSize && MakeArrayView(ChunkId) == Other.ChunkId)
-	{
-		return true;
-	}
-
-	return false;
+	return CookedData.ReferenceEquals(Other.CookedData);
 }
 
 uint32 FDerivedData::ReferenceHash() const
@@ -289,16 +270,13 @@ uint32 FDerivedData::ReferenceHash() const
 		return EditorData->ReferenceHash();
 	}
 #endif
-
-	FXxHash64Builder Builder;
-	Builder.Update(&ChunkOffset, sizeof(ChunkOffset));
-	Builder.Update(&ChunkSize, sizeof(ChunkSize));
-	Builder.Update(MakeMemoryView(ChunkId));
-	return uint32(Builder.Finalize().Hash);
+	return CookedData.ReferenceHash();
 }
 
 void FDerivedData::Serialize(FArchive& Ar, UObject* Owner)
 {
+	using namespace DerivedData::Private;
+
 	if (!Ar.IsPersistent() || Ar.IsObjectReferenceCollector() || Ar.ShouldSkipBulkData())
 	{
 		return;
@@ -307,15 +285,20 @@ void FDerivedData::Serialize(FArchive& Ar, UObject* Owner)
 #if WITH_EDITORONLY_DATA
 	if (EditorData)
 	{
-		EditorData->Serialize(Ar, Owner);
+		checkf(Ar.IsSaving() && Ar.IsCooking(),
+			TEXT("FEditorData for FDerivedData only supports saving to cooked packages. %s"), *WriteToString<256>(*this));
+
+		FLinkerSave* Linker = Cast<FLinkerSave>(Ar.GetLinker());
+		checkf(Linker, TEXT("Serializing FDerivedData requires a linker. %s"), *WriteToString<256>(*this));
+
+		FCookedData CookedEditorData = Linker->AddDerivedData(*this);
+		CookedEditorData.Serialize(Ar);
 		Ar << Flags;
 		return;
 	}
 #endif
 
-	Ar << ChunkOffset;
-	Ar << ChunkSize;
-	Ar.Serialize(ChunkId, sizeof(ChunkId));
+	CookedData.Serialize(Ar);
 	Ar << Flags;
 }
 
@@ -443,13 +426,16 @@ enum class EIoRequestType : uint8
 #endif
 };
 
-struct FIoChunkState
+struct FIoCookedState
 {
 	FIoRequest Request;
-	uint64 ChunkOffset = 0;
-	uint64 ChunkSize = 0;
-	FIoChunkId ChunkId;
+	FCookedData CookedData;
 	std::atomic<bool> bCanceled = false;
+
+	inline explicit FIoCookedState(const FCookedData& InCookedData)
+		: CookedData(InCookedData)
+	{
+	}
 };
 
 struct FIoEditorState
@@ -470,7 +456,7 @@ struct FIoEditorState
 
 struct FIoRequestState
 {
-	TVariant<TYPE_OF_NULLPTR, FIoChunkState, FIoEditorState> State;
+	TVariant<TYPE_OF_NULLPTR, FIoCookedState, FIoEditorState> State;
 	FDerivedDataIoOptions Options;
 	FSharedBuffer Data;
 	uint64 Size = 0;
@@ -534,9 +520,9 @@ public:
 private:
 	FIoResponseDispatcher() = default;
 
-	void DispatchChunk(FIoResponse& Response, FIoRequestState& Request, FIoChunkState& Chunk);
+	void DispatchCooked(FIoResponse& Response, FIoRequestState& Request, FIoCookedState& State);
 
-	static FIoReadOptions MakeIoReadOptions(const FIoChunkState& State, const FDerivedDataIoOptions& Options);
+	static FIoReadOptions MakeIoReadOptions(const FIoCookedState& State, const FDerivedDataIoOptions& Options);
 	static void OnIoRequestComplete(FIoResponse& Response, FIoRequestState& Request, TIoStatusOr<FIoBuffer> Buffer);
 
 #if WITH_EDITORONLY_DATA
@@ -579,17 +565,12 @@ FDerivedDataIoRequest FIoResponse::Queue(
 	if (Data.EditorData)
 	{
 		RequestState.State.Emplace<FIoEditorState>(*Data.EditorData);
-		return Handle;
 	}
+	else
 #endif
-
-	if (Data)
+	if (Data.CookedData)
 	{
-		RequestState.State.Emplace<FIoChunkState>();
-		FIoChunkState& ChunkState = RequestState.State.Get<FIoChunkState>();
-		ChunkState.ChunkOffset = Data.ChunkOffset;
-		ChunkState.ChunkSize = Data.ChunkSize;
-		ChunkState.ChunkId.Set(MakeMemoryView(Data.ChunkId));
+		RequestState.State.Emplace<FIoCookedState>(Data.CookedData);
 	}
 
 	return Handle;
@@ -637,9 +618,9 @@ void FIoResponse::SetPriority(FDerivedDataIoPriority Priority)
 	const int32 IoDispatcherPriority = ConvertToIoDispatcherPriority(Priority);
 	for (FIoRequestState& Request : Requests)
 	{
-		if (FIoChunkState* ChunkState = Request.State.TryGet<FIoChunkState>())
+		if (FIoCookedState* CookedState = Request.State.TryGet<FIoCookedState>())
 		{
-			ChunkState->Request.UpdatePriority(IoDispatcherPriority);
+			CookedState->Request.UpdatePriority(IoDispatcherPriority);
 		}
 	}
 }
@@ -658,12 +639,12 @@ bool FIoResponse::Cancel()
 
 	for (FIoRequestState& Request : Requests)
 	{
-		if (FIoChunkState* ChunkState = Request.State.TryGet<FIoChunkState>())
+		if (FIoCookedState* CookedState = Request.State.TryGet<FIoCookedState>())
 		{
 			// Request cancellation only once because every call wakes the dispatcher.
-			if (!ChunkState->bCanceled.exchange(true, std::memory_order_relaxed))
+			if (!CookedState->bCanceled.exchange(true, std::memory_order_relaxed))
 			{
-				ChunkState->Request.Cancel();
+				CookedState->Request.Cancel();
 			}
 		}
 	}
@@ -719,9 +700,9 @@ void FIoResponseDispatcher::Dispatch(FIoResponse& Response, FDerivedDataIoPriori
 	for (FIoRequestState& Request : Response.Requests)
 	{
 		++RequestIndex;
-		if (FIoChunkState* ChunkState = Request.State.TryGet<FIoChunkState>())
+		if (FIoCookedState* CookedState = Request.State.TryGet<FIoCookedState>())
 		{
-			Dispatcher.DispatchChunk(Response, Request, *ChunkState);
+			Dispatcher.DispatchCooked(Response, Request, *CookedState);
 		}
 	#if WITH_EDITORONLY_DATA
 		else if (FIoEditorState* EditorState = Request.State.TryGet<FIoEditorState>())
@@ -767,14 +748,16 @@ void FIoResponseDispatcher::Dispatch(FIoResponse& Response, FDerivedDataIoPriori
 #endif
 }
 
-void FIoResponseDispatcher::DispatchChunk(FIoResponse& Response, FIoRequestState& Request, FIoChunkState& Chunk)
+void FIoResponseDispatcher::DispatchCooked(FIoResponse& Response, FIoRequestState& Request, FIoCookedState& State)
 {
+	FIoChunkId ChunkId;
+	ChunkId.Set(MakeMemoryView(State.CookedData.ChunkId));
 	if (Request.Type == EIoRequestType::Read)
 	{
 		Response.BeginRequest();
-		Chunk.Request = Batch.ReadWithCallback(
-			Chunk.ChunkId,
-			MakeIoReadOptions(Chunk, Request.Options),
+		State.Request = Batch.ReadWithCallback(
+			ChunkId,
+			MakeIoReadOptions(State, Request.Options),
 			ConvertToIoDispatcherPriority(Priority),
 			[Response = &Response, &Request](TIoStatusOr<FIoBuffer> Buffer)
 			{
@@ -783,7 +766,7 @@ void FIoResponseDispatcher::DispatchChunk(FIoResponse& Response, FIoRequestState
 	}
 	else
 	{
-		TIoStatusOr<uint64> Size = FIoDispatcher::Get().GetSizeForChunk(Chunk.ChunkId);
+		TIoStatusOr<uint64> Size = FIoDispatcher::Get().GetSizeForChunk(ChunkId);
 		if (Size.IsOk())
 		{
 			const uint64 TotalSize = Size.ValueOrDie();
@@ -800,24 +783,25 @@ void FIoResponseDispatcher::DispatchChunk(FIoResponse& Response, FIoRequestState
 }
 
 FIoReadOptions FIoResponseDispatcher::MakeIoReadOptions(
-	const FIoChunkState& State,
+	const FIoCookedState& State,
 	const FDerivedDataIoOptions& Options)
 {
+	const FCookedData& CookedData = State.CookedData;
 	const uint64 LocalOffset = Options.GetOffset();
-	const uint64 TotalOffset = State.ChunkOffset + LocalOffset;
+	const uint64 TotalOffset = CookedData.ChunkOffset + LocalOffset;
 
 	FIoReadOptions ReadOptions;
 	ReadOptions.SetTargetVa(Options.GetTarget());
 
 	if (Options.GetSize() == MAX_uint64)
 	{
-		if (State.ChunkSize == MAX_uint64)
+		if (CookedData.ChunkSize == MAX_uint64)
 		{
 			ReadOptions.SetRange(TotalOffset, MAX_uint64);
 		}
 		else
 		{
-			ReadOptions.SetRange(TotalOffset, LocalOffset <= State.ChunkSize ? State.ChunkSize - LocalOffset : 0);
+			ReadOptions.SetRange(TotalOffset, LocalOffset <= CookedData.ChunkSize ? CookedData.ChunkSize - LocalOffset : 0);
 		}
 	}
 	else
@@ -833,7 +817,7 @@ void FIoResponseDispatcher::OnIoRequestComplete(
 	FIoRequestState& Request,
 	TIoStatusOr<FIoBuffer> StatusOrBuffer)
 {
-	FIoChunkState& Chunk = Request.State.Get<FIoChunkState>();
+	FIoCookedState& State = Request.State.Get<FIoCookedState>();
 
 	if (StatusOrBuffer.IsOk())
 	{
@@ -860,7 +844,7 @@ void FIoResponseDispatcher::OnIoRequestComplete(
 		Request.SetStatus(bCanceled ? EDerivedDataIoStatus::Canceled : EDerivedDataIoStatus::Error);
 	}
 
-	Chunk.Request.Release();
+	State.Request.Release();
 	Response.EndRequest();
 }
 
diff --git a/Engine/Source/Runtime/CoreUObject/Private/UObject/LinkerSave.cpp b/Engine/Source/Runtime/CoreUObject/Private/UObject/LinkerSave.cpp
--- a/Engine/Source/Runtime/CoreUObject/Private/UObject/LinkerSave.cpp
+++ b/Engine/Source/Runtime/CoreUObject/Private/UObject/LinkerSave.cpp
@@ -16,8 +16,8 @@
 #include "HAL/PlatformStackWalk.h"
 
 #if WITH_EDITORONLY_DATA
-#include "DerivedDataValueId.h"
 #include "IO/IoDispatcher.h"
+#include "Serialization/DerivedData.h"
 #endif
 
 /*----------------------------------------------------------------------------
@@ -491,33 +491,22 @@ void FLinkerSave::SetFilterEditorOnly(bool bInFilterEditorOnly)
 }
 
 #if WITH_EDITORONLY_DATA
-
-static FIoChunkId CreateDerivedDataChunkId(const FPackageId PackageId, const int32 ChunkIndex)
+UE::DerivedData::Private::FCookedData FLinkerSave::AddDerivedData(const UE::FDerivedData& Data)
 {
+	UE_LOG(LogLinker, Warning, TEXT("Data will not be able to load because derived data is not saved yet."));
+
+	UE::DerivedData::Private::FCookedData CookedData;
+
+	const FPackageId PackageId = LinkerRoot->GetPackageId();
+	const int32 ChunkIndex = ++LastDerivedDataIndex;
 	checkf(ChunkIndex >= 0 && ChunkIndex < (1 << 24), TEXT("ChunkIndex %d is out of range."), ChunkIndex);
 
 	// PackageId                 ChunkIndex Type
 	// [00 01 02 03 04 05 06 07] [08 09 10] [11]
-	uint8 Data[12]{};
-	*reinterpret_cast<uint8*>(&Data[11]) = static_cast<uint8>(EIoChunkType::DerivedData);
-	*reinterpret_cast<uint32*>(&Data[7]) = NETWORK_ORDER32(ChunkIndex);
-	*reinterpret_cast<uint64*>(&Data[0]) = PackageId.Value();
-
-	FIoChunkId ChunkId;
-	ChunkId.Set(Data, 12);
-	return ChunkId;
-}
-
-FIoChunkId FLinkerSave::AddDerivedData(const FCompressedBuffer& Data)
-{
-	UE_LOG(LogLinker, Warning, TEXT("Data will not be able to load because derived data is not saved yet."));
-	return CreateDerivedDataChunkId(LinkerRoot->GetPackageId(), ++LastDerivedDataIndex);
-}
+	*reinterpret_cast<uint8*>(&CookedData.ChunkId[11]) = static_cast<uint8>(EIoChunkType::DerivedData);
+	*reinterpret_cast<uint32*>(&CookedData.ChunkId[7]) = NETWORK_ORDER32(ChunkIndex);
+	*reinterpret_cast<uint64*>(&CookedData.ChunkId[0]) = PackageId.Value();
 
-FIoChunkId FLinkerSave::AddDerivedData(const UE::DerivedData::FCacheKey& Key, const UE::DerivedData::FValueId& ValueId)
-{
-	UE_LOG(LogLinker, Warning, TEXT("Data will not be able to load because derived data is not saved yet."));
-	return CreateDerivedDataChunkId(LinkerRoot->GetPackageId(), ++LastDerivedDataIndex);
+	return CookedData;
 }
-
 #endif // WITH_EDITORONLY_DATA
diff --git a/Engine/Source/Runtime/CoreUObject/Public/Serialization/DerivedData.h b/Engine/Source/Runtime/CoreUObject/Public/Serialization/DerivedData.h
--- a/Engine/Source/Runtime/CoreUObject/Public/Serialization/DerivedData.h
+++ b/Engine/Source/Runtime/CoreUObject/Public/Serialization/DerivedData.h
@@ -22,6 +22,7 @@ namespace UE::DerivedData { struct FCacheKey; }
 namespace UE::DerivedData { struct FValueId; }
 namespace UE::DerivedData { template <typename CharType> class TSharedString; }
 namespace UE::DerivedData { using FSharedString = TSharedString<TCHAR>; }
+namespace UE::DerivedData::Private { class FCookedData; }
 namespace UE::DerivedData::Private { class FEditorData; }
 namespace UE::DerivedData::Private { class FIoResponse; }
 
@@ -42,6 +43,28 @@ enum class EDerivedDataFlags : uint32
 
 ENUM_CLASS_FLAGS(EDerivedDataFlags);
 
+class DerivedData::Private::FCookedData
+{
+public:
+	inline explicit operator bool() const
+	{
+		// Check the last byte of ChunkId, which is type of the FIoChunkId, where 0 is EIoChunkType::Invalid.
+		return ChunkId[11] != 0;
+	}
+
+	bool ReferenceEquals(const FCookedData& Other) const;
+	uint32 ReferenceHash() const;
+
+	void Serialize(FArchive& Ar);
+
+	/** The offset of this data within its chunk. */
+	uint64 ChunkOffset = 0;
+	/** The size of this data within its chunk. */
+	uint64 ChunkSize = 0;
+	/** The FIoChunkId. */
+	uint8 ChunkId[12]{};
+};
+
 /**
  * Derived Data Reference
  *
@@ -110,14 +133,7 @@ private:
 #endif // WITH_EDITORONLY_DATA
 
 private:
-	// Chunk members can only be non-zero when loading from a cooked package.
-
-	/** The offset of this data within its chunk. */
-	uint64 ChunkOffset = 0;
-	/** The size of this data within its chunk. */
-	uint64 ChunkSize = 0;
-	/** The FIoChunkId. */
-	uint8 ChunkId[12]{};
+	DerivedData::Private::FCookedData CookedData;
 
 	/** The flags that modify how this derived data is saved and how it may be loaded. */
 	EDerivedDataFlags Flags = EDerivedDataFlags::None;
@@ -130,10 +146,9 @@ inline const FDerivedData FDerivedData::Null;
 inline bool FDerivedData::IsNull() const
 {
 #if WITH_EDITORONLY_DATA
-	return !EditorData && ChunkId[11] == 0;
+	return !EditorData && !CookedData;
 #else
-	// Check the last byte of ChunkId, which is type of the FIoChunkId, where 0 is EIoChunkType::Invalid.
-	return ChunkId[11] == 0;
+	return !CookedData;
 #endif
 }
 
diff --git a/Engine/Source/Runtime/CoreUObject/Public/UObject/LinkerSave.h b/Engine/Source/Runtime/CoreUObject/Public/UObject/LinkerSave.h
--- a/Engine/Source/Runtime/CoreUObject/Public/UObject/LinkerSave.h
+++ b/Engine/Source/Runtime/CoreUObject/Public/UObject/LinkerSave.h
@@ -23,15 +23,14 @@
 #include "UObject/UObjectThreadContext.h"
 
 class FBulkData;
-class FIoChunkId;
 class FObjectPostSaveContext;
 class FOutputDevice;
 class FPackagePath;
 class UObject;
 class UPackage;
-namespace UE::DerivedData { struct FCacheKey; }
-namespace UE::DerivedData { struct FValueId; }
+namespace UE { class FDerivedData; }
 namespace UE { class FPackageTrailerBuilder; }
+namespace UE::DerivedData::Private { class FCookedData; }
 struct FLazyObjectPtr;
 struct FUObjectSerializeContext;
 
@@ -237,29 +236,12 @@ public:
 	}
 
 #if WITH_EDITORONLY_DATA
-
-	/**
-	 * Adds the derived data to the package.
-	 *
-	 * This is only supported when saving a cooked package.
-	 *
-	 * @return An ID that can be used to load the derived data from the cooked package.
-	 */
-	FIoChunkId AddDerivedData(const FCompressedBuffer& Data);
-
 	/**
-	 * Adds the derived data value referenced by the cache key and ID.
+	 * Adds the derived data to the package. This is only supported when saving a cooked package.
 	 *
-	 * This is only supported when saving a cooked package.
-	 * The key and a null ID must reference a value that was saved using ICache::PutValue.
-	 * The key and a non-null ID must reference a value in a record that was saved using ICache::Put.
-	 *
-	 * @param Key       Reference to a record or value in the cache.
-	 * @param ValueId   Reference to a value within a record, otherwise FValueId::Null.
-	 * @return An ID that can be used to load the derived data from the cooked package.
+	 * @return A cooked reference that can be used to load the derived data from the cooked package.
 	 */
-	FIoChunkId AddDerivedData(const UE::DerivedData::FCacheKey& Key, const UE::DerivedData::FValueId& ValueId);
-
+	UE::DerivedData::Private::FCookedData AddDerivedData(const UE::FDerivedData& Data);
 #endif // WITH_EDITORONLY_DATA
 
 protected:
