GC: Make sure to set GIsGarbageCollecting=true for the complete IncrementalPurgeGarbage phase (cleanup step left out of the bigger change 19898875).
Fix for destroying ChildActorComponent sometimes fails assert in UObjectBase destructor checking for NAME_None.
Previously GIsGarbageCollecting was only set during the internal IncrementalDestroyGarbage function, and was missing for the UnhashUnreachableObjects function that is responsible for calling UObject::BeginDestroy.
This was probably a mistake when incremental GC was first introduced, and not what client code like UChildActorComponent::DestroyChildActor is expecting.
We always want BeginDestroy to the behave the same no matter if it is called from a full blocking GC (when GIsGarbageCollecting=true is already true) or from an incremental GC step.
Change type of GIsGarbageCollecting from FThreadSafeBool to a normal bool. It is only set from GT, and it should only be read from GT (it does not make sense on other threads, nothing that cares about it should run at this time).
Explicitly silence ~SObjectWidget ensure on engine shutdown since GIsGarbageCollecting is now set in this scneario too (no actual change in behaviour).

#rnx

diff --git a/Engine/Source/Runtime/CoreUObject/Private/UObject/GarbageCollection.cpp b/Engine/Source/Runtime/CoreUObject/Private/UObject/GarbageCollection.cpp
--- a/Engine/Source/Runtime/CoreUObject/Private/UObject/GarbageCollection.cpp
+++ b/Engine/Source/Runtime/CoreUObject/Private/UObject/GarbageCollection.cpp
@@ -78,21 +78,7 @@ static double GLastGCTime = 0;
 /** Helpful constant for determining how many token slots we need to store a pointer **/
 static const uint32 GNumTokensPerPointer = sizeof(void*) / sizeof(uint32); //-V514
 
-FThreadSafeBool GIsGarbageCollecting(false);
-class FIsGarbageCollectingScope
-{
-	bool bOldValue = false;
-public:
-	FIsGarbageCollectingScope()
-		: bOldValue(GIsGarbageCollecting)
-	{
-		GIsGarbageCollecting = true;
-	}
-	~FIsGarbageCollectingScope()
-	{
-		GIsGarbageCollecting = bOldValue;
-	}
-};
+bool GIsGarbageCollecting = false;
 
 int32 FGCReferenceTokenStream::MaxStackSize = 1;
 
@@ -1760,6 +1746,7 @@ void IncrementalPurgeGarbage(bool bUseTimeLimit, double TimeLimit)
 		// Lock before settting GCStartTime as it could be slow to lock if async loading is in progress
 		// but we still want to perform some GC work otherwise we'd be keeping objects in memory for a long time
 		FLegacyConditionalIncrementalGCLock ScopedGCLock;
+		TGuardValue<bool> GuardIsGarbageCollecting(GIsGarbageCollecting, true);
 
 		// Keep track of start time to enforce time limit unless bForceFullPurge is true;
 		GCStartTime = FPlatformTime::Seconds();
@@ -1838,7 +1825,7 @@ bool IncrementalDestroyGarbage(bool bUseTimeLimit, double TimeLimit)
 
 	// Set 'I'm garbage collecting' flag - might be checked inside UObject::Destroy etc.
 	FGCLegacyHashTableScopeLock GCLegacyHashTableLock;
-	FIsGarbageCollectingScope IsGarbageCollectingScope;
+	TGuardValue<bool> GuardIsGarbageCollecting(GIsGarbageCollecting, true);
 
 	if( !GObjFinishDestroyHasBeenRoutedToAllObjects && !bTimeLimitReached )
 	{
@@ -2301,7 +2288,7 @@ void CollectGarbageInternal(EObjectFlags KeepFlags, bool bPerformFullPurge)
 		// Set 'I'm garbage collecting' flag - might be checked inside various functions.
 		// This has to be unlocked before we call post GC callbacks
 		FGCLegacyHashTableScopeLock GCLegacyHashTableLock;
-		FIsGarbageCollectingScope IsGarbageCollectingScope;
+		TGuardValue<bool> GuardIsGarbageCollecting(GIsGarbageCollecting, true);
 
 		UE_LOG(LogGarbage, Log, TEXT("Collecting garbage%s"), IsAsyncLoading() ? TEXT(" while async loading") : TEXT(""));
 
diff --git a/Engine/Source/Runtime/CoreUObject/Public/UObject/GarbageCollection.h b/Engine/Source/Runtime/CoreUObject/Public/UObject/GarbageCollection.h
--- a/Engine/Source/Runtime/CoreUObject/Public/UObject/GarbageCollection.h
+++ b/Engine/Source/Runtime/CoreUObject/Public/UObject/GarbageCollection.h
@@ -701,7 +701,7 @@ public:
 };
 
 /** True if Garbage Collection is running. Use IsGarbageCollecting() functio n instead of using this variable directly */
-extern COREUOBJECT_API FThreadSafeBool GIsGarbageCollecting;
+extern COREUOBJECT_API bool GIsGarbageCollecting;
 
 /**
  * Gets the last time that the GC was run.
diff --git a/Engine/Source/Runtime/UMG/Private/Slate/SObjectWidget.cpp b/Engine/Source/Runtime/UMG/Private/Slate/SObjectWidget.cpp
--- a/Engine/Source/Runtime/UMG/Private/Slate/SObjectWidget.cpp
+++ b/Engine/Source/Runtime/UMG/Private/Slate/SObjectWidget.cpp
@@ -40,7 +40,7 @@ SObjectWidget::~SObjectWidget(void)
 		// This is only a concern during a running game - design-time instances can be destroyed from GC quite often when recompiling
 		if (!WidgetObject || !WidgetObject->IsDesignTime())
 		{
-			ensureMsgf(!IsGarbageCollecting(), TEXT("SObjectWidget for '%s' destroyed while collecting garbage.  This can lead to multiple GCs being required to cleanup the object.  Possible causes might be,\n1) ReleaseSlateResources not being implemented for the owner of this pointer.\n2) You may just be holding onto some slate pointers on an actor that don't get reset until the actor is Garbage Collected.  You should avoid doing this, and instead reset those references when the actor is Destroyed."), *DebugName);
+			ensureMsgf(GExitPurge || !IsGarbageCollecting(), TEXT("SObjectWidget for '%s' destroyed while collecting garbage.  This can lead to multiple GCs being required to cleanup the object.  Possible causes might be,\n1) ReleaseSlateResources not being implemented for the owner of this pointer.\n2) You may just be holding onto some slate pointers on an actor that don't get reset until the actor is Garbage Collected.  You should avoid doing this, and instead reset those references when the actor is Destroyed."), *DebugName);
 		}
 	}
 #endif
