Test code removed.

diff --git a/Engine/Source/Runtime/CoreUObject/Private/UObject/UObjectBaseUtility.cpp b/Engine/Source/Runtime/CoreUObject/Private/UObject/UObjectBaseUtility.cpp
--- a/Engine/Source/Runtime/CoreUObject/Private/UObject/UObjectBaseUtility.cpp
+++ b/Engine/Source/Runtime/CoreUObject/Private/UObject/UObjectBaseUtility.cpp
@@ -275,12 +275,6 @@ UE::Core::FVersePath UObjectBaseUtility::GetVersePath() const
 		return {};
 	}
 
-	if (VerseModule.StartsWith(TEXT("/localhost/")) && !UE::IsSavingPackage(nullptr) && !IsGarbageCollectingAndLockingUObjectHashTables())
-	{
-		UObject* Found = StaticFindObject(this->GetClass(), Result);
-		check(this == Found);
-	}
-
 	return Result;
 }
 
