Avoid postloading duplicates during async loading as they will be postloaded on the GT instead

diff --git a/Engine/Source/Runtime/CoreUObject/Private/UObject/UObjectGlobals.cpp b/Engine/Source/Runtime/CoreUObject/Private/UObject/UObjectGlobals.cpp
--- a/Engine/Source/Runtime/CoreUObject/Private/UObject/UObjectGlobals.cpp
+++ b/Engine/Source/Runtime/CoreUObject/Private/UObject/UObjectGlobals.cpp
@@ -2913,10 +2913,21 @@ UObject* StaticDuplicateObjectEx( FObjectDuplicationParameters& Parameters )
 				DuplicatedObject->PostDuplicate(Parameters.DuplicateMode);
 				if (!Parameters.bSkipPostLoad && !DuplicatedObject->IsTemplate())
 				{
-					// Don't want to call PostLoad on class duplicated CDOs
-					TGuardValue<bool> GuardIsRoutingPostLoad(FUObjectThreadContext::Get().IsRoutingPostLoad, true);
-					DuplicatedObject->ConditionalPostLoad();
+					// We skip post-loading during async loading if on the loader thread as we're going to handle it deferred on GT instead.
+					if (IsInGameThread())
+					{
+						// Don't want to call PostLoad on class duplicated CDOs
+						TGuardValue<bool> GuardIsRoutingPostLoad(FUObjectThreadContext::Get().IsRoutingPostLoad, true);
+						DuplicatedObject->ConditionalPostLoad();
+					}
+					else
+					{
+						// The only other thread that we allow to go through here is ALT because we know
+						// it is going to call post-load on new objects.
+						check(IsInAsyncLoadingThread());
+					}
 				}
+
 				DuplicatedObject->CheckDefaultSubobjects();
 			}
 		}
@@ -2963,6 +2974,8 @@ bool SaveToTransactionBuffer(UObject* Object, bool bMarkDirty)
 
 	if ( GUndo && bIsTransactional && bIsNotScriptPackage)
 	{
+		check(IsInGameThread());
+
 		// Mark the package dirty, if requested
 		if ( bMarkDirty )
 		{
