Optimization in the LinkerLoad, when it has many bulk data loaders. This could do like a linear search through the large array thousands of times, which got very slow. Turned the TArray into a TSet. This speeds up loading certain geometry caches significantly.

diff --git a/Engine/Source/Runtime/CoreUObject/Private/UObject/LinkerLoad.cpp b/Engine/Source/Runtime/CoreUObject/Private/UObject/LinkerLoad.cpp
--- a/Engine/Source/Runtime/CoreUObject/Private/UObject/LinkerLoad.cpp
+++ b/Engine/Source/Runtime/CoreUObject/Private/UObject/LinkerLoad.cpp
@@ -5531,14 +5531,16 @@ void FLinkerLoad::Detach()
  */
 void FLinkerLoad::AttachBulkData( UObject* Owner, FBulkData* BulkData )
 {
-	check( BulkDataLoaders.Find(BulkData)==INDEX_NONE );
-	BulkDataLoaders.Add( BulkData );
+	bool bAlreadyInSet = false;
+	BulkDataLoaders.Add(BulkData, &bAlreadyInSet);
+	check(!bAlreadyInSet);
 }
 
 void FLinkerLoad::AttachBulkData(UE::Serialization::FEditorBulkData* BulkData)
 {
-	check(EditorBulkDataLoaders.Find(BulkData) == INDEX_NONE);
-	EditorBulkDataLoaders.Add(BulkData);
+	bool bAlreadyInSet = false;
+	EditorBulkDataLoaders.Add(BulkData, &bAlreadyInSet);
+	check(!bAlreadyInSet);
 }
 
 /**
@@ -5576,8 +5578,7 @@ void FLinkerLoad::DetachAllBulkData(bool bEnsureAllBulkDataIsLoaded)
 {
 	// Old style bulkdata first
 	{
-		TArray<FBulkData*> BulkDataToDetach = BulkDataLoaders;
-		for (FBulkData* BulkData : BulkDataToDetach)
+		for (FBulkData* BulkData : BulkDataLoaders)
 		{
 			check(BulkData != nullptr);
 			BulkData->DetachFromArchive(this, bEnsureAllBulkDataIsLoaded);
@@ -5587,8 +5588,7 @@ void FLinkerLoad::DetachAllBulkData(bool bEnsureAllBulkDataIsLoaded)
 
 	// Then virtualized bulkdata
 	{
-		TArray<UE::Serialization::FEditorBulkData*> BulkDataToDetach = EditorBulkDataLoaders;
-		for (UE::Serialization::FEditorBulkData* BulkData : BulkDataToDetach)
+		for (UE::Serialization::FEditorBulkData* BulkData : EditorBulkDataLoaders)
 		{
 			check(BulkData != nullptr);
 			BulkData->DetachFromDisk(this, bEnsureAllBulkDataIsLoaded);
diff --git a/Engine/Source/Runtime/CoreUObject/Public/UObject/LinkerLoad.h b/Engine/Source/Runtime/CoreUObject/Public/UObject/LinkerLoad.h
--- a/Engine/Source/Runtime/CoreUObject/Public/UObject/LinkerLoad.h
+++ b/Engine/Source/Runtime/CoreUObject/Public/UObject/LinkerLoad.h
@@ -267,8 +267,8 @@ public:
 	void* AsyncRoot;
 #if WITH_EDITOR
 	/** Bulk data that does not need to be loaded when the linker is loaded.												*/
-	TArray<FBulkData*> BulkDataLoaders;
-	TArray<UE::Serialization::FEditorBulkData*> EditorBulkDataLoaders;
+	TSet<FBulkData*> BulkDataLoaders;
+	TSet<UE::Serialization::FEditorBulkData*> EditorBulkDataLoaders;
 #endif // WITH_EDITOR
 
 	/** Hash table for exports.																								*/
