Add a softobject path table to the package summary header
This is to faliciate validation by inspecting the header as well as allow asset patching using solely the header without export loading.

Deprecate FolderName in UPackage and replace it in PackageSummary with the PackageName
Use to replaced PackageName to generate relative path loading when detecting a package has been relocated which allows some support for asset manipulation through normal file operation

diff --git a/Engine/Source/Runtime/Core/Public/UObject/ObjectVersion.h b/Engine/Source/Runtime/Core/Public/UObject/ObjectVersion.h
--- a/Engine/Source/Runtime/Core/Public/UObject/ObjectVersion.h
+++ b/Engine/Source/Runtime/Core/Public/UObject/ObjectVersion.h
@@ -66,6 +66,10 @@ enum class EUnrealEngineObjectUE5Version : uint32
 	// Replace FName asset path in FSoftObjectPath with (package name, asset name) pair FTopLevelAssetPath
 	FSOFTOBJECTPATH_REMOVE_ASSET_PATH_FNAMES,
 	
+
+	// Add a soft object path list to the package summary for fast remap
+	ADD_SOFTOBJECTPATH_LIST,
+
 	// -----<new versions can be added before this line>-------------------------------------------------
 	// - this needs to be the last line (see note below)
 	AUTOMATIC_VERSION_PLUS_ONE,
diff --git a/Engine/Source/Runtime/CoreUObject/Private/UObject/LinkerLoad.cpp b/Engine/Source/Runtime/CoreUObject/Private/UObject/LinkerLoad.cpp
--- a/Engine/Source/Runtime/CoreUObject/Private/UObject/LinkerLoad.cpp
+++ b/Engine/Source/Runtime/CoreUObject/Private/UObject/LinkerLoad.cpp
@@ -721,6 +721,13 @@ FLinkerLoad::ELinkerStatus FLinkerLoad::ProcessPackageSummary(TMap<TPair<FName,
 		Status = SerializeNameMap();
 	}
 
+	// Serialize the soft object path list and register the paths.
+	if (Status == LINKER_Loaded)
+	{
+		SCOPED_LOADTIMER(LinkerLoad_SerializeSoftObjectPathList);
+		Status = SerializeSoftObjectPathList();
+	}
+
 	// Serialize the gatherable text data map.
 	if( Status == LINKER_Loaded )
 	{
@@ -765,6 +772,13 @@ FLinkerLoad::ELinkerStatus FLinkerLoad::ProcessPackageSummary(TMap<TPair<FName,
 		Status = PopulateInstancingContext();
 	}
 
+	// Populate the linker instancing context for relocation loading if needed.
+	if (Status == LINKER_Loaded)
+	{
+		SCOPED_LOADTIMER(LinkerLoad_PopulateRelocationContext);
+		Status = PopulateRelocationContext();
+	}
+
 	// Fix up export map for object class conversion 
 	if( Status == LINKER_Loaded )
 	{	
@@ -903,6 +917,7 @@ FLinkerLoad::FLinkerLoad(UPackage* InParent, const FPackagePath& InPackagePath,
 , InstancingContext(MoveTemp(InInstancingContext))
 , PackageTrailer(nullptr)
 , AsyncRoot(nullptr)
+, SoftObjectPathListIndex(0)
 , GatherableTextDataMapIndex(0)
 , ImportMapIndex(0)
 , ExportMapIndex(0)
@@ -914,6 +929,7 @@ FLinkerLoad::FLinkerLoad(UPackage* InParent, const FPackagePath& InPackagePath,
 , bHasSerializedPreloadDependencies(false)
 , bHasFixedUpImportMap(false)
 , bHasPopulatedInstancingContext(false)
+, bHasPopulatedRelocationContext(false)
 , bFixupExportMapDone(false)
 , bHasFoundExistingExports(false)
 , bHasFinishedInitialization(false)
@@ -1532,11 +1548,6 @@ FLinkerLoad::ELinkerStatus FLinkerLoad::UpdateFromPackageFileSummary()
 		// Propagate package flags
 		LinkerRootPackage->SetPackageFlagsTo(NewPackageFlags);
 
-#if WITH_EDITORONLY_DATA
-		// Propagate package folder name
-		LinkerRootPackage->SetFolderName(*Summary.FolderName);
-#endif
-
 		// Propagate streaming install ChunkID
 		LinkerRootPackage->SetChunkIDs(Summary.ChunkIDs);
 
@@ -1664,6 +1675,28 @@ FLinkerLoad::ELinkerStatus FLinkerLoad::SerializeNameMap()
 	return LINKER_Loaded;
 }
 
+FLinkerLoad::ELinkerStatus FLinkerLoad::SerializeSoftObjectPathList()
+{
+	if (!IsTextFormat() && SoftObjectPathListIndex == 0 && Summary.SoftObjectPathsCount > 0)
+	{
+		Seek(Summary.SoftObjectPathsOffset);
+	}
+
+	FStructuredArchive::FStream Stream = StructuredArchiveRootRecord->EnterStream(TEXT("SoftObjectPathList"));
+	while (SoftObjectPathListIndex < Summary.SoftObjectPathsCount && !IsTimeLimitExceeded(TEXT("serializing soft object path list"), 100))
+	{
+		FSoftObjectPath& SoftObjectPath = SoftObjectPathList.AddDefaulted_GetRef();
+		FStructuredArchive::FSlot Slot = Stream.EnterElement();
+		SoftObjectPath.SerializePath(Slot.GetUnderlyingArchive());
+
+		FixupSoftObjectPathForInstancedPackage(SoftObjectPath);
+		++SoftObjectPathListIndex;
+	}
+
+	// Return whether we finished this step and it's safe to start with the next.
+	return ((SoftObjectPathListIndex == Summary.SoftObjectPathsCount) && !IsTimeLimitExceeded(TEXT("serializing soft object path list"))) ? LINKER_Loaded : LINKER_TimedOut;
+}
+
 /**
  * Serializes the gatherable text data container.
  */
@@ -1734,7 +1767,7 @@ FLinkerLoad::ELinkerStatus FLinkerLoad::FixupImportMap()
 #if WITH_EDITOR
 		if (LoadProgressScope)
 		{
-		LoadProgressScope->EnterProgressFrame(1);
+			LoadProgressScope->EnterProgressFrame(1);
 		}
 #endif
 		// Fix up imports, not required if everything is cooked.
@@ -1972,6 +2005,120 @@ FLinkerLoad::ELinkerStatus FLinkerLoad::PopulateInstancingContext()
 	return IsTimeLimitExceeded(TEXT("populating instancing context")) ? LINKER_TimedOut : LINKER_Loaded;
 }
 
+/**
+ * Generate a package name with the same relative pathing from the current package path as from the original package path
+ * i.e.
+ * OriginalPath: /Game/Folder/
+ * CurrentPath: /Game/OtherFolder/SubFolder/
+ * Example 1:
+ * PackageName: /Game/Folder/Sub/Asset
+ * Result: /Game/OtherFolder/SubFolder/Sub/Asset
+ * Example 2:
+ * PackageName: /Game/OtherAsset
+ * Result: /Game/OtherFolder/OtherAsset
+ * @param InOriginalPackagePath The original package path where this package was located
+ * @param InCurrentPackagePath The current location of the package
+ * @param InPackageName The package name we want to remap relative to the current location
+ * @return FName new package name keeping the same relative location to the current path as the original
+ */
+FName GenerateCurrentRelativePackageName(FStringView InOriginalPackagePath, FStringView InCurrentPackagePath, FStringView InPackageName)
+{
+	FName Result = NAME_None;
+	FString RelPackageName(InPackageName);
+	//@todo 5.2: Create and use a string view version of MakePathRelativeTo
+	FString OriginalPackagePath(InOriginalPackagePath);
+	FPaths::MakePathRelativeTo(RelPackageName, *OriginalPackagePath);
+
+	// Strip the mount point from the current path prior to collapsing as we do not want to allow it to be collapsed
+	FStringView CurrentMount, MountlessCurrentPath;
+	FPathViews::SplitFirstComponent(FStringView(InCurrentPackagePath.RightChop(1)), CurrentMount, MountlessCurrentPath);
+	TStringBuilder<64> MountBuilder;
+	MountBuilder << '/' << CurrentMount << '/';
+
+	// Build the new path to collapse against the current path
+	TStringBuilder<256> Builder;
+	if (MountlessCurrentPath.Len() > 0)
+	{
+		Builder << MountlessCurrentPath << '/' << RelPackageName;
+	}
+	else
+	{
+		Builder << RelPackageName;
+	}
+	bool bSuccess = FPathViews::CollapseRelativeDirectories(Builder);
+
+	// if we successfully collapsed the path prepend the mount point back
+	if (bSuccess)
+	{
+		Builder.Prepend(MountBuilder);
+		Result = FName(*Builder);
+	}
+	// if we couldn't, at least fix up the mount point if needed
+	else if (!InPackageName.StartsWith(MountBuilder.ToView()))
+	{
+		FStringView PackageMount, MountlessPackagePath;
+		FPathViews::SplitFirstComponent(FStringView(InPackageName.RightChop(1)), PackageMount, MountlessPackagePath);
+		Builder.Reset();
+		Builder << MountBuilder << MountlessPackagePath;
+		Result = FName(*Builder);
+	}
+	return Result;
+}
+
+FLinkerLoad::ELinkerStatus FLinkerLoad::PopulateRelocationContext()
+{
+	DECLARE_SCOPE_CYCLE_COUNTER(TEXT("FLinkerLoad::PopulateRelocationContext"), STAT_LinkerLoad_PopulateRelocationContext, STATGROUP_LinkerLoad);
+
+	if (!bHasPopulatedRelocationContext)
+	{
+#if WITH_EDITOR
+		// Validate if the package was moved and we want to generate fix up for references
+		FString PackageNameToLoad = GetPackagePath().GetPackageName();
+		// Check the summmary versioning alongside, PackageName used to be FolderName which wasn't actively used and got deprecated
+		bool bRelocated = PackageNameToLoad != Summary.PackageName && Summary.GetFileVersionUE() >= EUnrealEngineObjectUE5Version::ADD_SOFTOBJECTPATH_LIST;
+		if (bRelocated)
+		{
+			// Only Generate relative fix up if the file was actually moved not just renamed
+			FStringView PackagePathToLoad = FPathViews::GetPath(PackageNameToLoad);
+			FStringView OriginalPackagePath = FPathViews::GetPath(Summary.PackageName);
+			//@todo 5.2: switch to use the string view version
+			//FStringView PackageMount = FPathViews::GetMountPointNameFromPath(OriginalPackagePath, nullptr, false);
+			FString PackageMount = FPackageName::GetPackageMountPoint(FString(OriginalPackagePath.GetData(), OriginalPackagePath.Len()), false).ToString();
+
+			if (PackagePathToLoad != OriginalPackagePath)
+			{
+				for (const FObjectImport& Import : ImportMap)
+				{
+					if (!Import.OuterIndex.IsNull())
+					{
+						continue;
+					}
+
+					// Generate relative fix-up for package path from the same original package mount point of the package we are loading
+					FString ImportPackageName = Import.ObjectName.ToString();
+					if (FPathViews::IsParentPathOf(PackageMount, ImportPackageName))
+					{
+						FName GeneratedName = GenerateCurrentRelativePackageName(OriginalPackagePath, PackagePathToLoad, ImportPackageName);
+						if (!GeneratedName.IsNone())
+						{
+							FName& RelocatedName = InstancingContext.RelocatedPackageMapping.FindOrAdd(Import.ObjectName);
+							// if there isn't already a remapping for that package, create one
+							if (RelocatedName.IsNone())
+							{
+								RelocatedName = GeneratedName;
+							}
+						}
+					}
+				}
+			}
+		}
+#endif
+		// Avoid duplicate work in async case.
+		bHasPopulatedRelocationContext = true;
+	}
+	return IsTimeLimitExceeded(TEXT("populating relocation context")) ? LINKER_TimedOut : LINKER_Loaded;
+}
+
 /**
  * Serializes the export map.
  */
@@ -3085,7 +3232,9 @@ bool FLinkerLoad::VerifyImportInner(const int32 ImportIndex, FString& WarningSuf
 		FUObjectSerializeContext* SerializeContext = GetSerializeContext();
 
 		// Resolve the package name for the import, potentially remapping it, if instancing
-		FName PackageToLoad = !Import.HasPackageName() ? Import.ObjectName : Import.GetPackageName();
+		FName OriginalPackageToLoad = !Import.HasPackageName() ? Import.ObjectName : Import.GetPackageName();
+		FName PackageToLoad = InstancingContext.RelocatePackage(OriginalPackageToLoad);
+
 		FName PackageToLoadInto = InstancingContext.RemapPackage(PackageToLoad);
 #if WITH_EDITOR
 		if (SlowTask)
@@ -3174,6 +3323,13 @@ bool FLinkerLoad::VerifyImportInner(const int32 ImportIndex, FString& WarningSuf
 		// to be linked to any other package's ImportMaps
 		if (!Package || Package->HasAnyPackageFlags(PKG_Compiling))
 		{
+			FName RelocatedName;
+			if (InstancingContext.RelocatedPackageMapping.RemoveAndCopyValue(OriginalPackageToLoad, RelocatedName))
+			{
+				UE_ASSET_LOG(LogLinker, Warning, PackagePath, TEXT("VerifyImport: Failed to load package for import object '%s' through using relative path '%s'. Using orginal location as fallback."), *GetImportFullName(ImportIndex), *RelocatedName.ToString());
+				return nullptr;
+			}
+
 			if (!FLinkerLoad::IsKnownMissingPackage(PackageToLoad))
 			{
 				FLinkerLoad::AddKnownMissingPackage(PackageToLoad);
@@ -5638,9 +5794,46 @@ FArchive& FLinkerLoad::operator<<( FObjectPtr& ObjectPtr )
 	return *this;
 }
 
+FArchive& FLinkerLoad::operator<<(FSoftObjectPath& Value)
+{
+	FArchive& Ar = *this;
+	// if we have items in the soft object path list consider soft object path saved as index into it.
+	// Otherwise deserialize them as usual, cooking for example might not want soft object path serialized as index
+	if (SoftObjectPathList.Num() > 0)
+	{
+		int32 SoftObjectPathIndex = INDEX_NONE;
+		Ar << SoftObjectPathIndex;
+		if (SoftObjectPathList.IsValidIndex(SoftObjectPathIndex))
+		{
+			Value = SoftObjectPathList[SoftObjectPathIndex];
+		}
+		else
+		{
+			Value = FSoftObjectPath();
+			BadSoftObjectPathError(SoftObjectPathIndex);
+			SetCriticalError();
+		}
+	}
+	else
+	{
+		FArchiveUObject::operator<<(Value);
+		FixupSoftObjectPathForInstancedPackage(Value);
+	}
+	return Ar;
+}
+
+void FLinkerLoad::BadSoftObjectPathError(int32 SoftObjIndex)
+{
+	UE_ASSET_LOG(LogLinker, Error, PackagePath,
+		TEXT("Serialization error - FSoftObjectPath are serialized as an index recorded in the package header, the current deserialized index has value %i, outside of the valid range [0, %i)."),
+		SoftObjIndex, SoftObjectPathList.Num());
+}
+
 void FLinkerLoad::BadNameIndexError(int32 NameIndex)
 {
-	UE_ASSET_LOG(LogLinker, Error, PackagePath, TEXT("Bad name index %i/%i"), NameIndex, NameMap.Num());
+	UE_ASSET_LOG(LogLinker, Error, PackagePath,
+		TEXT("Serialization error - FName are serialized as an index recorded in the package header, the current deserialized index has value %i, outside of the valid range [0, %i)."),
+		NameIndex, NameMap.Num());
 }
 
 /**
diff --git a/Engine/Source/Runtime/CoreUObject/Private/UObject/LinkerSave.cpp b/Engine/Source/Runtime/CoreUObject/Private/UObject/LinkerSave.cpp
--- a/Engine/Source/Runtime/CoreUObject/Private/UObject/LinkerSave.cpp
+++ b/Engine/Source/Runtime/CoreUObject/Private/UObject/LinkerSave.cpp
@@ -72,7 +72,7 @@ FLinkerSave::FLinkerSave(UPackage* InParent, const TCHAR* InFilename, bool bForc
 		if (Package)
 		{
 #if WITH_EDITORONLY_DATA
-			Summary.FolderName = Package->GetFolderName().ToString();
+			Summary.PackageName = Package->GetName();
 #endif
 			Summary.ChunkIDs = Package->GetChunkIDs();
 		}
@@ -126,7 +126,7 @@ FLinkerSave::FLinkerSave(UPackage* InParent, FArchive *InSaver, bool bForceByteS
 		if (Package)
 		{
 #if WITH_EDITORONLY_DATA
-			Summary.FolderName = Package->GetFolderName().ToString();
+			Summary.PackageName = Package->GetName();
 #endif
 			Summary.ChunkIDs = Package->GetChunkIDs();
 		}
@@ -175,7 +175,7 @@ FLinkerSave::FLinkerSave(UPackage* InParent, bool bForceByteSwapping, bool bInSa
 		if (Package)
 		{
 #if WITH_EDITORONLY_DATA
-			Summary.FolderName = Package->GetFolderName().ToString();
+			Summary.PackageName = Package->GetName();
 #endif
 			Summary.ChunkIDs = Package->GetChunkIDs();
 		}
@@ -225,6 +225,19 @@ int32 FLinkerSave::MapName(FNameEntryId Id) const
 	return INDEX_NONE;
 }
 
+int32 FLinkerSave::MapSoftObjectPath(const FSoftObjectPath& SoftObjectPath) const
+{
+	const int32* IndexPtr = SoftObjectPathIndices.Find(SoftObjectPath);
+
+	if (IndexPtr)
+	{
+		return *IndexPtr;
+	}
+
+	return INDEX_NONE;
+}
+
+
 FPackageIndex FLinkerSave::MapObject( const UObject* Object ) const
 {
 	if (Object)
@@ -382,6 +395,37 @@ FArchive& FLinkerSave::operator<<( UObject*& Obj )
 	return *this << Save;
 }
 
+FArchive& FLinkerSave::operator<<(FSoftObjectPath& SoftObjectPath)
+{
+	// Map soft object path to indices if we aren't currently serializing the list itself
+	// and we actually built one, cooking might want to serialize soft object path directly for example
+	if (!bIsWritingHeader && SoftObjectPathList.Num() > 0)
+	{
+		int32 Save = MapSoftObjectPath(SoftObjectPath);
+		bool bPathMapped = Save != INDEX_NONE;
+		if (!bPathMapped)
+		{
+			// Set an error on the archive and record the error on the log output if one is set.
+			SetCriticalError();
+			FString ErrorMessage = FString::Printf(TEXT("SoftObjectPath \"%s\" is not mapped when saving %s (object: %s, property: %s). This can mean that this object serialize function is not deterministic between reference harvesting and serialization."),
+				*SoftObjectPath.ToString(),
+				*GetArchiveName(),
+				*GetSerializeContext()->SerializedObject->GetFullName(),
+				*GetFullNameSafe(GetSerializedProperty()));
+			ensureMsgf(false, TEXT("%s"), *ErrorMessage);
+			if (LogOutput)
+			{
+				LogOutput->Logf(ELogVerbosity::Error, TEXT("%s"), *ErrorMessage);
+			}
+		}
+		return *this << Save;
+	}
+	else
+	{
+		return FArchiveUObject::operator<<(SoftObjectPath);
+	}
+}
+
 FArchive& FLinkerSave::operator<<(FLazyObjectPtr& LazyObjectPtr)
 {
 	FUniqueObjectGuid ID;
diff --git a/Engine/Source/Runtime/CoreUObject/Private/UObject/PackageFileSummary.cpp b/Engine/Source/Runtime/CoreUObject/Private/UObject/PackageFileSummary.cpp
--- a/Engine/Source/Runtime/CoreUObject/Private/UObject/PackageFileSummary.cpp
+++ b/Engine/Source/Runtime/CoreUObject/Private/UObject/PackageFileSummary.cpp
@@ -192,7 +192,7 @@ void operator<<(FStructuredArchive::FSlot Slot, FPackageFileSummary& Sum)
 			}
 		}
 		Record << SA_VALUE(TEXT("TotalHeaderSize"), Sum.TotalHeaderSize);
-		Record << SA_VALUE(TEXT("FolderName"), Sum.FolderName);
+		Record << SA_VALUE(TEXT("PackageName"), Sum.PackageName); //@note used to be FolderName, now unused and deprecated
 
 		if (BaseArchive.IsCooking())
 		{
@@ -211,7 +211,14 @@ void operator<<(FStructuredArchive::FSlot Slot, FPackageFileSummary& Sum)
 		{
 			BaseArchive.SetFilterEditorOnly(true);
 		}
+
 		Record << SA_VALUE(TEXT("NameCount"), Sum.NameCount) << SA_VALUE(TEXT("NameOffset"), Sum.NameOffset);
+
+		if (BaseArchive.IsSaving() || Sum.FileVersionUE >= EUnrealEngineObjectUE5Version::ADD_SOFTOBJECTPATH_LIST)
+		{
+			Record << SA_VALUE(TEXT("SoftObjectPathsCount"), Sum.SoftObjectPathsCount) << SA_VALUE(TEXT("SoftObjectPathsOffset"), Sum.SoftObjectPathsOffset);
+		}
+
 		if (!BaseArchive.IsFilterEditorOnly())
 		{
 			if (BaseArchive.IsSaving() || Sum.FileVersionUE >= VER_UE4_ADDED_PACKAGE_SUMMARY_LOCALIZATION_ID)
@@ -219,6 +226,7 @@ void operator<<(FStructuredArchive::FSlot Slot, FPackageFileSummary& Sum)
 				Record << SA_VALUE(TEXT("LocalizationId"), Sum.LocalizationId);
 			}
 		}
+
 		if (Sum.FileVersionUE >= VER_UE4_SERIALIZE_TEXT_IN_PACKAGES)
 		{
 			Record << SA_VALUE(TEXT("GatherableTextDataCount"), Sum.GatherableTextDataCount) << SA_VALUE(TEXT("GatherableTextDataOffset"), Sum.GatherableTextDataOffset);
diff --git a/Engine/Source/Runtime/CoreUObject/Private/UObject/SavePackage/PackageHarvester.cpp b/Engine/Source/Runtime/CoreUObject/Private/UObject/SavePackage/PackageHarvester.cpp
--- a/Engine/Source/Runtime/CoreUObject/Private/UObject/SavePackage/PackageHarvester.cpp
+++ b/Engine/Source/Runtime/CoreUObject/Private/UObject/SavePackage/PackageHarvester.cpp
@@ -531,6 +531,8 @@ FArchive& FPackageHarvester::operator<<(FSoftObjectPath& Value)
 {
 	// We need to harvest NAME_None even if the path isn't valid
 	Value.SerializePath(*this);
+	// Add the soft object path to the list, we need to map invalid soft object path too
+	SaveContext.GetHarvestedRealm(CurrentExportHarvestingRealm).GetSoftObjectPathList().Add(Value);
 	if (Value.IsValid())
 	{
 		FSoftObjectPathThreadContext& ThreadContext = FSoftObjectPathThreadContext::Get();
@@ -539,7 +541,6 @@ FArchive& FPackageHarvester::operator<<(FSoftObjectPath& Value)
 		ESoftObjectPathSerializeType SerializeType = ESoftObjectPathSerializeType::AlwaysSerialize;
 
 		ThreadContext.GetSerializationOptions(ReferencingPackageName, ReferencingPropertyName, CollectType, SerializeType, this);
-
 		if (CollectType != ESoftObjectPathCollectType::NeverCollect && CollectType != ESoftObjectPathCollectType::NonPackage)
 		{
 			// Don't track if this is a never collect path
diff --git a/Engine/Source/Runtime/CoreUObject/Private/UObject/SavePackage/SaveContext.h b/Engine/Source/Runtime/CoreUObject/Private/UObject/SavePackage/SaveContext.h
--- a/Engine/Source/Runtime/CoreUObject/Private/UObject/SavePackage/SaveContext.h
+++ b/Engine/Source/Runtime/CoreUObject/Private/UObject/SavePackage/SaveContext.h
@@ -21,6 +21,7 @@
 #include "UObject/Package.h"
 #include "UObject/SavePackage.h"
 #include "UObject/SavePackage/SavePackageUtilities.h"
+#include "UObject/SoftObjectPath.h"
 #include "UObject/UObjectHash.h"
 #include "UObject/UObjectThreadContext.h"
 
@@ -202,6 +203,16 @@ struct FHarvestedRealm
 		return NamesReferencedFromPackageHeader;
 	}
 
+	const TSet<FSoftObjectPath>& GetSoftObjectPathList() const
+	{
+		return SoftObjectPathList;
+	}
+
+	TSet<FSoftObjectPath>& GetSoftObjectPathList() 
+	{
+		return SoftObjectPathList;
+	}
+
 	const TMap<UObject*, TSet<UObject*>>& GetObjectDependencies() const
 	{
 		return ExportObjectDependencies;
@@ -339,6 +350,8 @@ private:
 	TSet<FNameEntryId> NamesReferencedFromExportData;
 	// Set of names referenced from the package header (import and export table object names etc)
 	TSet<FNameEntryId> NamesReferencedFromPackageHeader;
+	// Set of SoftObjectPath harvested in this realm
+	TSet<FSoftObjectPath> SoftObjectPathList;
 	// List of soft package reference found
 	TSet<FName> SoftPackageReferenceList;
 	// Map of objects to their list of searchable names
@@ -803,6 +816,11 @@ public:
 		return GetHarvestedRealm().GetNamesReferencedFromPackageHeader();
 	}
 
+	const TSet<FSoftObjectPath>& GetSoftObjectPathList() const
+	{
+		return GetHarvestedRealm().GetSoftObjectPathList();
+	}
+
 	const TMap<UObject*, TSet<UObject*>>& GetObjectDependencies() const
 	{
 		return GetHarvestedRealm().GetObjectDependencies();
diff --git a/Engine/Source/Runtime/CoreUObject/Private/UObject/SavePackage2.cpp b/Engine/Source/Runtime/CoreUObject/Private/UObject/SavePackage2.cpp
--- a/Engine/Source/Runtime/CoreUObject/Private/UObject/SavePackage2.cpp
+++ b/Engine/Source/Runtime/CoreUObject/Private/UObject/SavePackage2.cpp
@@ -1048,6 +1048,28 @@ ESavePackageResult BuildLinker(FSaveContext& SaveContext)
 		}
 	}
 
+	// Build SoftObjectPathList
+	{
+		SCOPED_SAVETIMER(UPackage_Save_BuildSoftObjectPathList);
+		Linker->Summary.SoftObjectPathsOffset = 0;
+		Linker->Summary.SoftObjectPathsCount = 0;
+
+		// Do not serialize a soft object path list when cooking.
+		// iostore for example does not keep that list as part of its header information
+		if (!SaveContext.IsCooking())
+		{
+			Linker->SoftObjectPathList = SaveContext.GetSoftObjectPathList().Array();
+		}
+
+		if (!SaveContext.IsTextFormat())
+		{
+			for (int32 Index = 0; Index < Linker->SoftObjectPathList.Num(); ++Index)
+			{
+				Linker->SoftObjectPathIndices.Add(Linker->SoftObjectPathList[Index], Index);
+			}
+		}
+	}
+
 	// Build GatherableText
 	{
 		Linker->Summary.GatherableTextDataOffset = 0;
@@ -1170,7 +1192,7 @@ ESavePackageResult BuildLinker(FSaveContext& SaveContext)
 		}
 	}
 
-	// Build Searchable Name Map
+	// Build SoftPackageReference & Searchable Name Map
 	{
 		Linker->SoftPackageReferenceList = SaveContext.GetSoftPackageReferenceList().Array();
 
@@ -1587,6 +1609,7 @@ ESavePackageResult WritePackageHeader(FStructuredArchive::FRecord& StructuredArc
 #if WITH_EDITOR
 	FArchiveStackTraceIgnoreScope IgnoreDiffScope(SaveContext.IsIgnoringHeaderDiff());
 #endif
+	TGuardValue<bool> Guard(Linker->bIsWritingHeader, true);
 
 	// Write Dummy Summary
 	{
@@ -1605,6 +1628,20 @@ ESavePackageResult WritePackageHeader(FStructuredArchive::FRecord& StructuredArc
 		}
 	}
 
+	// Write Soft Object Paths
+	{
+		SCOPED_SAVETIMER(UPackage_Save_SaveSoftObjectPaths);
+		// Save soft object paths references
+		Linker->Summary.SoftObjectPathsOffset = (int32)Linker->Tell();
+		Linker->Summary.SoftObjectPathsCount = Linker->SoftObjectPathList.Num();
+		// Do not map soft object path during the table serialization itself
+		FStructuredArchive::FStream SoftObjectPathListStream = StructuredArchiveRoot.EnterStream(TEXT("SoftObjectPathList"));
+		for (FSoftObjectPath& Path : Linker->SoftObjectPathList)
+		{
+			SoftObjectPathListStream.EnterElement() << Path;
+		}
+	}
+
 	// Write GatherableText
 	{
 		SCOPED_SAVETIMER(UPackage_Save_WriteGatherableTextData);
diff --git a/Engine/Source/Runtime/CoreUObject/Public/UObject/Linker.h b/Engine/Source/Runtime/CoreUObject/Public/UObject/Linker.h
--- a/Engine/Source/Runtime/CoreUObject/Public/UObject/Linker.h
+++ b/Engine/Source/Runtime/CoreUObject/Public/UObject/Linker.h
@@ -440,6 +440,9 @@ public:
 	/** Names used by objects contained within this package */
 	TArray<FNameEntryId>	NameMap;
 
+	/** List of SoftObjectPath contained in this package,  */
+	TArray<FSoftObjectPath> SoftObjectPathList;
+
 	/** Gatherable text data contained within this package */
 	TArray<FGatherableTextData> GatherableTextDataMap;
 
diff --git a/Engine/Source/Runtime/CoreUObject/Public/UObject/LinkerInstancingContext.h b/Engine/Source/Runtime/CoreUObject/Public/UObject/LinkerInstancingContext.h
--- a/Engine/Source/Runtime/CoreUObject/Public/UObject/LinkerInstancingContext.h
+++ b/Engine/Source/Runtime/CoreUObject/Public/UObject/LinkerInstancingContext.h
@@ -110,12 +110,25 @@ public:
 	}
 
 private:
+	/** Used internally by the linker to try to fix references on relocated packages. */
+	FName RelocatePackage(const FName& PackageName) const
+	{
+		if (const FName* RemappedName = RelocatedPackageMapping.Find(PackageName))
+		{
+			return *RemappedName;
+		}
+		return PackageName;
+	}
+
 	friend class FLinkerLoad;
 
 	/** Map of original package name to their instance counterpart. */
 	TMap<FName, FName> PackageMapping;
 	/** Map of original top level asset path to their instance counterpart. */
 	TMap<FTopLevelAssetPath, FTopLevelAssetPath> PathMapping;
+	/** Map of original package name to their potential relocated counterpart. */
+	TMap<FName, FName> RelocatedPackageMapping;
+
 	/** Tags can be used to determine some loading behavior. */
 	TSet<FName> Tags;
 	/** Remap soft object paths */
diff --git a/Engine/Source/Runtime/CoreUObject/Public/UObject/LinkerLoad.h b/Engine/Source/Runtime/CoreUObject/Public/UObject/LinkerLoad.h
--- a/Engine/Source/Runtime/CoreUObject/Public/UObject/LinkerLoad.h
+++ b/Engine/Source/Runtime/CoreUObject/Public/UObject/LinkerLoad.h
@@ -384,6 +384,8 @@ private:
 
 	// Variables used during async linker creation.
 
+	/** Current index into soft object path list, used by async linker creation for spreading out serializing soft object path entries.	*/
+	int32						SoftObjectPathListIndex;
 	/** Current index into gatherable text data map, used by async linker creation for spreading out serializing text entries.	*/
 	int32						GatherableTextDataMapIndex;
 	/** Current index into import map, used by async linker creation for spreading out serializing importmap entries.			*/
@@ -406,8 +408,10 @@ private:
 	bool					bHasSerializedPreloadDependencies:1;
 	/** Whether we already fixed up import map.																				*/
 	bool					bHasFixedUpImportMap:1;
-	/** Whether we already fixed up import map.																				*/
+	/** Whether we already populated the instancing context.																*/
 	bool					bHasPopulatedInstancingContext:1;
+	/** Whether we already populated the relocation context.																*/
+	bool					bHasPopulatedRelocationContext:1;
 	/** Used for ActiveClassRedirects functionality */
 	bool					bFixupExportMapDone:1;
 	/** Whether we already matched up existing exports.																		*/
@@ -955,25 +959,11 @@ private:
 	}
 
 	virtual FArchive& operator<<(FObjectPtr& ObjectPtr) override;
+	virtual FArchive& operator<<(FSoftObjectPath& Value) override;
 
-	FORCEINLINE virtual FArchive& operator<<(FSoftObjectPtr& Value) override
-	{
-		FArchive& Ar = *this;
-		FSoftObjectPath ID;
-		ID.Serialize(Ar);
-		FixupSoftObjectPathForInstancedPackage(ID);
-		Value = ID;
-		return Ar;
-	}
-
-	FORCEINLINE virtual FArchive& operator<<(FSoftObjectPath& Value) override
-	{
-		FArchive& Ar = FArchiveUObject::operator<<(Value);
-		FixupSoftObjectPathForInstancedPackage(Value);
-		return Ar;
-	}
-
+	void BadSoftObjectPathError(int32 SoftObjIndex);
 	void BadNameIndexError(int32 NameIndex);
+
 	FORCEINLINE virtual FArchive& operator<<(FName& Name) override
 	{
 		FArchive& Ar = *this;
@@ -1098,6 +1088,11 @@ private:
 	 */
 	ELinkerStatus SerializeNameMap();
 
+	/**
+	 * Serializes the soft object map.
+	 */
+	ELinkerStatus SerializeSoftObjectPathList();
+
 	/**
 	 * Serializes the import map.
 	 */
@@ -1113,6 +1108,11 @@ private:
 	 */
 	ELinkerStatus PopulateInstancingContext();
 
+	/**
+	 * Generate remapping for the relocation context if this is a relocated package.
+	 */
+	ELinkerStatus PopulateRelocationContext();
+
 	/**
 	 * Serializes the export map.
 	 */
diff --git a/Engine/Source/Runtime/CoreUObject/Public/UObject/LinkerSave.h b/Engine/Source/Runtime/CoreUObject/Public/UObject/LinkerSave.h
--- a/Engine/Source/Runtime/CoreUObject/Public/UObject/LinkerSave.h
+++ b/Engine/Source/Runtime/CoreUObject/Public/UObject/LinkerSave.h
@@ -20,6 +20,7 @@
 #include "UObject/NameTypes.h"
 #include "UObject/ObjectResource.h"
 #include "UObject/PackageTrailer.h"
+#include "UObject/SoftObjectPath.h"
 #include "UObject/UObjectThreadContext.h"
 
 class FBulkData;
@@ -65,9 +66,15 @@ public:
 	/** List of Searchable Names, by object containing them. This gets turned into package indices later */
 	TMap<const UObject *, TArray<FName> > SearchableNamesObjectMap;
 
-	/** Index array - location of the name in the NameMap array for each FName is stored in the NameIndices array using the FName's Index */
+	/* Map from FName to the index of the name in the name array written into the package header. */
 	TMap<FNameEntryId, int32> NameIndices;
 
+	/* Map from FSoftObjectPath to the index of the path in the soft object path array written into the package header. */
+	TMap<FSoftObjectPath, int32> SoftObjectPathIndices;
+
+	/** Flag that indicate if we are currently serializing the package header. Used to disable mapping while serializing the header itself. */
+	bool bIsWritingHeader = false;
+
 	/** Save context associated with this linker */
 	TRefCountPtr<FUObjectSerializeContext> SaveContext;
 
@@ -157,9 +164,12 @@ public:
 	/** Constructor for custom savers. The linker assumes ownership of the custom saver. */
 	FLinkerSave(UPackage* InParent, FArchive *InSaver, bool bForceByteSwapping, bool bInSaveUnversioned = false);
 
-	/** Returns the appropriate name index for the source name, or 0 if not found in NameIndices */
+	/** Returns the appropriate name index for the source name, or INDEX_NONE if not found in NameIndices */
 	int32 MapName( FNameEntryId Name) const;
 
+	/** Returns the appropriate soft object path index for the source soft object path, or INDEX_NONE if not found. */
+	int32 MapSoftObjectPath(const FSoftObjectPath& SoftObjectPath) const;
+
 	/** Returns the appropriate package index for the source object, or default value if not found in ObjectIndicesMap */
 	FPackageIndex MapObject(const UObject* Object) const;
 
@@ -167,6 +177,7 @@ public:
 	using FArchiveUObject::operator<<; // For visibility of the overloads we don't override
 	FArchive& operator<<( FName& InName );
 	FArchive& operator<<( UObject*& Obj );
+	FArchive& operator<<(FSoftObjectPath& SoftObjectPath);
 	FArchive& operator<<( FLazyObjectPtr& LazyObjectPtr );
 	virtual void SetSerializeContext(FUObjectSerializeContext* InLoadContext) override;
 	FUObjectSerializeContext* GetSerializeContext() override;
diff --git a/Engine/Source/Runtime/CoreUObject/Public/UObject/Package.h b/Engine/Source/Runtime/CoreUObject/Public/UObject/Package.h
--- a/Engine/Source/Runtime/CoreUObject/Public/UObject/Package.h
+++ b/Engine/Source/Runtime/CoreUObject/Public/UObject/Package.h
@@ -234,7 +234,8 @@ private:
 #endif
 
 #if WITH_EDITORONLY_DATA
-	/** Indicates which folder to display this package under in the Generic Browser's list of packages. If not specified, package is added to the root level.	*/
+	/** Indicates which folder to display this package under in the Generic Browser's list of packages. If not specified, package is added to the root level. deprecated	*/
+	UE_DEPRECATED(5.1, "Unused property to be removed. Use something like FPackageName::GetLongPackagePath if requiring something similar.")
 	FName FolderName;
 
 	/** Persistent GUID of package if it was loaded from disk. Persistent across saves. */
@@ -459,6 +460,7 @@ public:
 	* Get the package's folder name
 	* @return		Folder name
 	*/
+	UE_DEPRECATED(5.1, "Unused property to be removed. Use something like FPackageName::GetLongPackagePath if requiring something similar.")
 	FName GetFolderName() const
 	{
 		return FolderName;
@@ -467,6 +469,7 @@ public:
 	/**
 	* Set the package's folder name
 	*/
+	UE_DEPRECATED(5.1, "Unused property to be removed. Use something like FPackageName::GetLongPackagePath if requiring something similar.")
 	void SetFolderName (FName name)
 	{
 		FolderName = name;
diff --git a/Engine/Source/Runtime/CoreUObject/Public/UObject/PackageFileSummary.h b/Engine/Source/Runtime/CoreUObject/Public/UObject/PackageFileSummary.h
--- a/Engine/Source/Runtime/CoreUObject/Public/UObject/PackageFileSummary.h
+++ b/Engine/Source/Runtime/CoreUObject/Public/UObject/PackageFileSummary.h
@@ -76,9 +76,9 @@ public:
 	int32		TotalHeaderSize;
 
 	/**
-	* The Generic Browser folder name that this package lives in
+	* The package name the file was last saved with.
 	*/
-	FString	FolderName;
+	FString	PackageName;
 
 	/**
 	* Number of names used in this package
@@ -90,6 +90,17 @@ public:
 	*/
 	int32 	NameOffset;
 
+	/**
+	* Number of soft object paths references contained in this package
+	* @note: the soft object path lists slightly duplicate information found in the SoftPackageReferences list
+	*/
+	int32		SoftObjectPathsCount;
+
+	/**
+	* Location into the file on disk for the soft object paths reference list
+	*/
+	int32		SoftObjectPathsOffset;
+
 	/**
 	* Localization ID of this package
 	* @note This is copy of the version stored in the package meta-data. It exists here so we can query it without having to load the whole package
