EditorBulkData will now generate a new identifier when loaded into a transient package

#rnx

- This fix an issue where Niagara emitters being loaded from templates could end up with duplicate unique identifier as they were first being duplicated to a transient package before being moved (via UObject::Rename) to the final package. This meant we had no way to produce a unique guid for the editor bulkdata.
- This solution does mean that we are not producing stable guids between multiple runs of the editor (the above scenario occurs every time the Niagara package is loaded) but that requirement is for future features that are not yet online. We will revisit this code in 5.2 to implement a solution that does produce stable guids.

diff --git a/Engine/Source/Runtime/CoreUObject/Private/Serialization/EditorBulkData.cpp b/Engine/Source/Runtime/CoreUObject/Private/Serialization/EditorBulkData.cpp
--- a/Engine/Source/Runtime/CoreUObject/Private/Serialization/EditorBulkData.cpp
+++ b/Engine/Source/Runtime/CoreUObject/Private/Serialization/EditorBulkData.cpp
@@ -211,6 +211,33 @@ FString GetDebugNameFromOwner(UObject* Owner)
 	}
 }
 
+/** 
+ * Utility to check if we need to generate a new unique identifier for the editor bulkdata or not.
+ * At the moment the instancing context (if any) can request that we generate new guids. At the time
+ * of writing the use case for this is creating a new map from a template, and the bulkdata is being
+ * loaded from the template as we don't want every map created from the template to have the same 
+ * identifiers.
+ * The second use case is if we are loading the editor bulkdata to a transient package, which can 
+ * occur when some assets duplicate a number of UObjects from a template package that will eventually
+ * be added to themselves. The duplication will load the editor bulkdata to a transient package before
+ * it is re-parented.
+ */
+bool ShouldGenerateNewIdentifier(FLinkerLoad* LinkerLoad, UObject* Owner)
+{
+	if (LinkerLoad && LinkerLoad->GetInstancingContext().ShouldRegenerateUniqueBulkDataGuids())
+	{
+		return true;
+	}
+
+	UPackage* Package = Owner != nullptr ? Owner->GetPackage() : nullptr;
+	if (Package != nullptr && Package->HasAnyFlags(RF_Transient))
+	{
+		return true;
+	}
+
+	return false;
+}
+
 /** Utility for hashing a payload, will return a default FIoHash if the payload is invalid or of zero length */
 static FIoHash HashPayload(const FSharedBuffer& InPayload)
 {
@@ -963,7 +990,7 @@ void FEditorBulkData::Serialize(FArchive& Ar, UObject* Owner, bool bAllowRegiste
 			{
 				FLinkerLoad* LinkerLoad = Cast<FLinkerLoad>(Ar.GetLinker());
 
-				if (LinkerLoad && LinkerLoad->GetInstancingContext().ShouldRegenerateUniqueBulkDataGuids())
+				if (ShouldGenerateNewIdentifier(LinkerLoad, Owner))
 				{
 					BulkDataId = FGuid::NewGuid();
 				}
