Avoid async loading thread to touch the transaction buffer or to trigger an object modified event

diff --git a/Engine/Source/Runtime/CoreUObject/Private/UObject/Obj.cpp b/Engine/Source/Runtime/CoreUObject/Private/UObject/Obj.cpp
--- a/Engine/Source/Runtime/CoreUObject/Private/UObject/Obj.cpp
+++ b/Engine/Source/Runtime/CoreUObject/Private/UObject/Obj.cpp
@@ -1264,25 +1264,29 @@ bool UObject::Modify( bool bAlwaysMarkDirty/*=true*/ )
 
 	if (CanModify())
 	{
-		// Do not consider script packages, as they should never end up in the
-		// transaction buffer and we don't want to mark them dirty here either.
-		// We do want to consider PIE objects however
-		if ((GetOutermost()->HasAnyPackageFlags(PKG_ContainsScript | PKG_CompiledIn) == false || GetClass()->HasAnyClassFlags(CLASS_DefaultConfig | CLASS_Config)) &&
-			!HasAnyInternalFlags(EInternalObjectFlags::Async | EInternalObjectFlags::AsyncLoading))
-		{
-			// Attempt to mark the package dirty and save a copy of the object to the transaction
-			// buffer. The save will fail if there isn't a valid transactor, the object isn't
-			// transactional, etc.
-			bSavedToTransactionBuffer = SaveToTransactionBuffer(this, bAlwaysMarkDirty);
-
-			// If we failed to save to the transaction buffer, but the user requested the package
-			// marked dirty anyway, do so
-			if (!bSavedToTransactionBuffer && bAlwaysMarkDirty)
+		// Only the game-thread should be allowed to touch the transaction buffer at all
+		if (!HasAnyInternalFlags(EInternalObjectFlags::Async | EInternalObjectFlags::AsyncLoading) && IsInGameThread())
+		{
+			// Do not consider script packages, as they should never end up in the
+			// transaction buffer and we don't want to mark them dirty here either.
+			// We do want to consider PIE objects however
+			if (GetOutermost()->HasAnyPackageFlags(PKG_ContainsScript | PKG_CompiledIn) == false || GetClass()->HasAnyClassFlags(CLASS_DefaultConfig | CLASS_Config))
 			{
-				MarkPackageDirty();
+				// Attempt to mark the package dirty and save a copy of the object to the transaction
+				// buffer. The save will fail if there isn't a valid transactor, the object isn't
+				// transactional, etc.
+				bSavedToTransactionBuffer = SaveToTransactionBuffer(this, bAlwaysMarkDirty);
+
+				// If we failed to save to the transaction buffer, but the user requested the package
+				// marked dirty anyway, do so
+				if (!bSavedToTransactionBuffer && bAlwaysMarkDirty)
+				{
+					MarkPackageDirty();
+				}
 			}
+
+			FCoreUObjectDelegates::BroadcastOnObjectModified(this);
 		}
-		FCoreUObjectDelegates::BroadcastOnObjectModified(this);
 	}
 
 	return bSavedToTransactionBuffer;
