Suppress the BulkData duplicate guid warning until we have eliminated all sources of duplication.
#rnx

diff --git a/Engine/Source/Runtime/CoreUObject/Private/Serialization/EditorBulkData.cpp b/Engine/Source/Runtime/CoreUObject/Private/Serialization/EditorBulkData.cpp
--- a/Engine/Source/Runtime/CoreUObject/Private/Serialization/EditorBulkData.cpp
+++ b/Engine/Source/Runtime/CoreUObject/Private/Serialization/EditorBulkData.cpp
@@ -656,7 +656,9 @@ void FEditorBulkData::LogRegisterError(UE::BulkDataRegistry::ERegisterResult Val
 						TEXT("\"-autocheckout -resaveondemand=bulkdataduplicates -SaveAll -Package=%s -Package=%s\"."),
 						*FPackageName::ObjectPathToPackageName(OwnerPathName), *OtherOwnerPackageName);
 				}
-				bool bSuppressWarning = false; // Allow projects to suppress the warning because they might not be able to resave packages during an integration
+				// Suppress the warning by default for now. Once we have eliminated all sources of duplication, turn it back on, but allow
+				// projects to suppress the warning even then because they might not be able to resave packages during an integration
+				bool bSuppressWarning = true; 
 				GConfig->GetBool(TEXT("CookSettings"), TEXT("BulkDataRegistrySuppressDuplicateWarning"), bSuppressWarning, GEditorIni);
 				FString LogMessage = FString::Printf(TEXT("%s updated BulkData %s on load because it collided with an ID in package %s.%s"),
 					*OwnerPathName, *FailedBulkDataId.ToString(), *OtherOwnerPackageName, *SilenceWarningMessage);
