Add more logging to investigate shader decompression errors (UE-159777).

- Hexdump of the whole compressed shader will be dumped if small enough (which is what we're seeing mostly).

diff --git a/Engine/Source/Runtime/Core/Private/Compression/CompressionUtil.cpp b/Engine/Source/Runtime/Core/Private/Compression/CompressionUtil.cpp
--- a/Engine/Source/Runtime/Core/Private/Compression/CompressionUtil.cpp
+++ b/Engine/Source/Runtime/Core/Private/Compression/CompressionUtil.cpp
@@ -94,5 +94,27 @@ namespace FCompressionUtil
 		}
 	}
 
+	void CORE_API LogHexDump(const uint8* Bytes, int64 BytesNum, int64 OffsetStart, int64 OffsetEnd, int32 BytesPerLine)
+	{
+		if (UE_LOG_ACTIVE(LogCore, Display))
+		{
+			OffsetStart = FMath::Max(0ll, OffsetStart);
+			OffsetEnd = FMath::Min(BytesNum, OffsetEnd);
 
+			for (int64 Idx = OffsetStart; Idx < OffsetEnd;)
+			{
+				int64 LineOffset = OffsetStart;
+				FString HexString;
+				for (int64 Idx2 = 0; Idx2 < BytesPerLine && Idx < OffsetEnd; ++Idx, ++Idx2, ++OffsetStart)
+				{
+					HexString += FString::Printf(TEXT("%02X "), Bytes[Idx]);
+					if ((Idx2 & 7) == 7)
+					{
+						HexString += TEXT(" ");
+					}
+				}
+				UE_LOG(LogCore, Display, TEXT("%016X: %s"), LineOffset, *HexString);
+			}
+		}
+	}
 };
diff --git a/Engine/Source/Runtime/Core/Private/Misc/Compression.cpp b/Engine/Source/Runtime/Core/Private/Misc/Compression.cpp
--- a/Engine/Source/Runtime/Core/Private/Misc/Compression.cpp
+++ b/Engine/Source/Runtime/Core/Private/Misc/Compression.cpp
@@ -15,6 +15,7 @@
 #include "Misc/MemoryReadStream.h"
 
 #include "Compression/OodleDataCompression.h"
+#include "Compression/CompressionUtil.h" // for extra logging only
 
 // #include "TargetPlatformBase.h"
 THIRD_PARTY_INCLUDES_START
@@ -760,6 +761,13 @@ bool FCompression::UncompressMemory(FName FormatName, void* UncompressedBuffer,
 		}
 		// Always log an error
 		UE_LOG(LogCompression, Error, TEXT("FCompression::UncompressMemory - Failed to uncompress memory (%d/%d) from address %p using format %s, this may indicate the asset is corrupt!"), CompressedSize, UncompressedSize, CompressedBuffer, *FormatName.ToString());
+		// this extra logging is added to understand shader decompression errors, see UE-159777
+		const int32 MaxSizeToLogOut = 16384;
+		if (FormatName == NAME_Oodle && CompressedSize <= MaxSizeToLogOut)
+		{
+			UE_LOG(LogCompression, Error, TEXT("FCompression::UncompressMemory - Logging compressed data (%d bytes) as a hex dump for investigation"), CompressedSize);
+			FCompressionUtil::LogHexDump(reinterpret_cast<const uint8*>(CompressedBuffer), CompressedSize, 0, CompressedSize);
+		}
 	}
 
 #if	STATS
diff --git a/Engine/Source/Runtime/Core/Public/Compression/CompressionUtil.h b/Engine/Source/Runtime/Core/Public/Compression/CompressionUtil.h
--- a/Engine/Source/Runtime/Core/Public/Compression/CompressionUtil.h
+++ b/Engine/Source/Runtime/Core/Public/Compression/CompressionUtil.h
@@ -16,4 +16,15 @@ namespace FCompressionUtil
 	void CORE_API SerializeCompressorName(FArchive & Archive,FName & Compressor);
 
 
+	/** Logs out a memory region as a hex dump (log category: LogCore, level: Display).
+	 *  Bytes, BytesNum are used to designate the overall memory region of interest (e.g. a file in memory), but only bytes from OffsetStart to OffsetEnd will be logged out,
+	 * and current offset will be printed at each line's start.
+	 * 
+	 * @param Bytes memory location of the data 
+	 * @param BytesNum total size of the data
+	 * @param OffsetStart offset to first byte to be logged out (pass 0 if you want to log out whole region)
+	 * @param OffsetEnd offset to the end of log window (pass BytesNum if you want to log out whole region)
+	 * @param BytesPerLine how many bytes to print per line of hex dump
+	 */
+	void CORE_API LogHexDump(const uint8* Bytes, int64 BytesNum, int64 OffsetStart, int64 OffsetEnd, int32 BytesPerLine = 32);
 };
\ No newline at end of file
diff --git a/Engine/Source/Runtime/CoreUObject/Private/Serialization/ArchiveStackTrace.cpp b/Engine/Source/Runtime/CoreUObject/Private/Serialization/ArchiveStackTrace.cpp
--- a/Engine/Source/Runtime/CoreUObject/Private/Serialization/ArchiveStackTrace.cpp
+++ b/Engine/Source/Runtime/CoreUObject/Private/Serialization/ArchiveStackTrace.cpp
@@ -21,6 +21,7 @@
 #include "Templates/UniquePtr.h"
 #include "UObject/UObjectGlobals.h"
 #include "Misc/ScopeExit.h"
+#include "Compression/CompressionUtil.h"
 
 DEFINE_LOG_CATEGORY_STATIC(LogArchiveDiff, Log, All);
 
@@ -501,30 +502,6 @@ namespace
 	}
 }
 
-namespace ArchiveStackTraceUtils
-{
-	void LogHexDump(const uint8* Bytes, int64 BytesNum, int64 OffsetStart, int64 OffsetEnd)
-	{
-		OffsetStart = FMath::Max(0ll, OffsetStart);
-		OffsetEnd = FMath::Min(BytesNum, OffsetEnd);
-
-		for (int64 Idx = OffsetStart; Idx < OffsetEnd;)
-		{
-			int64 LineOffset = OffsetStart;
-			FString HexString;
-			for (int64 Idx2 = 0; Idx2 < 32 && Idx < OffsetEnd; ++Idx, ++Idx2, ++OffsetStart)
-			{
-				HexString += FString::Printf(TEXT("%02X "), Bytes[Idx]);
-				if ((Idx2 & 7) == 7)
-				{
-					HexString += TEXT(" ");
-				}
-			}
-			UE_LOG(LogArchiveDiff, Display, TEXT("%016X: %s"), LineOffset, *HexString);
-		}
-	}
-}
-
 void FArchiveStackTrace::CompareWithInternal(const FPackageData& SourcePackage, const FPackageData& DestPackage, const TCHAR* AssetFilename, const TCHAR* CallstackCutoffText, const int64 MaxDiffsToLog, int32& InOutDiffsLogged, TMap<FName, FArchiveDiffStats>& OutStats)
 {
 #if !NO_LOGGING
@@ -689,7 +666,7 @@ void FArchiveStackTrace::CompareWithInternal(const FPackageData& SourcePackage,
 					DestAbsoluteOffset,
 					DestAbsoluteOffset
 				);
-				ArchiveStackTraceUtils::LogHexDump(SourcePackage.Data, SourcePackage.Size, SourceAbsoluteOffset - BytesToLog / 2, SourceAbsoluteOffset + BytesToLog / 2);
+				FCompressionUtil::LogHexDump(SourcePackage.Data, SourcePackage.Size, SourceAbsoluteOffset - BytesToLog / 2, SourceAbsoluteOffset + BytesToLog / 2);
 
 				UE_LOG(
 					LogArchiveDiff,
@@ -700,7 +677,7 @@ void FArchiveStackTrace::CompareWithInternal(const FPackageData& SourcePackage,
 					DestAbsoluteOffset,
 					DestAbsoluteOffset
 				);
-				ArchiveStackTraceUtils::LogHexDump(DestPackage.Data, DestPackage.Size, DestAbsoluteOffset - BytesToLog / 2, DestAbsoluteOffset + BytesToLog / 2);
+				FCompressionUtil::LogHexDump(DestPackage.Data, DestPackage.Size, DestAbsoluteOffset - BytesToLog / 2, DestAbsoluteOffset + BytesToLog / 2);
 
 				bDifferenceLogged = true;
 			}
