[AssetRegistry] added two new methods `TryGetAssetByObjectPath` and `TryGetAssetPackageData` to allow lazy load to be able to accurately validate that objects and paths exist
otherwise an unresolve `TObjectPtr` could resolve to nullptr causing determinism issues

#rnx

diff --git a/Engine/Source/Runtime/AssetRegistry/Private/AssetRegistry.cpp b/Engine/Source/Runtime/AssetRegistry/Private/AssetRegistry.cpp
--- a/Engine/Source/Runtime/AssetRegistry/Private/AssetRegistry.cpp
+++ b/Engine/Source/Runtime/AssetRegistry/Private/AssetRegistry.cpp
@@ -120,6 +120,26 @@ public:
 		IAssetRegistry::GetChecked().GetDependencies(InPackageName, OutDependencies, Category, Flags);
 	}
 
+	virtual UE::AssetRegistry::EExists TryGetAssetByObjectPath(const FName ObjectPath, FAssetData& OutAssetData) const override
+	{
+		auto AssetRegistry = IAssetRegistry::Get();
+		if (!AssetRegistry)
+		{
+			return UE::AssetRegistry::EExists::Unknown;
+		}
+		return AssetRegistry->TryGetAssetByObjectPath(ObjectPath, OutAssetData);
+	}
+
+	virtual UE::AssetRegistry::EExists TryGetAssetPackageData(FName PackageName, class FAssetPackageData& OutPackageData) const override
+	{
+		auto AssetRegistry = IAssetRegistry::Get();
+		if (!AssetRegistry)
+		{
+			return UE::AssetRegistry::EExists::Unknown;
+		}
+		return AssetRegistry->TryGetAssetPackageData(PackageName, OutPackageData);
+	}
+
 protected:
 	/* This function is a workaround for platforms that don't support disable of deprecation warnings on override functions*/
 	virtual void GetDependenciesDeprecated(FName InPackageName, TArray<FName>& OutDependencies, EAssetRegistryDependencyType::Type InDependencyType) override
@@ -2061,6 +2081,42 @@ PRAGMA_DISABLE_DEPRECATION_WARNINGS;
 PRAGMA_ENABLE_DEPRECATION_WARNINGS;
 }
 
+UE::AssetRegistry::EExists UAssetRegistryImpl::TryGetAssetByObjectPath(const FName ObjectPath, FAssetData& OutAssetData) const
+{
+	FReadScopeLock InterfaceScopeLock(InterfaceLock);
+	bool bAssetRegistryReady = GuardedData.IsInitialSearchStarted() && GuardedData.IsInitialSearchCompleted();
+	const FAssetRegistryState& State = GuardedData.GetState();
+	const FAssetData* FoundData = State.GetAssetByObjectPath(ObjectPath);
+	if (!FoundData)
+	{
+		if (!bAssetRegistryReady)
+		{
+			return UE::AssetRegistry::EExists::Unknown;
+		}
+		return UE::AssetRegistry::EExists::DoesNotExist;
+	}
+	OutAssetData = *FoundData;
+	return UE::AssetRegistry::EExists::Exists;
+}
+
+UE::AssetRegistry::EExists UAssetRegistryImpl::TryGetAssetPackageData(const FName PackageName, FAssetPackageData& OutAssetPackageData) const
+{
+	FReadScopeLock InterfaceScopeLock(InterfaceLock);
+	bool bAssetRegistryReady = GuardedData.IsInitialSearchStarted() && GuardedData.IsInitialSearchCompleted();
+	const FAssetRegistryState& State = GuardedData.GetState();
+	const FAssetPackageData* FoundData = State.GetAssetPackageData(PackageName);
+	if (!FoundData)
+	{
+		if (!bAssetRegistryReady)
+		{
+			return UE::AssetRegistry::EExists::Unknown;
+		}
+		return UE::AssetRegistry::EExists::DoesNotExist;
+	}
+	OutAssetPackageData = *FoundData;
+	return UE::AssetRegistry::EExists::Exists;
+}
+
 bool UAssetRegistryImpl::GetAllAssets(TArray<FAssetData>& OutAssetData, bool bIncludeOnlyOnDiskAssets) const
 {
 	const double GetAllAssetsStartTime = FPlatformTime::Seconds();
diff --git a/Engine/Source/Runtime/AssetRegistry/Private/AssetRegistry.h b/Engine/Source/Runtime/AssetRegistry/Private/AssetRegistry.h
--- a/Engine/Source/Runtime/AssetRegistry/Private/AssetRegistry.h
+++ b/Engine/Source/Runtime/AssetRegistry/Private/AssetRegistry.h
@@ -45,6 +45,8 @@ public:
 	UE_DEPRECATED(5.1, "Asset path FNames have been deprecated, use FSoftObjectPath instead.")
 	virtual FAssetData GetAssetByObjectPath( const FName ObjectPath, bool bIncludeOnlyOnDiskAssets = false ) const override;
 	virtual FAssetData GetAssetByObjectPath(const FSoftObjectPath& ObjectPath, bool bIncludeOnlyOnDiskAssets = false) const override;
+	virtual UE::AssetRegistry::EExists TryGetAssetByObjectPath(const FName ObjectPath, FAssetData& OutAssetData) const override;
+	virtual UE::AssetRegistry::EExists TryGetAssetPackageData(const FName PackageName, FAssetPackageData& OutAssetPackageData) const override;
 	virtual bool GetAllAssets(TArray<FAssetData>& OutAssetData, bool bIncludeOnlyOnDiskAssets = false) const override;
 	virtual bool EnumerateAllAssets(TFunctionRef<bool(const FAssetData&)> Callback, bool bIncludeOnlyOnDiskAssets = false) const override;
 	virtual void GetPackagesByName(FStringView PackageName, TArray<FName>& OutPackageNames) const override;
diff --git a/Engine/Source/Runtime/AssetRegistry/Public/AssetRegistry/AssetRegistryModule.h b/Engine/Source/Runtime/AssetRegistry/Public/AssetRegistry/AssetRegistryModule.h
--- a/Engine/Source/Runtime/AssetRegistry/Public/AssetRegistry/AssetRegistryModule.h
+++ b/Engine/Source/Runtime/AssetRegistry/Public/AssetRegistry/AssetRegistryModule.h
@@ -96,6 +96,25 @@ public:
 		IAssetRegistry::GetChecked().GetDependencies(InPackageName, OutDependencies, Category, Flags);
 	}
 
+	virtual UE::AssetRegistry::EExists TryGetAssetByObjectPath(const FName ObjectPath, FAssetData& OutAssetData) const override
+	{
+		IAssetRegistry* AssetRegistry = IAssetRegistry::Get();
+		if (!AssetRegistry)
+		{
+			return UE::AssetRegistry::EExists::Unknown;
+		}
+		return AssetRegistry->TryGetAssetByObjectPath(ObjectPath, OutAssetData);
+	}
+
+	virtual UE::AssetRegistry::EExists TryGetAssetPackageData(FName PackageName, FAssetPackageData& OutAssetPackageData) const override
+	{
+		IAssetRegistry* AssetRegistry = IAssetRegistry::Get();
+		if (!AssetRegistry)
+		{
+			return UE::AssetRegistry::EExists::Unknown;
+		}
+		return AssetRegistry->TryGetAssetPackageData(PackageName, OutAssetPackageData);
+	}
 
 protected:
 	/* This function is a workaround for platforms that don't support disable of deprecation warnings on override functions*/
diff --git a/Engine/Source/Runtime/AssetRegistry/Public/AssetRegistry/IAssetRegistry.h b/Engine/Source/Runtime/AssetRegistry/Public/AssetRegistry/IAssetRegistry.h
--- a/Engine/Source/Runtime/AssetRegistry/Public/AssetRegistry/IAssetRegistry.h
+++ b/Engine/Source/Runtime/AssetRegistry/Public/AssetRegistry/IAssetRegistry.h
@@ -242,6 +242,24 @@ public:
 	 */
 	virtual FAssetData GetAssetByObjectPath(const FSoftObjectPath& ObjectPath, bool bIncludeOnlyOnDiskAssets = false) const = 0;
 
+	/**
+	 * Tries to get the asset data for the specified object path
+	 * 
+	 * @param ObjectPath the path of the object to be looked up
+	 * @param OutAssetData out FAssetData 
+	 * @return Enum return code
+	 */
+	virtual UE::AssetRegistry::EExists TryGetAssetByObjectPath(const FName ObjectPath, FAssetData& OutAssetData) const = 0;
+
+	/**
+	 * Tries to get the pacakge data for a specified path
+	 *
+	 * @param PackageName name of the package
+	 * @param OutAssetPackageData out FAssetPackageData
+	 * @return Enum return code
+	 */
+	virtual UE::AssetRegistry::EExists TryGetAssetPackageData(FName PackageName, FAssetPackageData& OutAssetPackageData) const = 0;
+
 	/**
 	 * Gets asset data for all assets in the registry.
 	 * This method may be slow, use a filter if possible to avoid iterating over the entire registry.
diff --git a/Engine/Source/Runtime/CoreUObject/Internal/UObject/LinkerLoadImportBehavior.h b/Engine/Source/Runtime/CoreUObject/Internal/UObject/LinkerLoadImportBehavior.h
--- a/Engine/Source/Runtime/CoreUObject/Internal/UObject/LinkerLoadImportBehavior.h
+++ b/Engine/Source/Runtime/CoreUObject/Internal/UObject/LinkerLoadImportBehavior.h
@@ -8,6 +8,7 @@
 
 class FLinkerLoad;
 struct FObjectImport;
+class IAssetRegistryInterface;
 
 namespace UE::LinkerLoad
 {
@@ -21,6 +22,24 @@ enum class EImportBehavior : uint8
 
 
 EImportBehavior GetPropertyImportLoadBehavior(const FObjectImport& Import, const FLinkerLoad& LinkerLoad);
+
+/// @brief tries to lazy load an imported FObjectPtr
+/// 
+/// @param AssetRegistry to see if import exists
+/// @param Import to try to lazy load
+/// @param LinkerLoad for the Import
+/// @param ObjectPtr out parameter
+/// @return true if import was lazy loaded otherwise false
+bool TryLazyImport(const IAssetRegistryInterface& AssetRegistry, const FObjectImport& Import, const FLinkerLoad& LinkerLoad, FObjectPtr& ObjectPtr);
+
+/// @brief tries to resovle an object path to a lazy loaded TObjectPtr
+///
+/// @param Class type of the object that will be resolved
+/// @param ObjectPath object path of to resolve
+/// @param OutObjectPtr out parameter
+/// @return true if was lazy loaded otherwise false
+bool TryLazyLoad(const UClass& Class, FName ObjectPath, TObjectPtr<UObject>& OutObjectPtr);
+
 }
 
 #undef UE_API
diff --git a/Engine/Source/Runtime/CoreUObject/Private/UObject/LinkerLoad.cpp b/Engine/Source/Runtime/CoreUObject/Private/UObject/LinkerLoad.cpp
--- a/Engine/Source/Runtime/CoreUObject/Private/UObject/LinkerLoad.cpp
+++ b/Engine/Source/Runtime/CoreUObject/Private/UObject/LinkerLoad.cpp
@@ -1,6 +1,7 @@
 // Copyright Epic Games, Inc. All Rights Reserved.
 
 #include "UObject/LinkerLoad.h"
+#include "AssetRegistry/AssetData.h"
 #include "HAL/FileManager.h"
 #include "Misc/Paths.h"
 #include "Stats/StatsMisc.h"
@@ -50,6 +51,7 @@
 #include "Serialization/EditorBulkData.h"
 #include "HAL/FileManager.h"
 #include "UObject/CoreRedirects.h"
+#include "Misc/AssetRegistryInterface.h"
 #include "Misc/StringBuilder.h"
 #include "Misc/EngineBuildSettings.h"
 #include "Internationalization/GatherableTextData.h"
@@ -5766,32 +5768,34 @@ FArchive& FLinkerLoad::operator<<( UObject*& Object )
 	return *this;
 }
 
-FArchive& FLinkerLoad::operator<<( FObjectPtr& ObjectPtr )
+FArchive& FLinkerLoad::operator<<(FObjectPtr& ObjectPtr)
 {
 	FPackageIndex Index;
 	FArchive& Ar = *this;
 	Ar << Index;
 
-#if UE_WITH_OBJECT_HANDLE_LATE_RESOLVE
-	if (Index.IsImport())
-	{
-		using namespace UE::LinkerLoad;
-		if (GetPropertyImportLoadBehavior(Imp(Index), *this) == EImportBehavior::LazyOnDemand)
-		{
-			const FObjectImport& Import = Imp(Index);
-			FObjectPathId ObjectPath;
-			FName PackageName = FObjectPathId::MakeImportPathIdAndPackageName(Import, *this, ObjectPath);
-			FObjectRef ImportRef{PackageName, Import.ClassPackage, Import.ClassName, ObjectPath};
-			ObjectPtr = FObjectPtr(ImportRef);
+#if !UE_WITH_OBJECT_HANDLE_LATE_RESOLVE
+	ObjectPtr = FObjectPtr(ResolveResource(Index));
+	return *this;
 
-			return *this;
-		}
+#else
+	auto AssetRegistry = IAssetRegistryInterface::GetPtr();
+
+	if (!Index.IsImport() || !AssetRegistry)
+	{
+		ObjectPtr = FObjectPtr(ResolveResource(Index));
+		return *this;
 	}
-#endif
 
-	ObjectPtr = FObjectPtr(ResolveResource(Index));
-	
+	using namespace UE::LinkerLoad;
+	FObjectImport& Import = Imp(Index);
+	if (!TryLazyImport(*AssetRegistry, Import, *this, ObjectPtr))
+	{
+		ObjectPtr = FObjectPtr(ResolveResource(Index));
+	}
 	return *this;
+	
+#endif
 }
 
 FArchive& FLinkerLoad::operator<<(FSoftObjectPath& Value)
diff --git a/Engine/Source/Runtime/CoreUObject/Private/UObject/LinkerLoadImportBehavior.cpp b/Engine/Source/Runtime/CoreUObject/Private/UObject/LinkerLoadImportBehavior.cpp
--- a/Engine/Source/Runtime/CoreUObject/Private/UObject/LinkerLoadImportBehavior.cpp
+++ b/Engine/Source/Runtime/CoreUObject/Private/UObject/LinkerLoadImportBehavior.cpp
@@ -1,9 +1,12 @@
 // Copyright Epic Games, Inc. All Rights Reserved.
 
 #include "UObject/LinkerLoadImportBehavior.h"
+#include "AssetRegistry/AssetData.h"
+#include "UObject/ConstructorHelpers.h"
 #include "UObject/LinkerLoad.h"
 #include "UObject/ObjectResource.h"
 #include "UObject/Package.h"
+#include "Misc/AssetRegistryInterface.h"
 #include "Misc/CommandLine.h"
 
 #if UE_WITH_OBJECT_HANDLE_LATE_RESOLVE
@@ -15,6 +18,12 @@ namespace UE::LinkerLoad
 /// @return Eager by default in not found
 static EImportBehavior FindLoadBehavior(const UClass& Class)
 {
+	//Package class can't have meta data because of UHT
+	if (&Class == UPackage::StaticClass())
+	{
+		return EImportBehavior::LazyOnDemand;
+	}
+
 	static const FName Name_LoadBehavior("LoadBehavior");
 	if (const FString* LoadBehaviorMeta = Class.FindMetaData(Name_LoadBehavior))
 	{
@@ -26,6 +35,7 @@ static EImportBehavior FindLoadBehavior(const UClass& Class)
 	}
 	else
 	{
+		//look in super class to see if it has lazy load on
 		const UClass* Super = Class.GetSuperClass();
 		if (Super != nullptr)
 		{
@@ -55,12 +65,6 @@ EImportBehavior GetPropertyImportLoadBehavior(const FObjectImport& Import, const
 		return EImportBehavior::LazyOnDemand;
 	}
 
-	//Packages can't have meta data because they 
-	if (Import.ClassName == TEXT("Package"))
-	{
-		return EImportBehavior::LazyOnDemand;
-	}
-
 	TRACE_CPUPROFILER_EVENT_SCOPE(LinkerLoader::GetPropertyImportLoadBehavior);
 
 	UObject* ClassPackage = FindObjectFast<UPackage>(nullptr, Import.ClassPackage);
@@ -71,6 +75,144 @@ EImportBehavior GetPropertyImportLoadBehavior(const FObjectImport& Import, const
 	}
 	return EImportBehavior::Eager;
 }
+
+// recursively handles FAssetData redirectors
+static bool HandleRedirector(const IAssetRegistryInterface& AssetRegistry, const FAssetData& InAssetData, TSet<FName>& SeenPaths, FObjectPtr& OutObjectPtr)
+{
+	FString RedirectedPath;
+	if (!InAssetData.GetTagValue("DestinationObject", RedirectedPath))
+	{
+		return false;
+	}
+
+	ConstructorHelpers::StripObjectClass(RedirectedPath);
+	FName RedirectedPathFName(RedirectedPath);
+
+	bool bIsAlreadyInSet;
+	SeenPaths.FindOrAdd(RedirectedPathFName, &bIsAlreadyInSet);
+	if(bIsAlreadyInSet)
+	{
+		//Recursive redirectors
+		return false;
+	}
+
+	FAssetData AssetData;
+	UE::AssetRegistry::EExists Exists = AssetRegistry.TryGetAssetByObjectPath(FName(RedirectedPath), AssetData);
+	if (Exists == UE::AssetRegistry::EExists::Unknown)
+	{
+		return false;
+	}
+	if (Exists == UE::AssetRegistry::EExists::DoesNotExist)
+	{
+		//redirector pointed to something that doesn't exist
+		//return true to resolve the ObjectPtr to null
+		return true; 
+	}
+	if (AssetData.IsRedirector())
+	{
+		return HandleRedirector(AssetRegistry, AssetData, SeenPaths, OutObjectPtr);
+	}
+
+	FNameBuilder NameBuilder;
+	AssetData.AssetName.AppendString(NameBuilder);
+	FObjectRef ImportRef = { AssetData.PackageName, AssetData.AssetClassPath.GetPackageName(), AssetData.AssetClassPath.GetAssetName(), FObjectPathId(NameBuilder) };
+	FObjectPtr Ptr(ImportRef);
+	OutObjectPtr = Ptr;
+	return true;
+}
+
+static bool TryLazyLoad(const IAssetRegistryInterface& AssetRegistry, FName ObjectPath, FObjectPtr& OutObjectPtr)
+{
+	OutObjectPtr = nullptr;
+	
+	FAssetData AssetData;
+	UE::AssetRegistry::EExists Exists = AssetRegistry.TryGetAssetByObjectPath(ObjectPath, AssetData);
+	if (Exists == UE::AssetRegistry::EExists::Unknown)
+	{
+		return false;
+	}
+	if (Exists == UE::AssetRegistry::EExists::DoesNotExist)
+	{
+		//object path doesn't exist. return true that it can be lazy loaded
+		return true;
+	}
+	if (AssetData.IsRedirector())
+	{
+		TSet<FName> SeenPaths;
+		return HandleRedirector(AssetRegistry, AssetData, SeenPaths, OutObjectPtr);
+	}
+	FNameBuilder NameBuilder;
+	AssetData.AssetName.AppendString(NameBuilder);
+
+	FObjectRef ImportRef = { AssetData.PackageName, AssetData.AssetClassPath.GetPackageName(), AssetData.AssetClassPath.GetAssetName(), FObjectPathId(NameBuilder) };
+	FObjectPtr Ptr(ImportRef);
+	return true;
+}
+
+bool TryLazyImport(const IAssetRegistryInterface& AssetRegistry, const FObjectImport& Import, const FLinkerLoad& LinkerLoad, FObjectPtr& ObjectPtr)
+{
+	EImportBehavior Behavior = GetPropertyImportLoadBehavior(Import, LinkerLoad);
+	if (Behavior != EImportBehavior::LazyOnDemand)
+	{
+		return false;
+	}
+
+	ObjectPtr = nullptr;
+
+	//packages need to handles differently since the can't be found by object path in the AssetRegistry
+	if (Import.ClassPackage == TEXT("/Script/CoreUObject") && Import.ClassName == TEXT("Package"))
+	{
+		FAssetPackageData PackageData;
+		UE::AssetRegistry::EExists Exists = AssetRegistry.TryGetAssetPackageData(Import.ObjectName, PackageData);
+		if (Exists == UE::AssetRegistry::EExists::Unknown)
+		{
+			return false;
+		}
+		if (Exists == UE::AssetRegistry::EExists::DoesNotExist)
+		{
+			//package doesn't exist. resolve to nullptr
+			return true;
+		}
+
+		FObjectPathId ObjectPath;
+		FObjectRef ImportRef{ Import.ObjectName, Import.ClassPackage, Import.ClassName, ObjectPath };
+		ObjectPtr = FObjectPtr(ImportRef);
+		return true;
+	}
+
+	//build the a full objectpath for the AssetRegistry
+	FObjectPathId ObjectPath;
+	FName PackageName = FObjectPathId::MakeImportPathIdAndPackageName(Import, LinkerLoad, ObjectPath);
+	FObjectRef ImportRef{ PackageName, Import.ClassPackage, Import.ClassName, ObjectPath };
+
+	TStringBuilder<FName::StringBufferSize> PathName;
+	ImportRef.AppendPathName(PathName);
+
+	return TryLazyLoad(AssetRegistry, FName(PathName), ObjectPtr);
+}
+
+
+bool TryLazyLoad(const UClass& Class, FName ObjectPath, TObjectPtr<UObject>& OutObjectPtr)
+{
+	if (!FLinkerLoad::IsImportLazyLoadEnabled())
+	{
+		return false;
+	}
+
+	EImportBehavior Behavior = FindLoadBehavior(Class);
+	if (Behavior != EImportBehavior::LazyOnDemand)
+	{
+		return false;
+	}
+	auto AssetRegistry = IAssetRegistryInterface::GetPtr();
+	if (!AssetRegistry)
+	{
+		return false;
+	}
+	FObjectPtr& ObjectPtr = reinterpret_cast<FObjectPtr&>(OutObjectPtr);
+	return TryLazyLoad(*AssetRegistry, ObjectPath, ObjectPtr);
+}
+
 }
 
 #endif // UE_WITH_OBJECT_HANDLE_LATE_RESOLVE
diff --git a/Engine/Source/Runtime/CoreUObject/Private/UObject/ObjectPathId.cpp b/Engine/Source/Runtime/CoreUObject/Private/UObject/ObjectPathId.cpp
--- a/Engine/Source/Runtime/CoreUObject/Private/UObject/ObjectPathId.cpp
+++ b/Engine/Source/Runtime/CoreUObject/Private/UObject/ObjectPathId.cpp
@@ -194,8 +194,7 @@ FObjectPathId::FObjectPathId(const FObjectImport& Import, const FLinkerTables& L
 FName FObjectPathId::MakeImportPathIdAndPackageName(const FObjectImport& Import, const FLinkerTables& LinkerTables, FObjectPathId& OutPathId)
 {
 	TRACE_CPUPROFILER_EVENT_SCOPE(FObjectPathId::MakeImportPathIdAndPackageName);
-	// @TODO: OBJPTR: Need to handle redirects.  FCoreRedirectObjectName could be used, but it doesn't fit
-	//		conveniently with the FName walk approach that is currently here. 
+	
 	struct FImportObjectNamePathProducer
 	{
 		FImportObjectNamePathProducer(const FObjectImport& InImport, const FLinkerTables& InLinkerTables)
@@ -231,7 +230,7 @@ FName FObjectPathId::MakeImportPathIdAndPackageName(const FObjectImport& Import,
 	return NamePathProducer.GetPackageName();
 }
 
-#if WITH_LOW_LEVEL_TESTS
+#if !(UE_BUILD_SHIPPING || UE_BUILD_TEST)
 template <typename CharType>
 struct TStringViewNamePathProducer
 {
@@ -277,22 +276,15 @@ private:
 
 FObjectPathId::FObjectPathId(FWideStringView StringPath)
 {
-	// @TODO: OBJPTR: Need to handle redirects.  FCoreRedirectObjectName could be used, but it doesn't fit
-	//		conveniently with the FName walk approach that is currently here.
 	TStringViewNamePathProducer<WIDECHAR> NamePathProducer(StringPath);
-
 	StoreObjectPathId(NamePathProducer, static_cast<uint64>(EPathId::FlagSimple), PathId);
 }
 
 FObjectPathId::FObjectPathId(FAnsiStringView StringPath)
 {
-	// @TODO: OBJPTR: Need to handle redirects.  FCoreRedirectObjectName could be used, but it doesn't fit
-	//		conveniently with the FName walk approach that is currently here.
 	TStringViewNamePathProducer<ANSICHAR> NamePathProducer(StringPath);
-
 	StoreObjectPathId(NamePathProducer, static_cast<uint64>(EPathId::FlagSimple), PathId);
 }
-#endif
 
 void FObjectPathId::Resolve(ResolvedNameContainerType& OutContainer) const
 {
diff --git a/Engine/Source/Runtime/CoreUObject/Public/Misc/AssetRegistryInterface.h b/Engine/Source/Runtime/CoreUObject/Public/Misc/AssetRegistryInterface.h
--- a/Engine/Source/Runtime/CoreUObject/Public/Misc/AssetRegistryInterface.h
+++ b/Engine/Source/Runtime/CoreUObject/Public/Misc/AssetRegistryInterface.h
@@ -139,6 +139,16 @@ namespace UE::AssetRegistry
 		ForDevelopment
 	};
 
+	/**
+	 * Return values for AssetRegistry functions
+	 */
+	enum class EExists
+	{
+		DoesNotExist,	// Does not exist on disk
+		Exists,			// Exists on disk
+		Unknown,		// Not known. AssetRegistry might still be indexing
+	};
+
 	/**
 	 * A struct that is equivalent to EDependencyQuery, but is more useful for performance in filtering operations.
 	 * This is used by the filter implementations inside of GetDependency/GetReferencer calls; callers of those functions can instead use the more convenient values in EDependencyQuery.
@@ -280,6 +290,24 @@ public:
 	 */
 	virtual void GetDependencies(FName InPackageName, TArray<FName>& OutDependencies, UE::AssetRegistry::EDependencyCategory Category = UE::AssetRegistry::EDependencyCategory::Package, const UE::AssetRegistry::FDependencyQuery& Flags = UE::AssetRegistry::FDependencyQuery()) = 0;
 
+	/**
+	 * Tries to get the asset data for the specified object path
+	 *
+	 * @param ObjectPath the path of the object to be looked up
+	 * @param OutAssetData out FAssetData
+	 * @return Return code enum
+	 */
+	virtual UE::AssetRegistry::EExists TryGetAssetByObjectPath(const FName ObjectPath, struct FAssetData& OutAssetData) const = 0;
+
+	/**
+	 * Tries to get the pacakge data for a specified path
+	 *
+	 * @param PackageName name of the package
+	 * @param OutAssetPackageData out FAssetPackageData
+	 * @return Return code enum
+	 */
+	virtual UE::AssetRegistry::EExists TryGetAssetPackageData(FName PackageName, class FAssetPackageData& OutPackageData) const = 0;
+
 protected:
 	/* This function is a workaround for platforms that don't support disable of deprecation warnings on override functions*/
 	virtual void GetDependenciesDeprecated(FName InPackageName, TArray<FName>& OutDependencies, EAssetRegistryDependencyType::Type InDependencyType) = 0;
diff --git a/Engine/Source/Runtime/CoreUObject/Public/UObject/ObjectPathId.h b/Engine/Source/Runtime/CoreUObject/Public/UObject/ObjectPathId.h
--- a/Engine/Source/Runtime/CoreUObject/Public/UObject/ObjectPathId.h
+++ b/Engine/Source/Runtime/CoreUObject/Public/UObject/ObjectPathId.h
@@ -70,11 +70,8 @@ public:
 
 	COREUOBJECT_API static FName MakeImportPathIdAndPackageName(const FObjectImport& Import, const FLinkerTables& LinkerTables, FObjectPathId& OutPathId);
 
-#if WITH_LOW_LEVEL_TESTS
-	// Should only be used by unit tests
 	explicit FObjectPathId(FWideStringView StringPath);
 	explicit FObjectPathId(FAnsiStringView StringPath);
-#endif
 
 	FObjectPathId& operator=(const FObjectPathId& Other) = default;
 	inline bool operator==(const FObjectPathId& Other) const { return PathId == Other.PathId; }
