[CoreUObject] compile fix for bad merge

diff --git a/Engine/Source/Runtime/CoreUObject/Private/UObject/ObjectPathId.cpp b/Engine/Source/Runtime/CoreUObject/Private/UObject/ObjectPathId.cpp
--- a/Engine/Source/Runtime/CoreUObject/Private/UObject/ObjectPathId.cpp
+++ b/Engine/Source/Runtime/CoreUObject/Private/UObject/ObjectPathId.cpp
@@ -230,7 +230,6 @@ FName FObjectPathId::MakeImportPathIdAndPackageName(const FObjectImport& Import,
 	return NamePathProducer.GetPackageName();
 }
 
-#if !(UE_BUILD_SHIPPING || UE_BUILD_TEST)
 template <typename CharType>
 struct TStringViewNamePathProducer
 {
