Move BP reinstanciation to happen on async loading thread before postloads are called

diff --git a/Engine/Source/Runtime/CoreUObject/Private/Serialization/AsyncLoading2.cpp b/Engine/Source/Runtime/CoreUObject/Private/Serialization/AsyncLoading2.cpp
--- a/Engine/Source/Runtime/CoreUObject/Private/Serialization/AsyncLoading2.cpp
+++ b/Engine/Source/Runtime/CoreUObject/Private/Serialization/AsyncLoading2.cpp
@@ -5756,9 +5756,14 @@ EEventLoadNodeExecutionResult FAsyncPackage2::Event_PostLoadExportBundle(FAsyncL
 	UE_ASYNC_PACKAGE_DEBUG(Package->Desc);
 	check(Package->AsyncPackageLoadingState == EAsyncPackageLoadingState2::PostLoad);
 	check(Package->ExternalReadDependencies.Num() == 0);
-	
+
 	FAsyncPackageScope2 PackageScope(Package);
 
+#if WITH_EDITOR
+	// This needs to happen after loading new blueprints in the editor, and this is handled in EndLoad for synchronous loads
+	FBlueprintSupport::FlushReinstancingQueue();
+#endif
+
 #if ALT2_ENABLE_LINKERLOAD_SUPPORT
 	if (Package->LinkerLoadState.IsSet())
 	{
@@ -6394,11 +6399,6 @@ EAsyncPackageState::Type FAsyncLoadingThread2::ProcessLoadedPackagesFromGameThre
 
 	if (Result == EAsyncPackageState::Complete)
 	{
-#if WITH_EDITORONLY_DATA
-		// This needs to happen after loading new blueprints in the editor, and this is handled in EndLoad for synchronous loads
-		FBlueprintSupport::FlushReinstancingQueue();
-#endif
-
 #if WITH_EDITOR
 		// In editor builds, call the asset load callback. This happens in both editor and standalone to match EndLoad
 		TSet<FWeakObjectPtr> TempLoadedAssets = LoadedAssets;
