Propagate the LoaderImport internal flag when replacing an object in the object array

diff --git a/Engine/Source/Runtime/CoreUObject/Private/UObject/UObjectGlobals.cpp b/Engine/Source/Runtime/CoreUObject/Private/UObject/UObjectGlobals.cpp
--- a/Engine/Source/Runtime/CoreUObject/Private/UObject/UObjectGlobals.cpp
+++ b/Engine/Source/Runtime/CoreUObject/Private/UObject/UObjectGlobals.cpp
@@ -3209,7 +3209,7 @@ UObject* StaticAllocateObject
 		// Remember linker, flags, index, and native class info.
 		Linker		= Obj->GetLinker();
 		LinkerIndex = Obj->GetLinkerIndex();
-		InternalSetFlags |= (Obj->GetInternalFlags() & (EInternalObjectFlags::Native | EInternalObjectFlags::RootSet));
+		InternalSetFlags |= (Obj->GetInternalFlags() & (EInternalObjectFlags::Native | EInternalObjectFlags::RootSet | EInternalObjectFlags::LoaderImport));
 
 		if ( bCreatingCDO )
 		{
