Fixed an issue where a uobject intended to replace another uobject in the object array could get a new internal index if other threads were creating objects at the same time

diff --git a/Engine/Source/Runtime/CoreUObject/Private/UObject/UObjectArray.cpp b/Engine/Source/Runtime/CoreUObject/Private/UObject/UObjectArray.cpp
--- a/Engine/Source/Runtime/CoreUObject/Private/UObject/UObjectArray.cpp
+++ b/Engine/Source/Runtime/CoreUObject/Private/UObject/UObjectArray.cpp
@@ -186,15 +186,19 @@ void FUObjectArray::DisableDisregardForGC()
 	}
 }
 
-void FUObjectArray::AllocateUObjectIndex(UObjectBase* Object, bool bMergingThreads /*= false*/)
+void FUObjectArray::AllocateUObjectIndex(UObjectBase* Object, bool bMergingThreads /*= false*/, int32 AlreadyAllocatedIndex /*= -1*/)
 {
 	int32 Index = INDEX_NONE;
 	check(Object->InternalIndex == INDEX_NONE || bMergingThreads);
 
 	LockInternalArray();
 
+	if (AlreadyAllocatedIndex >= 0)
+	{
+		Index = AlreadyAllocatedIndex;
+	}
 	// Special non- garbage collectable range.
-	if (OpenForDisregardForGC && DisregardForGCEnabled())
+	else if (OpenForDisregardForGC && DisregardForGCEnabled())
 	{
 		Index = ++ObjLastNonGCIndex;
 		// Check if we're not out of bounds, unless there hasn't been any gc objects yet
@@ -296,7 +300,7 @@ void FUObjectArray::FreeUObjectIndex(UObjectBase* Object)
 
 	// You cannot safely recycle indicies in the non-GC range
 	// No point in filling this list when doing exit purge. Nothing should be allocated afterwards anyway.
-	if (Index > ObjLastNonGCIndex && !GExitPurge)  
+	if (Index > ObjLastNonGCIndex && !GExitPurge && bShouldRecycleObjectIndices)
 	{
 		ObjAvailableList.Add(Index);
 #if UE_GC_TRACK_OBJ_AVAILABLE
diff --git a/Engine/Source/Runtime/CoreUObject/Private/UObject/UObjectBase.cpp b/Engine/Source/Runtime/CoreUObject/Private/UObject/UObjectBase.cpp
--- a/Engine/Source/Runtime/CoreUObject/Private/UObject/UObjectBase.cpp
+++ b/Engine/Source/Runtime/CoreUObject/Private/UObject/UObjectBase.cpp
@@ -118,7 +118,7 @@ UObjectBase::UObjectBase( EObjectFlags InFlags )
  * @param	InName				name of the new object
  * @param	InObjectArchetype	archetype to assign
  */
-UObjectBase::UObjectBase(UClass* InClass, EObjectFlags InFlags, EInternalObjectFlags InInternalFlags, UObject *InOuter, FName InName)
+UObjectBase::UObjectBase(UClass* InClass, EObjectFlags InFlags, EInternalObjectFlags InInternalFlags, UObject *InOuter, FName InName, int32 InInternalIndex)
 :	ObjectFlags			(InFlags)
 ,	InternalIndex		(INDEX_NONE)
 ,	ClassPrivate		(InClass)
@@ -126,7 +126,7 @@ UObjectBase::UObjectBase(UClass* InClass, EObjectFlags InFlags, EInternalObjectF
 {
 	check(ClassPrivate);
 	// Add to global table.
-	AddObject(InName, InInternalFlags);
+	AddObject(InName, InInternalFlags, InInternalIndex);
 	
 #if CSV_PROFILER && CSV_TRACK_UOBJECT_COUNT
 	UObjectStats::IncrementUObjectCount();
@@ -191,7 +191,7 @@ void UObjectBase::DeferredRegister(UClass *UClassStaticClass,const TCHAR* Packag
  *
  * @param Name name to assign to this uobject
  */
-void UObjectBase::AddObject(FName InName, EInternalObjectFlags InSetInternalFlags)
+void UObjectBase::AddObject(FName InName, EInternalObjectFlags InSetInternalFlags, int32 InInternalIndex)
 {
 	NamePrivate = InName;
 	EInternalObjectFlags InternalFlagsToSet = InSetInternalFlags;
@@ -209,7 +209,7 @@ void UObjectBase::AddObject(FName InName, EInternalObjectFlags InSetInternalFlag
 		InternalFlagsToSet |= EInternalObjectFlags::Native;
 		ObjectFlags &= ~RF_MarkAsNative;
 	}
-	GUObjectArray.AllocateUObjectIndex(this);
+	GUObjectArray.AllocateUObjectIndex(this, /* bMergingThreads */ false, InInternalIndex);
 	check(InName != NAME_None && InternalIndex >= 0);
 	if (InternalFlagsToSet != EInternalObjectFlags::None)
 	{
diff --git a/Engine/Source/Runtime/CoreUObject/Private/UObject/UObjectGlobals.cpp b/Engine/Source/Runtime/CoreUObject/Private/UObject/UObjectGlobals.cpp
--- a/Engine/Source/Runtime/CoreUObject/Private/UObject/UObjectGlobals.cpp
+++ b/Engine/Source/Runtime/CoreUObject/Private/UObject/UObjectGlobals.cpp
@@ -3193,6 +3193,8 @@ UObject* StaticAllocateObject
 	int32 TotalSize = InClass->GetPropertiesSize();
 	checkSlow(TotalSize);
 
+	int32 OldIndex = -1;
+
 	if( Obj == nullptr )
 	{	
 		int32 Alignment	= FMath::Max( 4, InClass->GetMinAlignment() );
@@ -3248,6 +3250,8 @@ UObject* StaticAllocateObject
 		// Subobjects are always created in the constructor, no need to re-create them here unless their archetype != CDO or they're blueprint generated.	
 		if (!bCreatingCDO && (!bCanRecycleSubobjects || !Obj->IsDefaultSubobject()))
 		{
+			OldIndex = GUObjectArray.ObjectToIndex(Obj);
+
 			// Destroy the object.
 			SCOPE_CYCLE_COUNTER(STAT_DestroyObject);
 			// Check that the object hasn't been destroyed yet.
@@ -3291,6 +3295,7 @@ UObject* StaticAllocateObject
 				Obj->ConditionalFinishDestroy();
 			}
 			GUObjectArray.LockInternalArray();
+			TGuardValue<bool> _(GUObjectArray.bShouldRecycleObjectIndices, false);
 			Obj->~UObject();
 			GUObjectArray.UnlockInternalArray();
 			bWasConstructedOnOldObject	= true;
@@ -3310,7 +3315,7 @@ UObject* StaticAllocateObject
 	if (!bSubObject)
 	{
 		FMemory::Memzero((void *)Obj, TotalSize);
-		new ((void *)Obj) UObjectBase(const_cast<UClass*>(InClass), InFlags|RF_NeedInitialization, InternalSetFlags, InOuter, InName);
+		new ((void *)Obj) UObjectBase(const_cast<UClass*>(InClass), InFlags|RF_NeedInitialization, InternalSetFlags, InOuter, InName, OldIndex);
 	}
 	else
 	{
diff --git a/Engine/Source/Runtime/CoreUObject/Public/UObject/UObjectArray.h b/Engine/Source/Runtime/CoreUObject/Public/UObject/UObjectArray.h
--- a/Engine/Source/Runtime/CoreUObject/Public/UObject/UObjectArray.h
+++ b/Engine/Source/Runtime/CoreUObject/Public/UObject/UObjectArray.h
@@ -579,6 +579,8 @@ public:
 class COREUOBJECT_API FUObjectArray
 {
 	friend class UObject;
+	friend COREUOBJECT_API UObject* StaticAllocateObject(const UClass*, UObject*, FName, EObjectFlags, EInternalObjectFlags, bool, bool*, UPackage*);
+
 private:
 	/**
 	 * Reset the serial number from the game thread to invalidate all weak object pointers to it
@@ -687,7 +689,7 @@ public:
 	 *
 	 * @param	Object Object to allocate an index for
 	 */
-	void AllocateUObjectIndex(class UObjectBase* Object, bool bMergingThreads = false);
+	void AllocateUObjectIndex(class UObjectBase* Object, bool bMergingThreads = false, int32 AlreadyAllocatedIndex = -1);
 
 	/**
 	 * Returns a UObject index top to the global uobject array
@@ -1106,6 +1108,9 @@ private:
 	/** Current primary serial number **/
 	FThreadSafeCounter	PrimarySerialNumber;
 
+	/** If set to false object indices won't be recycled to the global pool and can be explicitly reused when creating new objects */
+	bool bShouldRecycleObjectIndices = true;
+
 public:
 
 	/** INTERNAL USE ONLY: gets the internal FUObjectItem array */
diff --git a/Engine/Source/Runtime/CoreUObject/Public/UObject/UObjectBase.h b/Engine/Source/Runtime/CoreUObject/Public/UObject/UObjectBase.h
--- a/Engine/Source/Runtime/CoreUObject/Public/UObject/UObjectBase.h
+++ b/Engine/Source/Runtime/CoreUObject/Public/UObject/UObjectBase.h
@@ -66,8 +66,9 @@ public:
 	 * @param	InInternalFlags EInternalObjectFlags to assign
 	 * @param	InOuter				outer for this object
 	 * @param	InName				name of the new object
+	 * @param	InInternalIndex		internal index to use (if already allocated), negative value means allocate a new index
 	 */
-	UObjectBase( UClass* InClass, EObjectFlags InFlags, EInternalObjectFlags InInternalFlags, UObject *InOuter, FName InName );
+	UObjectBase( UClass* InClass, EObjectFlags InFlags, EInternalObjectFlags InInternalFlags, UObject *InOuter, FName InName, int32 InInternalIndex = -1 );
 
 	/**
 	 * Final destructor, removes the object from the object array, and indirectly, from any annotations
@@ -107,8 +108,9 @@ private:
 	 *
 	 * @param Name name to assign to this uobject
 	 * @param InSetInternalFlags Internal object flags to be set on the object once it's been added to the array
+	 * @param InInternalIndex already allocated internal index to use, negative value means allocate a new index
 	 */
-	void AddObject(FName Name, EInternalObjectFlags InSetInternalFlags);
+	void AddObject(FName Name, EInternalObjectFlags InSetInternalFlags, int32 InInternalIndex = -1);
 
 public:
 	/**
