Expansions to LoadMap world leak tracking.
- configurable severity
- separate cvar for shipping configuration
- strip sublevels from list to check if the owning world was also leaked
- optional minimal mode that searches for direct, external references to the target objects or one of their inner objects (e.g. actors in a world) for less memory overhead - uses about 225kb in one test (vs hundreds of mb)
- fix fatal logging option to log all errors before crashing
- optimize reference graph to skip leaf nodes, reduces cost from 500mb -> 335mb in one test
- only store a callstack in nodes that actually need it, reduces cost from 335mb -> 100mb in testing

diff --git a/Engine/Source/Runtime/CoreUObject/Private/UObject/ReferenceChainSearch.cpp b/Engine/Source/Runtime/CoreUObject/Private/UObject/ReferenceChainSearch.cpp
--- a/Engine/Source/Runtime/CoreUObject/Private/UObject/ReferenceChainSearch.cpp
+++ b/Engine/Source/Runtime/CoreUObject/Private/UObject/ReferenceChainSearch.cpp
@@ -40,6 +40,19 @@ FReferenceChainSearch::FGraphNode* FReferenceChainSearch::FindOrAddNode(UObject*
 	return FindOrAddNode(FGCObjectInfo::FindOrAddInfoHelper(InObjectToFindNodeFor, ObjectToInfoMap));
 }
 
+void FReferenceChainSearch::LinkNodes(FGCObjectInfo* From, FGCObjectInfo* To, EReferenceType ReferenceType, FName PropertyName, TConstArrayView<uint64> StackFrames)
+{
+	FGraphNode* ToNode = FindOrAddNode(To);
+	FGraphNode* FromNode = FindOrAddNode(From);
+	FromNode->ReferencedObjects.Emplace(FNodeReferenceInfo(ToNode, ReferenceType, PropertyName, StackFrames));
+	ToNode->ReferencedByObjects.Add(FromNode);
+}
+
+void FReferenceChainSearch::LinkNodes(UObject* From, UObject* To, EReferenceType ReferenceType, FName PropertyName, TConstArrayView<uint64> StackFrames)
+{
+	LinkNodes(FGCObjectInfo::FindOrAddInfoHelper(From, ObjectToInfoMap), FGCObjectInfo::FindOrAddInfoHelper(To, ObjectToInfoMap), ReferenceType, PropertyName, StackFrames);
+}
+
 int32 FReferenceChainSearch::BuildReferenceChainsRecursive(FGraphNode* TargetNode, TArray<FReferenceChain*>& ProducedChains, int32 ChainDepth, const int32 VisitCounter, EReferenceChainSearchMode SearchMode, FGraphNode* GCObjReferencerNode)
 {
 	int32 ProducedChainsCount = 0;
@@ -285,6 +298,7 @@ void FReferenceChainSearch::BuildReferenceChainsForDirectReferences(TConstArrayV
 			if (!(SearchMode & EReferenceChainSearchMode::ExternalOnly) || !ReferencedByNode->ObjectInfo->IsIn(TargetNode->ObjectInfo))
 			{
 				FReferenceChain* Chain = new FReferenceChain();
+				Chain->TargetNode = TargetNode;
 				Chain->AddNode(TargetNode);
 				Chain->AddNode(ReferencedByNode);
 				Chain->FillReferenceInfo();
@@ -351,7 +365,7 @@ FString FReferenceChainSearch::GetObjectFlags(FGCObjectInfo* InObject)
 	return Flags;
 }
 
-static void ConvertStackFramesToCallstack(uint64* StackFrames, int32 NumStackFrames, int32 Indent, FOutputDevice& Out)
+static void ConvertStackFramesToCallstack(const uint64* StackFrames, int32 NumStackFrames, int32 Indent, FOutputDevice& Out)
 {
 	// Convert the stack trace to text
 	for (int32 Idx = 0; Idx < NumStackFrames; Idx++)
@@ -389,7 +403,7 @@ void FReferenceChainSearch::DumpChain(FReferenceChainSearch::FReferenceChain* Ch
 	{
 		bool bPostCallbackContinue = true;
 		const int32 RootIndex = Chain->Num() - 1;
-		FNodeReferenceInfo ReferenceInfo = Chain->GetReferenceInfo(RootIndex);
+		const FNodeReferenceInfo* ReferenceInfo = Chain->GetReferenceInfo(RootIndex);
 		FGCObjectInfo* ReferencerObject = Chain->GetNode(RootIndex)->ObjectInfo;
 		{
 			FCallbackParams Params;
@@ -415,17 +429,17 @@ void FReferenceChainSearch::DumpChain(FReferenceChainSearch::FReferenceChain* Ch
 			FCallbackParams Params;
 			Params.Referencer = ReferencerObject;
 			Params.Object = Object;
-			Params.ReferenceInfo = &ReferenceInfo;
+			Params.ReferenceInfo = ReferenceInfo;
 			Params.Indent = FMath::Min<int32>(TCStringSpcHelper<TCHAR>::MAX_SPACES, Chain->Num() - NodeIndex - 1);
 			Params.Out = &Out;
 
-			if (ReferenceInfo.Type == EReferenceType::Property)
+			if (ReferenceInfo->Type == EReferenceType::Property)
 			{
 				FString ReferencingPropertyName;
 				UClass* ReferencerClass = Cast<UClass>(ReferencerObject->GetClass()->TryResolveObject());
 				TArray<FProperty*> ReferencingProperties;
 
-				if (ReferencerClass && FGCStackSizeHelper::ConvertPathToProperties(ReferencerClass, ReferenceInfo.ReferencerName, ReferencingProperties))
+				if (ReferencerClass && FGCStackSizeHelper::ConvertPathToProperties(ReferencerClass, ReferenceInfo->ReferencerName, ReferencingProperties))
 				{
 					FProperty* InnermostProperty = ReferencingProperties.Last();
 					FProperty* OutermostProperty = ReferencingProperties[0];
@@ -434,7 +448,7 @@ void FReferenceChainSearch::DumpChain(FReferenceChainSearch::FReferenceChain* Ch
 						*InnermostProperty->GetCPPType(),
 						OutermostProperty->GetOwnerClass()->GetPrefixCPP(),
 						*OutermostProperty->GetOwnerClass()->GetName(),
-						*ReferenceInfo.ReferencerName.ToString());
+						*ReferenceInfo->ReferencerName.ToString());
 				}
 				else
 				{
@@ -443,7 +457,7 @@ void FReferenceChainSearch::DumpChain(FReferenceChainSearch::FReferenceChain* Ch
 					static const FName OuterPropertyName(TEXT("Outer"));
 					
 					FString ClassName;
-					if (ReferenceInfo.ReferencerName == ClassPropertyName || ReferenceInfo.ReferencerName == OuterPropertyName)
+					if (ReferenceInfo->ReferencerName == ClassPropertyName || ReferenceInfo->ReferencerName == OuterPropertyName)
 					{
 						ClassName = TEXT("UObject");
 					}
@@ -458,7 +472,7 @@ void FReferenceChainSearch::DumpChain(FReferenceChainSearch::FReferenceChain* Ch
 						// Revert to the internal class name if not
 						ClassName = ReferencerObject->GetClassName();
 					}
-					ReferencingPropertyName = FString::Printf(TEXT("UObject* %s::%s"), *ClassName, *ReferenceInfo.ReferencerName.ToString());
+					ReferencingPropertyName = FString::Printf(TEXT("UObject* %s::%s"), *ClassName, *ReferenceInfo->ReferencerName.ToString());
 				}
 
 				Out.Logf(ELogVerbosity::Log, TEXT("%s-> %s = %s %s"),
@@ -467,10 +481,10 @@ void FReferenceChainSearch::DumpChain(FReferenceChainSearch::FReferenceChain* Ch
 					*GetObjectFlags(Object),
 					*Object->GetFullName());
 			}
-			else if (ReferenceInfo.Type == EReferenceType::AddReferencedObjects)
+			else if (ReferenceInfo->Type == EReferenceType::AddReferencedObjects)
 			{
 				FString UObjectOrGCObjectName;
-				if (ReferenceInfo.ReferencerName.IsNone())
+				if (ReferenceInfo->ReferencerName.IsNone())
 				{
 					UClass* ReferencerClass = Cast<UClass>(ReferencerObject->GetClass()->TryResolveObject());
 					if (ReferencerClass)
@@ -485,7 +499,7 @@ void FReferenceChainSearch::DumpChain(FReferenceChainSearch::FReferenceChain* Ch
 				}
 				else
 				{
-					UObjectOrGCObjectName = ReferenceInfo.ReferencerName.ToString();
+					UObjectOrGCObjectName = ReferenceInfo->ReferencerName.ToString();
 				}
 
 				Out.Logf(ELogVerbosity::Log, TEXT("%s-> %s::AddReferencedObjects(%s %s)"),
@@ -494,11 +508,27 @@ void FReferenceChainSearch::DumpChain(FReferenceChainSearch::FReferenceChain* Ch
 					*GetObjectFlags(Object),
 					*Object->GetFullName());
 
-				if (ReferenceInfo.NumStackFrames)
+				if (ReferenceInfo->StackFrames.Num())
 				{
-					ConvertStackFramesToCallstack(ReferenceInfo.StackFrames, ReferenceInfo.NumStackFrames, Params.Indent, Out);
+					ConvertStackFramesToCallstack(ReferenceInfo->StackFrames.GetData(), ReferenceInfo->StackFrames.Num(), Params.Indent, Out);
 				}
 			}
+			else if (ReferenceInfo->Type == EReferenceType::OuterChain)
+			{
+				Out.Logf(ELogVerbosity::Log, TEXT("%s-> %s = %s %s"),
+					FCString::Spc(Params.Indent),
+					TEXT("Outer Chain"),
+					*GetObjectFlags(Object),
+					*Object->GetFullName());
+			}
+			else
+			{
+				Out.Logf(ELogVerbosity::Log, TEXT("%s-> %s = %s %s"),
+					FCString::Spc(Params.Indent),
+					TEXT("UNKNOWN"),
+					*GetObjectFlags(Object),
+					*Object->GetFullName());
+			}
 
 			bPostCallbackContinue = ReferenceCallback(Params);
 
@@ -512,7 +542,7 @@ void FReferenceChainSearch::DumpChain(FReferenceChainSearch::FReferenceChain* Ch
 void FReferenceChainSearch::FReferenceChain::FillReferenceInfo()
 {
 	// The first entry is the object we were looking for references to so add an empty entry for it
-	ReferenceInfos.Add(FNodeReferenceInfo());
+	ReferenceInfos.Emplace(nullptr);
 
 	// Iterate over all nodes and add reference info based on the next node (which is the object that referenced the current node)
 	for (int32 NodeIndex = 1; NodeIndex < Nodes.Num(); ++NodeIndex)
@@ -521,8 +551,8 @@ void FReferenceChainSearch::FReferenceChain::FillReferenceInfo()
 		FGraphNode* CurrentNode = Nodes[NodeIndex];
 
 		// Found the PreviousNode in the list of objects referenced by the CurrentNode
-		FNodeReferenceInfo* FoundInfo = nullptr;
-		for (FNodeReferenceInfo& Info : CurrentNode->ReferencedObjects)
+		const FNodeReferenceInfo* FoundInfo = nullptr;
+		for (const FNodeReferenceInfo& Info : CurrentNode->ReferencedObjects)
 		{
 			if (Info.Object == PreviousNode)
 			{
@@ -532,7 +562,7 @@ void FReferenceChainSearch::FReferenceChain::FillReferenceInfo()
 		}
 		check(FoundInfo); // because there must have been a reference since we created this chain using it
 		check(FoundInfo->Object == PreviousNode);
-		ReferenceInfos.Add(*FoundInfo);
+		ReferenceInfos.Add(FoundInfo);
 	}
 	check(ReferenceInfos.Num() == Nodes.Num());
 }
@@ -555,20 +585,30 @@ bool FReferenceChainSearch::FReferenceChain::IsExternal() const
 */
 class FDirectReferenceProcessor : public FSimpleReferenceProcessorBase
 {	
-	TSet<FReferenceChainSearch::FObjectReferenceInfo>& ReferencedObjects;
+protected:
+	TSet<FReferenceChainSearch::FObjectReferenceInfo> ReferencedObjects;
 	TMap<UObject*, FGCObjectInfo*>& ObjectToInfoMap;
 
 public:
 
-	FDirectReferenceProcessor(TSet<FReferenceChainSearch::FObjectReferenceInfo>& InReferencedObjects, TMap<UObject*, FGCObjectInfo*>& InObjectToInfoMap)
-		: ReferencedObjects(InReferencedObjects)
-		, ObjectToInfoMap(InObjectToInfoMap)
+	FDirectReferenceProcessor(TMap<UObject*, FGCObjectInfo*>& InObjectToInfoMap)
+	: ObjectToInfoMap(InObjectToInfoMap)
 	{
 	}
 
-	FORCEINLINE void HandleTokenStreamObjectReference(FGCArrayStruct& ObjectsToSerializeStruct, UObject* ReferencingObject, UObject*& Object, const int32 TokenIndex, const EGCTokenType TokenType, bool bAllowReferenceElimination)
+	void Reset()
+	{
+		ReferencedObjects.Reset();
+	}
+
+	const TSet<FReferenceChainSearch::FObjectReferenceInfo>& GetReferencedObjects() const
+	{
+		return ReferencedObjects;
+	}
+
+	void HandleTokenStreamObjectReference(FGCArrayStruct& ObjectsToSerializeStruct, UObject* ReferencingObject, UObject*& Object, const int32 TokenIndex, const EGCTokenType TokenType, bool bAllowReferenceElimination)
 	{		
-		if (Object)
+		if (Object && Object != ReferencingObject) // Skip self-references just in case 
 		{
 			FGCObjectInfo* ObjectInfo = FGCObjectInfo::FindOrAddInfoHelper(Object, ObjectToInfoMap);
 
@@ -584,7 +624,10 @@ public:
 				else
 				{
 					RefInfo.Type = FReferenceChainSearch::EReferenceType::AddReferencedObjects;
-					RefInfo.NumStackFrames = FPlatformStackWalk::CaptureStackBackTrace(RefInfo.StackFrames, FReferenceChainSearch::FObjectReferenceInfo::MaxStackFrames);
+
+					RefInfo.StackFrames.AddUninitialized(FReferenceChainSearch::FObjectReferenceInfo::MaxStackFrames);
+					int32 NumStackFrames = FPlatformStackWalk::CaptureStackBackTrace(RefInfo.StackFrames.GetData(), RefInfo.StackFrames.Num());
+					RefInfo.StackFrames.SetNum(NumStackFrames);
 
 					if (FGCObject::GGCObjectReferencer && (!ReferencingObject || ReferencingObject == FGCObject::GGCObjectReferencer))
 					{
@@ -610,17 +653,79 @@ public:
 					}
 				}
 
-				ReferencedObjects.Add(RefInfo);
+				ReferencedObjects.Emplace(MoveTemp(RefInfo));
 			}
 		}
 	}
 };
 
-class FDirectReferenceCollector : public TDefaultReferenceCollector<FDirectReferenceProcessor>
+class FMinimalReferenceProcessor : public FDirectReferenceProcessor
+{	
+	TSet<UObject*>& TargetObjects;
+	bool bSearchInners = false;
+
+public:
+	FMinimalReferenceProcessor(
+		TSet<UObject*>& InTargetObjects,
+		TMap<UObject*, FGCObjectInfo*>& InObjectToInfoMap
+	)
+	: FDirectReferenceProcessor(InObjectToInfoMap)
+	, TargetObjects(InTargetObjects)
+	{
+	}
+
+	void SetSearchInners(bool InValue)
+	{
+		bSearchInners = InValue;
+	}
+
+	void HandleTokenStreamObjectReference(
+		FGCArrayStruct& ObjectsToSerializeStruct,
+		UObject* ReferencingObject,
+		UObject*& Object,
+		const int32 TokenIndex,
+		const EGCTokenType TokenType,
+		bool bAllowReferenceElimination
+	)
+	{		
+		if (Object && Object != ReferencingObject) // Skip self-references just in case 
+		{
+			FGCObjectInfo* ObjectInfo = nullptr;
+			if (bSearchInners)
+			{
+				// Walk the outer chain to see if this is the inner of an object we're looking for which would be referenced via this outer chain
+				for (UObject* OuterObject = Object->GetOuter(); OuterObject; OuterObject = OuterObject->GetOuter())
+				{
+					if (TargetObjects.Contains(OuterObject))
+					{
+						// Create a node to reference this inner object, which we will later connect to the target object
+						ObjectInfo = FGCObjectInfo::FindOrAddInfoHelper(Object, ObjectToInfoMap);
+						break;
+					}
+				}
+			}
+			else if (TargetObjects.Contains(Object))
+			{
+				ObjectInfo = FGCObjectInfo::FindOrAddInfoHelper(Object, ObjectToInfoMap);
+			}
+
+			// We don't care about this reference
+			if (!ObjectInfo) 
+			{
+				return;
+			}
+
+			FDirectReferenceProcessor::HandleTokenStreamObjectReference(ObjectsToSerializeStruct, ReferencingObject, Object, TokenIndex, TokenType, bAllowReferenceElimination);
+		}
+	}
+};
+
+template<typename ProcessorType>
+class TReferenceSearchReferenceCollector : public TDefaultReferenceCollector<ProcessorType>
 {
 public:
-	FDirectReferenceCollector(FDirectReferenceProcessor& InProcessor, FGCArrayStruct& InObjectArrayStruct)
-		: TDefaultReferenceCollector<FDirectReferenceProcessor>(InProcessor, InObjectArrayStruct)
+	TReferenceSearchReferenceCollector (ProcessorType& InProcessor, FGCArrayStruct& InObjectArrayStruct)
+		: TDefaultReferenceCollector<ProcessorType>(InProcessor, InObjectArrayStruct)
 	{
 	}
 
@@ -650,7 +755,14 @@ FReferenceChainSearch::FReferenceChainSearch(TConstArrayView<UObject*> InObjects
 	}
 
 	// First pass is to find all direct references for each object
-	FindDirectReferencesForObjects();
+	if (!!(SearchMode & EReferenceChainSearchMode::Minimal))
+	{
+		FindMinimalDirectReferencesForObjects();
+	}
+	else
+	{
+		FindDirectReferencesForObjects();
+	}
 
 	// Second pass creates all reference chains
 	PerformSearch();
@@ -671,6 +783,27 @@ FReferenceChainSearch::~FReferenceChainSearch()
 	Cleanup();
 }
 
+int64 FReferenceChainSearch::GetAllocatedSize() const
+{
+	int64 Size = 0;
+	Size += ObjectsToFindReferencesTo.GetAllocatedSize();
+	Size += ObjectInfosToFindReferencesTo.GetAllocatedSize();
+	Size += ReferenceChains.GetAllocatedSize();
+	for (const FReferenceChain* Chain : ReferenceChains)
+	{
+		Size += Chain->GetAllocatedSize();
+	}
+	Size += AllNodes.GetAllocatedSize();
+	Size += AllNodes.Num() * sizeof(FGCObjectInfo); // GC object infos are heap allocated and owned by this structure 
+	for (const TPair<FGCObjectInfo*, FGraphNode*> Pair : AllNodes)
+	{
+		// FGCObjectInfo makes no allocations
+		Size += Pair.Value->GetAllocatedSize();
+	}
+	Size += ObjectToInfoMap.GetAllocatedSize();
+	return Size;
+}
+
 void FReferenceChainSearch::PerformSearch()
 {
 	TArray<FGraphNode*> GraphNodesToFind;
@@ -739,7 +872,7 @@ void FReferenceChainSearch::PerformSearchFromGCSnapshot(TConstArrayView<UObject*
 			{
 				ReferenceType = EReferenceType::Property;
 			}
-			ObjectNode->ReferencedObjects.Add(FNodeReferenceInfo(ReferencedObjectNode, ReferenceType, ReferenceInfo.ReferencerName, nullptr, 0));
+			ObjectNode->ReferencedObjects.Emplace(FNodeReferenceInfo(ReferencedObjectNode, ReferenceType, ReferenceInfo.ReferencerName));
 			ReferencedObjectNode->ReferencedByObjects.Add(ObjectNode);
 		}
 	}
@@ -757,41 +890,183 @@ void FReferenceChainSearch::PerformSearchFromGCSnapshot(TConstArrayView<UObject*
 }
 #endif // ENABLE_GC_HISTORY
 
-void FReferenceChainSearch::FindDirectReferencesForObjects()
+template<typename ProcessorType>
+struct TReferenceSearchHelper
 {
-	TSet<FObjectReferenceInfo> ReferencedObjects;
-	FDirectReferenceProcessor Processor(ReferencedObjects, ObjectToInfoMap);
+	ProcessorType Processor;
 	TFastReferenceCollector<
-		FDirectReferenceProcessor, 
-		FDirectReferenceCollector, 
+		ProcessorType, 
+		TReferenceSearchReferenceCollector<ProcessorType>, 
 		FGCArrayPool, 
 		EFastReferenceCollectorOptions::AutogenerateTokenStream | EFastReferenceCollectorOptions::ProcessNoOpTokens
-	> ReferenceCollector(Processor, FGCArrayPool::Get());
+	> ReferenceCollector;
 	FGCArrayStruct ArrayStruct;
-	TArray<UObject*>& ObjectsToProcess = ArrayStruct.ObjectsToSerialize;
 
-	for (FRawObjectIterator It; It; ++It)
+	TReferenceSearchHelper(ProcessorType&& InProcessor)
+	: Processor(MoveTemp(InProcessor))
+	, ReferenceCollector(Processor, FGCArrayPool::Get())
 	{
-		FUObjectItem* ObjItem = *It;
-		UObject* Object = static_cast<UObject*>(ObjItem->Object);
-		FGraphNode* ObjectNode = FindOrAddNode(Object);
+	}
 
-		// Find direct references
-		ReferencedObjects.Reset();
-		ObjectsToProcess.Reset();
-		ObjectsToProcess.Add(Object);
-		ReferenceCollector.CollectReferences(ArrayStruct);
+	static bool DoNotFilterObjects(UObject* Object) { return true; }
 
-		// Build direct reference tree
-		for (FObjectReferenceInfo& ReferenceInfo : ReferencedObjects)
+	// HANDLE_REFERENCES is a function bool(const TSet<FReferenceChainSearch>&) function returning true to continue the search
+	// FILTER_OBJECTS is a function bool(UObject*) that returns true if the object's outgoing references should be collected
+	template<typename HANDLE_REFERENCES, typename FILTER_OBJECTS = decltype(DoNotFilterObjects)>
+	void CollectAllReferences(bool bGCOnly, HANDLE_REFERENCES HandleReferences, FILTER_OBJECTS FilterObjects = &DoNotFilterObjects)
+	{
+		TArray<UObject*>& ObjectsToProcess = ArrayStruct.ObjectsToSerialize;
+		for (FRawObjectIterator It; It; ++It)
 		{
-			FGraphNode* ReferencedObjectNode = FindOrAddNode(ReferenceInfo.Object);
-			ObjectNode->ReferencedObjects.Add(FNodeReferenceInfo(ReferencedObjectNode, ReferenceInfo.Type, ReferenceInfo.ReferencerName, ReferenceInfo.StackFrames, ReferenceInfo.NumStackFrames));
-			ReferencedObjectNode->ReferencedByObjects.Add(ObjectNode);
+			FUObjectItem* ObjItem = *It;
+			UObject* Object = static_cast<UObject*>(ObjItem->Object);
+
+			// We can't ask the iterator for only GC objects because that would skip the GC Object referencer 
+			if (bGCOnly && GUObjectArray.IsDisregardForGC(Object) && (Object != FGCObject::GGCObjectReferencer))
+			{
+				continue;
+			}
+
+			if (!FilterObjects(Object))
+			{
+				continue;
+			}
+
+			// Find direct references
+			Processor.Reset();
+			ObjectsToProcess.Reset();
+			ObjectsToProcess.Add(Object);
+			ReferenceCollector.CollectReferences(ArrayStruct);
+
+			if (!HandleReferences(Object, Processor.GetReferencedObjects()))
+			{
+				break;
+			}
 		}
 	}
+};
+
+void FReferenceChainSearch::FindDirectReferencesForObjects()
+{
+	TReferenceSearchHelper<FDirectReferenceProcessor> Search{FDirectReferenceProcessor(ObjectToInfoMap)};
+	const bool bGCOnly = !!(SearchMode & EReferenceChainSearchMode::GCOnly);
+
+	// First pass, create nodes for any objects which have outgoing references 
+	Search.CollectAllReferences(bGCOnly, [&](UObject* SourceObject, const TSet<FObjectReferenceInfo>& References)
+	{
+		// Skip node creation if there are no outgoing references that we care about 
+		if (References.Num() > 0)
+		{
+			FindOrAddNode(FGCObjectInfo::FindOrAddInfoHelper(SourceObject, ObjectToInfoMap));
+		}
+		return true;
+	});
+
+	Search.CollectAllReferences(bGCOnly, [&](UObject* SourceObject, const TSet<FObjectReferenceInfo>& References)
+	{
+		if (References.Num() > 0)
+		{
+			FGCObjectInfo* SourceObjectInfo = FGCObjectInfo::FindOrAddInfoHelper(SourceObject, ObjectToInfoMap);
+
+			// Build direct reference tree filtering out leaf notes
+			for (const FObjectReferenceInfo& ReferenceInfo : References)
+			{
+				FGCObjectInfo* ReferencedObject = ReferenceInfo.Object;
+				if (AllNodes.Contains(ReferencedObject))
+				{
+					LinkNodes(SourceObjectInfo, ReferencedObject, ReferenceInfo.Type, ReferenceInfo.ReferencerName, ReferenceInfo.StackFrames);
+				}
+			}
+		}
+		return true;
+	});
 }
 
+void FReferenceChainSearch::FindMinimalDirectReferencesForObjects()
+{
+	TSet<UObject*> TargetObjects;
+	TReferenceSearchHelper<FMinimalReferenceProcessor> Search{FMinimalReferenceProcessor(TargetObjects, ObjectToInfoMap)};
+	const bool bGCOnly = !!(SearchMode & EReferenceChainSearchMode::GCOnly);
+
+	auto FilterObjects = [&TargetObjects](UObject* Object)
+	{
+		// Skip objects which are in any of what we're looking for so we only report reference external to that entire group
+		for( UObject* OuterObject = Object; OuterObject; OuterObject = OuterObject->GetOuter())
+		{
+			if (TargetObjects.Contains(OuterObject))
+			{
+				return false;
+			}
+		}
+		return true;
+	};
+
+	for (UObject* Object: ObjectsToFindReferencesTo)
+	{
+		TargetObjects.Add(Object);
+	}
+	
+	// First pass, create nodes anything that directly references any of the target objects 
+	Search.CollectAllReferences(bGCOnly, [&](UObject* SourceObject, const TSet<FObjectReferenceInfo>& References)
+	{
+		if (References.Num() > 0)
+		{
+			FGCObjectInfo* SourceObjectInfo = FGCObjectInfo::FindOrAddInfoHelper(SourceObject, ObjectToInfoMap);
+
+			// Build direct reference tree filtering out leaf notes
+			for (const FObjectReferenceInfo& ReferenceInfo : References)
+			{
+				FGCObjectInfo* ReferencedObject = ReferenceInfo.Object;
+				LinkNodes(SourceObjectInfo, ReferencedObject, ReferenceInfo.Type, ReferenceInfo.ReferencerName, ReferenceInfo.StackFrames);
+			}
+		}
+		return true;
+	}, FilterObjects);
+
+	// If we didn't find a reference to any of the target objects in our first pass, loosen requirements slightly
+	TargetObjects.Reset();
+	for (FGCObjectInfo* Info : ObjectInfosToFindReferencesTo)
+	{
+		if (FGraphNode* Node = AllNodes.FindRef(Info); !Node || Node->ReferencedByObjects.Num() == 0)
+		{
+			TargetObjects.Add(Info->TryResolveObject());
+		}
+	}
+
+	// We failed to find direct references to some objects, try searching for direct references to one of their inners 
+	if (TargetObjects.Num())
+	{
+		Search.Processor.SetSearchInners(true);
+		Search.CollectAllReferences(bGCOnly, [&](UObject* SourceObject, const TSet<FObjectReferenceInfo>& References)
+		{
+			if (References.Num() > 0)
+			{
+				FGCObjectInfo* SourceObjectInfo = FGCObjectInfo::FindOrAddInfoHelper(SourceObject, ObjectToInfoMap);
+
+				// Build direct reference tree
+				for (const FObjectReferenceInfo& ReferenceInfo : References)
+				{
+					LinkNodes(SourceObjectInfo, ReferenceInfo.Object, ReferenceInfo.Type, ReferenceInfo.ReferencerName, ReferenceInfo.StackFrames);
+					UObject* ReferencedObject = ReferenceInfo.Object->TryResolveObject();
+					for (UObject* OuterObject = ReferencedObject; OuterObject; OuterObject = OuterObject->GetOuter())
+					{
+						if (TargetObjects.Contains(OuterObject))
+						{
+							LinkNodes(ReferencedObject, OuterObject, EReferenceType::OuterChain, NAME_None);
+							TargetObjects.Remove(OuterObject);
+						}
+					}
+				}
+
+				if (TargetObjects.Num() == 0)
+				{
+					return false;
+				}
+			}
+			return true;
+		}, FilterObjects);
+	}
+}
 int32 FReferenceChainSearch::PrintResults(bool bDumpAllChains /*= false*/, UObject* TargetObject /*= nullptr*/) const
 {
 	return PrintResults([](FCallbackParams& Params) { return true; }, bDumpAllChains, TargetObject);
diff --git a/Engine/Source/Runtime/CoreUObject/Public/UObject/ReferenceChainSearch.h b/Engine/Source/Runtime/CoreUObject/Public/UObject/ReferenceChainSearch.h
--- a/Engine/Source/Runtime/CoreUObject/Public/UObject/ReferenceChainSearch.h
+++ b/Engine/Source/Runtime/CoreUObject/Public/UObject/ReferenceChainSearch.h
@@ -44,6 +44,12 @@ enum class EReferenceChainSearchMode
 	FullChain = 1 << 4,
 	// Returns the shortest path to a garbage object from which the target object is reachable
 	ShortestToGarbage = 1 << 5,
+	// Attempts to find a plausible path to the target object with minimal memory usage
+	// E.g. returns a direct external reference to the target object if one is found 
+	//	otherwise returns an external reference to an inner of the target object
+	Minimal = 1 << 6,
+	// Skips the disregard-for-GC set that will never be GCd and whose outgoing references are not checked during GC
+	GCOnly = 1 << 7,
 
 	// Print results
 	PrintResults = 1 << 16,
@@ -68,51 +74,56 @@ public:
 	{
 		Unknown = 0,
 		Property = 1,
-		AddReferencedObjects
+		AddReferencedObjects = 2,
+		OuterChain = 3,
 	};
 
 	/** Extended information about a reference */
-	template <typename T>
+	template<typename T> 
 	struct TReferenceInfo
 	{
 		// Maximum number of stack frames to keep for AddReferencedObjects function calls
 		static constexpr uint32 MaxStackFrames = 30;
 
 		/** Object that is being referenced */
-		T* Object;
+		T* Object = nullptr;
+		TArray<uint64> StackFrames;
 		/** Type of reference to the object being referenced */
 		EReferenceType Type;
 		/** Name of the object or property that is referencing this object */
 		FName ReferencerName;
-		uint64 StackFrames[MaxStackFrames];
-		int32 NumStackFrames;
 
 		/** Default ctor */
 		TReferenceInfo()
-			: Object(nullptr)
-			, Type(EReferenceType::Unknown)
+			: Type(EReferenceType::Unknown)
 		{
-			InitStackFrames(nullptr, 0);
 		}
 
-		/** Simple refernce constructor. Probably will be filled with more info later */
+		/** Simple reference constructor. Probably will be filled with more info later */
 		TReferenceInfo(T* InObject)
 			: Object(InObject)
 			, Type(EReferenceType::Unknown)
 		{
-			InitStackFrames(nullptr, 0);
 		}
 
-		/** Full reference infor constructor */
-		TReferenceInfo(T* InObject, EReferenceType InType, const FName& InReferencerName, uint64* InStackFrames, int32 InNumStackFrames)
+		/** Full reference info constructor */
+		TReferenceInfo(
+			T* InObject, 
+			EReferenceType InType, 
+			FName InReferencerName,
+			TConstArrayView<uint64> InStackFrames = {}
+		)
 			: Object(InObject)
+			, StackFrames(InStackFrames)
 			, Type(InType)
 			, ReferencerName(InReferencerName)
 		{
-			InitStackFrames(InStackFrames, InNumStackFrames);
 		}
 
-		bool operator == (const TReferenceInfo& Other) const
+		TReferenceInfo(TReferenceInfo&&) = default;
+		TReferenceInfo& operator=(TReferenceInfo&&) = default;
+
+		bool operator ==(const TReferenceInfo& Other) const
 		{
 			return Object == Other.Object;
 		}
@@ -142,19 +153,13 @@ public:
 			}
 			return FString();
 		}
-	private:
-
-		void InitStackFrames(uint64* InStackFrames,int32 InNumStackFrames)
-		{
-			check(InNumStackFrames <= MaxStackFrames);
-			NumStackFrames = InNumStackFrames;
-			FMemory::Memset(StackFrames, 0, sizeof(StackFrames));
-			if (InStackFrames && InNumStackFrames)
-			{
-				FMemory::Memcpy(StackFrames, InStackFrames, InNumStackFrames * sizeof(uint64));
-			}
-		}
 	};
+	
+	struct FGraphNode;
+
+	/** Convenience type definitions to avoid template hell */
+	typedef TReferenceInfo<FGCObjectInfo> FObjectReferenceInfo;
+	typedef TReferenceInfo<FGraphNode> FNodeReferenceInfo;
 
 	/** Single node in the reference graph */
 	struct FGraphNode
@@ -165,16 +170,18 @@ public:
 		/** Object pointer */
 		FGCObjectInfo* ObjectInfo = nullptr;
 		/** Objects referenced by this object with reference info */
-		TSet< TReferenceInfo<FGraphNode> > ReferencedObjects;
+		TSet<FNodeReferenceInfo> ReferencedObjects;
 		/** Objects that have references to this object */
 		TSet<FGraphNode*> ReferencedByObjects;
 		/** Non-zero if this node has been already visited during reference search */
 		int32 Visited = 0;
+
+		int64 GetAllocatedSize() const
+		{
+			return ReferencedObjects.GetAllocatedSize() + ReferencedByObjects.GetAllocatedSize();
+		}
 	};
 
-	/** Convenience type definitions to avoid template hell */
-	typedef TReferenceInfo<FGCObjectInfo> FObjectReferenceInfo;
-	typedef TReferenceInfo<FGraphNode> FNodeReferenceInfo;
 
 	/** Reference chain. The first object in the list is the target object and the last object is a root object */
 	class FReferenceChain
@@ -182,12 +189,12 @@ public:
 		friend class FReferenceChainSearch;
 
 		/** The target nodes that caused this chain to be created. Usually this will be Nodes[0] unless the chain was truncated by request. */
-		FGraphNode* TargetNode;
+		FGraphNode* TargetNode = nullptr;
 
 		/** Nodes in this reference chain. The first node is the target object and the last one is a root object */
 		TArray<FGraphNode*> Nodes;
 		/** Reference information for Nodes */
-		TArray<FNodeReferenceInfo> ReferenceInfos;
+		TArray<const FNodeReferenceInfo*> ReferenceInfos;
 
 		/** Fills this chain with extended reference info for each node */
 		void FillReferenceInfo();
@@ -199,6 +206,11 @@ public:
 			Nodes.Reserve(ReserveDepth);
 		}
 
+		int64 GetAllocatedSize() const
+		{
+			return Nodes.GetAllocatedSize() + ReferenceInfos.GetAllocatedSize();
+		}
+
 		/** Adds a new node to the chain */
 		void AddNode(FGraphNode* InNode)
 		{
@@ -243,7 +255,7 @@ public:
 			return Nodes.Contains(InNode);
 		}
 		/** Gets extended reference info for the specified node index */
-		const FNodeReferenceInfo& GetReferenceInfo(int32 NodeIndex) const
+		const FNodeReferenceInfo* GetReferenceInfo(int32 NodeIndex) const
 		{
 			return ReferenceInfos[NodeIndex];
 		}
@@ -312,6 +324,8 @@ public:
 		return ReferenceChains;
 	}
 
+	int64 GetAllocatedSize() const;
+
 private:
 
 	/** The objects we're going to look for references to */
@@ -332,9 +346,12 @@ private:
 	/** Performs the search */
 	void PerformSearch();
 
-	/** Performs the search */
+	/** Constructs a complete graph of all references between objects */
 	void FindDirectReferencesForObjects();
 
+	/** Constructs a subset of the graph of references between objects to try and find the targets using minimal memory */
+	void FindMinimalDirectReferencesForObjects();
+
 	/** Frees memory */
 	void Cleanup();
 
@@ -342,6 +359,20 @@ private:
 	FGraphNode* FindOrAddNode(UObject* InObjectToFindNodeFor);
 	FGraphNode* FindOrAddNode(FGCObjectInfo* InObjectInfo);
 
+	/** Link two nodes together with the given reference info */
+	void LinkNodes(
+		UObject* From,
+		UObject* To,
+		EReferenceType ReferenceType,
+		FName PropertyName,
+		TConstArrayView<uint64> StackFrames = {});
+	void LinkNodes(
+		FGCObjectInfo* From,
+		FGCObjectInfo* To,
+		EReferenceType ReferenceType,
+		FName PropertyName,
+		TConstArrayView<uint64> StackFrames = {});
+
 	/** Builds reference chains */
 	static int32 BuildReferenceChainsRecursive(FGraphNode* TargetNodes, TArray<FReferenceChain*>& ProducedChains, int32 ChainDepth, const int32 VisitCounter, EReferenceChainSearchMode SearchMode, FGraphNode* GCObjReferencerNode);
 	/** Builds reference chains */
diff --git a/Engine/Source/Runtime/Engine/Classes/Engine/Engine.h b/Engine/Source/Runtime/Engine/Classes/Engine/Engine.h
--- a/Engine/Source/Runtime/Engine/Classes/Engine/Engine.h
+++ b/Engine/Source/Runtime/Engine/Classes/Engine/Engine.h
@@ -673,6 +673,9 @@ enum class EPrintStaleReferencesOptions
 	Fatal = 5,
 	Ensure = 0x00000100,
 
+	// Only search for direct references to the object or one of its inners, not a full reference chain 
+	Minimal = 0x00000200, 
+
 	VerbosityMask = 0x000000ff
 };
 ENUM_CLASS_FLAGS(EPrintStaleReferencesOptions);
@@ -3275,7 +3278,7 @@ public:
 	static void FindAndPrintStaleReferencesToObject(UObject* ObjectToFindReferencesTo, ELogVerbosity::Type Verbosity);
 
 	/**
-	 * Attempts to find what is referencing a world that should have been garbage collected
+	 * Attempts to find a reference chain leading to a world that should have been garbage collected
 	 * @param ObjectToFindReferencesTo World or its package (or any object from the world package that should've been destroyed)
 	 * @param Options Determines how the stale references messages should be logged
 	 */
diff --git a/Engine/Source/Runtime/Engine/Private/UnrealEngine.cpp b/Engine/Source/Runtime/Engine/Private/UnrealEngine.cpp
--- a/Engine/Source/Runtime/Engine/Private/UnrealEngine.cpp
+++ b/Engine/Source/Runtime/Engine/Private/UnrealEngine.cpp
@@ -444,6 +444,36 @@ static FAutoConsoleVariableRef GShouldLogReferencesToLeakedWorldObjectsCVar(
 	ECVF_Default
 );
 
+static int32 GVerifyLoadMapWorldCleanup_Severity = 2;
+static FAutoConsoleVariableRef CVarVerifyLoadMapWorldCleanup_Severity(
+#if UE_BUILD_SHIPPING
+	TEXT("Engine.VerifyLoadMapWorldCleanup.Severity.Shipping"),
+#else
+	TEXT("Engine.VerifyLoadMapWorldCleanup.Severity"),
+#endif
+	GVerifyLoadMapWorldCleanup_Severity,
+	TEXT("Controls severity of logging when the engine detects that a UWorld was leaked during LoadMap.\n")
+	TEXT("0 - all reference tracing and logging is disabled\n")
+	TEXT("1 - logs an error\n")
+	TEXT("2 - ensure\n")
+	TEXT("3 - fatal error\n"),
+	ECVF_Default
+);
+
+static int32 GVerifyLoadMapWorldCleanup_TraceMode = (UE_BUILD_SHIPPING ? 0 : 1);
+static FAutoConsoleVariableRef CVarVerifyLoadMapWorldCleanup_TraceMode(
+#if UE_BUILD_SHIPPING
+	TEXT("Engine.VerifyLoadMapWorldCleanup.TraceMode.Shipping"),
+#else
+	TEXT("Engine.VerifyLoadMapWorldCleanup.TraceMode"),
+#endif
+	GVerifyLoadMapWorldCleanup_TraceMode,
+	TEXT("Controls detail level of reference tracing when the engine detects that a world was leaked during LoadMap.\n")
+	TEXT("0 - direct references only\n")
+	TEXT("1 - full reference trace"),
+	ECVF_Default
+);
+
 #if !UE_BUILD_SHIPPING
 TAutoConsoleVariable<bool> CVarCsvAlwaysShowFrameCount(
 	TEXT("csv.AlwaysShowFrameCount"),
@@ -15598,10 +15628,57 @@ void UEngine::CheckAndHandleStaleWorldObjectReferences(FWorldContext* WorldConte
 		UE_LOG(LogLoad, Error, TEXT("Some previously active worlds or related objects were not cleaned up by garbage collection!"));
 		UE_LOG(LogLoad, Error, TEXT("Once a world has become active, it cannot be reused and must be destroyed and reloaded."));
 
-		if (GShouldLogReferencesToLeakedWorldObjects)
+		if (GVerifyLoadMapWorldCleanup_Severity != 0)
 		{
+			EPrintStaleReferencesOptions Options = EPrintStaleReferencesOptions::None;
+			switch (GVerifyLoadMapWorldCleanup_Severity)
+			{
+			case 3:
+				Options = EPrintStaleReferencesOptions::Fatal;
+				break;
+			case 2: 
+				Options = EPrintStaleReferencesOptions::Ensure | EPrintStaleReferencesOptions::Error;
+				break;
+			case 1:
+			default:
+				Options = EPrintStaleReferencesOptions::Error;
+				break;
+			}	
+			// Force fatal error if reference elimination is enabled as legacy behavior
+			if (UObjectBaseUtility::IsPendingKillEnabled())
+			{
+				Options = EPrintStaleReferencesOptions::Fatal;
+			}
+			
+			if (GVerifyLoadMapWorldCleanup_TraceMode == 0)
+			{
+				Options |= EPrintStaleReferencesOptions::Minimal;
+			}
+
+			// Remove sublevels from the list if their owning world also leaked
+			// Unless tracing is in minimal mode; in that case we may miss useful info
+			if (!(Options & EPrintStaleReferencesOptions::Minimal))
+			{
+				for (int32 i = LeakedObjects.Num() - 1; i >= 0; --i)
+				{
+					UObject* LeakedObject = LeakedObjects[i];
+					if (UWorld* World = Cast<UWorld>(LeakedObject))
+					{
+						if (!World->PersistentLevel || World->PersistentLevel->OwningWorld == World)
+						{
+							continue;
+						}
+
+						if (LeakedObjects.Contains(World->PersistentLevel->OwningWorld))
+						{
+							LeakedObjects.RemoveAt(i, 1, false);
+						}
+					}
+				}
+			}
+
 			UE_LOG(LogLoad, Error, TEXT("Dumping reference chains:"));
-			FindAndPrintStaleReferencesToObjects(LeakedObjects, UObjectBaseUtility::IsPendingKillEnabled() ? EPrintStaleReferencesOptions::Fatal : (EPrintStaleReferencesOptions::Error | EPrintStaleReferencesOptions::Ensure));
+			FindAndPrintStaleReferencesToObjects(LeakedObjects, Options);
 		}
 
 		if (!UObjectBaseUtility::IsPendingKillEnabled())
@@ -15723,10 +15800,27 @@ TArray<FString> UEngine::FindAndPrintStaleReferencesToObjects(TConstArrayView<UO
 		checkf(Obj, TEXT("Object to find references to cannot be null"));
 	}
 
-	const ELogVerbosity::Type OptionsToVerbosityMapping[] = { ELogVerbosity::NoLogging, ELogVerbosity::Log, ELogVerbosity::Display, ELogVerbosity::Warning, ELogVerbosity::Error, ELogVerbosity::Fatal };
-	const int32 VerbosityIndex = int32(Options & EPrintStaleReferencesOptions::VerbosityMask);
-	checkf(VerbosityIndex < UE_ARRAY_COUNT(OptionsToVerbosityMapping), TEXT("Unsupported verbosity option provided: %d"), VerbosityIndex);
-	ELogVerbosity::Type Verbosity = OptionsToVerbosityMapping[VerbosityIndex];
+	ELogVerbosity::Type Verbosity = ELogVerbosity::NoLogging;
+	switch (Options & EPrintStaleReferencesOptions::VerbosityMask)
+	{
+	case EPrintStaleReferencesOptions::None:
+		Verbosity = ELogVerbosity::NoLogging; 
+		break;
+	case EPrintStaleReferencesOptions::Log:
+		Verbosity = ELogVerbosity::Log; 
+		break;
+	case EPrintStaleReferencesOptions::Display:
+		Verbosity = ELogVerbosity::Display; 
+		break;
+	case EPrintStaleReferencesOptions::Warning:
+		Verbosity = ELogVerbosity::Warning; 
+		break;
+	case EPrintStaleReferencesOptions::Error:
+	case EPrintStaleReferencesOptions::Fatal: // Use verbosity error for reference chains so we can log all errors before crashing 
+	default:
+		Verbosity = ELogVerbosity::Error; 
+		break;
+	}
 
 	UE_LOG(LogLoad, Log, TEXT("Beginning reference chain search..."));
 	for (UObject* Obj : ObjectsToFindReferencesTo)
@@ -15734,11 +15828,22 @@ TArray<FString> UEngine::FindAndPrintStaleReferencesToObjects(TConstArrayView<UO
 		UE_LOG(LogLoad, Log, TEXT(" - %s"), *Obj->GetFullName());
 	}
 
-	FReferenceChainSearch RefChainSearch(ObjectsToFindReferencesTo, EReferenceChainSearchMode::ShortestToGarbage);
+	EReferenceChainSearchMode SearchMode = EReferenceChainSearchMode::GCOnly;
+	if(!!(Options & EPrintStaleReferencesOptions::Minimal))
+	{
+		SearchMode |= EReferenceChainSearchMode::Minimal;
+	}
+	else
+	{
+		SearchMode |= EReferenceChainSearchMode::ShortestToGarbage;
+	}
+
+	FReferenceChainSearch RefChainSearch(ObjectsToFindReferencesTo, SearchMode);
 #if ENABLE_GC_HISTORY
+	// If we have the full history already, don't only search for direct referencers 
 	FReferenceChainSearch HistorySearch(EReferenceChainSearchMode::ShortestToGarbage); // we may not need it but it has to live in this scope so that FGCObjectInfos are being kept alive
 #endif
-	
+
 	for (UObject* ObjectToFindReferencesTo : ObjectsToFindReferencesTo)
 	{
 		FGCObjectInfo* OutGarbageObject = nullptr;
@@ -15821,6 +15926,13 @@ TArray<FString> UEngine::FindAndPrintStaleReferencesToObjects(TConstArrayView<UO
 				*PathToCulprit);
 		}
 	}
+
+	if (!!(Options & EPrintStaleReferencesOptions::Fatal))
+	{
+		UE_LOG(LogLoad, Fatal, TEXT("Fatal world leaks detected. Logging first error, check logs for additional information") LINE_TERMINATOR TEXT("%s"), 
+			Paths.Num() ? *Paths[0] : TEXT("No paths to leaked objects found!"));
+	}
+
 	return Paths;
 }
 
