Bulk Data - deprecate GetPackagePath/PackageSegment

diff --git a/Engine/Plugins/Runtime/HairStrands/Source/HairStrandsCore/Private/GroomResources.cpp b/Engine/Plugins/Runtime/HairStrands/Source/HairStrandsCore/Private/GroomResources.cpp
--- a/Engine/Plugins/Runtime/HairStrands/Source/HairStrandsCore/Private/GroomResources.cpp
+++ b/Engine/Plugins/Runtime/HairStrands/Source/HairStrandsCore/Private/GroomResources.cpp
@@ -1002,7 +1002,7 @@ void FHairStrandsClusterCullingBulkData::Validate(bool bIsSaving)
 	}
 	if (!bIsValid)
 	{
-		FString Name = ClusterLODInfos.GetPackagePath().GetPackageName();
+		FString Name = ClusterLODInfos.GetDebugName();
 		UE_LOG(LogHairStrands, Error, TEXT("[Groom/DDC] Strands - Invalid ClusterCullingBulkData when %s bulk data - %s"), bIsSaving ? TEXT("Saving") : TEXT("Loading"), *Name);
 	}
 
diff --git a/Engine/Source/Runtime/CoreUObject/Private/Serialization/BulkData.cpp b/Engine/Source/Runtime/CoreUObject/Private/Serialization/BulkData.cpp
--- a/Engine/Source/Runtime/CoreUObject/Private/Serialization/BulkData.cpp
+++ b/Engine/Source/Runtime/CoreUObject/Private/Serialization/BulkData.cpp
@@ -1222,6 +1222,11 @@ FIoChunkId FBulkData::CreateChunkId() const
 	return UE::BulkData::Private::CreateBulkDataIoChunkId(BulkMeta, BulkChunkId.GetPackageId());
 }
 
+FString FBulkData::GetDebugName() const
+{
+	return BulkChunkId.ToDebugString();
+}
+
 /*-----------------------------------------------------------------------------
 	Serialization.
 -----------------------------------------------------------------------------*/
@@ -1717,7 +1722,7 @@ FCustomVersionContainer FBulkData::GetCustomVersions(FArchive& InlineArchive)
 	else
 	{
 		// Read the CustomVersions out of the separate package file
-		const FPackagePath& PackagePath = GetPackagePath();
+		const FPackagePath& PackagePath = BulkChunkId.GetPackagePath();
 		TUniquePtr<FArchive> ExternalArchive = IPackageResourceManager::Get().OpenReadExternalResource(
 			EPackageExternalResource::WorkspaceDomainFile, PackagePath.GetPackageName());
 		if (ExternalArchive.IsValid())
diff --git a/Engine/Source/Runtime/CoreUObject/Private/Serialization/EditorBulkData.cpp b/Engine/Source/Runtime/CoreUObject/Private/Serialization/EditorBulkData.cpp
--- a/Engine/Source/Runtime/CoreUObject/Private/Serialization/EditorBulkData.cpp
+++ b/Engine/Source/Runtime/CoreUObject/Private/Serialization/EditorBulkData.cpp
@@ -713,14 +713,14 @@ void FEditorBulkData::CreateFromBulkData(FBulkData& InBulkData, const FGuid& InG
 	bool bWasKeyGuidDerived = false;
 	if (InBulkData.GetBulkDataSize() > 0)
 	{
-		BulkDataId = CreateUniqueGuid(InGuid, Owner, *InBulkData.GetPackagePath().GetDebugName());
+		BulkDataId = CreateUniqueGuid(InGuid, Owner, *InBulkData.GetDebugName());
 		PayloadContentId = GuidToIoHash(BulkDataId);
 		bWasKeyGuidDerived = true;
 	}
 	
 	PayloadSize = InBulkData.GetBulkDataSize();
 	
-	PackagePath = InBulkData.GetPackagePath();
+	PackagePath = GetPackagePathFromOwner(Owner); 
 	
 	OffsetInFile = InBulkData.GetBulkDataOffsetInFile();
 
diff --git a/Engine/Source/Runtime/CoreUObject/Public/Serialization/BulkData.h b/Engine/Source/Runtime/CoreUObject/Public/Serialization/BulkData.h
--- a/Engine/Source/Runtime/CoreUObject/Public/Serialization/BulkData.h
+++ b/Engine/Source/Runtime/CoreUObject/Public/Serialization/BulkData.h
@@ -819,12 +819,19 @@ public:
 	void ClearBulkDataFlags( uint32 BulkDataFlagsToClear );
 
 	UE_DEPRECATED(5.0, "Use GetPackagePath instead")
-	FString GetFilename() const { return GetPackagePath().GetLocalFullPath(GetPackageSegment()); }
+	FString GetFilename() const
+	{ 
+		PRAGMA_DISABLE_DEPRECATION_WARNINGS
+		return GetPackagePath().GetLocalFullPath(GetPackageSegment());
+		PRAGMA_ENABLE_DEPRECATION_WARNINGS
+	}
 
 	/** Returns the PackagePath this bulkdata resides in */
+	UE_DEPRECATED(5.1, "Deprecated")
 	const FPackagePath& GetPackagePath() const { return BulkChunkId.GetPackagePath(); }
 
 	/** Returns which segment of its PackagePath this bulkdata resides in */
+	UE_DEPRECATED(5.1, "Deprecated")
 	EPackageSegment GetPackageSegment() const { return UE::BulkData::Private::GetPackageSegmentFromFlags(BulkMeta); }
 
 	/** 
@@ -837,6 +844,9 @@ public:
 	/** Returns a FIoChunkId for the bulkdata payload, this will be invalid if the bulkdata is not stored in the IoStore */
 	FIoChunkId CreateChunkId() const;
 
+	/** Returns a string representing the bulk data for debugging purposes. */
+	FString GetDebugName() const;
+
 	/*-----------------------------------------------------------------------------
 		Data retrieval and manipulation.
 	-----------------------------------------------------------------------------*/
diff --git a/Engine/Source/Runtime/Engine/Private/Animation/AnimationStreaming.cpp b/Engine/Source/Runtime/Engine/Private/Animation/AnimationStreaming.cpp
--- a/Engine/Source/Runtime/Engine/Private/Animation/AnimationStreaming.cpp
+++ b/Engine/Source/Runtime/Engine/Private/Animation/AnimationStreaming.cpp
@@ -222,7 +222,7 @@ void FStreamingAnimationData::BeginPendingRequests(const TArray<uint32>& Indices
 			};
 
 			UE_LOG(LogAnimation, Warning, TEXT("Loading Stream Anim %s Chunk:%i Length: %.3f Offset:%i Size:%i File:%s\n"),
-				*StreamableAnim->GetName(), ChunkIndex, Chunk.SequenceLength, Chunk.BulkData.GetBulkDataOffsetInFile(), Chunk.BulkData.GetBulkDataSize(), *Chunk.BulkData.GetPackagePath().GetDebugName());
+				*StreamableAnim->GetName(), ChunkIndex, Chunk.SequenceLength, Chunk.BulkData.GetBulkDataOffsetInFile(), Chunk.BulkData.GetBulkDataSize(), *Chunk.BulkData.GetDebugName());
 			ChunkStorage.IORequest = Chunk.BulkData.CreateStreamingRequest(AsyncIOPriority, &AsyncFileCallBack, nullptr);
 			if (!ChunkStorage.IORequest)
 			{
@@ -540,4 +540,4 @@ const FCompressedAnimSequence* FAnimationStreamingManager::GetLoadedChunk(const
 	}
 
 	return nullptr;
-}
\ No newline at end of file
+}
diff --git a/Engine/Source/Runtime/Engine/Private/TextureDerivedData.cpp b/Engine/Source/Runtime/Engine/Private/TextureDerivedData.cpp
--- a/Engine/Source/Runtime/Engine/Private/TextureDerivedData.cpp
+++ b/Engine/Source/Runtime/Engine/Private/TextureDerivedData.cpp
@@ -2024,7 +2024,7 @@ bool FTexturePlatformData::TryLoadMipsWithSizes(int32 FirstMipToLoad, void** Out
 				if (!FPlatformProperties::HasEditorOnlyData() && CVarSetTextureStreaming.GetValueOnAnyThread() != 0)
 				{
 					UE_CLOG(Mip.BulkData.IsInSeparateFile(), LogTexture, Error, TEXT("Loading non-streamed mips from an external bulk file.  This is not desireable.  File %s"),
-						*(Mip.BulkData.GetPackagePath().GetDebugName()) );
+						*(Mip.BulkData.GetDebugName()) );
 				}
 #endif
 				Mip.BulkData.GetCopy(&OutMipData[MipIndex - FirstMipToLoad], true);
diff --git a/Engine/Source/Runtime/Engine/Private/VT/UploadingVirtualTexture.cpp b/Engine/Source/Runtime/Engine/Private/VT/UploadingVirtualTexture.cpp
--- a/Engine/Source/Runtime/Engine/Private/VT/UploadingVirtualTexture.cpp
+++ b/Engine/Source/Runtime/Engine/Private/VT/UploadingVirtualTexture.cpp
@@ -391,7 +391,7 @@ FVTDataAndStatus FUploadingVirtualTexture::ReadData(FGraphEventArray& OutComplet
 		{
 			if (!InvalidChunks[ChunkIndex])
 			{
-				UE_LOG(LogConsoleResponse, Display, TEXT("BulkData for chunk %d in file '%s' is empty."), ChunkIndex, *BulkData.GetPackagePath().GetDebugName());
+				UE_LOG(LogConsoleResponse, Display, TEXT("BulkData for chunk %d in file '%s' is empty."), ChunkIndex, *BulkData.GetDebugName());
 				InvalidChunks[ChunkIndex] = true;
 			}
 			return EVTRequestPageStatus::Invalid;
@@ -461,7 +461,7 @@ void FUploadingVirtualTexture::DumpToConsole(bool verbose)
 		}
 		else
 #endif
-			BulkDataFiles.Add(Chunk.BulkData.GetPackagePath().GetLocalFullPath());
+			BulkDataFiles.Add(Chunk.BulkData.GetDebugName());
 	}
 
 	for (const auto& FileName : BulkDataFiles)
