AsyncLoading2: Only process/create pending CDOs when absolutely required to unblock the current FlushAsyncLoading stack during initial load.
Fix for FAsyncLoadingThread2::FlushLoading getting stuck in an infinite loop (constructur reentrancy) on projects where native class CDOs are loading blueprints.

#rnx

diff --git a/Engine/Source/Runtime/CoreUObject/Private/Serialization/AsyncLoading2.cpp b/Engine/Source/Runtime/CoreUObject/Private/Serialization/AsyncLoading2.cpp
--- a/Engine/Source/Runtime/CoreUObject/Private/Serialization/AsyncLoading2.cpp
+++ b/Engine/Source/Runtime/CoreUObject/Private/Serialization/AsyncLoading2.cpp
@@ -4313,9 +4313,10 @@ void FAsyncPackage2::StartLoading(FIoBatch& IoBatch)
 						TEXT("Failed reading optional chunk for package: %s"), *Result.Status().ToString());
 					bLoadHasFailed = true;
 				}
-				int32 LocalPendingIoRequestsCounter = AsyncLoadingThread.PendingIoRequestsCounter.DecrementExchange() - 1;
-				TRACE_COUNTER_SET(AsyncLoadingPendingIoRequests, LocalPendingIoRequestsCounter);
+				FAsyncLoadingThread2* LocalAsyncLoadingThread = &AsyncLoadingThread;
 				GetPackageNode(EEventLoadNode2::Package_ProcessSummary).ReleaseBarrier();
+				int32 LocalPendingIoRequestsCounter = LocalAsyncLoadingThread->PendingIoRequestsCounter.DecrementExchange() - 1;
+				TRACE_COUNTER_SET(AsyncLoadingPendingIoRequests, LocalPendingIoRequestsCounter);
 			});
 	}
 #endif
@@ -4339,9 +4340,10 @@ void FAsyncPackage2::StartLoading(FIoBatch& IoBatch)
 					TEXT("Failed reading chunk for package: %s"), *Result.Status().ToString());
 				bLoadHasFailed = true;
 			}
-			int32 LocalPendingIoRequestsCounter = AsyncLoadingThread.PendingIoRequestsCounter.DecrementExchange() - 1;
-			TRACE_COUNTER_SET(AsyncLoadingPendingIoRequests, LocalPendingIoRequestsCounter);
+			FAsyncLoadingThread2* LocalAsyncLoadingThread = &AsyncLoadingThread;
 			GetPackageNode(EEventLoadNode2::Package_ProcessSummary).ReleaseBarrier();
+			int32 LocalPendingIoRequestsCounter = LocalAsyncLoadingThread->PendingIoRequestsCounter.DecrementExchange() - 1;
+			TRACE_COUNTER_SET(AsyncLoadingPendingIoRequests, LocalPendingIoRequestsCounter);
 		});
 
 	if (!Data.ShaderMapHashes.IsEmpty())
@@ -4356,9 +4358,10 @@ void FAsyncPackage2::StartLoading(FIoBatch& IoBatch)
 				[this, GraphEvent](TIoStatusOr<FIoBuffer> Result)
 				{
 					GraphEvent->DispatchSubsequents();
-					int32 LocalPendingIoRequestsCounter = AsyncLoadingThread.PendingIoRequestsCounter.DecrementExchange() - 1;
-					TRACE_COUNTER_SET(AsyncLoadingPendingIoRequests, LocalPendingIoRequestsCounter);
+					FAsyncLoadingThread2* LocalAsyncLoadingThread = &AsyncLoadingThread;
 					GetPackageNode(Package_ExportsSerialized).ReleaseBarrier();
+					int32 LocalPendingIoRequestsCounter = LocalAsyncLoadingThread->PendingIoRequestsCounter.DecrementExchange() - 1;
+					TRACE_COUNTER_SET(AsyncLoadingPendingIoRequests, LocalPendingIoRequestsCounter);
 				});
 		};
 		FCoreDelegates::PreloadPackageShaderMaps.ExecuteIfBound(Data.ShaderMapHashes, ReadShaderMapFunc);
@@ -6452,6 +6455,7 @@ EAsyncPackageState::Type FAsyncLoadingThread2::TickAsyncLoadingFromGameThread(FA
 		ThreadState.SetTimeLimit(bUseTimeLimit, TimeLimit);
 
 		const bool bIsMultithreaded = FAsyncLoadingThread2::IsMultithreaded();
+		bool bMayProcessPendingCDOs = PendingCDOs.Num() > 0 && (PendingIoRequestsCounter.Load() == 0);
 		double TickStartTime = FPlatformTime::Seconds();
 
 		{
@@ -6467,9 +6471,14 @@ EAsyncPackageState::Type FAsyncLoadingThread2::TickAsyncLoadingFromGameThread(FA
 
 		if (Result != EAsyncPackageState::TimeOut)
 		{
-			if (!bDidSomething && PendingCDOs.Num() > 0)
+			if (!bDidSomething && bMayProcessPendingCDOs)
 			{
-				bDidSomething = ProcessPendingCDOs();
+				// Only process/create pending CDOs when absolutely required to unblock the current FlushAsyncLoading stack during initial load,
+				// i.e. when we are not making any progress and we are sure that there are no outstanding IO requests.
+				if (PendingIoRequestsCounter.Load() == 0)
+				{
+					bDidSomething = ProcessPendingCDOs();
+				}
 			}
 
 			// Flush deferred messages
