[AssetRegistry] changed TryGetAssetByObjectPath to use FSoftObjectPath to match the rest of the API

diff --git a/Engine/Source/Runtime/AssetRegistry/Private/AssetRegistry.cpp b/Engine/Source/Runtime/AssetRegistry/Private/AssetRegistry.cpp
--- a/Engine/Source/Runtime/AssetRegistry/Private/AssetRegistry.cpp
+++ b/Engine/Source/Runtime/AssetRegistry/Private/AssetRegistry.cpp
@@ -120,7 +120,7 @@ public:
 		IAssetRegistry::GetChecked().GetDependencies(InPackageName, OutDependencies, Category, Flags);
 	}
 
-	virtual UE::AssetRegistry::EExists TryGetAssetByObjectPath(const FName ObjectPath, FAssetData& OutAssetData) const override
+	virtual UE::AssetRegistry::EExists TryGetAssetByObjectPath(const FSoftObjectPath& ObjectPath, FAssetData& OutAssetData) const override
 	{
 		auto AssetRegistry = IAssetRegistry::Get();
 		if (!AssetRegistry)
@@ -2081,14 +2081,12 @@ PRAGMA_DISABLE_DEPRECATION_WARNINGS;
 PRAGMA_ENABLE_DEPRECATION_WARNINGS;
 }
 
-UE::AssetRegistry::EExists UAssetRegistryImpl::TryGetAssetByObjectPath(const FName ObjectPath, FAssetData& OutAssetData) const
+UE::AssetRegistry::EExists UAssetRegistryImpl::TryGetAssetByObjectPath(const FSoftObjectPath& ObjectPath, FAssetData& OutAssetData) const
 {
 	FReadScopeLock InterfaceScopeLock(InterfaceLock);
 	bool bAssetRegistryReady = GuardedData.IsInitialSearchStarted() && GuardedData.IsInitialSearchCompleted();
 	const FAssetRegistryState& State = GuardedData.GetState();
-	PRAGMA_DISABLE_DEPRECATION_WARNINGS;
-	const FAssetData* FoundData = State.GetAssetByObjectPath(FSoftObjectPath(ObjectPath.ToString()));
-	PRAGMA_ENABLE_DEPRECATION_WARNINGS;
+	const FAssetData* FoundData = State.GetAssetByObjectPath(ObjectPath);
 	if (!FoundData)
 	{
 		if (!bAssetRegistryReady)
diff --git a/Engine/Source/Runtime/AssetRegistry/Private/AssetRegistry.h b/Engine/Source/Runtime/AssetRegistry/Private/AssetRegistry.h
--- a/Engine/Source/Runtime/AssetRegistry/Private/AssetRegistry.h
+++ b/Engine/Source/Runtime/AssetRegistry/Private/AssetRegistry.h
@@ -45,7 +45,7 @@ public:
 	UE_DEPRECATED(5.1, "Asset path FNames have been deprecated, use FSoftObjectPath instead.")
 	virtual FAssetData GetAssetByObjectPath( const FName ObjectPath, bool bIncludeOnlyOnDiskAssets = false ) const override;
 	virtual FAssetData GetAssetByObjectPath(const FSoftObjectPath& ObjectPath, bool bIncludeOnlyOnDiskAssets = false) const override;
-	virtual UE::AssetRegistry::EExists TryGetAssetByObjectPath(const FName ObjectPath, FAssetData& OutAssetData) const override;
+	virtual UE::AssetRegistry::EExists TryGetAssetByObjectPath(const FSoftObjectPath& ObjectPath, FAssetData& OutAssetData) const override;
 	virtual UE::AssetRegistry::EExists TryGetAssetPackageData(const FName PackageName, FAssetPackageData& OutAssetPackageData) const override;
 	virtual bool GetAllAssets(TArray<FAssetData>& OutAssetData, bool bIncludeOnlyOnDiskAssets = false) const override;
 	virtual bool EnumerateAllAssets(TFunctionRef<bool(const FAssetData&)> Callback, bool bIncludeOnlyOnDiskAssets = false) const override;
diff --git a/Engine/Source/Runtime/AssetRegistry/Public/AssetRegistry/AssetRegistryModule.h b/Engine/Source/Runtime/AssetRegistry/Public/AssetRegistry/AssetRegistryModule.h
--- a/Engine/Source/Runtime/AssetRegistry/Public/AssetRegistry/AssetRegistryModule.h
+++ b/Engine/Source/Runtime/AssetRegistry/Public/AssetRegistry/AssetRegistryModule.h
@@ -96,7 +96,7 @@ public:
 		IAssetRegistry::GetChecked().GetDependencies(InPackageName, OutDependencies, Category, Flags);
 	}
 
-	virtual UE::AssetRegistry::EExists TryGetAssetByObjectPath(const FName ObjectPath, FAssetData& OutAssetData) const override
+	virtual UE::AssetRegistry::EExists TryGetAssetByObjectPath(const FSoftObjectPath& ObjectPath, FAssetData& OutAssetData) const override
 	{
 		IAssetRegistry* AssetRegistry = IAssetRegistry::Get();
 		if (!AssetRegistry)
diff --git a/Engine/Source/Runtime/AssetRegistry/Public/AssetRegistry/IAssetRegistry.h b/Engine/Source/Runtime/AssetRegistry/Public/AssetRegistry/IAssetRegistry.h
--- a/Engine/Source/Runtime/AssetRegistry/Public/AssetRegistry/IAssetRegistry.h
+++ b/Engine/Source/Runtime/AssetRegistry/Public/AssetRegistry/IAssetRegistry.h
@@ -249,7 +249,7 @@ public:
 	 * @param OutAssetData out FAssetData 
 	 * @return Enum return code
 	 */
-	virtual UE::AssetRegistry::EExists TryGetAssetByObjectPath(const FName ObjectPath, FAssetData& OutAssetData) const = 0;
+	virtual UE::AssetRegistry::EExists TryGetAssetByObjectPath(const FSoftObjectPath& ObjectPath, FAssetData& OutAssetData) const = 0;
 
 	/**
 	 * Tries to get the pacakge data for a specified path
diff --git a/Engine/Source/Runtime/CoreUObject/Internal/UObject/LinkerLoadImportBehavior.h b/Engine/Source/Runtime/CoreUObject/Internal/UObject/LinkerLoadImportBehavior.h
--- a/Engine/Source/Runtime/CoreUObject/Internal/UObject/LinkerLoadImportBehavior.h
+++ b/Engine/Source/Runtime/CoreUObject/Internal/UObject/LinkerLoadImportBehavior.h
@@ -38,7 +38,7 @@ bool TryLazyImport(const IAssetRegistryInterface& AssetRegistry, const FObjectIm
 /// @param ObjectPath object path of to resolve
 /// @param OutObjectPtr out parameter
 /// @return true if was lazy loaded otherwise false
-bool TryLazyLoad(const UClass& Class, FName ObjectPath, TObjectPtr<UObject>& OutObjectPtr);
+bool TryLazyLoad(const UClass& Class, const FSoftObjectPath& ObjectPath, TObjectPtr<UObject>& OutObjectPtr);
 
 }
 
diff --git a/Engine/Source/Runtime/CoreUObject/Private/UObject/LinkerLoadImportBehavior.cpp b/Engine/Source/Runtime/CoreUObject/Private/UObject/LinkerLoadImportBehavior.cpp
--- a/Engine/Source/Runtime/CoreUObject/Private/UObject/LinkerLoadImportBehavior.cpp
+++ b/Engine/Source/Runtime/CoreUObject/Private/UObject/LinkerLoadImportBehavior.cpp
@@ -77,7 +77,7 @@ EImportBehavior GetPropertyImportLoadBehavior(const FObjectImport& Import, const
 }
 
 // recursively handles FAssetData redirectors
-static bool HandleRedirector(const IAssetRegistryInterface& AssetRegistry, const FAssetData& InAssetData, TSet<FName>& SeenPaths, FObjectPtr& OutObjectPtr)
+static bool HandleRedirector(const IAssetRegistryInterface& AssetRegistry, const FAssetData& InAssetData, TSet<FSoftObjectPath>& SeenPaths, FObjectPtr& OutObjectPtr)
 {
 	FString RedirectedPath;
 	if (!InAssetData.GetTagValue("DestinationObject", RedirectedPath))
@@ -86,10 +86,10 @@ static bool HandleRedirector(const IAssetRegistryInterface& AssetRegistry, const
 	}
 
 	ConstructorHelpers::StripObjectClass(RedirectedPath);
-	FName RedirectedPathFName(RedirectedPath);
+	FSoftObjectPath RedirectedPathName(RedirectedPath);
 
 	bool bIsAlreadyInSet;
-	SeenPaths.FindOrAdd(RedirectedPathFName, &bIsAlreadyInSet);
+	SeenPaths.FindOrAdd(RedirectedPathName, &bIsAlreadyInSet);
 	if(bIsAlreadyInSet)
 	{
 		//Recursive redirectors
@@ -97,7 +97,7 @@ static bool HandleRedirector(const IAssetRegistryInterface& AssetRegistry, const
 	}
 
 	FAssetData AssetData;
-	UE::AssetRegistry::EExists Exists = AssetRegistry.TryGetAssetByObjectPath(FName(RedirectedPath), AssetData);
+	UE::AssetRegistry::EExists Exists = AssetRegistry.TryGetAssetByObjectPath(RedirectedPathName, AssetData);
 	if (Exists == UE::AssetRegistry::EExists::Unknown)
 	{
 		return false;
@@ -121,7 +121,7 @@ static bool HandleRedirector(const IAssetRegistryInterface& AssetRegistry, const
 	return true;
 }
 
-static bool TryLazyLoad(const IAssetRegistryInterface& AssetRegistry, FName ObjectPath, FObjectPtr& OutObjectPtr)
+static bool TryLazyLoad(const IAssetRegistryInterface& AssetRegistry, const FSoftObjectPath& ObjectPath, FObjectPtr& OutObjectPtr)
 {
 	OutObjectPtr = nullptr;
 	
@@ -138,7 +138,7 @@ static bool TryLazyLoad(const IAssetRegistryInterface& AssetRegistry, FName Obje
 	}
 	if (AssetData.IsRedirector())
 	{
-		TSet<FName> SeenPaths;
+		TSet<FSoftObjectPath> SeenPaths;
 		return HandleRedirector(AssetRegistry, AssetData, SeenPaths, OutObjectPtr);
 	}
 	FNameBuilder NameBuilder;
@@ -188,11 +188,11 @@ bool TryLazyImport(const IAssetRegistryInterface& AssetRegistry, const FObjectIm
 	TStringBuilder<FName::StringBufferSize> PathName;
 	ImportRef.AppendPathName(PathName);
 
-	return TryLazyLoad(AssetRegistry, FName(PathName), ObjectPtr);
+	return TryLazyLoad(AssetRegistry, FSoftObjectPath(PathName), ObjectPtr);
 }
 
 
-bool TryLazyLoad(const UClass& Class, FName ObjectPath, TObjectPtr<UObject>& OutObjectPtr)
+bool TryLazyLoad(const UClass& Class, const FSoftObjectPath& ObjectPath, TObjectPtr<UObject>& OutObjectPtr)
 {
 	if (!FLinkerLoad::IsImportLazyLoadEnabled())
 	{
diff --git a/Engine/Source/Runtime/CoreUObject/Public/Misc/AssetRegistryInterface.h b/Engine/Source/Runtime/CoreUObject/Public/Misc/AssetRegistryInterface.h
--- a/Engine/Source/Runtime/CoreUObject/Public/Misc/AssetRegistryInterface.h
+++ b/Engine/Source/Runtime/CoreUObject/Public/Misc/AssetRegistryInterface.h
@@ -297,7 +297,7 @@ public:
 	 * @param OutAssetData out FAssetData
 	 * @return Return code enum
 	 */
-	virtual UE::AssetRegistry::EExists TryGetAssetByObjectPath(const FName ObjectPath, struct FAssetData& OutAssetData) const = 0;
+	virtual UE::AssetRegistry::EExists TryGetAssetByObjectPath(const FSoftObjectPath& ObjectPath, struct FAssetData& OutAssetData) const = 0;
 
 	/**
 	 * Tries to get the pacakge data for a specified path
