Properly skip soft object path list for text asset

diff --git a/Engine/Source/Runtime/CoreUObject/Private/UObject/LinkerLoad.cpp b/Engine/Source/Runtime/CoreUObject/Private/UObject/LinkerLoad.cpp
--- a/Engine/Source/Runtime/CoreUObject/Private/UObject/LinkerLoad.cpp
+++ b/Engine/Source/Runtime/CoreUObject/Private/UObject/LinkerLoad.cpp
@@ -1679,7 +1679,13 @@ FLinkerLoad::ELinkerStatus FLinkerLoad::SerializeNameMap()
 
 FLinkerLoad::ELinkerStatus FLinkerLoad::SerializeSoftObjectPathList()
 {
-	if (!IsTextFormat() && SoftObjectPathListIndex == 0 && Summary.SoftObjectPathsCount > 0)
+	// Text archives don't have soft object path tables at the moment
+	if (IsTextFormat())
+	{
+		return LINKER_Loaded;
+	}
+
+	if (SoftObjectPathListIndex == 0 && Summary.SoftObjectPathsCount > 0)
 	{
 		Seek(Summary.SoftObjectPathsOffset);
 	}
