Fix FAssetData::GetObjectPathString to return the empty string rather than None.None for empty/default initialized asset data.

diff --git a/Engine/Source/Runtime/CoreUObject/Private/AssetRegistry/AssetData.cpp b/Engine/Source/Runtime/CoreUObject/Private/AssetRegistry/AssetData.cpp
--- a/Engine/Source/Runtime/CoreUObject/Private/AssetRegistry/AssetData.cpp
+++ b/Engine/Source/Runtime/CoreUObject/Private/AssetRegistry/AssetData.cpp
@@ -972,3 +972,19 @@ FArchive& operator<<(FArchive& Ar, UE::AssetRegistry::FPackageCustomVersionsHand
 
 }
 }
+
+#if WITH_DEV_AUTOMATION_TESTS 
+
+#include "Misc/AutomationTest.h"
+IMPLEMENT_SIMPLE_AUTOMATION_TEST(FAssetDataTests, "System.CoreUObject.AssetData", EAutomationTestFlags::ApplicationContextMask | EAutomationTestFlags::EngineFilter);
+bool FAssetDataTests::RunTest(const FString& Parameters)
+{
+	FAssetData EmptyAssetData;
+
+	TestEqual(TEXT("Empty Asset Data: Object path string is empty"), EmptyAssetData.GetObjectPathString(), FString());
+
+	return true;
+}
+
+
+#endif // WITH_DEV_AUTOMATION_TESTS
diff --git a/Engine/Source/Runtime/CoreUObject/Public/AssetRegistry/AssetData.h b/Engine/Source/Runtime/CoreUObject/Public/AssetRegistry/AssetData.h
--- a/Engine/Source/Runtime/CoreUObject/Public/AssetRegistry/AssetData.h
+++ b/Engine/Source/Runtime/CoreUObject/Public/AssetRegistry/AssetData.h
@@ -456,6 +456,11 @@ PRAGMA_ENABLE_DEPRECATION_WARNINGS
 	/** Append the object path to the given string builder. */
 	void AppendObjectPath(FStringBuilderBase& Builder) const
 	{
+		if (!IsValid())
+		{
+			return;
+		}
+
 #if WITH_EDITORONLY_DATA
 		if (!OptionalOuterPath.IsNone())
 		{
