Fix potential deadlocks when serializing objects while holding the object hash lock

Do not force flush asset compilation while fixing up soft object path through serialization

diff --git a/Engine/Source/Runtime/CoreUObject/Public/UObject/SoftObjectPath.h b/Engine/Source/Runtime/CoreUObject/Public/UObject/SoftObjectPath.h
--- a/Engine/Source/Runtime/CoreUObject/Public/UObject/SoftObjectPath.h
+++ b/Engine/Source/Runtime/CoreUObject/Public/UObject/SoftObjectPath.h
@@ -557,6 +557,7 @@ struct FSoftObjectPathFixupArchive : public FArchiveUObject
 	{
 		this->SetIsSaving(true);
 		this->ArShouldSkipBulkData = true;
+		this->SetShouldSkipCompilingAssets(true);
 	}
 
 	FSoftObjectPathFixupArchive(const FString& InOldAssetPathString, const FString& InNewAssetPathString)
@@ -579,7 +580,12 @@ struct FSoftObjectPathFixupArchive : public FArchiveUObject
 	void Fixup(UObject* Root)
 	{
 		Root->Serialize(*this);
-		ForEachObjectWithOuter(Root, [this](UObject* InObject) { InObject->Serialize(*this);  });
+		TArray<UObject*> SubObjects;
+		GetObjectsWithOuter(Root, SubObjects);
+		for (UObject* Obj : SubObjects)
+		{
+			Obj->Serialize(*this);
+		}
 	}
 
 	TFunction<void(FSoftObjectPath&)> FixupFunction;
