ZenLoader: Ensure that packages for sync load requests are created as soon as possible so that we don't waste time looking for them in the event queue

diff --git a/Engine/Source/Runtime/CoreUObject/Private/Serialization/AsyncLoading2.cpp b/Engine/Source/Runtime/CoreUObject/Private/Serialization/AsyncLoading2.cpp
--- a/Engine/Source/Runtime/CoreUObject/Private/Serialization/AsyncLoading2.cpp
+++ b/Engine/Source/Runtime/CoreUObject/Private/Serialization/AsyncLoading2.cpp
@@ -3732,6 +3732,8 @@ bool FAsyncLoadEventQueue2::PopAndExecute(FAsyncLoadingThreadState2& ThreadState
 
 bool FAsyncLoadEventQueue2::ExecuteSyncLoadEvents(FAsyncLoadingThreadState2& ThreadState, FAsyncLoadingSyncLoadContext& SyncLoadContext)
 {
+	TRACE_CPUPROFILER_EVENT_SCOPE(ExecuteSyncLoadEvents);
+
 	int32 ThisCallCounter = ++ExecuteSyncLoadEventsCallCounter;
 
 	auto ShouldExecuteNode = [&SyncLoadContext](FEventLoadNode2* Node) -> bool
@@ -6062,6 +6064,8 @@ FAsyncLoadingSyncLoadContext* FAsyncLoadingThread2::UpdateSyncLoadContext(FAsync
 	}
 	if (ThreadState.bCanAccessAsyncLoadingThreadData && !SyncLoadContext->bHasFoundRequestedPackage)
 	{
+		// Ensure that we've created the package we're waiting for
+		CreateAsyncPackagesFromQueue(ThreadState);
 		if (FAsyncPackage2* RequestedPackage = RequestIdToPackageMap.FindRef(SyncLoadContext->RequestId))
 		{
 			SyncLoadContext->bHasFoundRequestedPackage = true;
