ZenLoader: Fixed incorrect calling of CreateAsyncPackagesFromQueue when running without the ALT

diff --git a/Engine/Source/Runtime/CoreUObject/Private/Serialization/AsyncLoading2.cpp b/Engine/Source/Runtime/CoreUObject/Private/Serialization/AsyncLoading2.cpp
--- a/Engine/Source/Runtime/CoreUObject/Private/Serialization/AsyncLoading2.cpp
+++ b/Engine/Source/Runtime/CoreUObject/Private/Serialization/AsyncLoading2.cpp
@@ -6064,8 +6064,11 @@ FAsyncLoadingSyncLoadContext* FAsyncLoadingThread2::UpdateSyncLoadContext(FAsync
 	}
 	if (ThreadState.bCanAccessAsyncLoadingThreadData && !SyncLoadContext->bHasFoundRequestedPackage)
 	{
-		// Ensure that we've created the package we're waiting for
-		CreateAsyncPackagesFromQueue(ThreadState);
+		if (ThreadState.bIsAsyncLoadingThread)
+		{
+			// Ensure that we've created the package we're waiting for
+			CreateAsyncPackagesFromQueue(ThreadState);
+		}
 		if (FAsyncPackage2* RequestedPackage = RequestIdToPackageMap.FindRef(SyncLoadContext->RequestId))
 		{
 			SyncLoadContext->bHasFoundRequestedPackage = true;
