ZenLoader: Ensure that packages for sync load requests are created as soon as possible so that we don't waste time looking for them in the event queue

diff --git a/Engine/Source/Runtime/CoreUObject/Private/Serialization/AsyncLoading2.cpp b/Engine/Source/Runtime/CoreUObject/Private/Serialization/AsyncLoading2.cpp
--- a/Engine/Source/Runtime/CoreUObject/Private/Serialization/AsyncLoading2.cpp
+++ b/Engine/Source/Runtime/CoreUObject/Private/Serialization/AsyncLoading2.cpp
@@ -6062,6 +6062,11 @@ FAsyncLoadingSyncLoadContext* FAsyncLoadingThread2::UpdateSyncLoadContext(FAsync
 	}
 	if (ThreadState.bCanAccessAsyncLoadingThreadData && !SyncLoadContext->bHasFoundRequestedPackage)
 	{
+		if (ThreadState.bIsAsyncLoadingThread)
+		{
+			// Ensure that we've created the package we're waiting for
+			CreateAsyncPackagesFromQueue(ThreadState);
+		}
 		if (FAsyncPackage2* RequestedPackage = RequestIdToPackageMap.FindRef(SyncLoadContext->RequestId))
 		{
 			SyncLoadContext->bHasFoundRequestedPackage = true;
