ZenLoader in editor: Use bIsFullyLoaded directly instead of calling IsFullyLoaded on externally created packages

diff --git a/Engine/Source/Runtime/CoreUObject/Private/Serialization/AsyncLoading2.cpp b/Engine/Source/Runtime/CoreUObject/Private/Serialization/AsyncLoading2.cpp
--- a/Engine/Source/Runtime/CoreUObject/Private/Serialization/AsyncLoading2.cpp
+++ b/Engine/Source/Runtime/CoreUObject/Private/Serialization/AsyncLoading2.cpp
@@ -924,7 +924,7 @@ public:
 					PackageRef.SetPackage(FoundPackage);
 					FoundPackage->SetCanBeImportedFlag(true);
 					FoundPackage->SetPackageId(PackageId);
-					if (FoundPackage->IsFullyLoaded())
+					if (FoundPackage->bHasBeenFullyLoaded)
 					{
 						PackageRef.SetAllPublicExportsLoaded();
 					}
