ZenLoader in editor: Create package metadata objects earlier

diff --git a/Engine/Source/Runtime/CoreUObject/Private/Serialization/AsyncLoading2.cpp b/Engine/Source/Runtime/CoreUObject/Private/Serialization/AsyncLoading2.cpp
--- a/Engine/Source/Runtime/CoreUObject/Private/Serialization/AsyncLoading2.cpp
+++ b/Engine/Source/Runtime/CoreUObject/Private/Serialization/AsyncLoading2.cpp
@@ -2309,7 +2309,7 @@ private:
 		int32 SerializeExportIndex = 0;
 		int32 PostLoadExportIndex = 0;
 #if WITH_EDITORONLY_DATA
-		TOptional<int32> MetaDataIndex;
+		int32 MetaDataIndex = -1;
 #endif
 		bool bIsCurrentlyProcessingImports = false;
 
@@ -4567,7 +4567,32 @@ EEventLoadNodeExecutionResult FAsyncPackage2::ProcessLinkerLoadPackageSummary(FA
 	HeaderData.ImportedPublicExportHashes = LinkerLoadState->LinkerLoadHeaderData.ImportedPublicExportHashes;
 	HeaderData.ImportMap = LinkerLoadState->LinkerLoadHeaderData.ImportMap;
 
-	if (bLoadHasFailed && Desc.bCanBeImported)
+	if (!bLoadHasFailed)
+	{
+#if WITH_EDITORONLY_DATA
+		// Create metadata object, this needs to happen before any other package wants to use our exports
+		LinkerLoadState->MetaDataIndex = LinkerLoadState->Linker->LoadMetaDataFromExportMap(false);
+		if (LinkerLoadState->MetaDataIndex >= 0)
+		{
+			FObjectExport& LinkerExport = LinkerLoadState->Linker->ExportMap[LinkerLoadState->MetaDataIndex];
+			FExportObject& ExportObject = Data.Exports[LinkerLoadState->MetaDataIndex];
+			ExportObject.Object = LinkerExport.Object;
+			if (ExportObject.Object)
+			{
+				ExportObject.bWasFoundInMemory = true; // Make sure that the async flags are cleared in ClearConstructedObjects
+			}
+			else
+			{
+				ExportObject.bExportLoadFailed = LinkerExport.bExportLoadFailed;
+				if (!ExportObject.bExportLoadFailed)
+				{
+					ExportObject.bFiltered = true;
+				}
+			}
+		}
+#endif
+	}
+	else if (Desc.bCanBeImported)
 	{
 		FLoadedPackageRef& PackageRef = ImportStore.LoadedPackageStore.FindPackageRefChecked(Desc.UPackageId);
 		PackageRef.SetHasFailed();
@@ -4687,33 +4712,6 @@ bool FAsyncPackage2::ProcessLinkerLoadPackageImports(FAsyncLoadingThreadState2&
 bool FAsyncPackage2::ProcessLinkerLoadPackageExports(FAsyncLoadingThreadState2& ThreadState)
 {
 	check(AsyncPackageLoadingState == EAsyncPackageLoadingState2::ProcessExportBundles);
-#if WITH_EDITORONLY_DATA
-	// Create metadata object
-	int32 MetaDataIndex = LinkerLoadState->MetaDataIndex.IsSet() ? LinkerLoadState->MetaDataIndex.GetValue() : -1;
-	if (MetaDataIndex < 0)
-	{
-		MetaDataIndex = LinkerLoadState->Linker->LoadMetaDataFromExportMap(false);
-		if (MetaDataIndex >= 0)
-		{
-			FObjectExport& LinkerExport = LinkerLoadState->Linker->ExportMap[MetaDataIndex];
-			FExportObject& ExportObject = Data.Exports[MetaDataIndex];
-			ExportObject.Object = LinkerExport.Object;
-			if (ExportObject.Object)
-			{
-				ExportObject.bWasFoundInMemory = true; // Make sure that the async flags are cleared in ClearConstructedObjects
-			}
-			else
-			{
-				ExportObject.bExportLoadFailed = LinkerExport.bExportLoadFailed;
-				if (!ExportObject.bExportLoadFailed)
-				{
-					ExportObject.bFiltered = true;
-				}
-			}
-		}
-		LinkerLoadState->MetaDataIndex = MetaDataIndex;
-	}
-#endif
 
 	// Create exports
 	const int32 ExportCount = LinkerLoadState->Linker->ExportMap.Num();
@@ -4721,7 +4719,7 @@ bool FAsyncPackage2::ProcessLinkerLoadPackageExports(FAsyncLoadingThreadState2&
 	{
 		const int32 ExportIndex = LinkerLoadState->CreateExportIndex++;
 #if WITH_EDITORONLY_DATA
-		if (ExportIndex == MetaDataIndex)
+		if (ExportIndex == LinkerLoadState->MetaDataIndex)
 		{
 			continue;
 		}
