BulkData.GetBulkDataVersions: Add GetBulkDataVersions to FBulkData so that SkeletalMesh and Cloth can use it. In FEditorBulkData, update the implementation to correctly read the versions from the WorkspaceDomain file for BulkDatas that are in EditorDomain packages and reference bytes out of their WorkspaceDomain package.
#rnx

diff --git a/Engine/Plugins/Runtime/HairStrands/Source/HairStrandsCore/Private/HairDescription.cpp b/Engine/Plugins/Runtime/HairStrands/Source/HairStrandsCore/Private/HairDescription.cpp
--- a/Engine/Plugins/Runtime/HairStrands/Source/HairStrandsCore/Private/HairDescription.cpp
+++ b/Engine/Plugins/Runtime/HairStrands/Source/HairStrandsCore/Private/HairDescription.cpp
@@ -113,7 +113,10 @@ void FHairDescriptionBulkData::Serialize(FArchive& Ar, UObject* Owner)
 	{
 		// If loading, take the package custom version so it can be applied to the bulk data archive
 		// when unpacking HairDescription from it
-		CustomVersions = BulkData.GetCustomVersions(Ar);
+		// TODO: Save the UEVersion and LicenseeVersion as well
+		FPackageFileVersion UEVersion;
+		int32 LicenseeUEVersion;
+		BulkData.GetBulkDataVersions(Ar, UEVersion, LicenseeUEVersion, CustomVersions);
 	}
 }
 
diff --git a/Engine/Source/Runtime/CoreUObject/Private/Serialization/BulkData.cpp b/Engine/Source/Runtime/CoreUObject/Private/Serialization/BulkData.cpp
--- a/Engine/Source/Runtime/CoreUObject/Private/Serialization/BulkData.cpp
+++ b/Engine/Source/Runtime/CoreUObject/Private/Serialization/BulkData.cpp
@@ -1707,17 +1707,31 @@ void FBulkData::SetFlagsFromDiskWrittenValues(
 }
 #endif
 
-FCustomVersionContainer FBulkData::GetCustomVersions(FArchive& InlineArchive)
+FCustomVersionContainer FBulkData::GetCustomVersions(FArchive& InlineArchive) const
+{
+	FPackageFileVersion OutUEVersion;
+	int32 OutLicenseeUEVersion;
+	FCustomVersionContainer OutCustomVersions;
+	GetBulkDataVersions(InlineArchive, OutUEVersion, OutLicenseeUEVersion, OutCustomVersions);
+	return OutCustomVersions;
+}
+
+void FBulkData::GetBulkDataVersions(FArchive& InlineArchive, FPackageFileVersion& OutUEVersion,
+	int32& OutLicenseeUEVersion, FCustomVersionContainer& OutCustomVersions) const
 {
 	if (!IsInSeparateFile())
 	{
-		return InlineArchive.GetCustomVersions();
+		OutUEVersion = InlineArchive.UEVer();
+		OutLicenseeUEVersion = InlineArchive.LicenseeUEVer();
+		OutCustomVersions = InlineArchive.GetCustomVersions();
 	}
 	else if (!IsInExternalResource())
 	{
 		// The BulkData is in a sidecar file. These files were created with the same custom versions that
 		// the package file containing the BulkData used
-		return InlineArchive.GetCustomVersions();
+		OutUEVersion = InlineArchive.UEVer();
+		OutLicenseeUEVersion = InlineArchive.LicenseeUEVer();
+		OutCustomVersions = InlineArchive.GetCustomVersions();
 	}
 	else
 	{
@@ -1731,10 +1745,15 @@ FCustomVersionContainer FBulkData::GetCustomVersions(FArchive& InlineArchive)
 			*ExternalArchive << PackageFileSummary;
 			if (PackageFileSummary.Tag == PACKAGE_FILE_TAG && !ExternalArchive->IsError())
 			{
-				return PackageFileSummary.GetCustomVersionContainer();
+				OutUEVersion = PackageFileSummary.GetFileVersionUE();
+				OutLicenseeUEVersion = PackageFileSummary.GetFileVersionLicenseeUE();
+				OutCustomVersions = PackageFileSummary.GetCustomVersionContainer();
+				return;
 			}
 		}
-		return FCustomVersionContainer();
+		OutUEVersion = InlineArchive.UEVer();
+		OutLicenseeUEVersion = InlineArchive.LicenseeUEVer();
+		OutCustomVersions = InlineArchive.GetCustomVersions();
 	}
 }
 
diff --git a/Engine/Source/Runtime/CoreUObject/Private/Serialization/EditorBulkData.cpp b/Engine/Source/Runtime/CoreUObject/Private/Serialization/EditorBulkData.cpp
--- a/Engine/Source/Runtime/CoreUObject/Private/Serialization/EditorBulkData.cpp
+++ b/Engine/Source/Runtime/CoreUObject/Private/Serialization/EditorBulkData.cpp
@@ -1957,6 +1957,24 @@ FCustomVersionContainer FEditorBulkData::GetCustomVersions(FArchive& InlineArchi
 void FEditorBulkData::GetBulkDataVersions(FArchive& InlineArchive, FPackageFileVersion& OutUEVersion,
 	int32& OutLicenseeUEVersion, FCustomVersionContainer& OutCustomVersions) const
 {
+	if (EnumHasAnyFlags(Flags, EFlags::ReferencesWorkspaceDomain))
+	{
+		// Read the version data out of the separate package file
+		TUniquePtr<FArchive> ExternalArchive = IPackageResourceManager::Get().OpenReadExternalResource(
+			EPackageExternalResource::WorkspaceDomainFile, PackagePath.GetPackageName());
+		if (ExternalArchive.IsValid())
+		{
+			FPackageFileSummary PackageFileSummary;
+			*ExternalArchive << PackageFileSummary;
+			if (PackageFileSummary.Tag == PACKAGE_FILE_TAG && !ExternalArchive->IsError())
+			{
+				OutUEVersion = PackageFileSummary.GetFileVersionUE();
+				OutLicenseeUEVersion = PackageFileSummary.GetFileVersionLicenseeUE();
+				OutCustomVersions = PackageFileSummary.GetCustomVersionContainer();
+				return;
+			}
+		}
+	}
 	OutUEVersion = InlineArchive.UEVer();
 	OutLicenseeUEVersion = InlineArchive.LicenseeUEVer();
 	OutCustomVersions = InlineArchive.GetCustomVersions();
diff --git a/Engine/Source/Runtime/CoreUObject/Public/Serialization/BulkData.h b/Engine/Source/Runtime/CoreUObject/Public/Serialization/BulkData.h
--- a/Engine/Source/Runtime/CoreUObject/Public/Serialization/BulkData.h
+++ b/Engine/Source/Runtime/CoreUObject/Public/Serialization/BulkData.h
@@ -991,16 +991,19 @@ public:
 		return DataAllocation.StealFileMapping(this); 
 	}
 
+	UE_DEPRECATED(5.1, "Call GetBulkDataVersions instead.")
+	FCustomVersionContainer GetCustomVersions(FArchive& InlineArchive) const;
+
 	/**
 	 * Get the CustomVersions used in the file containing the BulkData payload. If !IsInSeparateFile, this will be
 	 * the custom versions from the archive used to serialize the FBulkData, which the caller must provide.
 	 * Otherwise, the CustomVersions come from the separate file and this function will look them up.
-	 * 
-	 * @param InlineArchive The archive that was used to load this BulkData object
-	 * 
-	 * @return The CustomVersions that apply to the interpretation of the BulkData's payload.
+	 *
+	 * @param InlineArchive The archive that was used to load this object
+	 *
 	 */
-	FCustomVersionContainer GetCustomVersions(FArchive& InlineArchive);
+	void GetBulkDataVersions(FArchive& InlineArchive, FPackageFileVersion& OutUEVersion, int32& OutLicenseeUEVersion,
+		FCustomVersionContainer& OutCustomVersions) const;
 
 #if WITH_EDITOR
 	/**
diff --git a/Engine/Source/Runtime/CoreUObject/Public/Serialization/EditorBulkData.h b/Engine/Source/Runtime/CoreUObject/Public/Serialization/EditorBulkData.h
--- a/Engine/Source/Runtime/CoreUObject/Public/Serialization/EditorBulkData.h
+++ b/Engine/Source/Runtime/CoreUObject/Public/Serialization/EditorBulkData.h
@@ -279,13 +279,10 @@ public:
 	FCustomVersionContainer GetCustomVersions(FArchive& InlineArchive);
 
 	/**
-	* Get the versions used in the file containing the payload. Currently this is assumed
-	* to always be the versions in the InlineArchive
-	*
-	* @param InlineArchive The archive that was used to load this object
-	*
-	* @return The CustomVersions that apply to the interpretation of the payload.
-	*/
+	 * Get the versions used in the file containing the payload.
+	 *
+	 * @param InlineArchive The archive that was used to load this object
+	 */
 	void GetBulkDataVersions(FArchive& InlineArchive, FPackageFileVersion& OutUEVersion, int32& OutLicenseeUEVersion,
 		FCustomVersionContainer& OutCustomVersions) const;
 
diff --git a/Engine/Source/Runtime/Engine/Private/Rendering/SkeletalMeshLODImporterData.cpp b/Engine/Source/Runtime/Engine/Private/Rendering/SkeletalMeshLODImporterData.cpp
--- a/Engine/Source/Runtime/Engine/Private/Rendering/SkeletalMeshLODImporterData.cpp
+++ b/Engine/Source/Runtime/Engine/Private/Rendering/SkeletalMeshLODImporterData.cpp
@@ -502,9 +502,7 @@ void FReductionBaseSkeletalMeshBulkData::Serialize(FArchive& Ar, UObject* Owner)
 	if (!Ar.IsTransacting() && Ar.IsLoading())
 	{
 		//Save the custom version so we can load FReductionSkeletalMeshData later
-		UEVersion = Ar.UEVer();
-		LicenseeUEVersion = Ar.LicenseeUEVer();
-		SerializeLoadingCustomVersionContainer = BulkData.GetCustomVersions(Ar);
+		BulkData.GetBulkDataVersions(Ar, UEVersion, LicenseeUEVersion, SerializeLoadingCustomVersionContainer);
 		bUseSerializeLoadingCustomVersion = true;
 	}
 }
@@ -816,9 +814,7 @@ void FRawSkeletalMeshBulkData::Serialize(FArchive& Ar, UObject* Owner)
 	if (!Ar.IsTransacting() && Ar.IsLoading())
 	{
 		//Save the custom version so we can load FReductionSkeletalMeshData later
-		UEVersion = Ar.UEVer();
-		LicenseeUEVersion = Ar.LicenseeUEVer();
-		SerializeLoadingCustomVersionContainer = BulkData.GetCustomVersions(Ar);
+		BulkData.GetBulkDataVersions(Ar, UEVersion, LicenseeUEVersion, SerializeLoadingCustomVersionContainer);
 		bUseSerializeLoadingCustomVersion = true;
 	}
 }
