StaticAllocateObject: Fix stale comment to match the current code.
#rnx

diff --git a/Engine/Source/Runtime/CoreUObject/Private/UObject/UObjectGlobals.cpp b/Engine/Source/Runtime/CoreUObject/Private/UObject/UObjectGlobals.cpp
--- a/Engine/Source/Runtime/CoreUObject/Private/UObject/UObjectGlobals.cpp
+++ b/Engine/Source/Runtime/CoreUObject/Private/UObject/UObjectGlobals.cpp
@@ -3163,8 +3163,7 @@ UObject* StaticAllocateObject
 		// See if object already exists.
 		Obj = StaticFindObjectFastInternal( /*Class=*/ NULL, InOuter, InName, true );
 
-		// Temporary: If the object we found is of a different class, allow the object to be allocated.
-		// This breaks new UObject assumptions and these need to be fixed.
+		// It is an error if we are trying to replace an object of a different class
 		if (Obj && !Obj->GetClass()->IsChildOf(InClass))
 		{
 			const TCHAR* ErrorPrefix = TEXT("");
