ZenLoader in editor: Reserve space for public exports up front

diff --git a/Engine/Source/Runtime/CoreUObject/Private/Serialization/AsyncLoading2.cpp b/Engine/Source/Runtime/CoreUObject/Private/Serialization/AsyncLoading2.cpp
--- a/Engine/Source/Runtime/CoreUObject/Private/Serialization/AsyncLoading2.cpp
+++ b/Engine/Source/Runtime/CoreUObject/Private/Serialization/AsyncLoading2.cpp
@@ -4567,8 +4567,10 @@ EEventLoadNodeExecutionResult FAsyncPackage2::ProcessLinkerLoadPackageSummary(FA
 	HeaderData.ImportedPublicExportHashes = LinkerLoadState->LinkerLoadHeaderData.ImportedPublicExportHashes;
 	HeaderData.ImportMap = LinkerLoadState->LinkerLoadHeaderData.ImportMap;
 
+	FLoadedPackageRef& PackageRef = ImportStore.LoadedPackageStore.FindPackageRefChecked(Desc.UPackageId);
 	if (!bLoadHasFailed)
 	{
+		PackageRef.ReserveSpaceForPublicExports(LinkerLoadState->LinkerLoadHeaderData.PublicExportHashes.Num());
 #if WITH_EDITORONLY_DATA
 		// Create metadata object, this needs to happen before any other package wants to use our exports
 		LinkerLoadState->MetaDataIndex = LinkerLoadState->Linker->LoadMetaDataFromExportMap(false);
@@ -4594,7 +4596,6 @@ EEventLoadNodeExecutionResult FAsyncPackage2::ProcessLinkerLoadPackageSummary(FA
 	}
 	else if (Desc.bCanBeImported)
 	{
-		FLoadedPackageRef& PackageRef = ImportStore.LoadedPackageStore.FindPackageRefChecked(Desc.UPackageId);
 		PackageRef.SetHasFailed();
 	}
 
