Bulk Data - remove GetCopy warning when running tests

#rnx

diff --git a/Engine/Source/Runtime/CoreUObject/Private/Serialization/BulkData.cpp b/Engine/Source/Runtime/CoreUObject/Private/Serialization/BulkData.cpp
--- a/Engine/Source/Runtime/CoreUObject/Private/Serialization/BulkData.cpp
+++ b/Engine/Source/Runtime/CoreUObject/Private/Serialization/BulkData.cpp
@@ -849,7 +849,7 @@ void FBulkData::GetCopy( void** Dest, bool bDiscardInternalCopy )
 			// single use bulk data.
 			if( bDiscardInternalCopy && CanDiscardInternalData() )
 			{
-#if USE_RUNTIME_BULKDATA
+#if USE_RUNTIME_BULKDATA && !WITH_LOW_LEVEL_TESTS 
 				UE_LOG(LogSerialization, Warning, TEXT("FBulkData::GetCopy both copied and discarded it's data, passing in an empty pointer would avoid an extra allocate and memcpy!"));
 #endif
 				FreeData();
