Core - Introducing FindPackageReferencesInObject, this utility function provides a pretty path from the root object to any references to any of the packages provided.

Given a root object, this function will call the callback every time it encounters a reference in the root object (including all inlined/instanced UObjects.  The returned references will be in an easily understood debugging format so that they're easy to understand the context.

For example, Actions[2](GameFeatureAction_Activities)->ObjectActivitiesInfo[0].TargetReference

Where TargetReference was one of the listed PackagesReferenced.  The point of this function is to assist in the reporting of references between objects you need to show the developer.

diff --git a/Engine/Source/Runtime/CoreUObject/Private/Misc/ReferenceDebugging.cpp b/Engine/Source/Runtime/CoreUObject/Private/Misc/ReferenceDebugging.cpp
new file mode 100644
--- /dev/null
+++ b/Engine/Source/Runtime/CoreUObject/Private/Misc/ReferenceDebugging.cpp
@@ -0,0 +1,59 @@
+// Copyright Epic Games, Inc. All Rights Reserved.
+
+#include "Misc/ReferenceDebugging.h"
+#include "UObject/UnrealType.h"
+
+void FindPackageReferencesInObject(const UObject* RootObject, const TArray<FName>& PackagesReferenced, TFunctionRef<void(FName /*PackagePath*/, const FString& /*PropertyPath*/)> Callback)
+{
+	TArray<FString> PathStack;
+	TFunctionRef<void(const UObject*)> InnerLoop = [&InnerLoop, &PackagesReferenced, &PathStack, Callback](const UObject* RootObject)
+	{
+		for (TPropertyValueIterator<FObjectProperty> It(RootObject->GetClass(), RootObject, EPropertyValueIteratorFlags::FullRecursion, EFieldIteratorFlags::ExcludeDeprecated); It; ++It)
+		{
+			const FObjectProperty* Property = It.Key();
+			if (const UObject* Object = Property->GetObjectPtrPropertyValueRef(It.Value()))
+			{
+				if (Object->HasAnyFlags(RF_Public) && Object->IsAsset())
+				{
+					const FName ObjectPackageName = Object->GetOutermost()->GetFName();
+					if (PackagesReferenced.Contains(ObjectPackageName))
+					{
+						TStringBuilder<FName::StringBufferSize> FullPropertyPath;
+						for (const FString& Path : PathStack)
+						{
+							FullPropertyPath.Append(Path);
+							FullPropertyPath.Append(TEXT("→"));
+						}
+
+						FullPropertyPath.Append(It.GetPropertyPathDebugString());
+						
+						Callback(ObjectPackageName, FullPropertyPath.ToString());
+					}
+				}
+				else if (Property->HasAllPropertyFlags(CPF_ExportObject) && Object->GetClass()->HasAllClassFlags(CLASS_EditInlineNew))
+				{
+					TStringBuilder<FName::StringBufferSize> PathFragment;
+					PathFragment.Append(It.GetPropertyPathDebugString());
+
+					// When we encounter an embedded UObject, to make the path easy to understand we should include the
+					// class name in the expression, the idea being that the user will see something like,
+					// 
+					PathFragment.Append(TEXT("("));
+					PathFragment.Append(GetNameSafe(Object->GetClass()));
+					PathFragment.Append(TEXT(")"));
+
+					{
+						PathStack.Add(PathFragment.ToString());
+
+						// Add the inner properties of this instanced object.
+						InnerLoop(Object);
+					
+						PathStack.Pop();
+					}
+				}
+			}
+		}
+	};
+
+	InnerLoop(RootObject);
+}
\ No newline at end of file
diff --git a/Engine/Source/Runtime/CoreUObject/Private/UObject/UnrealType.cpp b/Engine/Source/Runtime/CoreUObject/Private/UObject/UnrealType.cpp
--- a/Engine/Source/Runtime/CoreUObject/Private/UObject/UnrealType.cpp
+++ b/Engine/Source/Runtime/CoreUObject/Private/UObject/UnrealType.cpp
@@ -5,6 +5,8 @@
 
 DEFINE_LOG_CATEGORY(LogType);
 
+ENUM_CLASS_FLAGS(FPropertyValueIterator::EPropertyValueFlags)
+
 FPropertyValueIterator::FPropertyValueIterator(
 	FFieldClass* InPropertyClass,
 	const UStruct* InStruct,
@@ -26,31 +28,31 @@ FPropertyValueIterator::FPropertyValueIterator(
 	}
 }
 
-uint8 FPropertyValueIterator::GetPropertyValueFlags(const FProperty* Property)
+FPropertyValueIterator::EPropertyValueFlags FPropertyValueIterator::GetPropertyValueFlags(const FProperty* Property) const
 {
-	uint8 Flags = 0;
+	EPropertyValueFlags Flags = EPropertyValueFlags::None;
 	if (RecursionFlags == EPropertyValueIteratorFlags::FullRecursion)
 	{
 		if (Property->IsA(FArrayProperty::StaticClass()))
 		{
-			Flags = (uint8)EPropertyValueFlags::IsArray;
+			Flags = EPropertyValueFlags::IsArray;
 		}
 		else if (Property->IsA(FMapProperty::StaticClass()))
 		{
-			Flags = (uint8)EPropertyValueFlags::IsMap;
+			Flags = EPropertyValueFlags::IsMap;
 		}
 		else if (Property->IsA(FSetProperty::StaticClass()))
 		{
-			Flags = (uint8)EPropertyValueFlags::IsSet;
+			Flags = EPropertyValueFlags::IsSet;
 		}
 		else if (Property->IsA(FStructProperty::StaticClass()))
 		{
-			Flags = (uint8)EPropertyValueFlags::IsStruct;
+			Flags = EPropertyValueFlags::IsStruct;
 		}
 	}
 	if (bMatchAll || Property->IsA(PropertyClass))
 	{
-		Flags |= (uint8)EPropertyValueFlags::IsMatch;
+		Flags |= EPropertyValueFlags::IsMatch;
 	}
 	return Flags;
 }
@@ -61,8 +63,8 @@ void FPropertyValueIterator::FillStructProperties(const UStruct* Struct, FProper
 	for (TFieldIterator<FProperty> It(Struct, EFieldIteratorFlags::IncludeSuper, DeprecatedPropertyFlags, EFieldIteratorFlags::ExcludeInterfaces); It; ++It)
 	{
 		const FProperty* Property  = *It;
-		uint8 PropertyValueFlags = GetPropertyValueFlags(Property);
-		if (PropertyValueFlags)
+		EPropertyValueFlags PropertyValueFlags = GetPropertyValueFlags(Property);
+		if (PropertyValueFlags != EPropertyValueFlags::None)
 		{
 			int32 Num = Property->ArrayDim;
 			ValueArray.Reserve(ValueArray.Num() + Num);
@@ -89,13 +91,13 @@ bool FPropertyValueIterator::NextValue(EPropertyValueIteratorFlags InRecursionFl
 
 		const FProperty* Property = Entry.ValueArray[Entry.ValueIndex].Key.Key;
 		const void* PropertyValue = Entry.ValueArray[Entry.ValueIndex].Key.Value;
-		const uint8 PropertyValueFlags = Entry.ValueArray[Entry.ValueIndex].Value;
-		check(PropertyValueFlags);
+		const EPropertyValueFlags PropertyValueFlags = Entry.ValueArray[Entry.ValueIndex].Value;
+		check(PropertyValueFlags != EPropertyValueFlags::None);
 
 		// Handle matching properties
-		if (!bIsPropertyMatchProcessed && (PropertyValueFlags & EPropertyValueFlags::IsMatch))
+		if (!bIsPropertyMatchProcessed && EnumHasAnyFlags(PropertyValueFlags, EPropertyValueFlags::IsMatch))
 		{
-			if (PropertyValueFlags & EPropertyValueFlags::ContainerMask)
+			if (EnumHasAnyFlags(PropertyValueFlags, EPropertyValueFlags::ContainerMask))
 			{
 				// this match is also a container/struct, so recurse into it next time
 				Entry.NextValueIndex = Entry.ValueIndex;
@@ -104,17 +106,17 @@ bool FPropertyValueIterator::NextValue(EPropertyValueIteratorFlags InRecursionFl
 		}
 
 		// Handle container properties
-		check(PropertyValueFlags & EPropertyValueFlags::ContainerMask);
+		check(EnumHasAnyFlags(PropertyValueFlags, EPropertyValueFlags::ContainerMask));
 		if (InRecursionFlags == EPropertyValueIteratorFlags::FullRecursion)
 		{
 			FPropertyValueStackEntry NewEntry(PropertyValue);
 
-			if (PropertyValueFlags & EPropertyValueFlags::IsArray)
+			if (EnumHasAnyFlags(PropertyValueFlags, EPropertyValueFlags::IsArray))
 			{
 				const FArrayProperty* ArrayProperty = CastFieldChecked<FArrayProperty>(Property);
 				const FProperty* InnerProperty = ArrayProperty->Inner;
-				uint8 InnerFlags = GetPropertyValueFlags(InnerProperty);
-				if (InnerFlags)
+				EPropertyValueFlags InnerFlags = GetPropertyValueFlags(InnerProperty);
+				if (InnerFlags != EPropertyValueFlags::None)
 				{
 					FScriptArrayHelper Helper(ArrayProperty, PropertyValue);
 					const int32 Num = Helper.Num();
@@ -126,14 +128,14 @@ bool FPropertyValueIterator::NextValue(EPropertyValueIteratorFlags InRecursionFl
 					}
 				}
 			}
-			else if (PropertyValueFlags & EPropertyValueFlags::IsMap)
+			else if (EnumHasAnyFlags(PropertyValueFlags, EPropertyValueFlags::IsMap))
 			{
 				const FMapProperty* MapProperty = CastFieldChecked<FMapProperty>(Property);
 				const FProperty* KeyProperty = MapProperty->KeyProp;
 				const FProperty* ValueProperty = MapProperty->ValueProp;
-				uint8 KeyFlags = GetPropertyValueFlags(KeyProperty);
-				uint8 ValueFlags = GetPropertyValueFlags(ValueProperty);
-				if (KeyFlags | ValueFlags)
+				EPropertyValueFlags KeyFlags = GetPropertyValueFlags(KeyProperty);
+				EPropertyValueFlags ValueFlags = GetPropertyValueFlags(ValueProperty);
+				if ((KeyFlags | ValueFlags) != EPropertyValueFlags::None)
 				{
 					FScriptMapHelper Helper(MapProperty, PropertyValue);
 					const int32 Num = Helper.Num();
@@ -141,12 +143,12 @@ bool FPropertyValueIterator::NextValue(EPropertyValueIteratorFlags InRecursionFl
 					{
 						if (Helper.IsValidIndex(DynamicIndex))
 						{
-							if (KeyFlags)
+							if (KeyFlags != EPropertyValueFlags::None)
 							{
 								NewEntry.ValueArray.Emplace(
 									BasePairType(KeyProperty, Helper.GetKeyPtr(DynamicIndex)), KeyFlags);
 							}
-							if (ValueFlags)
+							if (ValueFlags != EPropertyValueFlags::None)
 							{
 								NewEntry.ValueArray.Emplace(
 									BasePairType(ValueProperty, Helper.GetValuePtr(DynamicIndex)), ValueFlags);
@@ -155,12 +157,12 @@ bool FPropertyValueIterator::NextValue(EPropertyValueIteratorFlags InRecursionFl
 					}
 				}
 			}
-			else if (PropertyValueFlags & EPropertyValueFlags::IsSet)
+			else if (EnumHasAnyFlags(PropertyValueFlags, EPropertyValueFlags::IsSet))
 			{
 				const FSetProperty* SetProperty = CastFieldChecked<FSetProperty>(Property);
 				const FProperty* InnerProperty = SetProperty->ElementProp;
-				uint8 InnerFlags = GetPropertyValueFlags(InnerProperty);
-				if (InnerFlags)
+				EPropertyValueFlags InnerFlags = GetPropertyValueFlags(InnerProperty);
+				if (InnerFlags != EPropertyValueFlags::None)
 				{
 					FScriptSetHelper Helper(SetProperty, PropertyValue);
 					const int32 Num = Helper.Num();
@@ -174,7 +176,7 @@ bool FPropertyValueIterator::NextValue(EPropertyValueIteratorFlags InRecursionFl
 					}
 				}
 			}
-			else if (PropertyValueFlags & EPropertyValueFlags::IsStruct)
+			else if (EnumHasAnyFlags(PropertyValueFlags, EPropertyValueFlags::IsStruct))
 			{
 				const FStructProperty* StructProperty = CastFieldChecked<FStructProperty>(Property);
 				FillStructProperties(StructProperty->Struct, NewEntry);
@@ -226,3 +228,62 @@ void FPropertyValueIterator::GetPropertyChain(TArray<const FProperty*>& Property
 		PropertyChain.Add(Property);
 	}
 }
+
+FString FPropertyValueIterator::GetPropertyPathDebugString() const
+{
+	TStringBuilder<FName::StringBufferSize> PropertyPath;
+	for (int32 StackIndex = 0; StackIndex < PropertyIteratorStack.Num(); StackIndex++)
+	{
+		const FPropertyValueStackEntry& Entry = PropertyIteratorStack[StackIndex];
+
+		// Index should always be valid
+		const BasePairType& PropertyAndValue = Entry.ValueArray[Entry.ValueIndex].Key;
+		const FProperty* Property = PropertyAndValue.Key;
+		const void* ValuePtr = PropertyAndValue.Value;
+
+		PropertyPath.Append(Property->GetAuthoredName());
+
+		int32 NextStackIndex = StackIndex + 1;
+
+		if (NextStackIndex < PropertyIteratorStack.Num())
+		{
+			if (CastField<FArrayProperty>(Property))
+			{
+				const FPropertyValueStackEntry& NextEntry = PropertyIteratorStack[StackIndex+1];
+			
+				PropertyPath.Append(TEXT("["));
+				PropertyPath.Appendf(TEXT("%d"), NextEntry.ValueIndex);
+				PropertyPath.Append(TEXT("]"));
+
+				StackIndex++;
+			}
+			else if (CastField<FSetProperty>(Property))
+			{
+			
+			}
+			else if (CastField<FMapProperty>(Property))
+			{
+				const FPropertyValueStackEntry& NextEntry = PropertyIteratorStack[StackIndex+1];
+
+				const FProperty* NextProperty = NextEntry.GetPropertyValue().Key;
+				const void* NextValuePtr = NextEntry.GetPropertyValue().Value;
+
+				FString KeyStr;
+				NextProperty->ExportText_Direct(KeyStr, NextValuePtr, nullptr, nullptr, PPF_None);
+				
+				PropertyPath.Append(TEXT("["));
+				PropertyPath.Append(KeyStr);
+				PropertyPath.Append(TEXT("]"));
+			}
+
+			NextStackIndex = StackIndex + 1;
+
+			if (NextStackIndex < PropertyIteratorStack.Num())
+			{
+				PropertyPath.Append(TEXT("."));
+			}
+		}
+	}
+
+	return PropertyPath.ToString();
+}
diff --git a/Engine/Source/Runtime/CoreUObject/Public/Misc/ReferenceDebugging.h b/Engine/Source/Runtime/CoreUObject/Public/Misc/ReferenceDebugging.h
new file mode 100644
--- /dev/null
+++ b/Engine/Source/Runtime/CoreUObject/Public/Misc/ReferenceDebugging.h
@@ -0,0 +1,20 @@
+// Copyright Epic Games, Inc. All Rights Reserved.
+
+#pragma once
+
+#include "CoreMinimal.h"
+
+class UObject;
+
+/**
+ * Given a root object, this function will call the callback every time it encounters a reference in the root object
+ * (including all inlined/instanced UObjects.  The returned references will be in an easily understood debugging
+ * format so that they're easy to understand the context.
+ *
+ * For example,
+ * Actions[2](GameFeatureAction_Activities)->ObjectActivitiesInfo[0].TargetReference
+ *
+ * Where TargetReference was one of the listed PackagesReferenced.  The point of this function is to assist in the
+ * reporting of references between objects you need to show the developer.
+ */
+COREUOBJECT_API void FindPackageReferencesInObject(const UObject* RootObject, const TArray<FName>& PackagesReferenced, TFunctionRef<void(FName /*PackagePath*/, const FString& /*PropertyPath*/)> Callback);
\ No newline at end of file
diff --git a/Engine/Source/Runtime/CoreUObject/Public/UObject/UnrealType.h b/Engine/Source/Runtime/CoreUObject/Public/UObject/UnrealType.h
--- a/Engine/Source/Runtime/CoreUObject/Public/UObject/UnrealType.h
+++ b/Engine/Source/Runtime/CoreUObject/Public/UObject/UnrealType.h
@@ -7033,9 +7033,25 @@ public:
 	 */
 	COREUOBJECT_API void GetPropertyChain(TArray<const FProperty*>& PropertyChain) const;
 
+	/**
+	 * Returns a string of the property chain in a form that makes reading it easy.  The format may change over time,
+	 * do not depend on this being in a specific Unreal path format.  It's primary intention is to aid in debugging
+	 * and in reporting paths to the end developer in an editor environment.
+	 *
+	 * Now for the implementation details, normally paths involving an array, will do things like print the array
+	 * name twice, since an array in the path is actually 2 properties, one for the array and another for the index.
+	 * Maps are similar.  The end result is stuff like seeing, ActionsArray.ActionsArray.ActionIfo.ActionPower.
+	 *
+	 * What this function does is tries to make this stuff a lot more human readable, by printing,
+	 * ActionsArray[3].ActionIfo.ActionPower, so that you know the index it comes from.  Similarly Maps would print
+	 * ActionMap["Action Name"].ActionIfo.ActionPower, so that you know the index it comes from.  Similarly Maps would print
+	 */
+	COREUOBJECT_API FString GetPropertyPathDebugString() const;
+
 private:
-	enum EPropertyValueFlags : uint8
+	enum class EPropertyValueFlags : uint8
 	{
+		None = 0x0,
 		IsMatch = 0x01,
 
 		ContainerMask = 0xF0,
@@ -7044,6 +7060,7 @@ private:
 		IsSet = 0x40,
 		IsStruct = 0x80,
 	};
+	FRIEND_ENUM_CLASS_FLAGS(EPropertyValueFlags)
 
 	struct FPropertyValueStackEntry
 	{
@@ -7051,7 +7068,7 @@ private:
 		const void* Owner = nullptr;
 		
 		/** List of current root property+value pairs for the current top level FProperty */
-		typedef TPair<BasePairType, uint8> BasePairAndFlags;
+		typedef TPair<BasePairType, EPropertyValueFlags> BasePairAndFlags;
 		typedef TArray<BasePairAndFlags, TInlineAllocator<8>> FValueArrayType;
 		FValueArrayType ValueArray;
 
@@ -7100,7 +7117,7 @@ private:
 	bool bMatchAll = false;
 
 	/** Returns EPropertyValueFlags to describe if this Property is a match and/or a container/struct */
-	uint8 GetPropertyValueFlags(const FProperty* Property);
+	EPropertyValueFlags GetPropertyValueFlags(const FProperty* Property) const;
 
 	/** Fills the Entry.ValueArray with all relevant properties found in Struct */
 	void FillStructProperties(const UStruct* Struct, FPropertyValueStackEntry& Entry);
