Core - Removing an export that isn't suppose to be in the cpp file.

diff --git a/Engine/Source/Runtime/CoreUObject/Private/UObject/ObjectPtr.cpp b/Engine/Source/Runtime/CoreUObject/Private/UObject/ObjectPtr.cpp
--- a/Engine/Source/Runtime/CoreUObject/Private/UObject/ObjectPtr.cpp
+++ b/Engine/Source/Runtime/CoreUObject/Private/UObject/ObjectPtr.cpp
@@ -17,7 +17,7 @@ FString FObjectPtr::GetPathName() const
 	}
 }
 
-COREUOBJECT_API FName FObjectPtr::GetFName() const
+FName FObjectPtr::GetFName() const
 {
 	if (IsResolved())
 	{
@@ -31,7 +31,7 @@ COREUOBJECT_API FName FObjectPtr::GetFName() const
 	}
 }
 
-COREUOBJECT_API FString FObjectPtr::GetFullName(EObjectFullNameFlags Flags) const
+FString FObjectPtr::GetFullName(EObjectFullNameFlags Flags) const
 {
 	if (IsResolved())
 	{
diff --git a/Engine/Source/Runtime/CoreUObject/Public/UObject/ObjectPtr.h b/Engine/Source/Runtime/CoreUObject/Public/UObject/ObjectPtr.h
--- a/Engine/Source/Runtime/CoreUObject/Public/UObject/ObjectPtr.h
+++ b/Engine/Source/Runtime/CoreUObject/Public/UObject/ObjectPtr.h
@@ -154,7 +154,7 @@ public:
 	COREUOBJECT_API FName GetFName() const;
 
 	// Gets the string name of the object without resolving the object reference.
-	FString GetName() const
+	FORCEINLINE FString GetName() const
 	{
 		return GetFName().ToString();
 	}
