DerivedData: Fixed inadvertent size increase of FDerivedData

Flags have moved into FCookedData to avoid padding, and there is now a static_assert on the size.

A new constructor allows creating references from the private FCookedData, and FLinkerSave now returns FDerivedData instead of FCookedData.

#rnx

diff --git a/Engine/Source/Runtime/CoreUObject/Private/Serialization/DerivedData.cpp b/Engine/Source/Runtime/CoreUObject/Private/Serialization/DerivedData.cpp
--- a/Engine/Source/Runtime/CoreUObject/Private/Serialization/DerivedData.cpp
+++ b/Engine/Source/Runtime/CoreUObject/Private/Serialization/DerivedData.cpp
@@ -25,6 +25,8 @@
 #include "UObject/LinkerSave.h"
 #endif // WITH_EDITORONLY_DATA
 
+static_assert(sizeof(UE::FDerivedData) == 32 + sizeof(void*) * WITH_EDITORONLY_DATA);
+
 namespace UE::DerivedData::Private { class FIoResponseDispatcher; }
 
 namespace UE::DerivedData::Private
@@ -63,8 +65,9 @@ void FCookedData::Serialize(FArchive& Ar)
 {
 	Ar << ChunkOffset;
 	Ar << ChunkSize;
-	Ar.Serialize(ChunkId, sizeof(ChunkId));
 	static_assert(sizeof(ChunkId) == sizeof(FIoChunkId));
+	Ar.Serialize(ChunkId, sizeof(ChunkId));
+	Ar << Flags;
 }
 
 } // UE::DerivedData::Private
@@ -287,19 +290,13 @@ void FDerivedData::Serialize(FArchive& Ar, UObject* Owner)
 	{
 		checkf(Ar.IsSaving() && Ar.IsCooking(),
 			TEXT("FEditorData for FDerivedData only supports saving to cooked packages. %s"), *WriteToString<256>(*this));
-
 		FLinkerSave* Linker = Cast<FLinkerSave>(Ar.GetLinker());
 		checkf(Linker, TEXT("Serializing FDerivedData requires a linker. %s"), *WriteToString<256>(*this));
-
-		FCookedData CookedEditorData = Linker->AddDerivedData(*this);
-		CookedEditorData.Serialize(Ar);
-		Ar << Flags;
-		return;
+		return Linker->AddDerivedData(*this).CookedData.Serialize(Ar);
 	}
 #endif
 
 	CookedData.Serialize(Ar);
-	Ar << Flags;
 }
 
 #if WITH_EDITORONLY_DATA
@@ -307,29 +304,29 @@ void FDerivedData::Serialize(FArchive& Ar, UObject* Owner)
 FDerivedData::FDerivedData(const DerivedData::FSharedString& Name, const FSharedBuffer& Data)
 	: EditorData(MakePimpl<DerivedData::Private::FEditorData, EPimplPtrMode::DeepCopy>(
 		Name, TInPlaceType<DerivedData::Private::FCompositeBufferWithHash>(), Data))
-	, Flags(EDerivedDataFlags::Required)
 {
+	CookedData.Flags = EDerivedDataFlags::Required;
 }
 
 FDerivedData::FDerivedData(const DerivedData::FSharedString& Name, const FCompositeBuffer& Data)
 	: EditorData(MakePimpl<DerivedData::Private::FEditorData, EPimplPtrMode::DeepCopy>(
 		Name, TInPlaceType<DerivedData::Private::FCompositeBufferWithHash>(), Data))
-	, Flags(EDerivedDataFlags::Required)
 {
+	CookedData.Flags = EDerivedDataFlags::Required;
 }
 
 FDerivedData::FDerivedData(const DerivedData::FSharedString& Name, const FCompressedBuffer& Data)
 	: EditorData(MakePimpl<DerivedData::Private::FEditorData, EPimplPtrMode::DeepCopy>(
 		Name, TInPlaceType<FCompressedBuffer>(), Data))
-	, Flags(EDerivedDataFlags::Required)
 {
+	CookedData.Flags = EDerivedDataFlags::Required;
 }
 
 FDerivedData::FDerivedData(const DerivedData::FSharedString& Name, const DerivedData::FCacheKey& Key)
 	: EditorData(MakePimpl<DerivedData::Private::FEditorData, EPimplPtrMode::DeepCopy>(
 		Name, TInPlaceType<DerivedData::Private::FCacheKeyWithId>(), Key, DerivedData::FValueId::Null))
-	, Flags(EDerivedDataFlags::Required)
 {
+	CookedData.Flags = EDerivedDataFlags::Required;
 }
 
 FDerivedData::FDerivedData(
@@ -338,8 +335,8 @@ FDerivedData::FDerivedData(
 	const DerivedData::FValueId& ValueId)
 	: EditorData(MakePimpl<DerivedData::Private::FEditorData, EPimplPtrMode::DeepCopy>(
 		Name, TInPlaceType<DerivedData::Private::FCacheKeyWithId>(), Key, ValueId))
-	, Flags(EDerivedDataFlags::Required)
 {
+	CookedData.Flags = EDerivedDataFlags::Required;
 }
 
 const DerivedData::FSharedString& FDerivedData::GetName() const
@@ -349,7 +346,7 @@ const DerivedData::FSharedString& FDerivedData::GetName() const
 
 void FDerivedData::SetFlags(EDerivedDataFlags InFlags)
 {
-	Flags = InFlags;
+	CookedData.Flags = InFlags;
 }
 
 #endif // WITH_EDITORONLY_DATA
diff --git a/Engine/Source/Runtime/CoreUObject/Private/UObject/LinkerSave.cpp b/Engine/Source/Runtime/CoreUObject/Private/UObject/LinkerSave.cpp
--- a/Engine/Source/Runtime/CoreUObject/Private/UObject/LinkerSave.cpp
+++ b/Engine/Source/Runtime/CoreUObject/Private/UObject/LinkerSave.cpp
@@ -535,7 +535,7 @@ void FLinkerSave::SetFilterEditorOnly(bool bInFilterEditorOnly)
 }
 
 #if WITH_EDITORONLY_DATA
-UE::DerivedData::Private::FCookedData FLinkerSave::AddDerivedData(const UE::FDerivedData& Data)
+UE::FDerivedData FLinkerSave::AddDerivedData(const UE::FDerivedData& Data)
 {
 	UE_LOG(LogLinker, Warning, TEXT("Data will not be able to load because derived data is not saved yet."));
 
@@ -551,6 +551,7 @@ UE::DerivedData::Private::FCookedData FLinkerSave::AddDerivedData(const UE::FDer
 	*reinterpret_cast<uint32*>(&CookedData.ChunkId[7]) = NETWORK_ORDER32(ChunkIndex);
 	*reinterpret_cast<uint64*>(&CookedData.ChunkId[0]) = PackageId.Value();
 
-	return CookedData;
+	CookedData.Flags = Data.GetFlags();
+	return UE::FDerivedData(CookedData);
 }
 #endif // WITH_EDITORONLY_DATA
diff --git a/Engine/Source/Runtime/CoreUObject/Public/Serialization/DerivedData.h b/Engine/Source/Runtime/CoreUObject/Public/Serialization/DerivedData.h
--- a/Engine/Source/Runtime/CoreUObject/Public/Serialization/DerivedData.h
+++ b/Engine/Source/Runtime/CoreUObject/Public/Serialization/DerivedData.h
@@ -63,6 +63,9 @@ public:
 	uint64 ChunkSize = 0;
 	/** The FIoChunkId. */
 	uint8 ChunkId[12]{};
+
+	/** The flags that modify how this derived data is saved and how it may be loaded. */
+	EDerivedDataFlags Flags = EDerivedDataFlags::None;
 };
 
 /**
@@ -84,6 +87,9 @@ public:
 	/** Constructs a null reference. */
 	FDerivedData() = default;
 
+	/** Constructs a reference from the private representation of a cooked reference. */
+	inline explicit FDerivedData(const DerivedData::Private::FCookedData& CookedData);
+
 	/** Resets the reference to null. */
 	inline void Reset() { *this = FDerivedData(); }
 
@@ -97,7 +103,7 @@ public:
 	inline bool HasData() const { return !IsNull(); }
 
 	/** Returns the flags, which are mainly relevant for staged references. */
-	inline EDerivedDataFlags GetFlags() const { return Flags; }
+	inline EDerivedDataFlags GetFlags() const { return CookedData.Flags; }
 
 	/** Appends the name and description of this reference to the builder. */
 	UE_API friend FStringBuilderBase& operator<<(FStringBuilderBase& Builder, const FDerivedData& DerivedData);
@@ -135,14 +141,16 @@ private:
 private:
 	DerivedData::Private::FCookedData CookedData;
 
-	/** The flags that modify how this derived data is saved and how it may be loaded. */
-	EDerivedDataFlags Flags = EDerivedDataFlags::None;
-
 	friend DerivedData::Private::FIoResponse;
 };
 
 inline const FDerivedData FDerivedData::Null;
 
+inline FDerivedData::FDerivedData(const DerivedData::Private::FCookedData& InCookedData)
+	: CookedData(InCookedData)
+{
+}
+
 inline bool FDerivedData::IsNull() const
 {
 #if WITH_EDITORONLY_DATA
diff --git a/Engine/Source/Runtime/CoreUObject/Public/UObject/LinkerSave.h b/Engine/Source/Runtime/CoreUObject/Public/UObject/LinkerSave.h
--- a/Engine/Source/Runtime/CoreUObject/Public/UObject/LinkerSave.h
+++ b/Engine/Source/Runtime/CoreUObject/Public/UObject/LinkerSave.h
@@ -31,7 +31,6 @@ class UObject;
 class UPackage;
 namespace UE { class FDerivedData; }
 namespace UE { class FPackageTrailerBuilder; }
-namespace UE::DerivedData::Private { class FCookedData; }
 struct FLazyObjectPtr;
 struct FUObjectSerializeContext;
 
@@ -250,9 +249,9 @@ public:
 	/**
 	 * Adds the derived data to the package. This is only supported when saving a cooked package.
 	 *
-	 * @return A cooked reference that can be used to load the derived data from the cooked package.
+	 * @return A reference that can be used to load the derived data from the cooked package.
 	 */
-	UE::DerivedData::Private::FCookedData AddDerivedData(const UE::FDerivedData& Data);
+	UE::FDerivedData AddDerivedData(const UE::FDerivedData& Data);
 #endif // WITH_EDITORONLY_DATA
 
 protected:
