Core - Adding checks for cases in ObjectHandle that shouldn't ever be hit alerting to some malformed data being intepreted.

diff --git a/Engine/Source/Runtime/CoreUObject/Private/UObject/ObjectHandle.cpp b/Engine/Source/Runtime/CoreUObject/Private/UObject/ObjectHandle.cpp
--- a/Engine/Source/Runtime/CoreUObject/Private/UObject/ObjectHandle.cpp
+++ b/Engine/Source/Runtime/CoreUObject/Private/UObject/ObjectHandle.cpp
@@ -330,6 +330,7 @@ namespace ObjectHandle_Private
 		const uint32 PackageIndex = PackageId.ToIndex();
 		if (PackageIndex >= static_cast<uint32>(GObjectHandleIndex.PackageData.Num()))
 		{
+			checkf(false, TEXT("FObjectHandle: PackageIndex invalid.  This ObjectHandle is from malformed data."));
 			return;
 		}
 		FObjectHandlePackageData& FoundPackageData = GObjectHandleIndex.PackageData[PackageIndex];
@@ -338,8 +339,10 @@ namespace ObjectHandle_Private
 		if ((ObjectId.Components.ObjectPathId >= static_cast<uint32>(FoundPackageData.ObjectPaths.Num() + 1)) ||
 			(ObjectId.Components.DataClassDescriptorId >= static_cast<uint32>(FoundPackageData.DataClassDescriptors.Num() + 1)))
 		{
+			checkf(false, TEXT("FObjectHandle: ObjectId.Components.ObjectPathId or ObjectId.Components.DataClassDescriptorId invalid.  This ObjectHandle is from malformed data."));
 			return;
 		}
+
 		OutPackageName = FoundPackageData.PackageName;
 		OutPathId = FoundPackageData.ObjectPaths[ObjectId.Components.ObjectPathId - 1];
 
@@ -349,6 +352,12 @@ namespace ObjectHandle_Private
 			OutClassPackageName = Desc.PackageName;
 			OutClassName = Desc.ClassName;
 		}
+		else
+		{
+			//@TODO OBJPTR: Is this an error case we should check on?
+			//checkf(false, TEXT("FObjectHandle: ObjectId.Components.ObjectPathId or ObjectId.Components.DataClassDescriptorId invalid.  This ObjectHandle is from malformed data."));
+			return;
+		}
 	}
 
 	static inline FObjectRef MakeObjectRef(FPackedObjectRef PackedObjectRef)
