[CoreUObject] removed TObjectPtr.IsNull TObjectPtr.IsNullNoResolve FObjectPtr.IsNull and FObjectPtr.IsNullNoResolve

#p4v-preflight-copy 22194027

diff --git a/Engine/Source/Runtime/CoreUObject/Public/UObject/ObjectHandle.h b/Engine/Source/Runtime/CoreUObject/Public/UObject/ObjectHandle.h
--- a/Engine/Source/Runtime/CoreUObject/Public/UObject/ObjectHandle.h
+++ b/Engine/Source/Runtime/CoreUObject/Public/UObject/ObjectHandle.h
@@ -169,6 +169,9 @@ inline FPackedObjectRef ReadObjectHandlePackedObjectRefNoCheck(FObjectHandle Han
 
 inline UObject* ResolveObjectHandle(FObjectHandle& Handle);
 inline UObject* ResolveObjectHandleNoRead(FObjectHandle& Handle);
+
+/** Resolves an ObjectHandle without checking if already resolved. Invalid to call for resolved handles */
+inline UObject* ResolveObjectHandleNoReadNoCheck(FObjectHandle& Handle);
 inline UClass* ResolveObjectHandleClass(FObjectHandle Handle);
 
 /** Read the handle as a pointer if resolved, and otherwise return null. */
@@ -451,6 +454,20 @@ inline UObject* ResolveObjectHandleNoRead(FObjectHandle& Handle)
 #endif
 }
 
+inline UObject* ResolveObjectHandleNoReadNoCheck(FObjectHandle& Handle)
+{
+#if UE_WITH_OBJECT_HANDLE_LATE_RESOLVE
+	FObjectHandle LocalHandle = Handle;
+	LocalHandle = MakeObjectHandle(ResolvePackedObjectRef(ReadObjectHandlePackedObjectRefNoCheck(LocalHandle)));
+	UObject* ResolvedObject = ReadObjectHandlePointerNoCheck(LocalHandle);
+	Handle = LocalHandle;
+	return ResolvedObject;
+	
+#else
+	return ReadObjectHandlePointerNoCheck(Handle);
+#endif
+}
+
 inline UObject* ResolveObjectHandle(FObjectHandle& Handle)
 {
 #if UE_WITH_OBJECT_HANDLE_LATE_RESOLVE || UE_WITH_OBJECT_HANDLE_TRACKING
diff --git a/Engine/Source/Runtime/CoreUObject/Public/UObject/ObjectPtr.h b/Engine/Source/Runtime/CoreUObject/Public/UObject/ObjectPtr.h
--- a/Engine/Source/Runtime/CoreUObject/Public/UObject/ObjectPtr.h
+++ b/Engine/Source/Runtime/CoreUObject/Public/UObject/ObjectPtr.h
@@ -133,12 +133,14 @@ public:
 	FORCEINLINE UObject* operator->() const { return Get(); }
 	FORCEINLINE UObject& operator*() const { return *Get(); }
 
-	// Forcing resolve for null check.  Non null unresolved pointers can become null resolved pointers.
+	UE_DEPRECATED(5.1, "IsNull is deprecated, please use operator bool instead.")
 	FORCEINLINE bool IsNull() const { return ResolveObjectHandleNoRead(Handle) == nullptr; }
+	
+	UE_DEPRECATED(5.1, "IsNullNoResolve is deprecated, please use operator bool instead.")
 	FORCEINLINE bool IsNullNoResolve() const { return IsObjectHandleNull(Handle); }
-
-	FORCEINLINE bool operator!() const { return IsNull(); }
-	explicit FORCEINLINE operator bool() const { return !IsNull(); }
+	
+	FORCEINLINE bool operator!() const { return IsObjectHandleNull(Handle); }
+	explicit FORCEINLINE operator bool() const { return !IsObjectHandleNull(Handle); }
 
 	FORCEINLINE bool IsResolved() const { return IsObjectHandleResolved(Handle); }
 
@@ -350,6 +352,16 @@ public:
 		return ObjectPtr == Other.ObjectPtr;
 	}
 
+	bool operator==(TYPE_OF_NULLPTR) const
+	{
+		return !ObjectPtr.operator bool();
+	}
+
+	bool operator!=(TYPE_OF_NULLPTR) const
+	{
+		return ObjectPtr.operator bool();
+	}
+
 	template <
 		typename U,
 		typename Base = std::decay_t<decltype(false ? std::declval<std::decay_t<T*>>() : std::declval<std::decay_t<U*>>())>
@@ -379,8 +391,12 @@ public:
 	UE_OBJPTR_DEPRECATED(5.0, "Conversion to a mutable pointer is deprecated.  Please pass a TObjectPtr<T>& instead so that assignment can be tracked accurately.")
 	explicit FORCEINLINE operator T*& () { return GetInternalRef(); }
 
-	FORCEINLINE bool IsNull() const { return ObjectPtr.IsNull(); }
-	FORCEINLINE bool IsNullNoResolve() const { return ObjectPtr.IsNullNoResolve(); }
+	UE_DEPRECATED(5.1, "IsNull is deprecated, please use operator bool instead.")
+	FORCEINLINE bool IsNull() const { return !ObjectPtr.operator bool(); }
+
+	UE_DEPRECATED(5.1, "IsNullNoResolve is deprecated, please use operator bool instead.")
+	FORCEINLINE bool IsNullNoResolve() const { return !ObjectPtr.operator bool(); }
+
 	FORCEINLINE bool operator!() const { return ObjectPtr.operator!(); }
 	explicit FORCEINLINE operator bool() const { return ObjectPtr.operator bool(); }
 	FORCEINLINE bool IsResolved() const { return ObjectPtr.IsResolved(); }
@@ -414,7 +430,7 @@ public:
 	template <typename U, typename V> friend bool ObjectPtr_Private::IsObjectPtrEqualToRawPtrOfRelatedType(const TObjectPtr<U>& Ptr, const V* Other);
 
 private:
-	FORCEINLINE T* GetNoRead() const { return (T*)(ResolveObjectHandleNoRead(ObjectPtr.GetHandleRef())); }
+	FORCEINLINE T* GetNoReadNoCheck() const { return (T*)(ResolveObjectHandleNoReadNoCheck(ObjectPtr.GetHandleRef())); }
 
 	// @TODO: OBJPTR: There is a risk of a gap in access tracking here.  The caller may get a mutable pointer, write to it, then
 	//			read from it.  That last read would happen without an access being recorded.  Not sure if there is a good way
@@ -434,6 +450,19 @@ private:
 	};
 };
 
+// Equals against nullptr to optimize comparing against nullptr as it avoids resolving
+template <typename U>
+bool operator==(TYPE_OF_NULLPTR, const TObjectPtr<U>& Rhs)
+{
+	return Rhs == nullptr;
+}
+
+template <typename U>
+bool operator!=(TYPE_OF_NULLPTR, const TObjectPtr<U>& Rhs)
+{
+	return Rhs != nullptr;
+}
+
 template <typename T> struct TIsTObjectPtr<               TObjectPtr<T>> { enum { Value = true }; };
 template <typename T> struct TIsTObjectPtr<const          TObjectPtr<T>> { enum { Value = true }; };
 template <typename T> struct TIsTObjectPtr<      volatile TObjectPtr<T>> { enum { Value = true }; };
@@ -510,7 +539,17 @@ namespace ObjectPtr_Private
 	template <typename T, typename U>
 	FORCEINLINE bool IsObjectPtrEqualToRawPtrOfRelatedType(const TObjectPtr<T>& Ptr, const U* Other)
 	{
-		return Ptr.GetNoRead() == Other;
+#if UE_WITH_OBJECT_HANDLE_LATE_RESOLVE
+		if (Ptr.IsResolved())
+		{
+			return Ptr.GetHandle().PointerOrRef == uintptr_t(Other);
+		}
+		else if (!Other) //avoids resolving if Other is null
+		{
+			return !Ptr;
+		}
+#endif
+		return Ptr.GetNoReadNoCheck() == Other;
 	}
 
 	/** Perform shallow equality check between a TObjectPtr and another (non TObjectPtr) type that we can coerce to a pointer. */
diff --git a/Engine/Source/Runtime/CoreUObject/Tests/ObjectHandleTest.cpp b/Engine/Source/Runtime/CoreUObject/Tests/ObjectHandleTest.cpp
--- a/Engine/Source/Runtime/CoreUObject/Tests/ObjectHandleTest.cpp
+++ b/Engine/Source/Runtime/CoreUObject/Tests/ObjectHandleTest.cpp
@@ -4,11 +4,10 @@
 
 #include "UObject/ObjectHandle.h"
 #include "UObject/ObjectPtr.h"
-
+#include "UObject/MetaData.h"
 #include "HAL/PlatformProperties.h"
 #include "ObjectRefTrackingTestBase.h"
 #include "IO/IoDispatcher.h"
-
 #include "TestHarness.h"
 
 static_assert(sizeof(FObjectHandle) == sizeof(void*), "FObjectHandle type must always compile to something equivalent to a pointer size.");
@@ -205,6 +204,104 @@ TEST_CASE_METHOD(FObjectHandleTestBase, "CoreUObject::TObjectPtr::HandleNullGetC
 	TEST_TRUE(TEXT("TObjectPtr.GetClass should return null on a null object"), Ptr.GetClass() == nullptr);
 }
 
+TEST_CASE_METHOD(FObjectHandleTestBase, "CoreUObject::TObjectPtr::Null Behavior", "[CoreUObject][ObjectHandle]")
+{
+	TObjectPtr<UObject> Ptr = nullptr;
+	UMetaData* TestObject = nullptr;
+
+	uint32 ResolveCount = 0;
+	auto ResolveDelegate = FObjectHandleReferenceResolvedDelegate::CreateLambda([&ResolveCount](const FObjectRef& SourceRef, UPackage* ObjectPackage, UObject* Object)
+		{
+			++ResolveCount;
+		});
+	auto Handle = AddObjectHandleReferenceResolvedCallback(ResolveDelegate);
+
+	//compare against all flavours of nullptr, should not try and resolve this pointer
+	CHECK(Ptr == nullptr); CHECK(ResolveCount == 0u);
+	CHECK(nullptr == Ptr); CHECK(ResolveCount == 0u);
+	CHECK_FALSE(Ptr != nullptr); CHECK(ResolveCount == 0u);
+	CHECK_FALSE(nullptr != Ptr); CHECK(ResolveCount == 0u);
+	CHECK(!Ptr); CHECK(ResolveCount == 0u);
+
+	//using an if otherwise the macros try to convert to a pointer and not use the bool operator
+	if (Ptr)
+	{
+		CHECK(false);
+	}
+	else
+	{
+		CHECK(true);
+	}
+	CHECK(ResolveCount == 0u);
+
+	CHECK(Ptr == TestObject); CHECK(ResolveCount == 0u);
+	CHECK(TestObject == Ptr); CHECK(ResolveCount == 0u);
+	CHECK_FALSE(Ptr != TestObject); CHECK(ResolveCount == 0u);
+	CHECK_FALSE(TestObject != Ptr); CHECK(ResolveCount == 0u);
+
+	FObjectRef TargetRef{ FName("SomePackage"), FName("ClassPackageName"), FName("ClassName"), FObjectPathId("ObjectName") };
+	FObjectPtr ObjectPtr(TargetRef);
+	REQUIRE(!ObjectPtr.IsResolved()); //make sure not resolved
+
+	//an unresolved pointers compared against nullptr should still not resolve
+	Ptr = *reinterpret_cast<TObjectPtr<UObject>*>(&ObjectPtr);
+	CHECK_FALSE(Ptr == nullptr); CHECK(ResolveCount == 0u);
+	CHECK_FALSE(nullptr == Ptr); CHECK(ResolveCount == 0u);
+	CHECK(Ptr != nullptr); CHECK(ResolveCount == 0u);
+	CHECK(nullptr != Ptr); CHECK(ResolveCount == 0u);
+	CHECK_FALSE(!Ptr); CHECK(ResolveCount == 0u);
+
+	//using an if otherwise the macros try to convert to a pointer and not use the bool operator
+	if (Ptr)
+	{
+		CHECK(true);
+	}
+	else
+	{
+		CHECK(false);
+	}
+	CHECK(ResolveCount == 0u);
+
+	//test an unresolve pointer against a null raw pointer
+	CHECK_FALSE(Ptr == TestObject); CHECK(ResolveCount == 0u);
+	CHECK_FALSE(TestObject == Ptr);	CHECK(ResolveCount == 0u);
+	CHECK(Ptr != TestObject); CHECK(ResolveCount == 0u);
+	CHECK(TestObject != Ptr); CHECK(ResolveCount == 0u);
+
+	//creating a real object for something that can resolve
+	const FName TestPackageName(TEXT("/Engine/Test/ObjectPtrDefaultSerialize/Transient"));
+	UPackage* TestPackage = NewObject<UPackage>(nullptr, TestPackageName, RF_Transient);
+	TestPackage->AddToRoot();
+
+	const FName TestObjectName(TEXT("MyObject"));
+	TestObject = NewObject<UMetaData>(TestPackage, TestObjectName, RF_Transient);
+	TestObject->AddToRoot();
+
+	//ObjectPtr = FObjectPtr(MakePackedObjectRef(TestObject));
+	Ptr = TestObject;
+
+	REQUIRE(Ptr.IsResolved());
+	CHECK(Ptr == TestObject); CHECK(ResolveCount == 0u);
+	CHECK(TestObject == Ptr); CHECK(ResolveCount == 0u);
+	CHECK_FALSE(Ptr != TestObject); CHECK(ResolveCount == 0u);
+	CHECK_FALSE(TestObject != Ptr); CHECK(ResolveCount == 0u);
+
+	TestObject = nullptr;
+	CHECK_FALSE(Ptr == TestObject); CHECK(ResolveCount == 0u);
+	CHECK_FALSE(TestObject == Ptr); CHECK(ResolveCount == 0u);
+	CHECK(Ptr != TestObject); CHECK(ResolveCount == 0u);
+	CHECK(TestObject != Ptr); CHECK(ResolveCount == 0u);
+
+	TestObject = static_cast<UMetaData*>(Ptr.Get());
+	Ptr = nullptr;
+	CHECK_FALSE(Ptr == TestObject); CHECK(ResolveCount == 0u);
+	CHECK_FALSE(TestObject == Ptr); CHECK(ResolveCount == 0u);
+	CHECK(Ptr != TestObject); CHECK(ResolveCount == 0u);
+	CHECK(TestObject != Ptr); CHECK(ResolveCount == 0u);
+
+	RemoveObjectHandleReferenceResolvedCallback(Handle);
+}
+
 #if UE_WITH_OBJECT_HANDLE_LATE_RESOLVE
 
 TEST_CASE_METHOD(FObjectHandleTestBase, "CoreUObject::FObjectHandle::Resolve Malformed Handle", "[CoreUObject][ObjectHandle]")
