#UE Check for redirectors when evaluating if we should show a warning about an import class mismatch

diff --git a/Engine/Source/Runtime/CoreUObject/Private/UObject/LinkerLoad.cpp b/Engine/Source/Runtime/CoreUObject/Private/UObject/LinkerLoad.cpp
--- a/Engine/Source/Runtime/CoreUObject/Private/UObject/LinkerLoad.cpp
+++ b/Engine/Source/Runtime/CoreUObject/Private/UObject/LinkerLoad.cpp
@@ -5457,7 +5457,15 @@ UObject* FLinkerLoad::CreateImport( int32 Index )
 			const UClass* ExpectedImportClass = nullptr;
 			if (UPackage* ImportClassPackage = FindObjectFast<UPackage>(nullptr, Import.ClassPackage, false))
 			{
-				ExpectedImportClass = FindObjectFast<UClass>(ImportClassPackage, Import.ClassName, false);
+				const UObject* FoundObject = FindObjectFast<UObject>(ImportClassPackage, Import.ClassName, false);
+				ExpectedImportClass = Cast<UClass>(FoundObject);
+				if (!ExpectedImportClass)
+				{
+					if (const UObjectRedirector* FoundRedirector = Cast<UObjectRedirector>(FoundObject))
+					{
+						ExpectedImportClass = Cast<UClass>(FoundRedirector->DestinationObject);
+					}
+				}
 			}
 
 			// Verify that the resolved import object's class is serialization-compatible with the expected result. Data loss will otherwise occur on load if this is not satisfied, so we warn about it. A re-save is required to fix up the import table and suppress this warning.
