Core - Disabling the checks I have for malformed handles.  The low level tests are triggering these checks, need to investigate avoiding that.

diff --git a/Engine/Source/Runtime/CoreUObject/Private/UObject/ObjectHandle.cpp b/Engine/Source/Runtime/CoreUObject/Private/UObject/ObjectHandle.cpp
--- a/Engine/Source/Runtime/CoreUObject/Private/UObject/ObjectHandle.cpp
+++ b/Engine/Source/Runtime/CoreUObject/Private/UObject/ObjectHandle.cpp
@@ -330,7 +330,7 @@ namespace ObjectHandle_Private
 		const uint32 PackageIndex = PackageId.ToIndex();
 		if (PackageIndex >= static_cast<uint32>(GObjectHandleIndex.PackageData.Num()))
 		{
-			checkf(false, TEXT("FObjectHandle: PackageIndex invalid.  This ObjectHandle is from malformed data."));
+			//checkf(false, TEXT("FObjectHandle: PackageIndex invalid.  This ObjectHandle is from malformed data."));
 			return;
 		}
 		FObjectHandlePackageData& FoundPackageData = GObjectHandleIndex.PackageData[PackageIndex];
@@ -339,7 +339,7 @@ namespace ObjectHandle_Private
 		if ((ObjectId.Components.ObjectPathId >= static_cast<uint32>(FoundPackageData.ObjectPaths.Num() + 1)) ||
 			(ObjectId.Components.DataClassDescriptorId >= static_cast<uint32>(FoundPackageData.DataClassDescriptors.Num() + 1)))
 		{
-			checkf(false, TEXT("FObjectHandle: ObjectId.Components.ObjectPathId or ObjectId.Components.DataClassDescriptorId invalid.  This ObjectHandle is from malformed data."));
+			//checkf(false, TEXT("FObjectHandle: ObjectId.Components.ObjectPathId or ObjectId.Components.DataClassDescriptorId invalid.  This ObjectHandle is from malformed data."));
 			return;
 		}
 
