Fixed ensures firing for vpaths with dashes (and other illegal characters) in them.

diff --git a/Engine/Source/Runtime/CoreUObject/Private/UObject/UObjectBaseUtility.cpp b/Engine/Source/Runtime/CoreUObject/Private/UObject/UObjectBaseUtility.cpp
--- a/Engine/Source/Runtime/CoreUObject/Private/UObject/UObjectBaseUtility.cpp
+++ b/Engine/Source/Runtime/CoreUObject/Private/UObject/UObjectBaseUtility.cpp
@@ -270,9 +270,18 @@ UE::Core::FVersePath UObjectBaseUtility::GetVersePath() const
 	}
 
 	UE::Core::FVersePath Result;
-	if (!ensureAlwaysMsgf(UE::Core::FVersePath::TryMake(Result, /*want to move this, but then it won't be available for the ensure message - MoveTemp*/(VerseModule)), TEXT("Unable to make FVersePath from string: %s"), *VerseModule))
+	if (!UE::Core::FVersePath::TryMake(Result, VerseModule))
 	{
-		return {};
+#if !NO_LOGGING
+		static thread_local TSet<FString> AlreadyLogged;
+
+		bool bAlreadyInSet = false;
+		AlreadyLogged.Add(VerseModule, &bAlreadyInSet);
+		if (!bAlreadyInSet)
+		{
+			UE_LOG(LogCore, Warning, TEXT("Unable to make a VersePath from %s"), *GetPathName());
+		}
+#endif
 	}
 
 	return Result;
