[CoreUObject] compile fix for platforms where UE_WITH_OBJECT_HANDLE_TRACKING is not on.
tests was using features only enable with handle tracking

diff --git a/Engine/Source/Runtime/CoreUObject/Tests/ObjectHandleTest.cpp b/Engine/Source/Runtime/CoreUObject/Tests/ObjectHandleTest.cpp
--- a/Engine/Source/Runtime/CoreUObject/Tests/ObjectHandleTest.cpp
+++ b/Engine/Source/Runtime/CoreUObject/Tests/ObjectHandleTest.cpp
@@ -204,6 +204,8 @@ TEST_CASE_METHOD(FObjectHandleTestBase, "CoreUObject::TObjectPtr::HandleNullGetC
 	TEST_TRUE(TEXT("TObjectPtr.GetClass should return null on a null object"), Ptr.GetClass() == nullptr);
 }
 
+#if UE_WITH_OBJECT_HANDLE_TRACKING
+
 TEST_CASE_METHOD(FObjectHandleTestBase, "CoreUObject::TObjectPtr::Null Behavior", "[CoreUObject][ObjectHandle]")
 {
 	TObjectPtr<UObject> Ptr = nullptr;
@@ -302,6 +304,8 @@ TEST_CASE_METHOD(FObjectHandleTestBase, "CoreUObject::TObjectPtr::Null Behavior"
 	RemoveObjectHandleReferenceResolvedCallback(Handle);
 }
 
+#endif
+
 #if UE_WITH_OBJECT_HANDLE_LATE_RESOLVE
 
 TEST_CASE_METHOD(FObjectHandleTestBase, "CoreUObject::FObjectHandle::Resolve Malformed Handle", "[CoreUObject][ObjectHandle]")
