Fixes Guid Remapping inTLazyObjectPtr asserting when multiple PIE sessions

diff --git a/Engine/Source/Runtime/CoreUObject/Private/UObject/LazyObjectPtr.cpp b/Engine/Source/Runtime/CoreUObject/Private/UObject/LazyObjectPtr.cpp
--- a/Engine/Source/Runtime/CoreUObject/Private/UObject/LazyObjectPtr.cpp
+++ b/Engine/Source/Runtime/CoreUObject/Private/UObject/LazyObjectPtr.cpp
@@ -147,8 +147,13 @@ void FLazyObjectPtr::PossiblySerializeObjectGuid(UObject *Object, FStructuredArc
 #if WITH_EDITOR
 						else if (Object->GetOutermostObject()->GetPackage()->HasAnyPackageFlags(PKG_PlayInEditor))
 						{
-							Guid = RemapGuid(Guid, Object->GetOutermostObject()->GetPackage()->GetPIEInstanceID());
-							GuidAnnotation.AddAnnotation(Object, Guid);
+							int32 PIEInstanceID = Object->GetOutermostObject()->GetPackage()->GetPIEInstanceID();
+
+							if (PIEInstanceID != INDEX_NONE)
+							{
+								Guid = RemapGuid(Guid, PIEInstanceID);
+								GuidAnnotation.AddAnnotation(Object, Guid);
+							}
 						}
 #endif
 						else
