[CoreUObject]
* SoftObjectPtr - Made it possible to use extern template on SoftObjectPtr and have get function end up in implementation of extern template.

diff --git a/Engine/Source/Runtime/CoreUObject/Public/UObject/SoftObjectPtr.h b/Engine/Source/Runtime/CoreUObject/Public/UObject/SoftObjectPtr.h
--- a/Engine/Source/Runtime/CoreUObject/Public/UObject/SoftObjectPtr.h
+++ b/Engine/Source/Runtime/CoreUObject/Public/UObject/SoftObjectPtr.h
@@ -288,10 +288,7 @@ public:
 	 *
 	 * @return nullptr if this object is gone or the lazy pointer was null, otherwise a valid UObject pointer
 	 */
-	FORCEINLINE T* Get() const
-	{
-		return dynamic_cast<T*>(SoftObjectPtr.Get());
-	}
+	T* Get() const;
 
 	/** Dereference the soft pointer */
 	FORCEINLINE T& operator*() const
@@ -645,3 +642,9 @@ using TAssetPtr UE_DEPRECATED(5.0, "TAssetPtr was renamed to TSoftObjectPtr as i
 template<class TClass = UObject>
 using TAssetSubclassOf UE_DEPRECATED(5.0, "TAssetSubclassOf was renamed to TSoftClassPtr") = TSoftClassPtr<TClass>;
 
+/** Not directly inlined on purpose so compiler have the option of not inlining it. (and it also works with extern template) */
+template<class T>
+T* TSoftObjectPtr<T>::Get() const
+{
+	return dynamic_cast<T*>(SoftObjectPtr.Get());
+}
