Explicitly delete operator bool on object weak pointers to make it clear why one doesn't exist, and clarify that you shouldn't be calling IsValid all the time.

diff --git a/Engine/Source/Runtime/Core/Public/UObject/WeakObjectPtrTemplates.h b/Engine/Source/Runtime/Core/Public/UObject/WeakObjectPtrTemplates.h
--- a/Engine/Source/Runtime/Core/Public/UObject/WeakObjectPtrTemplates.h
+++ b/Engine/Source/Runtime/Core/Public/UObject/WeakObjectPtrTemplates.h
@@ -155,11 +155,18 @@ public:
 		return Get();
 	}
 
+	// This is explicitly not added to avoid resolving weak pointers too often - use Get() once in a function.
+	explicit operator bool() const = delete;
+
 	/**  
-	 * Test if this points to a live UObject
+	 * Test if this points to a live UObject.
+	 * This should be done only when needed as excess resolution of the underlying pointer can cause performance issues.
+	 * 
 	 * @param bEvenIfPendingKill if this is true, pendingkill objects are considered valid
 	 * @param bThreadsafeTest if true then function will just give you information whether referenced
 	 *							UObject is gone forever (return false) or if it is still there (return true, no object flags checked).
+	 *							This is required as without it IsValid can return false during the mark phase of the GC
+	 *							due to the presence of the Unreachable flag.
 	 * @return true if Get() would return a valid non-null pointer
 	 */
 	FORCEINLINE bool IsValid(bool bEvenIfPendingKill, bool bThreadsafeTest = false) const
@@ -169,6 +176,9 @@ public:
 
 	/**
 	 * Test if this points to a live UObject. This is an optimized version implying bEvenIfPendingKill=false, bThreadsafeTest=false.
+	 * This should be done only when needed as excess resolution of the underlying pointer can cause performance issues.
+	 * Note that IsValid can not be used on another thread as it will incorrectly return false during the mark phase of the GC
+	 * due to the Unreachable flag being set. (see bThreadsafeTest above)
 	 * @return true if Get() would return a valid non-null pointer
 	 */
 	FORCEINLINE bool IsValid(/*bool bEvenIfPendingKill = false, bool bThreadsafeTest = false*/) const
diff --git a/Engine/Source/Runtime/CoreUObject/Public/UObject/WeakObjectPtr.h b/Engine/Source/Runtime/CoreUObject/Public/UObject/WeakObjectPtr.h
--- a/Engine/Source/Runtime/CoreUObject/Public/UObject/WeakObjectPtr.h
+++ b/Engine/Source/Runtime/CoreUObject/Public/UObject/WeakObjectPtr.h
@@ -165,17 +165,28 @@ public:
 	/** Dereference the weak pointer even if it is RF_PendingKill or RF_Unreachable */
 	COREUOBJECT_API class UObject* GetEvenIfUnreachable() const;
 
+	// This is explicitly not added to avoid resolving weak pointers too often - use Get() once in a function.
+	explicit operator bool() const = delete;
+
 	/**  
 	 * Test if this points to a live UObject
+	 * This should be done only when needed as excess resolution of the underlying pointer can cause performance issues.
+	 *
 	 * @param bEvenIfPendingKill if this is true, pendingkill are not considered invalid
 	 * @param bThreadsafeTest if true then function will just give you information whether referenced
 	 *							UObject is gone forever (return false) or if it is still there (return true, no object flags checked).
+	 *							This is required as without it IsValid can return false during the mark phase of the GC
+	 *							due to the presence of the Unreachable flag.
 	 * @return true if Get() would return a valid non-null pointer
 	 */
 	COREUOBJECT_API bool IsValid(bool bEvenIfPendingKill, bool bThreadsafeTest = false) const;
 
 	/**
 	 * Test if this points to a live UObject. This is an optimized version implying bEvenIfPendingKill=false, bThreadsafeTest=false.
+	 * This should be done only when needed as excess resolution of the underlying pointer can cause performance issues.
+	 * Note that IsValid can not be used on another thread as it will incorrectly return false during the mark phase of the GC
+	 * due to the Unreachable flag being set. (see bThreadsafeTest above)
+
 	 * @return true if Get() would return a valid non-null pointer.
 	 */
 	COREUOBJECT_API bool IsValid(/*bool bEvenIfPendingKill = false, bool bThreadsafeTest = false*/) const;
