Core - Doing a bit of cleanup in the ObjectPtr resolution code.

diff --git a/Engine/Source/Runtime/CoreUObject/Private/UObject/LinkerLoadImportBehavior.cpp b/Engine/Source/Runtime/CoreUObject/Private/UObject/LinkerLoadImportBehavior.cpp
--- a/Engine/Source/Runtime/CoreUObject/Private/UObject/LinkerLoadImportBehavior.cpp
+++ b/Engine/Source/Runtime/CoreUObject/Private/UObject/LinkerLoadImportBehavior.cpp
@@ -158,10 +158,12 @@ bool TryLazyImport(const IAssetRegistryInterface& AssetRegistry, const FObjectIm
 		return false;
 	}
 
+	static const FName Name_CoreUObjectPackage("/Script/CoreUObject");
+
 	ObjectPtr = nullptr;
 
 	//packages need to handles differently since the can't be found by object path in the AssetRegistry
-	if (Import.ClassPackage == TEXT("/Script/CoreUObject") && Import.ClassName == TEXT("Package"))
+	if (Import.ClassPackage == Name_CoreUObjectPackage && Import.ClassName == NAME_Package)
 	{
 		FAssetPackageData PackageData;
 		UE::AssetRegistry::EExists Exists = AssetRegistry.TryGetAssetPackageData(Import.ObjectName, PackageData);
@@ -200,12 +202,12 @@ bool TryLazyLoad(const UClass& Class, const FSoftObjectPath& ObjectPath, TObject
 		return false;
 	}
 
-	EImportBehavior Behavior = FindLoadBehavior(Class);
+	const EImportBehavior Behavior = FindLoadBehavior(Class);
 	if (Behavior != EImportBehavior::LazyOnDemand)
 	{
 		return false;
 	}
-	auto AssetRegistry = IAssetRegistryInterface::GetPtr();
+	const IAssetRegistryInterface* AssetRegistry = IAssetRegistryInterface::GetPtr();
 	if (!AssetRegistry)
 	{
 		return false;
diff --git a/Engine/Source/Runtime/CoreUObject/Public/UObject/ObjectPtr.h b/Engine/Source/Runtime/CoreUObject/Public/UObject/ObjectPtr.h
--- a/Engine/Source/Runtime/CoreUObject/Public/UObject/ObjectPtr.h
+++ b/Engine/Source/Runtime/CoreUObject/Public/UObject/ObjectPtr.h
@@ -391,10 +391,10 @@ public:
 	UE_OBJPTR_DEPRECATED(5.0, "Conversion to a mutable pointer is deprecated.  Please pass a TObjectPtr<T>& instead so that assignment can be tracked accurately.")
 	explicit FORCEINLINE operator T*& () { return GetInternalRef(); }
 
-	UE_DEPRECATED(5.1, "IsNull is deprecated, please use operator bool instead.")
+	UE_DEPRECATED(5.1, "IsNull is deprecated, please use operator bool instead.  if (!MyObjectPtr) { ... }")
 	FORCEINLINE bool IsNull() const { return !ObjectPtr.operator bool(); }
 
-	UE_DEPRECATED(5.1, "IsNullNoResolve is deprecated, please use operator bool instead.")
+	UE_DEPRECATED(5.1, "IsNullNoResolve is deprecated, please use operator bool instead.  if (!MyObjectPtr) { ... }")
 	FORCEINLINE bool IsNullNoResolve() const { return !ObjectPtr.operator bool(); }
 
 	FORCEINLINE bool operator!() const { return ObjectPtr.operator!(); }
