Add FPackageTrailerBuilder::CalculateTrailerLength that allows the caller to query how large the package trailer would be if created by the builder with it's current entries.

#rnx

- Moved code that was calculating the length of different parts of the trailer to methods.
- Added an assert after building and appending a trailer to make sure that the ::CalculateTrailerLength calculated the correct value, this should catch any unintended changes in the future.
- Allowing the default constructor as we are moving to allow for building 'in memory' packages with the package trailer, so a package path makes no sense.

diff --git a/Engine/Source/Runtime/CoreUObject/Private/UObject/PackageTrailer.cpp b/Engine/Source/Runtime/CoreUObject/Private/UObject/PackageTrailer.cpp
--- a/Engine/Source/Runtime/CoreUObject/Private/UObject/PackageTrailer.cpp
+++ b/Engine/Source/Runtime/CoreUObject/Private/UObject/PackageTrailer.cpp
@@ -264,9 +264,7 @@ bool FPackageTrailerBuilder::BuildAndAppendTrailer(FLinkerSave* Linker, FArchive
 	Trailer.Header.Tag = FPackageTrailer::FHeader::HeaderTag;
 	Trailer.Header.Version = (int32)EPackageTrailerVersion::AUTOMATIC_VERSION;
 
-	const uint32 DynamicHeaderSizeOnDisk = (GetNumPayloads() * Private::FLookupTableEntry::SizeOnDisk); // Add the length of the lookup table
-	
-	Trailer.Header.HeaderLength = FPackageTrailer::FHeader::StaticHeaderSizeOnDisk + DynamicHeaderSizeOnDisk;
+	Trailer.Header.HeaderLength = CalculatePotentialHeaderSize();
 	
 	Trailer.Header.PayloadsDataLength = 0;
 	Trailer.Header.PayloadLookupTable.Reserve(LocalEntries.Num() + ReferencedEntries.Num() + VirtualizedEntries.Num());
@@ -352,6 +350,9 @@ bool FPackageTrailerBuilder::BuildAndAppendTrailer(FLinkerSave* Linker, FArchive
 		}
 	}
 
+	// Minor sanity check that ::GetTrailerLength works
+	check(CalculateTrailerLength() == (DataArchive.Tell() - Trailer.TrailerPositionInFile) || DataArchive.IsError());
+
 	return !DataArchive.IsError();
 }
 
@@ -375,6 +376,19 @@ bool FPackageTrailerBuilder::IsVirtualizedPayloadEntry(const FIoHash& Identifier
 	return VirtualizedEntries.Find(Identifier) != nullptr;
 }
 
+uint64 FPackageTrailerBuilder::CalculateTrailerLength()
+{
+	// For now we need to call ::RemoveDuplicateEntries to make sure that we are calculating
+	// the correct length, hence the method not being const
+	RemoveDuplicateEntries();
+
+	const uint64 HeaderLength = (uint64)CalculatePotentialHeaderSize();
+	const uint64 PayloadsLength = CalculatePotentialPayloadSize();
+	const uint64 FooterLength = FPackageTrailer::FFooter::SizeOnDisk;
+
+	return HeaderLength + PayloadsLength + FooterLength;
+}
+
 int32 FPackageTrailerBuilder::GetNumPayloads() const
 {
 	return GetNumLocalPayloads() + GetNumReferencedPayloads() + GetNumVirtualizedPayloads();
@@ -395,6 +409,26 @@ int32 FPackageTrailerBuilder::GetNumVirtualizedPayloads() const
 	return VirtualizedEntries.Num();
 }
 
+uint32 FPackageTrailerBuilder::CalculatePotentialHeaderSize() const
+{
+	const uint32 DynamicHeaderSizeOnDisk = (GetNumPayloads() * Private::FLookupTableEntry::SizeOnDisk); // Add the length of the lookup table
+
+	return FPackageTrailer::FHeader::StaticHeaderSizeOnDisk + DynamicHeaderSizeOnDisk;
+}
+
+uint64 FPackageTrailerBuilder::CalculatePotentialPayloadSize() const
+{
+	uint64 PayloadsLength = 0;
+	for (const TPair<FIoHash, LocalEntry>& KV : LocalEntries)
+	{
+		const FCompressedBuffer& Payload = KV.Value.Payload;
+
+		PayloadsLength += Payload.GetCompressedSize();
+	}
+
+	return PayloadsLength;
+}
+
 void FPackageTrailerBuilder::RemoveDuplicateEntries()
 {
 	for (auto Iter = LocalEntries.CreateIterator(); Iter; ++Iter)
diff --git a/Engine/Source/Runtime/CoreUObject/Public/UObject/PackageTrailer.h b/Engine/Source/Runtime/CoreUObject/Public/UObject/PackageTrailer.h
--- a/Engine/Source/Runtime/CoreUObject/Public/UObject/PackageTrailer.h
+++ b/Engine/Source/Runtime/CoreUObject/Public/UObject/PackageTrailer.h
@@ -234,7 +234,7 @@ public:
 		return FPackageTrailerBuilder::CreateReferenceToTrailer(Trailer, PackageName.ToString());
 	}
 
-	FPackageTrailerBuilder() = delete;
+	FPackageTrailerBuilder() = default;
 	UE_DEPRECATED(5.1, "Use the overload that takes a FString instead of an FName")
 	FPackageTrailerBuilder(const FName& InPackageName);
 	FPackageTrailerBuilder(FString&& DebugContext);
@@ -289,6 +289,14 @@ public:
 	[[nodiscard]] bool IsReferencedPayloadEntry(const FIoHash& Identifier) const;
 	[[nodiscard]] bool IsVirtualizedPayloadEntry(const FIoHash& Identifier) const;
 
+	/** 
+	 * Returns the length of the trailer (in bytes) that the builder would currently create.
+	 * 
+	 * NOTE: At the moment this is not const as we need to check for and remove duplicate
+	 * payload entries as we do this before building the trailer not when gathering the entry info.
+	 */
+	[[nodiscard]] uint64 CalculateTrailerLength();
+
 	/** Returns the total number of payload entries in the builder */
 	[[nodiscard]] int32 GetNumPayloads() const;
 	
@@ -356,6 +364,11 @@ private:
 		int64 RawSize = INDEX_NONE;
 	};
 
+	/** Returns the total length of the header if we were to build a trailer right now */
+	uint32 CalculatePotentialHeaderSize() const;
+	/** Returns the total length of all payloads combined if we were to build a trailer right now */
+	uint64 CalculatePotentialPayloadSize() const;
+
 	void RemoveDuplicateEntries();
 
 	// Members used when building the trailer
