ZenLoader in editor: Fixed issue with fully loaded externally created packages being marked as missing

diff --git a/Engine/Source/Runtime/CoreUObject/Private/Serialization/AsyncLoading2.cpp b/Engine/Source/Runtime/CoreUObject/Private/Serialization/AsyncLoading2.cpp
--- a/Engine/Source/Runtime/CoreUObject/Private/Serialization/AsyncLoading2.cpp
+++ b/Engine/Source/Runtime/CoreUObject/Private/Serialization/AsyncLoading2.cpp
@@ -4028,16 +4028,8 @@ void FAsyncPackage2::ImportPackagesRecursiveInner(FAsyncLoadingThreadState2& Thr
 #endif
 
 		FLoadedPackageRef& ImportedPackageRef = ImportStore.LoadedPackageStore.FindPackageRefChecked(ImportedPackageId);
-		if (ImportedPackageStatus == EPackageStoreEntryStatus::Missing)
-		{
-			UE_ASYNC_PACKAGE_LOG(Warning, Desc, TEXT("ImportPackages: SkipPackage"),
-				TEXT("Skipping non mounted imported package with id '0x%llX'"), ImportedPackageId.Value());
-			ImportedPackageRef.SetIsMissingPackage();
-			Data.ImportedAsyncPackages[ImportedPackageIndex++] = nullptr;
-			continue;
-		}
 #if WITH_EDITOR
-		else if (AsyncLoadingThread.UncookedPackageLoader && !ImportedPackageEntry.UncookedPackageName.IsNone())
+		if (AsyncLoadingThread.UncookedPackageLoader && ImportedPackageStatus == EPackageStoreEntryStatus::Ok && !ImportedPackageEntry.UncookedPackageName.IsNone())
 		{
 			UPackage* UncookedPackage = ImportedPackageRef.GetPackage();
 			if (!ImportedPackageRef.AreAllPublicExportsLoaded())
@@ -4105,7 +4097,6 @@ void FAsyncPackage2::ImportPackagesRecursiveInner(FAsyncLoadingThreadState2& Thr
 
 		FAsyncPackage2* ImportedPackage = nullptr;
 		bool bInserted = false;
-		FAsyncPackageDesc2 PackageDesc = FAsyncPackageDesc2::FromPackageImport(Desc, ImportedPackageUPackageName, ImportedPackageId, ImportedPackageIdToLoad, MoveTemp(ImportedPackagePath));
 		if (ImportedPackageRef.AreAllPublicExportsLoaded())
 		{
 			ImportedPackage = AsyncLoadingThread.FindAsyncPackage(ImportedPackageId);
@@ -4116,8 +4107,17 @@ void FAsyncPackage2::ImportPackagesRecursiveInner(FAsyncLoadingThreadState2& Thr
 			}
 			bInserted = false;
 		}
+		else if (ImportedPackageStatus == EPackageStoreEntryStatus::Missing)
+		{
+			UE_ASYNC_PACKAGE_LOG(Warning, Desc, TEXT("ImportPackages: SkipPackage"),
+				TEXT("Skipping non mounted imported package with id '0x%llX'"), ImportedPackageId.Value());
+			ImportedPackageRef.SetIsMissingPackage();
+			Data.ImportedAsyncPackages[ImportedPackageIndex++] = nullptr;
+			continue;
+		}
 		else
 		{
+			FAsyncPackageDesc2 PackageDesc = FAsyncPackageDesc2::FromPackageImport(Desc, ImportedPackageUPackageName, ImportedPackageId, ImportedPackageIdToLoad, MoveTemp(ImportedPackagePath));
 			ImportedPackage = AsyncLoadingThread.FindOrInsertPackage(PackageDesc, bInserted);
 		}
 
