Fixing up /Config/ root path resolution, and making it match the other game sub-directories (current behavior works as long as the game project is organized in this fashion: //<DEPOT>/<ProjFolder>, but it does not work if the project folder is nested somewhere else. This fix resolves that problem.).

diff --git a/Engine/Source/Runtime/CoreUObject/Private/Misc/PackageName.cpp b/Engine/Source/Runtime/CoreUObject/Private/Misc/PackageName.cpp
--- a/Engine/Source/Runtime/CoreUObject/Private/Misc/PackageName.cpp
+++ b/Engine/Source/Runtime/CoreUObject/Private/Misc/PackageName.cpp
@@ -415,10 +415,11 @@ private:
 		AddPathPair(ContentPathToRoot, GameRootPath,   GameContentPath);
 		AddPathPair(ContentPathToRoot, ScriptRootPath, GameScriptPath);
 		AddPathPair(ContentPathToRoot, TempRootPath,   GameSavedPath);
+		AddPathPair(ContentPathToRoot, ConfigRootPath, GameConfigPath);
 		AddPathPair(ContentPathToRoot, GameRootPath,   GameContentPathRebased);
 		AddPathPair(ContentPathToRoot, ScriptRootPath, GameScriptPathRebased);
 		AddPathPair(ContentPathToRoot, TempRootPath,   GameSavedPathRebased);
-		AddPathPair(ContentPathToRoot, ConfigRootPath, GameConfigPath);
+		AddPathPair(ContentPathToRoot, ConfigRootPath, GameConfigPathRebased);
 
 		ContentRootToPath.Empty(11);
 		AddPathPair(ContentRootToPath, EngineRootPath, EngineContentPath);
@@ -426,6 +427,7 @@ private:
 		AddPathPair(ContentRootToPath, GameRootPath,   GameContentPath);
 		AddPathPair(ContentRootToPath, ScriptRootPath, GameScriptPath);
 		AddPathPair(ContentRootToPath, TempRootPath,   GameSavedPath);
+		AddPathPair(ContentRootToPath, ConfigRootPath, GameConfigPath);
 		AddPathPair(ContentRootToPath, GameRootPath,   GameContentPathRebased);
 		AddPathPair(ContentRootToPath, ScriptRootPath, GameScriptPathRebased);
 		AddPathPair(ContentRootToPath, TempRootPath,   GameSavedPathRebased);
