DerivedData: Fixed serialization to calculate the PackageId

The PackageId is only assigned by the runtime package loader.

#rnx

diff --git a/Engine/Source/Runtime/CoreUObject/Private/UObject/LinkerSave.cpp b/Engine/Source/Runtime/CoreUObject/Private/UObject/LinkerSave.cpp
--- a/Engine/Source/Runtime/CoreUObject/Private/UObject/LinkerSave.cpp
+++ b/Engine/Source/Runtime/CoreUObject/Private/UObject/LinkerSave.cpp
@@ -541,7 +541,7 @@ UE::FDerivedData FLinkerSave::AddDerivedData(const UE::FDerivedData& Data)
 
 	UE::DerivedData::Private::FCookedData CookedData;
 
-	const FPackageId PackageId = LinkerRoot->GetPackageId();
+	const FPackageId PackageId = FPackageId::FromName(LinkerRoot->GetFName());
 	const int32 ChunkIndex = ++LastDerivedDataIndex;
 	checkf(ChunkIndex >= 0 && ChunkIndex < (1 << 24), TEXT("ChunkIndex %d is out of range."), ChunkIndex);
 
