Fix TSAN warnings by letting it know our intention is really to do relaxed reads on UObject flags

diff --git a/Engine/Source/Runtime/CoreUObject/Public/UObject/UObjectArray.h b/Engine/Source/Runtime/CoreUObject/Public/UObject/UObjectArray.h
--- a/Engine/Source/Runtime/CoreUObject/Public/UObject/UObjectArray.h
+++ b/Engine/Source/Runtime/CoreUObject/Public/UObject/UObjectArray.h
@@ -103,7 +103,7 @@ struct FUObjectItem
 
 	FORCEINLINE EInternalObjectFlags GetFlags() const
 	{
-		return EInternalObjectFlags(Flags);
+		return EInternalObjectFlags(GetFlagsInternal());
 	}
 
 	FORCEINLINE void ClearFlags(EInternalObjectFlags FlagsToClear)
@@ -123,7 +123,7 @@ struct FUObjectItem
 		bool bIChangedIt = false;
 		while (1)
 		{
-			int32 StartValue = int32(Flags);
+			int32 StartValue = GetFlagsInternal();
 			if (!(StartValue & int32(FlagToClear)))
 			{
 				break;
@@ -144,7 +144,7 @@ struct FUObjectItem
 		bool bIChangedIt = false;
 		while (1)
 		{
-			int32 StartValue = int32(Flags);
+			int32 StartValue = GetFlagsInternal();
 			if ((StartValue & int32(FlagToSet)) == int32(FlagToSet))
 			{
 				break;
@@ -161,7 +161,7 @@ struct FUObjectItem
 
 	FORCEINLINE bool HasAnyFlags(EInternalObjectFlags InFlags) const
 	{
-		return !!(Flags & int32(InFlags));
+		return !!(GetFlagsInternal() & int32(InFlags));
 	}
 
 	FORCEINLINE void SetUnreachable()
@@ -174,7 +174,7 @@ struct FUObjectItem
 	}
 	FORCEINLINE bool IsUnreachable() const
 	{
-		return !!(Flags & int32(EInternalObjectFlags::Unreachable));
+		return !!(GetFlagsInternal() & int32(EInternalObjectFlags::Unreachable));
 	}
 	FORCEINLINE bool ThisThreadAtomicallyClearedRFUnreachable()
 	{
@@ -196,7 +196,7 @@ struct FUObjectItem
 	FORCEINLINE bool IsPendingKill() const
 	{
 		PRAGMA_DISABLE_DEPRECATION_WARNINGS
-		return !!(Flags & int32(EInternalObjectFlags::PendingKill));
+		return !!(GetFlagsInternal() & int32(EInternalObjectFlags::PendingKill));
 		PRAGMA_ENABLE_DEPRECATION_WARNINGS
 	}
 
@@ -210,7 +210,7 @@ struct FUObjectItem
 	}
 	FORCEINLINE bool IsRootSet() const
 	{
-		return !!(Flags & int32(EInternalObjectFlags::RootSet));
+		return !!(GetFlagsInternal() & int32(EInternalObjectFlags::RootSet));
 	}
 
 	FORCEINLINE void ResetSerialNumberAndFlags()
@@ -223,6 +223,11 @@ struct FUObjectItem
 #if STATS || ENABLE_STATNAMEDEVENTS_UOBJECT
 	COREUOBJECT_API void CreateStatID() const;
 #endif
+private:
+	FORCEINLINE int32 GetFlagsInternal() const
+	{
+		return FPlatformAtomics::AtomicRead_Relaxed((int32*)&Flags);
+	}
 };
 
 /**
diff --git a/Engine/Source/Runtime/CoreUObject/Public/UObject/UObjectBaseUtility.h b/Engine/Source/Runtime/CoreUObject/Public/UObject/UObjectBaseUtility.h
--- a/Engine/Source/Runtime/CoreUObject/Public/UObject/UObjectBaseUtility.h
+++ b/Engine/Source/Runtime/CoreUObject/Public/UObject/UObjectBaseUtility.h
@@ -355,7 +355,7 @@ public:
 	{
 		FUObjectItem* ObjectItem = GUObjectArray.IndexToObject(InternalIndex);
 		PRAGMA_DISABLE_DEPRECATION_WARNINGS
-		checkf(!(FlagsToSet & (EInternalObjectFlags::PendingKill | EInternalObjectFlags::Garbage)) || int32(FlagsToSet & (EInternalObjectFlags::PendingKill | EInternalObjectFlags::Garbage)) == (ObjectItem->Flags & int32(EInternalObjectFlags::PendingKill | EInternalObjectFlags::Garbage)), TEXT("SetInternalFlags should not set the PendingKill or Garbage flag. Use MarkPendingKill or MarkAsGarbage instead"));
+		checkf(!(FlagsToSet & (EInternalObjectFlags::PendingKill | EInternalObjectFlags::Garbage)) || (FlagsToSet & (EInternalObjectFlags::PendingKill | EInternalObjectFlags::Garbage)) == (ObjectItem->GetFlags() & (EInternalObjectFlags::PendingKill | EInternalObjectFlags::Garbage)), TEXT("SetInternalFlags should not set the PendingKill or Garbage flag. Use MarkPendingKill or MarkAsGarbage instead"));
 		PRAGMA_ENABLE_DEPRECATION_WARNINGS
 		ObjectItem->SetFlags(FlagsToSet);
 	}
@@ -392,7 +392,7 @@ public:
 	{
 		FUObjectItem* ObjectItem = GUObjectArray.IndexToObject(InternalIndex);
 		PRAGMA_DISABLE_DEPRECATION_WARNINGS
-		checkf(!(FlagsToClear & (EInternalObjectFlags::PendingKill | EInternalObjectFlags::Garbage)) || (ObjectItem->Flags & int32(FlagsToClear & (EInternalObjectFlags::PendingKill | EInternalObjectFlags::Garbage))) == 0, TEXT("ClearInternalFlags should not clear PendingKill or Garbage flag. Use ClearGarbage() instead"));
+		checkf(!(FlagsToClear & (EInternalObjectFlags::PendingKill | EInternalObjectFlags::Garbage)) || (ObjectItem->GetFlags() & (FlagsToClear & (EInternalObjectFlags::PendingKill | EInternalObjectFlags::Garbage))) == EInternalObjectFlags::None, TEXT("ClearInternalFlags should not clear PendingKill or Garbage flag. Use ClearGarbage() instead"));
 		PRAGMA_ENABLE_DEPRECATION_WARNINGS
 		ObjectItem->ClearFlags(FlagsToClear);
 	}
@@ -407,7 +407,7 @@ public:
 	{
 		FUObjectItem* ObjectItem = GUObjectArray.IndexToObject(InternalIndex);
 		PRAGMA_DISABLE_DEPRECATION_WARNINGS
-		checkf((ObjectItem->Flags & int32(FlagsToClear & (EInternalObjectFlags::PendingKill | EInternalObjectFlags::Garbage))) == 0, TEXT("ClearInternalFlags should not clear PendingKill or Garbage flag. Use ClearGarbage() instead"));
+		checkf((ObjectItem->GetFlags() & (FlagsToClear & (EInternalObjectFlags::PendingKill | EInternalObjectFlags::Garbage))) == EInternalObjectFlags::None, TEXT("ClearInternalFlags should not clear PendingKill or Garbage flag. Use ClearGarbage() instead"));
 		PRAGMA_ENABLE_DEPRECATION_WARNINGS
 		return ObjectItem->ThisThreadAtomicallyClearedFlag(FlagsToClear);
 	}
