OnObjectConstructed delegate - make it a ThreadSafe delegate because StaticConstructObjectInternal will be callable from multiple threads in the future.
#rnx

diff --git a/Engine/Source/Runtime/CoreUObject/Public/UObject/UObjectGlobals.h b/Engine/Source/Runtime/CoreUObject/Public/UObject/UObjectGlobals.h
--- a/Engine/Source/Runtime/CoreUObject/Public/UObject/UObjectGlobals.h
+++ b/Engine/Source/Runtime/CoreUObject/Public/UObject/UObjectGlobals.h
@@ -2718,7 +2718,7 @@ struct COREUOBJECT_API FCoreUObjectDelegates
 	DECLARE_MULTICAST_DELEGATE_OneParam(FOnAssetLoaded, UObject*);
 	static FOnAssetLoaded OnAssetLoaded;
 
-	DECLARE_MULTICAST_DELEGATE_OneParam(FOnObjectConstructed, UObject*);
+	DECLARE_TS_MULTICAST_DELEGATE_OneParam(FOnObjectConstructed, UObject*);
 	static FOnObjectConstructed OnObjectConstructed;
 
 	/** Callback when packages end loading in LoadPackage or AsyncLoadPackage. All packages loaded recursively due to imports are included in the single call of the explicitly-loaded package. */
