ZenLoader: Fixed race causing event nodes to get lost

diff --git a/Engine/Source/Runtime/CoreUObject/Private/Serialization/AsyncLoading2.cpp b/Engine/Source/Runtime/CoreUObject/Private/Serialization/AsyncLoading2.cpp
--- a/Engine/Source/Runtime/CoreUObject/Private/Serialization/AsyncLoading2.cpp
+++ b/Engine/Source/Runtime/CoreUObject/Private/Serialization/AsyncLoading2.cpp
@@ -3636,7 +3636,7 @@ void FAsyncLoadEventQueue2::Push(FEventLoadNode2* Node)
 {
 	// switch `Tail` to the new node and only then link the old tail to the new one. The list is not fully linked between these ops,
 	// this is explicitly handled by the consumer by waiting for the link
-	FEventLoadNode2* Prev = Tail.exchange(Node, std::memory_order_release); // `release` to make sure the new node is fully constructed before
+	FEventLoadNode2* Prev = Tail.exchange(Node, std::memory_order_acq_rel); // sync with consumer in Deplete()
 	check(Prev->Next.load(std::memory_order_relaxed) == nullptr); // `Tail` is assigned before its Next
 	Prev->Next.store(Node, std::memory_order_relaxed);
 	if (Zenaphore)
@@ -3664,8 +3664,7 @@ void FAsyncLoadEventQueue2::Deplete()
 	// reset the queue to the empty state. this redirects producers to start from `Sentinel` again.
 	// take note of the tail on resetting it because the list can be still not fully linked and so `Node.Next == nullptr` can't be 
 	// used to detect the end of the list
-	FEventLoadNode2* Last = Tail.exchange(&Sentinel, std::memory_order_acquire); // `acquire` to sync with producers' tail modifications,
-	// `Sentinel.Next = nullptr` must happen before modifying `Tail` but it's important only for the single consumer so doesn't need extra sync
+	FEventLoadNode2* Last = Tail.exchange(&Sentinel, std::memory_order_acq_rel); // after setting Sentinel.Next to nullptr and before producers' tail modifications
 	check(Last->Next.load(std::memory_order_relaxed) == nullptr); // `Tail` is assigned before its Next
 	// the previously queued items are detached from the instance (as a linked list, though potentially not fully linked yet)
 
