Bulk Data - set correct status when issuing batch request

#hardlock Mark.Lintott
#rnx

diff --git a/Engine/Source/Runtime/CoreUObject/Private/Serialization/BulkDataStreaming.cpp b/Engine/Source/Runtime/CoreUObject/Private/Serialization/BulkDataStreaming.cpp
--- a/Engine/Source/Runtime/CoreUObject/Private/Serialization/BulkDataStreaming.cpp
+++ b/Engine/Source/Runtime/CoreUObject/Private/Serialization/BulkDataStreaming.cpp
@@ -1061,7 +1061,12 @@ public:
 
 	virtual bool Cancel() override
 	{
-		if (FBulkDataRequest::EStatus::Pending == GetStatus())
+		if (FBulkDataRequest::EStatus::None == GetStatus())
+		{
+			CompleteBatch(FBulkDataRequest::EStatus::Cancelled);
+			return true;
+		}
+		else if (FBulkDataRequest::EStatus::Pending == GetStatus())
 		{
 			for (FChunkBatchReadRequest& Request : Requests)
 			{
@@ -1096,14 +1101,17 @@ public:
 
 	virtual bool Issue(FBulkDataRequest::FCompletionCallback&& Callback) override
 	{
+		CompletionCallback = MoveTemp(Callback);
+
 		if (Requests.IsEmpty())
 		{
+			CompleteBatch(FBulkDataRequest::EStatus::Ok);
 			return false;
 		}
 
 		TRACE_COUNTER_INCREMENT(BulkDataBatchRequest_PendingCount);
 
-		CompletionCallback = MoveTemp(Callback);
+		SetStatus(FBulkDataRequest::EStatus::Pending);
 
 		IoBatch.IssueWithCallback([this]()
 		{
@@ -1244,12 +1252,16 @@ public:
 
 	virtual bool Issue(FBulkDataRequest::FCompletionCallback&& Callback) override
 	{
+		CompletionCallback = MoveTemp(Callback);
+
 		if (Requests.IsEmpty())
 		{
+			CompleteBatch(FBulkDataRequest::EStatus::Ok);
 			return false;
 		}
 
-		CompletionCallback = MoveTemp(Callback);
+		SetStatus(FBulkDataRequest::EStatus::Pending);
+
 		PendingRequestCount.store(Requests.Num(), std::memory_order_relaxed);
 
 		for (FRequestParams& Params : RequestParams)
@@ -1258,6 +1270,7 @@ public:
 
 			if (FileHandle == nullptr)
 			{
+				CompleteBatch(FBulkDataRequest::EStatus::Error);
 				return false;
 			}
 
