Use the deferred post load step for every objects when async loading without the loading thread enabled
This allows proper handling of objects sync loaded through post load.

diff --git a/Engine/Source/Runtime/CoreUObject/Private/Serialization/AsyncLoading.cpp b/Engine/Source/Runtime/CoreUObject/Private/Serialization/AsyncLoading.cpp
--- a/Engine/Source/Runtime/CoreUObject/Private/Serialization/AsyncLoading.cpp
+++ b/Engine/Source/Runtime/CoreUObject/Private/Serialization/AsyncLoading.cpp
@@ -6903,7 +6903,7 @@ EAsyncPackageState::Type FAsyncPackage::PostLoadObjects()
 				--PostLoadIndex;
 				break;
 			}
-			else if (/*!bIsMultithreaded ||*/ (bAsyncPostLoadEnabled && CanPostLoadOnAsyncLoadingThread(Object)))
+			else if (bAsyncPostLoadEnabled && CanPostLoadOnAsyncLoadingThread(Object))
 			{
 				SCOPED_ACCUM_LOADTIME(PostLoad, StaticGetNativeClassName(Object->GetClass()));
 
