AsyncLoading2: Add trace counters for tracking number of package loads

#rnx

diff --git a/Engine/Source/Runtime/CoreUObject/Private/Serialization/AsyncLoading2.cpp b/Engine/Source/Runtime/CoreUObject/Private/Serialization/AsyncLoading2.cpp
--- a/Engine/Source/Runtime/CoreUObject/Private/Serialization/AsyncLoading2.cpp
+++ b/Engine/Source/Runtime/CoreUObject/Private/Serialization/AsyncLoading2.cpp
@@ -327,6 +327,9 @@ static FAutoConsoleVariableRef CVarGOnlyProcessRequiredPackagesWhenSyncLoading(
 CSV_DECLARE_CATEGORY_MODULE_EXTERN(CORE_API, Basic);
 CSV_DECLARE_CATEGORY_MODULE_EXTERN(CORE_API, FileIO);
 
+TRACE_DECLARE_INT_COUNTER(AsyncLoadingQueuedPackages, TEXT("AsyncLoading/PackagesQueued"));
+TRACE_DECLARE_INT_COUNTER(AsyncLoadingLoadingPackages, TEXT("AsyncLoading/PackagesLoading"));
+TRACE_DECLARE_INT_COUNTER(AsyncLoadingPackagesWithRemainingWork, TEXT("AsyncLoading/PackagesWithRemainingWork"));
 TRACE_DECLARE_INT_COUNTER(AsyncLoadingPendingIoRequests, TEXT("AsyncLoading/PendingIoRequests"));
 TRACE_DECLARE_MEMORY_COUNTER(AsyncLoadingTotalLoaded, TEXT("AsyncLoading/TotalLoaded"));
 
@@ -3135,6 +3138,7 @@ private:
 		}
 		delete Package;
 		--PackagesWithRemainingWorkCounter;
+		TRACE_COUNTER_SET(AsyncLoadingPackagesWithRemainingWork, PackagesWithRemainingWorkCounter);
 	}
 
 	/** Number of times we re-entered the async loading tick, mostly used by singlethreaded ticking. Debug purposes only. */
@@ -3248,6 +3252,7 @@ FAsyncPackage2* FAsyncLoadingThread2::FindOrInsertPackage(FAsyncPackageDesc2& De
 			checkf(Package, TEXT("Failed to create async package %s"), *Desc.UPackageName.ToString());
 			Package->AddRef();
 			++LoadingPackagesCounter;
+			TRACE_COUNTER_SET(AsyncLoadingLoadingPackages, LoadingPackagesCounter);
 			AsyncPackageLookup.Add(Desc.UPackageId, Package);
 			bInserted = true;
 		}
@@ -3328,6 +3333,7 @@ bool FAsyncLoadingThread2::CreateAsyncPackagesFromQueue(FAsyncLoadingThreadState
 
 			--QueuedPackagesCounter;
 			++NumDequeued;
+			TRACE_COUNTER_SET(AsyncLoadingQueuedPackages, QueuedPackagesCounter);
 		
 			FPackageRequest& Request = OptionalRequest.GetValue();
 			EPackageStoreEntryStatus PackageStatus = EPackageStoreEntryStatus::Missing;
@@ -3445,6 +3451,7 @@ bool FAsyncLoadingThread2::CreateAsyncPackagesFromQueue(FAsyncLoadingThreadState
 					UE_ASYNC_PACKAGE_LOG_VERBOSE(Verbose, PackageDesc, TEXT("CreateAsyncPackages: UpdatePackage"),
 						TEXT("Package is alreay being loaded."));
 					--PackagesWithRemainingWorkCounter;
+					TRACE_COUNTER_SET(AsyncLoadingPackagesWithRemainingWork, PackagesWithRemainingWorkCounter);
 				}
 
 				RequestIdToPackageMap.Add(PackageDesc.RequestID, Package);
@@ -4128,6 +4135,7 @@ void FAsyncPackage2::ImportPackagesRecursiveInner(FAsyncLoadingThreadState2& Thr
 			UE_ASYNC_PACKAGE_LOG(Verbose, Desc, TEXT("ImportPackages: AddPackage"),
 			TEXT("Start loading imported package with id '0x%llX'"), ImportedPackageId.ValueForDebugging());
 			++AsyncLoadingThread.PackagesWithRemainingWorkCounter;
+			TRACE_COUNTER_SET(AsyncLoadingPackagesWithRemainingWork, AsyncLoadingThread.PackagesWithRemainingWorkCounter);
 		}
 		else
 		{
@@ -6324,6 +6332,7 @@ EAsyncPackageState::Type FAsyncLoadingThread2::ProcessLoadedPackagesFromGameThre
 			// Incremented on the Async Thread, now decrement as we're done with this package
 			--LoadingPackagesCounter;
 
+			TRACE_COUNTER_SET(AsyncLoadingLoadingPackages, LoadingPackagesCounter);
 			TRACE_LOADTIME_END_LOAD_ASYNC_PACKAGE(Package);
 
 			PackagesReadyForCallback.Add(Package);
@@ -6372,6 +6381,7 @@ EAsyncPackageState::Type FAsyncLoadingThread2::ProcessLoadedPackagesFromGameThre
 				FailedPackageRequest.Callback->ExecuteIfBound(FailedPackageRequest.PackageName, nullptr, EAsyncLoadingResult::Failed);
 				RemovePendingRequests(TArrayView<int32>(&FailedPackageRequest.RequestID, 1));
 				--PackagesWithRemainingWorkCounter;
+				TRACE_COUNTER_SET(AsyncLoadingPackagesWithRemainingWork, PackagesWithRemainingWorkCounter);
 			}
 		}
 
@@ -7584,6 +7594,9 @@ int32 FAsyncLoadingThread2::LoadPackage(const FPackagePath& InPackagePath, FName
 	++QueuedPackagesCounter;
 	++PackagesWithRemainingWorkCounter;
 
+	TRACE_COUNTER_SET(AsyncLoadingQueuedPackages, QueuedPackagesCounter);
+	TRACE_COUNTER_SET(AsyncLoadingPackagesWithRemainingWork, PackagesWithRemainingWorkCounter);
+
 	AltZenaphore.NotifyOne();
 
 	return RequestId;
@@ -7619,6 +7632,7 @@ void FAsyncLoadingThread2::QueueMissingPackage(FAsyncPackageDesc2& PackageDesc,
 	{
 		RemovePendingRequests(TArrayView<int32>(&PackageDesc.RequestID, 1));
 		--PackagesWithRemainingWorkCounter;
+		TRACE_COUNTER_SET(AsyncLoadingPackagesWithRemainingWork, PackagesWithRemainingWorkCounter);
 	}
 }
 
