#UE Fix load error due to relocation logic. Imports are not always relocated in the same way

diff --git a/Engine/Source/Runtime/CoreUObject/Private/UObject/LinkerLoad.cpp b/Engine/Source/Runtime/CoreUObject/Private/UObject/LinkerLoad.cpp
--- a/Engine/Source/Runtime/CoreUObject/Private/UObject/LinkerLoad.cpp
+++ b/Engine/Source/Runtime/CoreUObject/Private/UObject/LinkerLoad.cpp
@@ -3240,8 +3240,7 @@ bool FLinkerLoad::VerifyImportInner(const int32 ImportIndex, FString& WarningSuf
 		FUObjectSerializeContext* SerializeContext = GetSerializeContext();
 
 		// Resolve the package name for the import, potentially remapping it, if instancing
-		FName OriginalPackageToLoad = !Import.HasPackageName() ? Import.ObjectName : Import.GetPackageName();
-		FName PackageToLoad = InstancingContext.RelocatePackage(OriginalPackageToLoad);
+		FName PackageToLoad = !Import.HasPackageName() ? Import.ObjectName : Import.GetPackageName();
 
 		FName PackageToLoadInto = InstancingContext.RemapPackage(PackageToLoad);
 #if WITH_EDITOR
@@ -3331,12 +3330,14 @@ bool FLinkerLoad::VerifyImportInner(const int32 ImportIndex, FString& WarningSuf
 		// to be linked to any other package's ImportMaps
 		if (!Package || Package->HasAnyPackageFlags(PKG_Compiling))
 		{
+			/*
 			FName RelocatedName;
 			if (InstancingContext.RelocatedPackageMapping.RemoveAndCopyValue(OriginalPackageToLoad, RelocatedName))
 			{
 				UE_ASSET_LOG(LogLinker, Warning, PackagePath, TEXT("VerifyImport: Failed to load package for import object '%s' through using relative path '%s'. Using orginal location as fallback."), *GetImportFullName(ImportIndex), *RelocatedName.ToString());
 				return nullptr;
 			}
+			*/
 
 			if (!FLinkerLoad::IsKnownMissingPackage(PackageToLoad))
 			{
