Moved UObject hasmaps re/allocations from GC and UObject into separate LLM category.

#tests Compile/Run/Collect utrace+stacks/look LLM allocations changes

diff --git a/Engine/Source/Runtime/CoreUObject/Private/UObject/UObjectHash.cpp b/Engine/Source/Runtime/CoreUObject/Private/UObject/UObjectHash.cpp
--- a/Engine/Source/Runtime/CoreUObject/Private/UObject/UObjectHash.cpp
+++ b/Engine/Source/Runtime/CoreUObject/Private/UObject/UObjectHash.cpp
@@ -11,6 +11,7 @@
 #include "Misc/PackageName.h"
 #include "Async/ParallelFor.h"
 #include "HAL/IConsoleManager.h"
+#include "HAL/LowLevelMemStats.h"
 #include "UObject/AnyPackagePrivate.h"
 
 DEFINE_LOG_CATEGORY_STATIC(LogUObjectHash, Log, All);
@@ -26,6 +27,8 @@ DECLARE_CYCLE_STAT( TEXT( "UnhashObject" ), STAT_Hash_UnhashObject, STATGROUP_UO
 DEFINE_STAT( STAT_Hash_NumObjects );
 #endif
 
+LLM_DEFINE_TAG(UObjectHash, TEXT("UObject hashtables"));
+
 // Global UObject array instance
 FUObjectArray GUObjectArray;
 
@@ -1056,6 +1059,7 @@ FORCEINLINE static UPackage* UnassignExternalPackageFromObject(FUObjectHashTable
 
 void ShrinkUObjectHashTables()
 {
+	LLM_SCOPE_BYTAG(UObjectHash);
 	TRACE_CPUPROFILER_EVENT_SCOPE(ShrinkUObjectHashTables);
 	FUObjectHashTables& ThreadHash = FUObjectHashTables::Get();
 	FHashTableLock HashLock(ThreadHash);
@@ -1470,6 +1474,7 @@ void HashObject(UObjectBase* Object)
 	FName Name = Object->GetFName();
 	if (Name != NAME_None)
 	{
+		LLM_SCOPE_BYTAG(UObjectHash);
 #if !UE_BUILD_TEST && !UE_BUILD_SHIPPING
 		SCOPE_CYCLE_COUNTER(STAT_Hash_HashObject);
 #endif
@@ -1512,6 +1517,7 @@ void UnhashObject(UObjectBase* Object)
 	FName Name = Object->GetFName();
 	if (Name != NAME_None)
 	{
+		LLM_SCOPE_BYTAG(UObjectHash);
 #if !UE_BUILD_TEST && !UE_BUILD_SHIPPING
 		SCOPE_CYCLE_COUNTER(STAT_Hash_UnhashObject);
 #endif
@@ -1544,6 +1550,7 @@ void UnhashObject(UObjectBase* Object)
 
 void HashObjectExternalPackage(UObjectBase* Object, UPackage* Package)
 {
+	LLM_SCOPE_BYTAG(UObjectHash);
 	if (Package)
 	{
 		FUObjectHashTables& ThreadHash = FUObjectHashTables::Get();
@@ -1566,6 +1573,7 @@ void HashObjectExternalPackage(UObjectBase* Object, UPackage* Package)
 
 void UnhashObjectExternalPackage(class UObjectBase* Object)
 {
+	LLM_SCOPE_BYTAG(UObjectHash);
 	FUObjectHashTables& ThreadHash = FUObjectHashTables::Get();
 	FHashTableLock LockHash(ThreadHash);
 	UPackage* Package = UnassignExternalPackageFromObject(ThreadHash, Object);
