FAssetData - Adding a templated version of IsInstanceOf, and adding some additional deprecation comments to ObjectPath.

diff --git a/Engine/Source/Runtime/CoreUObject/Public/AssetRegistry/AssetData.h b/Engine/Source/Runtime/CoreUObject/Public/AssetRegistry/AssetData.h
--- a/Engine/Source/Runtime/CoreUObject/Public/AssetRegistry/AssetData.h
+++ b/Engine/Source/Runtime/CoreUObject/Public/AssetRegistry/AssetData.h
@@ -156,7 +156,7 @@ public:
 public:
 #if WITH_EDITORONLY_DATA
 	/** The object path for the asset in the form PackageName.ObjectName, or PackageName.ObjectName:SubObjectName. */
-	UE_DEPRECATED(5.1, "FName asset paths have been deprecated. Use GetSoftObjectPath to get the path this asset will use in memory when loaded.")
+	UE_DEPRECATED(5.1, "FName asset paths have been deprecated. Use GetSoftObjectPath to get the path this asset will use in memory when loaded or GetObjectPathString() if you were just doing ObjectPath.ToString()")
 	FName ObjectPath;
 #endif
 	/** The name of the package in which the asset is found, this is the full long package name such as /Game/Path/Package */
@@ -453,6 +453,12 @@ PRAGMA_ENABLE_DEPRECATION_WARNINGS
 		return ClassPointer && ClassPointer->IsChildOf(BaseClass);
 	}
 
+	template<typename BaseClass>
+	bool IsInstanceOf() const
+	{
+		return IsInstanceOf(BaseClass::StaticClass());
+	}
+
 	/** Append the object path to the given string builder. */
 	void AppendObjectPath(FStringBuilderBase& Builder) const
 	{
@@ -471,9 +477,8 @@ PRAGMA_ENABLE_DEPRECATION_WARNINGS
 		{
 			Builder << PackageName << '.' << AssetName;
 		}
-
-	}
-
+    }
+	
 	/** Append the object path to the given string. */
 	void AppendObjectPath(FString& String) const
 	{
