Bulk Data - I/O dispatcher only streaming

- This change removes the explicit usage of the package resource manager when loading bulk data from loose files or the Editor Domain. A new
I/O dispatcher backend reads data from the package resource manager with a new I/O chunk type constructed from a package name/package segment. When
serializing bulk data the I/O chunk type is set to the correct chunk type depending on whether the package is cooked or not.

diff --git a/Engine/Source/Runtime/Core/Private/IO/IoStore.cpp b/Engine/Source/Runtime/Core/Private/IO/IoStore.cpp
--- a/Engine/Source/Runtime/Core/Private/IO/IoStore.cpp
+++ b/Engine/Source/Runtime/Core/Private/IO/IoStore.cpp
@@ -53,6 +53,7 @@ FString LexToString(const EIoChunkType Type)
 		TEXT("PackageStoreEntry"),
 		TEXT("DerivedData"),
 		TEXT("EditorDerivedData"),
+		TEXT("PackageResource"),
 	};
 	static_assert(UE_ARRAY_COUNT(Strings) == (SIZE_T)EIoChunkType::MAX);
 	uint8 Index = (uint8)Type;
diff --git a/Engine/Source/Runtime/Core/Public/IO/IoDispatcher.h b/Engine/Source/Runtime/Core/Public/IO/IoDispatcher.h
--- a/Engine/Source/Runtime/Core/Public/IO/IoDispatcher.h
+++ b/Engine/Source/Runtime/Core/Public/IO/IoDispatcher.h
@@ -681,7 +681,7 @@ enum class EIoChunkType : uint8
 	PackageStoreEntry = 10,
 	DerivedData = 11,
 	EditorDerivedData = 12,
-	
+	PackageResource = 13,
 	MAX
 };
 
diff --git a/Engine/Source/Runtime/CoreUObject/Private/Serialization/BulkData.cpp b/Engine/Source/Runtime/CoreUObject/Private/Serialization/BulkData.cpp
--- a/Engine/Source/Runtime/CoreUObject/Private/Serialization/BulkData.cpp
+++ b/Engine/Source/Runtime/CoreUObject/Private/Serialization/BulkData.cpp
@@ -3,7 +3,6 @@
 #include "Async/MappedFileHandle.h"
 #include "HAL/IConsoleManager.h"
 #include "Misc/ConfigCacheIni.h"
-#include "Misc/TVariant.h"
 #include "ProfilingDebugging/LoadTimeTracker.h"
 #include "Serialization/LargeMemoryWriter.h"
 #include "Serialization/LargeMemoryReader.h"
@@ -12,6 +11,7 @@
 #include "UObject/LinkerLoad.h"
 #include "UObject/LinkerSave.h"
 #include "UObject/PackageResourceManager.h"
+#include "UObject/PackageResourceIoDispatcherBackend.h"
 
 /** Whether to track information of how bulk data is being used */
 #define TRACK_BULKDATA_USE 0
@@ -76,19 +76,19 @@ namespace UE::BulkData::Private
 /** Open the bulk data chunk fo reading. */
 bool OpenReadBulkData(
 	const FBulkMetaData& BulkMeta,
-	const FBulkDataChunkId& BulkChunkId,
+	const FIoChunkId& BulkChunkId,
 	int64 Offset,
 	int64 Size,
 	EAsyncIOPriorityAndFlags Priority,
 	TFunction<void(FArchive& Ar)>&& Read);
 
 /** Open async read file handle for the specified bulk data chunk ID. */
-TUniquePtr<IAsyncReadFileHandle> OpenAsyncReadBulkData(const FBulkMetaData& BulkMeta, const FBulkDataChunkId& BulkChunkId);
+TUniquePtr<IAsyncReadFileHandle> OpenAsyncReadBulkData(const FBulkMetaData& BulkMeta, const FIoChunkId& BulkChunkId);
 
 /** Create bulk data streaming request. */
 TUniquePtr<IBulkDataIORequest> CreateStreamingRequest(
 	const FBulkMetaData& BulkMeta,
-	const FBulkDataChunkId& BulkChunkId,
+	const FIoChunkId& BulkChunkId,
 	int64 Offset,
 	int64 Size,
 	EAsyncIOPriorityAndFlags Priority,
@@ -96,12 +96,12 @@ TUniquePtr<IBulkDataIORequest> CreateStreamingRequest(
 	uint8* UserSuppliedMemory);
 
 /** Returns whether the bulk data chunk exist or not. */
-bool DoesBulkDataExist(const FBulkMetaData& BulkMeta, const FBulkDataChunkId& BulkChunkId);
+bool DoesBulkDataExist(const FIoChunkId& BulkChunkId);
 
 /** Try memory map the chunk specified by the bulk data ID. */
 bool TryMemoryMapBulkData(
 	const FBulkMetaData& BulkMeta,
-	const FBulkDataChunkId& BulkChunkId,
+	const FIoChunkId& BulkChunkId,
 	int64 Offset,
 	int64 Size,
 	FIoMappedRegion& OutRegion);
@@ -110,7 +110,7 @@ bool TryMemoryMapBulkData(
 bool StartAsyncLoad(
 	FBulkData* Owner,
 	const FBulkMetaData& BulkMeta,
-	const FBulkDataChunkId& BulkChunkId,
+	const FIoChunkId& BulkChunkId,
 	int64 Offset,
 	int64 Size,
 	EAsyncIOPriorityAndFlags Priority,
@@ -187,154 +187,6 @@ FArchive& operator<<(FArchive& Ar, FBulkMetaResource& BulkMeta)
 	return Ar;
 }
 
-//////////////////////////////////////////////////////////////////////////////
-
-struct FBulkDataChunkId::FImpl
-{
-	using FPathOrId = TVariant<FPackagePath, FPackageId>;
-
-	FImpl(const FPackageId& PackageId)
-		: PathOrId(TInPlaceType<FPackageId>(), PackageId)
-	{ }
-
-	FImpl(const FPackagePath& PackagePath)
-		: PathOrId(TInPlaceType<FPackagePath>(), PackagePath)
-	{ }
-
-	FImpl(const FPathOrId& InPathOrId)
-		: PathOrId(InPathOrId)
-	{ }
-
-	FPathOrId PathOrId;
-};
-
-FBulkDataChunkId::FBulkDataChunkId()
-{
-}
-
-FBulkDataChunkId::FBulkDataChunkId(TPimplPtr<FImpl>&& InImpl)
-	: Impl(MoveTemp(InImpl))
-{
-}
-
-FBulkDataChunkId::~FBulkDataChunkId()
-{
-}
-
-FBulkDataChunkId::FBulkDataChunkId(const FBulkDataChunkId& Other)
-{
-	*this = Other;
-}
-
-FBulkDataChunkId& FBulkDataChunkId::operator=(const FBulkDataChunkId& Other)
-{
-	Impl.Reset();
-
-	if (Other.Impl)
-	{
-		Impl = MakePimpl<FImpl>(Other.Impl->PathOrId);
-	}
-
-	return *this;
-}
-
-bool FBulkDataChunkId::operator==(const FBulkDataChunkId& Other) const
-{
-	if (Impl.IsValid() && Other.Impl.IsValid())
-	{
-		if (const FPackageId* Id = Impl->PathOrId.TryGet<FPackageId>())
-		{
-			if (const FPackageId* OtherId = Other.Impl->PathOrId.TryGet<FPackageId>())
-			{
-				return *Id == *OtherId;
-			}
-		}
-
-		if (const FPackagePath* PackagePath = Impl->PathOrId.TryGet<FPackagePath>())
-		{
-			if (const FPackagePath* OtherPackagePath = Other.Impl->PathOrId.TryGet<FPackagePath>())
-			{
-				return *PackagePath == *OtherPackagePath;
-			}
-		}
-	}
-
-	return Impl.IsValid() == Other.Impl.IsValid();
-}
-
-FPackageId FBulkDataChunkId::GetPackageId() const
-{
-	if (Impl)
-	{
-		if (const FPackageId* Id = Impl->PathOrId.TryGet<FPackageId>())
-		{
-			return *Id;
-		}
-	}
-
-	return FPackageId();
-}
-
-const FPackagePath& FBulkDataChunkId::GetPackagePath() const
-{
-	if (Impl)
-	{
-		if (const FPackagePath* Path = Impl->PathOrId.TryGet<FPackagePath>())
-		{
-			return *Path;
-		}
-	}
-
-	static FPackagePath Empty;
-	return Empty;
-}
-
-FIoFilenameHash FBulkDataChunkId::GetIoFilenameHash(EBulkDataFlags BulkDataFlags) const
-{
-	if (Impl)
-	{
-		if (const FPackageId* Id = Impl->PathOrId.TryGet<FPackageId>())
-		{
-			return MakeIoFilenameHash(CreateBulkDataIoChunkId(UE::BulkData::Private::FBulkMetaData(BulkDataFlags), *Id));
-		}
-
-		if (const FPackagePath* PackagePath = Impl->PathOrId.TryGet<FPackagePath>())
-		{
-			return MakeIoFilenameHash(*PackagePath);
-		}
-	}
-
-	return INVALID_IO_FILENAME_HASH;
-}
-
-FString FBulkDataChunkId::ToDebugString() const
-{
-	if (Impl)
-	{
-		if (const FPackageId* Id = Impl->PathOrId.TryGet<FPackageId>())
-		{
-			return FString::Printf(TEXT("0x%llX"), Id->Value());
-		}
-
-		if (const FPackagePath* PackagePath = Impl->PathOrId.TryGet<FPackagePath>())
-		{
-			return PackagePath->GetDebugName();
-		}
-	}
-
-	return FString(TEXT("None"));
-}
-
-FBulkDataChunkId FBulkDataChunkId::FromPackagePath(const FPackagePath& PackagePath)
-{
-	return FBulkDataChunkId(MakePimpl<FImpl>(PackagePath));
-}
-
-FBulkDataChunkId FBulkDataChunkId::FromPackageId(const FPackageId& PackageId)
-{
-	return FBulkDataChunkId(MakePimpl<FImpl>(PackageId));
-}
-
 } // namespace UE::BulkData
 
 /*-----------------------------------------------------------------------------
@@ -773,7 +625,7 @@ bool FBulkData::DoesExist() const
 	}
 #endif
 
-	return UE::BulkData::Private::DoesBulkDataExist(BulkMeta, BulkChunkId);
+	return UE::BulkData::Private::DoesBulkDataExist(BulkChunkId);
 }
 
 /**
@@ -1083,7 +935,7 @@ bool FBulkData::LoadBulkDataWithFileReader()
 bool FBulkData::CanLoadBulkDataWithFileReader() const
 {
 #if WITH_EDITOR
-	return BulkChunkId.GetPackagePath().IsEmpty() == false;
+	return BulkChunkId.IsValid();
 #else
 	return false;
 #endif
@@ -1153,7 +1005,7 @@ bool FBulkData::StartAsyncLoading()
 		}
 		else
 		{
-			UE_LOG(LogSerialization, Error, TEXT("Async load bulk data '%s' FAILED, reason '%s'"), *BulkChunkId.ToDebugString(), GetIoErrorText(Result.Status().GetErrorCode()));
+			UE_LOG(LogSerialization, Error, TEXT("Async load bulk data '%s' FAILED, reason '%s'"), *LexToString(BulkChunkId), GetIoErrorText(Result.Status().GetErrorCode()));
 			RemoveBulkData();
 		}
 
@@ -1219,12 +1071,12 @@ void FBulkData::ClearBulkDataFlags( uint32 BulkDataFlagsToClear )
 
 FIoChunkId FBulkData::CreateChunkId() const
 {
-	return UE::BulkData::Private::CreateBulkDataIoChunkId(BulkMeta, BulkChunkId.GetPackageId());
+	return BulkChunkId; 
 }
 
 FString FBulkData::GetDebugName() const
 {
-	return BulkChunkId.ToDebugString();
+	return LexToString(BulkChunkId);
 }
 
 /*-----------------------------------------------------------------------------
@@ -1375,7 +1227,6 @@ void FBulkData::Serialize(FArchive& Ar, UObject* Owner, bool bAttemptFileMapping
 					check(IsInSeparateFile());
 					check(Package != nullptr);
 					BulkMeta.SetFlags(EBulkDataFlags(BulkMeta.GetFlags() | BULKDATA_UsesIoDispatcher));
-					BulkChunkId = FBulkDataChunkId::FromPackageId(Package->GetPackageIdToLoad());
 				}
 			}
 			else
@@ -1386,11 +1237,6 @@ void FBulkData::Serialize(FArchive& Ar, UObject* Owner, bool bAttemptFileMapping
 					{
 						LinkerLoad = FLinkerLoad::FindExistingLinkerForPackage(Package);
 					}
-					
-					if (LinkerLoad != nullptr)
-					{
-						BulkChunkId = FBulkDataChunkId::FromPackagePath(LinkerLoad->GetPackagePath());
-					}
 #if WITH_EDITOR
 					Linker = LinkerLoad;
 					if (bCanLazyLoad)
@@ -1403,6 +1249,8 @@ void FBulkData::Serialize(FArchive& Ar, UObject* Owner, bool bAttemptFileMapping
 				}
 			}
 
+			BulkChunkId = CreateBulkDataIoChunkId(BulkMeta, Owner);
+
 			if (IsInlined())
 			{
 				checkf(bAttemptFileMapping == false || bIsUsingIoDispatcher == false, TEXT("Trying to memory map inline bulk data '%s' which is not supported when loading from I/O store"), *Owner->GetPackage()->GetFName().ToString());
@@ -1428,13 +1276,15 @@ void FBulkData::Serialize(FArchive& Ar, UObject* Owner, bool bAttemptFileMapping
 
 					if (IsDuplicateNonOptional())
 					{
-						const bool bOptionalDataExist = DoesBulkDataExist(FBulkMetaData(BULKDATA_OptionalPayload), BulkChunkId);
+						const EBulkDataFlags OptionalFlags = static_cast<EBulkDataFlags>(BULKDATA_PayloadAtEndOfFile | BULKDATA_PayloadInSeperateFile | BULKDATA_OptionalPayload);
+						const FIoChunkId OptionalChunkId = CreateBulkDataIoChunkId(FBulkMetaData(OptionalFlags, BulkMeta.GetMetaFlags()), Owner);
 
-						if (bOptionalDataExist)
+						if (DoesBulkDataExist(OptionalChunkId))
 						{
 							BulkMeta.SetFlags(static_cast<EBulkDataFlags>((MetaResource.DuplicateFlags & ~BULKDATA_DuplicateNonOptionalPayload) | BULKDATA_OptionalPayload | BULKDATA_PayloadInSeperateFile | BULKDATA_PayloadAtEndOfFile));
 							BulkMeta.SetOffset(MetaResource.DuplicateOffset);
 							BulkMeta.SetSizeOnDisk(MetaResource.DuplicateSizeOnDisk);
+							BulkChunkId = OptionalChunkId;
 
 							if (bIsUsingIoDispatcher)
 							{
@@ -1453,10 +1303,14 @@ void FBulkData::Serialize(FArchive& Ar, UObject* Owner, bool bAttemptFileMapping
 					bool bFileMappingFailed = false;
 					if (bAttemptFileMapping && (bCanLazyLoad || bIsUsingIoDispatcher))
 					{
+						const EBulkDataFlags MemoryMappedFlags = static_cast<EBulkDataFlags>(BULKDATA_PayloadAtEndOfFile | BULKDATA_PayloadInSeperateFile | BULKDATA_MemoryMappedPayload);
+						const FIoChunkId MmapChunkId = CreateBulkDataIoChunkId(FBulkMetaData(MemoryMappedFlags, BulkMeta.GetMetaFlags()), Owner);
+
 						FIoMappedRegion MappedRegion;
-						if (TryMemoryMapBulkData(BulkMeta, BulkChunkId, BulkMeta.GetOffset(), BulkMeta.GetSize(), MappedRegion))
+						if (TryMemoryMapBulkData(BulkMeta, MmapChunkId, BulkMeta.GetOffset(), BulkMeta.GetSize(), MappedRegion))
 						{
 							DataAllocation.SetMemoryMappedData(this, MappedRegion.MappedFileHandle, MappedRegion.MappedFileRegion);
+							BulkChunkId = MmapChunkId;
 						}
 						else
 						{
@@ -1466,7 +1320,7 @@ void FBulkData::Serialize(FArchive& Ar, UObject* Owner, bool bAttemptFileMapping
 
 					if (bFileMappingFailed || (bCanLazyLoad == false && bIsUsingIoDispatcher == false))
 					{
-						UE_CLOG(bFileMappingFailed, LogSerialization, Warning, TEXT("Memory map bulk data '%s' FAILED"), *BulkChunkId.ToDebugString());
+						UE_CLOG(bFileMappingFailed, LogSerialization, Warning, TEXT("Memory map bulk data '%s' FAILED"), *LexToString(BulkChunkId));
 						ForceBulkDataResident();
 					}
 				}
@@ -1719,42 +1573,34 @@ FCustomVersionContainer FBulkData::GetCustomVersions(FArchive& InlineArchive) co
 void FBulkData::GetBulkDataVersions(FArchive& InlineArchive, FPackageFileVersion& OutUEVersion,
 	int32& OutLicenseeUEVersion, FCustomVersionContainer& OutCustomVersions) const
 {
-	if (!IsInSeparateFile())
-	{
-		OutUEVersion = InlineArchive.UEVer();
-		OutLicenseeUEVersion = InlineArchive.LicenseeUEVer();
-		OutCustomVersions = InlineArchive.GetCustomVersions();
-	}
-	else if (!IsInExternalResource())
+	FName PackageName;
+	EPackageSegment Segment;
+	bool bExternal;
+
+	if (UE::TryGetPackageNameFromChunkId(BulkChunkId, PackageName, Segment, bExternal) == false)
 	{
-		// The BulkData is in a sidecar file. These files were created with the same custom versions that
-		// the package file containing the BulkData used
-		OutUEVersion = InlineArchive.UEVer();
-		OutLicenseeUEVersion = InlineArchive.LicenseeUEVer();
-		OutCustomVersions = InlineArchive.GetCustomVersions();
+		return;
 	}
-	else
+
+	IPackageResourceManager& ResourceMgr = IPackageResourceManager::Get();
+
+	if (TUniquePtr<FArchive> Ar = ResourceMgr.OpenReadExternalResource(EPackageExternalResource::WorkspaceDomainFile, PackageName.ToString()))
 	{
-		// Read the CustomVersions out of the separate package file
-		const FPackagePath& PackagePath = BulkChunkId.GetPackagePath();
-		TUniquePtr<FArchive> ExternalArchive = IPackageResourceManager::Get().OpenReadExternalResource(
-			EPackageExternalResource::WorkspaceDomainFile, PackagePath.GetPackageName());
-		if (ExternalArchive.IsValid())
+		FPackageFileSummary Summary;
+		*Ar << Summary;
+
+		if (Ar->IsError() == false && Summary.Tag == PACKAGE_FILE_TAG)
 		{
-			FPackageFileSummary PackageFileSummary;
-			*ExternalArchive << PackageFileSummary;
-			if (PackageFileSummary.Tag == PACKAGE_FILE_TAG && !ExternalArchive->IsError())
-			{
-				OutUEVersion = PackageFileSummary.GetFileVersionUE();
-				OutLicenseeUEVersion = PackageFileSummary.GetFileVersionLicenseeUE();
-				OutCustomVersions = PackageFileSummary.GetCustomVersionContainer();
-				return;
-			}
+			OutUEVersion = Summary.GetFileVersionUE();
+			OutLicenseeUEVersion = Summary.GetFileVersionLicenseeUE();
+			OutCustomVersions = Summary.GetCustomVersionContainer();
+			return;
 		}
-		OutUEVersion = InlineArchive.UEVer();
-		OutLicenseeUEVersion = InlineArchive.LicenseeUEVer();
-		OutCustomVersions = InlineArchive.GetCustomVersions();
 	}
+
+	OutUEVersion = InlineArchive.UEVer();
+	OutLicenseeUEVersion = InlineArchive.LicenseeUEVer();
+	OutCustomVersions = InlineArchive.GetCustomVersions();
 }
 
 /*-----------------------------------------------------------------------------
@@ -1784,7 +1630,7 @@ void FBulkData::DetachFromArchive( FArchive* Ar, bool bEnsureBulkDataIsLoaded )
 	// Detach from archive.
 	AttachedAr = nullptr;
 	Linker = nullptr;
-	BulkChunkId = UE::BulkData::Private::FBulkDataChunkId();
+	BulkChunkId = FIoChunkId::InvalidChunkId;
 }
 #endif // WITH_EDITOR
 
@@ -1889,12 +1735,12 @@ IBulkDataIORequest* FBulkData::CreateStreamingRequestForRange(const BulkDataRang
 	checkf(
 		Start.IsUsingIODispatcher() == false || Start.IsInSeparateFile(),
 		TEXT("Create bulkdata stream request from package '%s' FAILED, inline bulk data cannot be streamed from I/O store"),
-		*Start.BulkChunkId.ToDebugString());
+		*LexToString(Start.BulkChunkId));
 
 	checkf(
-		Start.IsUsingIODispatcher() == false || (End.IsInSeparateFile() && Start.BulkChunkId.GetPackageId() == End.BulkChunkId.GetPackageId()),
+		Start.IsUsingIODispatcher() == false || (End.IsInSeparateFile() && Start.BulkChunkId == End.BulkChunkId),
 		TEXT("Create bulk data stream request FAILED, range spans from package '%s' to package '%s'"),
-		*Start.BulkChunkId.ToDebugString(), *End.BulkChunkId.ToDebugString());
+		*LexToString(Start.BulkChunkId), *LexToString(End.BulkChunkId));
 
 	return UE::BulkData::Private::CreateStreamingRequest(
 		Start.BulkMeta,
diff --git a/Engine/Source/Runtime/CoreUObject/Private/Serialization/BulkDataStreaming.cpp b/Engine/Source/Runtime/CoreUObject/Private/Serialization/BulkDataStreaming.cpp
--- a/Engine/Source/Runtime/CoreUObject/Private/Serialization/BulkDataStreaming.cpp
+++ b/Engine/Source/Runtime/CoreUObject/Private/Serialization/BulkDataStreaming.cpp
@@ -1,6 +1,5 @@
 // Copyright Epic Games, Inc. All Rights Reserved.
 #include "Serialization/BulkData.h"
-#include "Algo/AllOf.h"
 #include "Async/MappedFileHandle.h"
 #include "Containers/ChunkedArray.h"
 #include "Experimental/Async/LazyEvent.h"
@@ -10,7 +9,8 @@
 #include "ProfilingDebugging/CountersTrace.h"
 #include "Serialization/MemoryReader.h"
 #include "Templates/RefCounting.h"
-#include "UObject/PackageResourceManager.h"
+#include "UObject/LinkerLoad.h"
+#include "UObject/PackageResourceIoDispatcherBackend.h"
 
 //////////////////////////////////////////////////////////////////////////////
 
@@ -129,25 +129,6 @@ void EnsureCanStreamBulkData(const FBulkMetaData& BulkMeta)
 #endif 
 }
 
-FIoChunkId CreateBulkDataIoChunkId(const FBulkMetaData& BulkMeta, const FPackageId& PackageId)
-{
-	if (PackageId.IsValid() == false)
-	{
-		return FIoChunkId();
-	}
-
-	const EBulkDataFlags BulkDataFlags = BulkMeta.GetFlags();
-
-	const EIoChunkType ChunkType = BulkDataFlags & BULKDATA_OptionalPayload
-		? EIoChunkType::OptionalBulkData
-		: BulkDataFlags & BULKDATA_MemoryMappedPayload
-			? EIoChunkType::MemoryMappedBulkData
-			: EIoChunkType::BulkData;
-
-	const uint16 ChunkIndex = EnumHasAnyFlags(BulkMeta.GetMetaFlags(), FBulkMetaData::EMetaFlags::OptionalPackage) ? 1 : 0;
-	return CreateIoChunkId(PackageId.Value(), ChunkIndex, ChunkType);
-}
-
 //////////////////////////////////////////////////////////////////////////////
 
 EPackageSegment GetPackageSegmentFromFlags(const FBulkMetaData& BulkMeta)
@@ -185,6 +166,54 @@ EPackageSegment GetPackageSegmentFromFlags(const FBulkMetaData& BulkMeta)
 
 //////////////////////////////////////////////////////////////////////////////
 
+FIoChunkId CreateBulkDataIoChunkId(const FBulkMetaData& BulkMeta, UObject* Owner)
+{
+	UPackage* Pkg = Owner ? Owner->GetPackage() : nullptr;
+
+	if (Pkg == nullptr)
+	{
+		return FIoChunkId::InvalidChunkId;
+	}
+
+	const bool bCooked = EnumHasAnyFlags(BulkMeta.GetMetaFlags(), FBulkMetaData::EMetaFlags::CookedPackage);
+	
+	if (bCooked && Pkg->GetPackageId().IsValid())
+	{
+		const EBulkDataFlags BulkDataFlags = BulkMeta.GetFlags();
+
+		const EIoChunkType ChunkType = BulkDataFlags & BULKDATA_OptionalPayload
+			? EIoChunkType::OptionalBulkData
+			: BulkDataFlags & BULKDATA_MemoryMappedPayload
+				? EIoChunkType::MemoryMappedBulkData
+				: EIoChunkType::BulkData;
+
+		const uint16 ChunkIndex = EnumHasAnyFlags(BulkMeta.GetMetaFlags(), FBulkMetaData::EMetaFlags::OptionalPackage) ? 1 : 0;
+		
+		return CreateIoChunkId(Pkg->GetPackageId().Value(), ChunkIndex, ChunkType);
+	}
+	else
+	{
+		FLinkerLoad* Linker = Owner->GetLinker();
+		if (Linker == nullptr)
+		{
+			Linker = FLinkerLoad::FindExistingLinkerForPackage(Pkg);
+		}
+	
+		if (Linker != nullptr)
+		{
+			const FPackagePath& PkgPath = Linker->GetPackagePath();
+			const EPackageSegment Segment = GetPackageSegmentFromFlags(BulkMeta);
+			const bool bExternalResource = BulkMeta.HasAnyFlags(BULKDATA_WorkspaceDomainPayload);
+
+			return CreatePackageResourceChunkId(PkgPath.GetPackageFName(), Segment, bExternalResource);
+		}
+	}
+
+	return FIoChunkId::InvalidChunkId;
+}
+
+//////////////////////////////////////////////////////////////////////////////
+
 enum class EChunkRequestStatus : uint32
 {
 	None				= 0,
@@ -513,7 +542,7 @@ TUniquePtr<IBulkDataIORequest> CreateBulkDataIoDispatcherRequest(
 
 bool OpenReadBulkData(
 	const FBulkMetaData& BulkMeta,
-	const FBulkDataChunkId& BulkChunkId,
+	const FIoChunkId& BulkChunkId,
 	int64 Offset,
 	int64 Size,
 	EAsyncIOPriorityAndFlags Priority,
@@ -525,50 +554,19 @@ bool OpenReadBulkData(
 	}
 
 	EnsureCanStreamBulkData(BulkMeta);
+	
+	FIoBatch Batch = FIoDispatcher::Get().NewBatch();
+	FIoRequest Request = Batch.Read(BulkChunkId, FIoReadOptions(Offset, Size), ConvertToIoDispatcherPriority(Priority));
+	FEventRef Event;
+	Batch.IssueAndTriggerEvent(Event.Get());
+	Event->Wait();
 
-	if (FPackageId PackageId = BulkChunkId.GetPackageId(); PackageId.IsValid())
+	if (const FIoBuffer* Buffer = Request.GetResult())
 	{
-		const FIoChunkId ChunkId = CreateBulkDataIoChunkId(BulkMeta, PackageId);
-		FIoBatch Batch = FIoDispatcher::Get().NewBatch();
-
-		FIoRequest Request = Batch.Read(ChunkId, FIoReadOptions(Offset, Size), ConvertToIoDispatcherPriority(Priority));
-		FEventRef Event;
-		Batch.IssueAndTriggerEvent(Event.Get());
-		Event->Wait();
-
-		if (const FIoBuffer* Buffer = Request.GetResult())
-		{
-			FMemoryReaderView Ar(Buffer->GetView());
-			Read(Ar);
+		FMemoryReaderView Ar(Buffer->GetView());
+		Read(Ar);
 
-			return true;
-		}
-	}
-	else
-	{
-		IPackageResourceManager& ResourceMgr = IPackageResourceManager::Get();
-		
-		const bool bExternalResource = BulkMeta.HasAllFlags(static_cast<EBulkDataFlags>(BULKDATA_PayloadInSeperateFile | BULKDATA_WorkspaceDomainPayload));
-		const FPackagePath& Path = BulkChunkId.GetPackagePath();
-
-		TUniquePtr<FArchive> Ar;
-		if (bExternalResource)
-		{
-			Ar = ResourceMgr.OpenReadExternalResource(EPackageExternalResource::WorkspaceDomainFile, Path.GetPackageName());
-		}
-		else
-		{
-			const EPackageSegment Segment = GetPackageSegmentFromFlags(BulkMeta);
-			Ar = ResourceMgr.OpenReadPackage(Path, Segment).Archive;
-		}
-
-		if (Ar)
-		{
-			Ar->Seek(Offset);
-			Read(*Ar);
-
-			return true;
-		}
+		return true;
 	}
 
 	return false;
@@ -576,7 +574,7 @@ bool OpenReadBulkData(
 
 //////////////////////////////////////////////////////////////////////////////
 
-TUniquePtr<IAsyncReadFileHandle> OpenAsyncReadBulkData(const FBulkMetaData& BulkMeta, const FBulkDataChunkId& BulkChunkId)
+TUniquePtr<IAsyncReadFileHandle> OpenAsyncReadBulkData(const FBulkMetaData& BulkMeta, const FIoChunkId& BulkChunkId)
 {
 	if (BulkChunkId.IsValid() == false)
 	{
@@ -584,35 +582,14 @@ TUniquePtr<IAsyncReadFileHandle> OpenAsyncReadBulkData(const FBulkMetaData& Bulk
 	}
 
 	EnsureCanStreamBulkData(BulkMeta);
-
-	if (FPackageId PackageId = BulkChunkId.GetPackageId(); PackageId.IsValid())
-	{
-		return MakeUnique<FChunkReadFileHandle>(CreateBulkDataIoChunkId(BulkMeta, PackageId));
-	}
-	else
-	{
-		IPackageResourceManager& ResourceMgr = IPackageResourceManager::Get();
-		
-		const bool bExternalResource = BulkMeta.HasAllFlags(static_cast<EBulkDataFlags>(BULKDATA_PayloadInSeperateFile | BULKDATA_WorkspaceDomainPayload));
-		const FPackagePath& Path = BulkChunkId.GetPackagePath();
-
-		if (bExternalResource)
-		{
-			return ResourceMgr.OpenAsyncReadExternalResource(EPackageExternalResource::WorkspaceDomainFile, Path.GetPackageName()).Handle;
-		}
-		else
-		{
-			const EPackageSegment Segment = GetPackageSegmentFromFlags(BulkMeta);
-			return ResourceMgr.OpenAsyncReadPackage(Path, Segment).Handle;
-		}
-	}
+	return MakeUnique<FChunkReadFileHandle>(BulkChunkId);
 }
 
 //////////////////////////////////////////////////////////////////////////////
 
 TUniquePtr<IBulkDataIORequest> CreateStreamingRequest(
 	const FBulkMetaData& BulkMeta,
-	const FBulkDataChunkId& BulkChunkId,
+	const FIoChunkId& BulkChunkId,
 	int64 Offset,
 	int64 Size,
 	EAsyncIOPriorityAndFlags Priority,
@@ -626,102 +603,35 @@ TUniquePtr<IBulkDataIORequest> CreateStreamingRequest(
 
 	EnsureCanStreamBulkData(BulkMeta);
 
-	if (FPackageId PackageId = BulkChunkId.GetPackageId(); PackageId.IsValid())
-	{
-		FIoBuffer Buffer = UserSuppliedMemory ? FIoBuffer(FIoBuffer::Wrap, UserSuppliedMemory, Size) : FIoBuffer(Size);
-		const FIoChunkId ChunkId = CreateBulkDataIoChunkId(BulkMeta, PackageId);
-
-		FChunkBulkDataRequest* Request = new FChunkBulkDataRequest(Callback, MoveTemp(Buffer));
-		Request->Issue(ChunkId, FIoReadOptions(Offset, Size), ConvertToIoDispatcherPriority(Priority));
-		
-		return TUniquePtr<IBulkDataIORequest>(Request);
-	}
-	else if (TUniquePtr<IAsyncReadFileHandle> FileHandle = OpenAsyncReadBulkData(BulkMeta, BulkChunkId); FileHandle.IsValid())
-	{
-		TUniquePtr<FBulkDataIORequest> Request = MakeUnique<FBulkDataIORequest>(FileHandle.Release());
-
-		if (Request->MakeReadRequest(Offset, Size, Priority, Callback, UserSuppliedMemory))
-		{
-			return Request;
-		}
-	}
-
-	return TUniquePtr<IBulkDataIORequest>();
+	FIoBuffer Buffer = UserSuppliedMemory ? FIoBuffer(FIoBuffer::Wrap, UserSuppliedMemory, Size) : FIoBuffer(Size);
+	FChunkBulkDataRequest* Request = new FChunkBulkDataRequest(Callback, MoveTemp(Buffer));
+	Request->Issue(BulkChunkId, FIoReadOptions(Offset, Size), ConvertToIoDispatcherPriority(Priority));
+	
+	return TUniquePtr<IBulkDataIORequest>(Request);
 }
 
 //////////////////////////////////////////////////////////////////////////////
 
-bool DoesBulkDataExist(const FBulkMetaData& BulkMeta, const FBulkDataChunkId& BulkChunkId)
+bool DoesBulkDataExist(const FIoChunkId& BulkChunkId)
 {
-	if (BulkChunkId.IsValid() == false)
-	{
-		return false;
-	}
-
-	if (FPackageId Id = BulkChunkId.GetPackageId(); Id.IsValid())
-	{
-		const FIoChunkId ChunkId = CreateBulkDataIoChunkId(BulkMeta, Id);
-		return FIoDispatcher::Get().DoesChunkExist(ChunkId);
-	}
-	else
-	{
-		IPackageResourceManager& ResourceMgr = IPackageResourceManager::Get();
-		const FPackagePath& PackagePath = BulkChunkId.GetPackagePath();
-		const bool bExternalResource = BulkMeta.HasAllFlags(static_cast<EBulkDataFlags>(BULKDATA_PayloadInSeperateFile | BULKDATA_WorkspaceDomainPayload));
-
-		if (bExternalResource)
-		{
-			return ResourceMgr.DoesExternalResourceExist(
-				EPackageExternalResource::WorkspaceDomainFile, PackagePath.GetPackageName());
-		}
-
-		const EPackageSegment PackageSegment = GetPackageSegmentFromFlags(BulkMeta);
-		return ResourceMgr.DoesPackageExist(PackagePath, PackageSegment);
-	}
+	return FIoDispatcher::Get().DoesChunkExist(BulkChunkId);
 }
 
 //////////////////////////////////////////////////////////////////////////////
 
 bool TryMemoryMapBulkData(
 	const FBulkMetaData& BulkMeta,
-	const FBulkDataChunkId& BulkChunkId,
+	const FIoChunkId& BulkChunkId,
 	int64 Offset,
 	int64 Size,
 	FIoMappedRegion& OutRegion)
 {
-	if (FPackageId Id = BulkChunkId.GetPackageId(); Id.IsValid())
-	{
-		const FIoChunkId ChunkId = CreateBulkDataIoChunkId(BulkMeta, Id);
-		TIoStatusOr<FIoMappedRegion> Status = FIoDispatcher::Get().OpenMapped(ChunkId, FIoReadOptions(Offset, Size));
-
-		if (Status.IsOk())
-		{
-			OutRegion = Status.ConsumeValueOrDie();
+	TIoStatusOr<FIoMappedRegion> Status = FIoDispatcher::Get().OpenMapped(BulkChunkId, FIoReadOptions(Offset, Size));
 
-			return true;
-		}
-	}
-	else
+	if (Status.IsOk())
 	{
-		IPackageResourceManager& ResourceMgr = IPackageResourceManager::Get();
-		const FPackagePath& Path = BulkChunkId.GetPackagePath();
-		const EPackageSegment Segment = EPackageSegment::BulkDataMemoryMapped;
-
-		TUniquePtr<IMappedFileHandle> MappedFile;
-		MappedFile.Reset(IPackageResourceManager::Get().OpenMappedHandleToPackage(Path, Segment));
-
-		if (!MappedFile)
-		{
-			return false;
-		}
-
-		if (IMappedFileRegion* MappedRegion = MappedFile->MapRegion(Offset, Size, true))
-		{
-			OutRegion.MappedFileHandle = MappedFile.Release();
-			OutRegion.MappedFileRegion = MappedRegion;
-
-			return true;
-		}
+		OutRegion = Status.ConsumeValueOrDie();
+		return true;
 	}
 
 	return false;
@@ -774,7 +684,7 @@ private:
 bool StartAsyncLoad(
 	FBulkData* Owner,
 	const FBulkMetaData& BulkMeta,
-	const FBulkDataChunkId& BulkChunkId,
+	const FIoChunkId& BulkChunkId,
 	int64 Offset,
 	int64 Size,
 	EAsyncIOPriorityAndFlags Priority,
@@ -834,24 +744,6 @@ public:
 
 //////////////////////////////////////////////////////////////////////////////
 
-class FBulkDataBatchRequest::IHandle : public FBulkDataRequest::IHandle
-{
-public:
-	virtual ~IHandle() = default;
-
-	virtual void Read(
-		const UE::BulkData::Private::FBulkMetaData& BulkMeta,
-		const UE::BulkData::Private::FBulkDataChunkId& BulkChunkId,
-		const FIoReadOptions& Options,
-		EAsyncIOPriorityAndFlags Priority,
-		FIoReadCallback&& Callback,
-		FBulkDataBatchReadRequest* OutRequest) = 0; 
-
-	virtual bool Issue(FBulkDataRequest::FCompletionCallback&& Callback) = 0;
-};
-
-//////////////////////////////////////////////////////////////////////////////
-
 class FHandleBase : public FBulkDataRequest::IHandle
 {
 public:
@@ -907,68 +799,6 @@ private:
 
 //////////////////////////////////////////////////////////////////////////////
 
-class FBatchHandleBase : public FBulkDataBatchRequest::IHandle
-{
-public:
-	FBatchHandleBase() = default;
-	virtual ~FBatchHandleBase() = default;
-
-	virtual void AddRef() override final
-	{
-		RefCount.fetch_add(1, std::memory_order_relaxed);
-	}
-
-	virtual void Release() override final
-	{
-		if (RefCount.fetch_sub(1, std::memory_order_acq_rel) == 1)
-		{
-			delete this;
-		}
-	}
-
-	virtual uint32 GetRefCount() const override final
-	{
-		return uint32(RefCount.load(std::memory_order_relaxed));
-	}
-
-	virtual FBulkDataBatchRequest::EStatus GetStatus() const override final
-	{
-		return FBulkDataBatchRequest::EStatus(Status.load(std::memory_order_consume));
-	}
-	
-	virtual bool Wait(uint32 Milliseconds) override final
-	{
-		return DoneEvent.Wait(Milliseconds);
-	}
-
-	void SetStatus(FBulkDataBatchRequest::EStatus InStatus)
-	{
-		Status.store(uint32(InStatus), std::memory_order_release);
-	}
-
-protected:
-	void CompleteBatch(FBulkDataRequest::EStatus CompletionStatus)
-	{
-		if (CompletionCallback)
-		{
-			CompletionCallback(CompletionStatus);
-		}
-
-		SetStatus(CompletionStatus);
-		DoneEvent.Trigger();
-	}
-
-private:
-	std::atomic<int32>	RefCount{0};
-	std::atomic<uint32> Status{0};
-	UE::FLazyEvent		DoneEvent{EEventMode::ManualReset};
-
-protected:
-	FBulkDataRequest::FCompletionCallback CompletionCallback;
-};
-
-//////////////////////////////////////////////////////////////////////////////
-
 static FBulkDataRequest::EStatus GetStatusFromIoErrorCode(const EIoErrorCode ErrorCode)
 {
 	if (ErrorCode == EIoErrorCode::Unknown)
@@ -993,9 +823,8 @@ class FChunkBatchReadRequest final : public FBulkDataRequest::IHandle
 {
 public:
 	FChunkBatchReadRequest() = default;
-	FChunkBatchReadRequest(FBulkDataBatchRequest::IHandle* InBatch, FIoRequest&& InIoHandle)
+	FChunkBatchReadRequest(FBulkDataRequest::IHandle* InBatch)
 		: Batch(InBatch)
-		, IoHandle(MoveTemp(InIoHandle))
 	{ }
 
 	virtual void AddRef() override
@@ -1034,14 +863,15 @@ public:
 		return true;
 	}
 
-	FBulkDataBatchRequest::IHandle* Batch = nullptr;
+	FBulkDataRequest::IHandle* Batch = nullptr;
 	FIoRequest IoHandle;
 };
 
-class FChunkBatchRequest final : public FBatchHandleBase
+class FBulkDataBatchRequest::FBatchHandle
+	: public FHandleBase
 {
 public:
-	FChunkBatchRequest(int32 BatchMaxCount)
+	FBatchHandle(int32 BatchMaxCount)
 	{
 		if (BatchMaxCount > 0)
 		{
@@ -1051,7 +881,7 @@ public:
 		TRACE_COUNTER_INCREMENT(BulkDataBatchRequest_Count);
 	}
 	
-	~FChunkBatchRequest()
+	~FBatchHandle()
 	{
 		Cancel();
 		Wait(MAX_uint32);
@@ -1059,7 +889,7 @@ public:
 		TRACE_COUNTER_DECREMENT(BulkDataBatchRequest_Count);
 	}
 
-	virtual bool Cancel() override
+	virtual bool Cancel() override final
 	{
 		if (FBulkDataRequest::EStatus::None == GetStatus())
 		{
@@ -1078,20 +908,24 @@ public:
 
 		return false;
 	}
+	
+	virtual bool Wait(uint32 Milliseconds) override final
+	{
+		check(GetStatus() != FBulkDataRequest::EStatus::None);
+		return DoneEvent.Wait(Milliseconds);
+	}
 
-	virtual void Read(
-		const UE::BulkData::Private::FBulkMetaData& BulkMeta,
-		const UE::BulkData::Private::FBulkDataChunkId& BulkChunkId,
+	void Read(
+		const FIoChunkId& BulkChunkId,
 		const FIoReadOptions& Options,
 		EAsyncIOPriorityAndFlags Priority,
 		FIoReadCallback&& Callback,
-		FBulkDataBatchReadRequest* OutRequest) override
+		FBulkDataBatchReadRequest* OutRequest) 
 	{
-		const FIoChunkId ChunkId	= CreateBulkDataIoChunkId(BulkMeta, BulkChunkId.GetPackageId());
-		const int32 IoPriority		= ConvertToIoDispatcherPriority(Priority);
+		const int32 IoPriority = ConvertToIoDispatcherPriority(Priority);
 
-		FIoRequest IoRequest = IoBatch.ReadWithCallback(ChunkId, Options, IoPriority, MoveTemp(Callback));
-		FChunkBatchReadRequest* Request = new (Requests) FChunkBatchReadRequest(this, MoveTemp(IoRequest));
+		FChunkBatchReadRequest* Request = new (Requests) FChunkBatchReadRequest(this);
+		Request->IoHandle = IoBatch.ReadWithCallback(BulkChunkId, Options, IoPriority, MoveTemp(Callback));
 
 		if (OutRequest)
 		{
@@ -1099,7 +933,7 @@ public:
 		}
 	}
 
-	virtual bool Issue(FBulkDataRequest::FCompletionCallback&& Callback) override
+	bool Issue(FBulkDataRequest::FCompletionCallback&& Callback) 
 	{
 		CompletionCallback = MoveTemp(Callback);
 
@@ -1117,7 +951,7 @@ public:
 		{
 			FBulkDataRequest::EStatus BatchStatus = FBulkDataRequest::EStatus::Ok;
 
-			for (auto& Request : Requests)
+			for (FChunkBatchReadRequest& Request : Requests)
 			{
 				if (EIoErrorCode ErrorCode = Request.IoHandle.Status().GetErrorCode(); ErrorCode != EIoErrorCode::Ok)
 				{
@@ -1134,245 +968,29 @@ public:
 		return true;
 	}
 
-	static constexpr int32 TargetBytesPerChunk = sizeof(FChunkBatchReadRequest) * 8;
-	using FRequests = TChunkedArray<FChunkBatchReadRequest, TargetBytesPerChunk>;
-
-	FIoBatch	IoBatch;
-	FRequests	Requests; //TOOD: Optimize if there's only a single read request 
-};
-
-//////////////////////////////////////////////////////////////////////////////
-
-class FFileSystemBatchReadRequest final : public FBulkDataRequest::IHandle 
-{
-public:
-	FFileSystemBatchReadRequest() = default;
-	FFileSystemBatchReadRequest(FBulkDataBatchRequest::IHandle* InBatch, FIoReadCallback&& InCallback)
-		: BatchHandle(InBatch)
-		, Callback(MoveTemp(InCallback))
-	{ }
-
-	virtual void AddRef() override
-	{
-		BatchHandle->AddRef();
-	}
-
-	virtual void Release() override
-	{
-		BatchHandle->Release();
-	}
-
-	virtual uint32 GetRefCount() const override
-	{
-		return BatchHandle->GetRefCount();
-	}
-
-	virtual FBulkDataRequest::EStatus GetStatus() const override
-	{
-		check(RequestHandle);
-		
-		if (bWasCancelled)
-		{
-			return FBulkDataRequest::EStatus::Cancelled;
-		}
-
-		return RequestHandle->PollCompletion() ? FBulkDataRequest::EStatus::Ok : FBulkDataRequest::EStatus::Pending;
-	}
-
-	virtual bool Cancel() override
-	{
-		check(RequestHandle);
-
-		if (GetStatus() == FBulkDataRequest::EStatus::Pending)
-		{
-			RequestHandle->Cancel();
-			return true;
-		}
-
-		return false;
-	}
-
-	virtual bool Wait(uint32 Milliseconds) override
-	{
-		checkNoEntry(); // Bulk read requests is currently not awaitable
-		return true;
-	}
-
-	FBulkDataBatchRequest::IHandle* BatchHandle = nullptr;
-	TUniquePtr<IAsyncReadRequest>	RequestHandle;
-	FIoReadCallback					Callback;
-	std::atomic<bool>				bWasCancelled{false};
-};
-
-class FFileSystemBatchRequest final : public FBatchHandleBase
-{
-	struct FRequestParams
-	{
-		FFileSystemBatchReadRequest*			Request;
-		FPackagePath							PackagePath;
-		FIoReadOptions							ReadOptions;
-		UE::BulkData::Private::FBulkMetaData	BulkMeta;
-		EAsyncIOPriorityAndFlags				Priority;
-	};
-
-public:
-	FFileSystemBatchRequest(int32 BatchMaxCount)
-	{
-		if (BatchMaxCount > 0)
-		{
-			Requests.Reserve(BatchMaxCount);
-			RequestParams.Reserve(BatchMaxCount);
-		}
-	}
-
-	~FFileSystemBatchRequest()
-	{
-		Cancel();
-		Wait(MAX_uint32);
-		Requests.Empty();
-	}
-
-	virtual void Read(
-		const UE::BulkData::Private::FBulkMetaData& BulkMeta,
-		const UE::BulkData::Private::FBulkDataChunkId& BulkChunkId,
-		const FIoReadOptions& Options,
-		EAsyncIOPriorityAndFlags Priority,
-		FIoReadCallback&& Callback,
-		FBulkDataBatchReadRequest* OutRequest) override
-	{
-		FFileSystemBatchReadRequest* Request = new (Requests) FFileSystemBatchReadRequest(this, MoveTemp(Callback));
-
-		RequestParams.Add(FRequestParams {Request, BulkChunkId.GetPackagePath(), Options, BulkMeta, Priority});
-
-		if (OutRequest)
-		{
-			*OutRequest = FBulkDataBatchReadRequest(Request);
-		}
-	}
-
-	virtual bool Issue(FBulkDataRequest::FCompletionCallback&& Callback) override
-	{
-		CompletionCallback = MoveTemp(Callback);
-
-		if (Requests.IsEmpty())
-		{
-			CompleteBatch(FBulkDataRequest::EStatus::Ok);
-			return false;
-		}
-
-		SetStatus(FBulkDataRequest::EStatus::Pending);
-
-		PendingRequestCount.store(Requests.Num(), std::memory_order_relaxed);
-
-		for (FRequestParams& Params : RequestParams)
-		{
-			IAsyncReadFileHandle* FileHandle = OpenFile(Params.PackagePath, Params.BulkMeta);
-
-			if (FileHandle == nullptr)
-			{
-				CompleteBatch(FBulkDataRequest::EStatus::Error);
-				return false;
-			}
-
-			const bool bMemoryIsOwned = Params.ReadOptions.GetTargetVa() != nullptr;
-
-			FAsyncFileCallBack ReadFileCallback = [
-				this,
-				Request = Params.Request,
-				bMemoryIsOwned,
-				Size = Params.ReadOptions.GetSize()](bool bWasCancelled, IAsyncReadRequest* ReadRequest)
-				{
-					if (Request->Callback)
-					{
-						if (bWasCancelled)
-						{
-							Request->Callback(FIoStatus(EIoErrorCode::Cancelled));
-						}
-						else
-						{
-							FIoBuffer Buffer = bMemoryIsOwned
-								? FIoBuffer(FIoBuffer::Wrap, ReadRequest->GetReadResults(), Size)
-								: FIoBuffer(FIoBuffer::AssumeOwnership, ReadRequest->GetReadResults(), Size);
-							
-							Request->Callback(Buffer);
-						}
-					}
-
-					if (bWasCancelled)
-					{
-						Request->bWasCancelled = true;
-						BatchStatus.store(uint32(FBulkDataRequest::EStatus::Cancelled), std::memory_order_release);
-					}
-
-					if (1 == PendingRequestCount.fetch_sub(1))
-					{
-						CompleteBatch(FBulkDataRequest::EStatus(BatchStatus.load(std::memory_order_consume)));
-					}
-				};
-
-			Params.Request->RequestHandle.Reset(FileHandle->ReadRequest(
-				Params.ReadOptions.GetOffset(),
-				Params.ReadOptions.GetSize(),
-				Params.Priority,
-				&ReadFileCallback,
-				reinterpret_cast<uint8*>(Params.ReadOptions.GetTargetVa())));
-		}
-		
-		RequestParams.Empty();
-
-		return true;
-	}
-
-	virtual bool Cancel() override
-	{
-		if (GetStatus() == FBulkDataRequest::EStatus::Pending)
-		{
-			for (FFileSystemBatchReadRequest& Request : Requests)
-			{
-				Request.Cancel();
-			}
-
-			return true;
-		}
-
-		return false;
-	}
-
 private:
-	IAsyncReadFileHandle* OpenFile(const FPackagePath& Path, const UE::BulkData::Private::FBulkMetaData& BulkMeta)
+	void CompleteBatch(FBulkDataRequest::EStatus CompletionStatus)
 	{
-		TUniquePtr<IAsyncReadFileHandle>& FileHandle = PathToFileHandle.FindOrAdd(Path.GetPackageFName());
-
-		if (FileHandle.IsValid() == false)
+		if (CompletionCallback)
 		{
-			IPackageResourceManager& ResourceMgr = IPackageResourceManager::Get();
-
-			const bool bExternalResource = BulkMeta.HasAllFlags(static_cast<EBulkDataFlags>(BULKDATA_PayloadInSeperateFile | BULKDATA_WorkspaceDomainPayload));
-
-			if (bExternalResource)
-			{
-				FileHandle = ResourceMgr.OpenAsyncReadExternalResource(EPackageExternalResource::WorkspaceDomainFile, Path.GetPackageName()).Handle;
-			}
-			else
-			{
-				const EPackageSegment Segment = GetPackageSegmentFromFlags(BulkMeta);
-				FileHandle = ResourceMgr.OpenAsyncReadPackage(Path, Segment).Handle;
-			}
+			CompletionCallback(CompletionStatus);
 		}
 
-		return FileHandle.Get();
+		SetStatus(CompletionStatus);
+		DoneEvent.Trigger();
 	}
 
-	using FPathToFileHandle = TMap<FName, TUniquePtr<IAsyncReadFileHandle>>;
-	using FRequests			= TChunkedArray<FFileSystemBatchReadRequest>;
+	static constexpr int32 TargetBytesPerChunk = sizeof(FChunkBatchReadRequest) * 8;
+	using FRequests = TChunkedArray<FChunkBatchReadRequest, TargetBytesPerChunk>;
 
-	FPathToFileHandle			PathToFileHandle;
-	FRequests					Requests;
-	TArray<FRequestParams>		RequestParams;
-	std::atomic<uint32>			BatchStatus;
-	std::atomic<int32>			PendingRequestCount;
+	FIoBatch		IoBatch;
+	FRequests		Requests; //TOOD: Optimize if there's only a single read request 
+	UE::FLazyEvent	DoneEvent{EEventMode::ManualReset};
+	FBulkDataRequest::FCompletionCallback CompletionCallback;
 };
 
+//////////////////////////////////////////////////////////////////////////////
+
 FBulkDataRequest::FBulkDataRequest()
 {
 }
@@ -1467,10 +1085,6 @@ bool FBulkDataBatchRequest::WaitFor(const FTimespan& WaitTime)
 
 //////////////////////////////////////////////////////////////////////////////
 
-FBulkDataBatchRequest::FBuilder::FBuilder()
-{
-}
-
 FBulkDataBatchRequest::FBuilder::FBuilder(int32 MaxCount)
 	: BatchMax(MaxCount)
 {
@@ -1480,20 +1094,11 @@ FBulkDataBatchRequest::FBuilder::~FBuilder()
 {
 }
 
-FBulkDataBatchRequest::IHandle& FBulkDataBatchRequest::FBuilder::GetBatch(const FBulkData& BulkData)
+FBulkDataBatchRequest::FBatchHandle& FBulkDataBatchRequest::FBuilder::GetBatch()
 {
 	if (Batch.IsValid() == false)
 	{
-		check(BulkData.BulkChunkId.IsValid());
-
-		if (FPackageId PackageId = BulkData.BulkChunkId.GetPackageId(); PackageId.IsValid())
-		{
-			Batch = new FChunkBatchRequest(BatchMax);
-		}
-		else
-		{
-			Batch = new FFileSystemBatchRequest(BatchMax);
-		}
+		Batch = new FBulkDataBatchRequest::FBatchHandle(BatchMax);
 	}
 
 	return *Batch;
@@ -1504,7 +1109,7 @@ FBulkDataRequest::EStatus FBulkDataBatchRequest::FBuilder::IssueBatch(FBulkDataB
 	check(Batch.IsValid());
 	checkf(OutRequest != nullptr || Batch->GetRefCount() > 1, TEXT("At least one request handle needs to be used when creating a batch request"));
 
-	TRefCountPtr<FBulkDataBatchRequest::IHandle> NewBatch = MoveTemp(Batch);
+	TRefCountPtr<FBulkDataBatchRequest::FBatchHandle> NewBatch = MoveTemp(Batch);
 	const bool bOk = NewBatch->Issue(MoveTemp(Callback));
 
 	if (OutRequest)
@@ -1529,16 +1134,13 @@ FBulkDataBatchRequest::FBatchBuilder::FBatchBuilder(int32 MaxCount)
 
 FBulkDataBatchRequest::FBatchBuilder& FBulkDataBatchRequest::FBatchBuilder::Read(FBulkData& BulkData, EAsyncIOPriorityAndFlags Priority)
 {
-	check(BatchMax == -1 || BatchCount < BatchMax);
-
 	if (BulkData.IsBulkDataLoaded())
 	{
 		++NumLoaded;
 		return *this;
 	}
 
-	GetBatch(BulkData).Read(
-		BulkData.BulkMeta,
+	GetBatch().Read(
 		BulkData.BulkChunkId,
 		FIoReadOptions(BulkData.GetBulkDataOffsetInFile(), BulkData.GetBulkDataSize()),
 		Priority,
@@ -1573,7 +1175,6 @@ FBulkDataBatchRequest::FBatchBuilder& FBulkDataBatchRequest::FBatchBuilder::Read
 	const uint64 ReadOffset = BulkData.GetBulkDataOffsetInFile() + Offset;
 	const uint64 ReadSize	= FMath::Min(uint64(BulkData.GetBulkDataSize()), Size);
 
-	check(BatchMax == -1 || BatchCount < BatchMax);
 	check(Dst.GetSize() == 0 || Dst.GetSize() == ReadSize);
 
 	if (Dst.GetSize() == 0)
@@ -1581,8 +1182,7 @@ FBulkDataBatchRequest::FBatchBuilder& FBulkDataBatchRequest::FBatchBuilder::Read
 		Dst = FIoBuffer(ReadSize);
 	}
 
-	GetBatch(BulkData).Read(
-		BulkData.BulkMeta,
+	GetBatch().Read(
 		BulkData.BulkChunkId,
 		FIoReadOptions(ReadOffset, ReadSize, Dst.GetData()),
 		Priority,
@@ -1671,13 +1271,10 @@ FBulkDataRequest::EStatus FBulkDataBatchRequest::FScatterGatherBuilder::Issue(FI
 		Dst = FIoBuffer(TotalSize);
 	}
 
-	FBulkDataBatchRequest::IHandle& BatchRequest = GetBatch(*Requests[0].BulkData);
-
 	FMutableMemoryView DstView = Dst.GetMutableView();
 	for (const FRequest& Request : Requests)
 	{
-		BatchRequest.Read(
-			Request.BulkData->BulkMeta,
+		GetBatch().Read(
 			Request.BulkData->BulkChunkId,
 			FIoReadOptions(Request.Offset, Request.Size, DstView.GetData()),
 			Priority,
diff --git a/Engine/Source/Runtime/CoreUObject/Private/UObject/PackageResourceIoDispatcherBackend.cpp b/Engine/Source/Runtime/CoreUObject/Private/UObject/PackageResourceIoDispatcherBackend.cpp
new file mode 100644
--- /dev/null
+++ b/Engine/Source/Runtime/CoreUObject/Private/UObject/PackageResourceIoDispatcherBackend.cpp
@@ -0,0 +1,347 @@
+// Copyright Epic Games, Inc. All Rights Reserved.
+
+#include "UObject/PackageResourceIoDispatcherBackend.h"
+#include "UObject/PackageResourceManager.h"
+#include "Async/AsyncFileHandle.h"
+#include "Async/MappedFileHandle.h"
+#include "Misc/PackageSegment.h"
+#include "Misc/ScopeLock.h"
+
+namespace UE
+{
+
+FIoChunkId CreatePackageResourceChunkId(const FName& PackageName, EPackageSegment Segment, bool bExternalResource)
+{
+	const int32 Index = PackageName.GetComparisonIndex().ToUnstableInt();
+	const int32 Number = PackageName.GetNumber();
+	
+	uint8 Id[12] = {0};
+	FMemory::Memcpy(Id, &Index, sizeof(int32));
+	FMemory::Memcpy(Id + sizeof(int32), &Number, sizeof(int32));
+	Id[9] = uint8(Segment);
+	Id[10] = uint8(bExternalResource);
+	Id[11] = uint8(EIoChunkType::PackageResource);
+
+	FIoChunkId ChunkId;
+	ChunkId.Set(Id, sizeof(Id));
+
+	return ChunkId;
+}
+
+bool TryGetPackageNameFromChunkId(const FIoChunkId& ChunkId, FName& OutPackageName, EPackageSegment& OutSegment, bool& bOutExternal)
+{
+	if (ChunkId.GetChunkType() != EIoChunkType::PackageResource)
+	{
+		return false;
+	}
+
+	const uint8* Id = ChunkId.GetData();
+	const int32 NameIndex = *reinterpret_cast<const int32*>(Id);
+	const int32 NameNumber = *reinterpret_cast<const int32*>(Id + sizeof(int32));
+
+	OutPackageName = FName::CreateFromDisplayId(FNameEntryId::FromUnstableInt(NameIndex), NameNumber);
+	OutSegment = static_cast<EPackageSegment>(Id[9]);
+	bOutExternal = static_cast<bool>(Id[10]);
+
+	return true;
+}
+
+bool TryGetPackagePathFromChunkId(const FIoChunkId& ChunkId, FPackagePath& OutPath, EPackageSegment& OutSegment, bool& bExternal)
+{
+	FName PackageName;
+	if (TryGetPackageNameFromChunkId(ChunkId, PackageName, OutSegment, bExternal))
+	{
+		return FPackagePath::TryFromPackageName(PackageName, OutPath);
+	}
+
+	return false;
+}
+
+inline EAsyncIOPriorityAndFlags ConvertToAsyncIoPriority(const int32 IoDispatcherProperty)
+{
+	if (IoDispatcherProperty < IoDispatcherPriority_Low)
+	{
+		return AIOP_MIN;
+	}
+
+	if (IoDispatcherProperty < (IoDispatcherPriority_Medium - 1))
+	{
+		return AIOP_Low;
+	}
+
+	if (IoDispatcherProperty < IoDispatcherPriority_Medium)
+	{
+		return AIOP_BelowNormal;
+	}
+
+	if (IoDispatcherProperty < IoDispatcherPriority_High)
+	{
+		return AIOP_Normal;
+	}
+
+	if (IoDispatcherProperty < IoDispatcherPriority_Max)
+	{
+		return AIOP_High;
+	}
+	
+	return AIOP_CriticalPath;
+}
+
+class FPackageResourceIoBackend final
+	: public IIoDispatcherBackend
+{
+	class FPendingRequests
+	{
+		struct FHandles
+		{
+			TUniquePtr<IAsyncReadFileHandle> FileHandle;
+			TUniquePtr<IAsyncReadRequest> RequestHandle;
+		};
+
+	public:
+		void Add(FIoRequestImpl* Request, TUniquePtr<IAsyncReadFileHandle>&& FileHandle, TFunction<IAsyncReadRequest*(IAsyncReadFileHandle&)>&& MakeReadRequest)
+		{
+			FScopeLock _(&CriticalSection);
+			FHandles& Handles = Lookup.Add(Request);
+			Handles.FileHandle = MoveTemp(FileHandle);
+			Handles.RequestHandle.Reset(MakeReadRequest(*Handles.FileHandle));
+		}
+
+		void Remove(FIoRequestImpl* Request)
+		{
+			FScopeLock _(&CriticalSection);
+			
+			FHandles& Handles = Lookup.FindChecked(Request);
+
+			Handles.RequestHandle->WaitCompletion();
+			Handles.RequestHandle.Reset();
+
+			Lookup.Remove(Request);
+		}
+
+		void Cancel(FIoRequestImpl* Request)
+		{
+			FScopeLock _(&CriticalSection);
+			if (FHandles* Handles = Lookup.Find(Request))
+			{
+				Handles->RequestHandle->Cancel();
+			}
+		}
+
+	private:
+		FCriticalSection CriticalSection;
+		TMap<FIoRequestImpl*, FHandles> Lookup;
+	};
+
+	class FCompletedRequests
+	{
+	public:
+		void Enqueue(FIoRequestImpl* Request)
+		{
+			FScopeLock _(&CriticalSection);
+
+			if (Tail)
+			{
+				Tail->NextRequest = Request;
+			}
+			else
+			{
+				Head = Request;
+			}
+			Tail = Request;
+		}
+
+		FIoRequestImpl* Dequeue()
+		{
+			FScopeLock _(&CriticalSection);
+
+			FIoRequestImpl* Completed = Head;
+			Head = Tail = nullptr;
+
+			return Completed;
+		}
+
+	private:
+		FCriticalSection CriticalSection;
+		FIoRequestImpl* Head = nullptr;
+		FIoRequestImpl* Tail = nullptr;
+	};
+
+public:
+	FPackageResourceIoBackend(IPackageResourceManager& Mgr);
+	~FPackageResourceIoBackend();
+
+	virtual void Initialize(TSharedRef<const FIoDispatcherBackendContext> Context) override;
+	virtual bool Resolve(FIoRequestImpl* Request) override;
+	virtual void CancelIoRequest(FIoRequestImpl* Request) override;
+	virtual void UpdatePriorityForIoRequest(FIoRequestImpl* Request) override { }
+	virtual bool DoesChunkExist(const FIoChunkId& ChunkId) const override;
+	virtual TIoStatusOr<uint64> GetSizeForChunk(const FIoChunkId& ChunkId) const override;
+	virtual FIoRequestImpl* GetCompletedRequests() override;
+	virtual TIoStatusOr<FIoMappedRegion> OpenMapped(const FIoChunkId& ChunkId, const FIoReadOptions& Options) override;
+
+private:
+	IPackageResourceManager& ResourceMgr;
+	TSharedPtr<const FIoDispatcherBackendContext> BackendContext;
+	FPendingRequests PendingRequests;
+	FCompletedRequests CompletedRequests;
+};
+
+FPackageResourceIoBackend::FPackageResourceIoBackend(IPackageResourceManager& Mgr)
+	: ResourceMgr(Mgr)
+{
+}
+
+FPackageResourceIoBackend::~FPackageResourceIoBackend()
+{
+}
+
+void FPackageResourceIoBackend::Initialize(TSharedRef<const FIoDispatcherBackendContext> Context)
+{
+	BackendContext = Context;
+}
+
+bool FPackageResourceIoBackend::Resolve(FIoRequestImpl* Request)
+{
+	const FIoChunkId& ChunkId = Request->ChunkId;
+
+	FPackagePath Path;
+	EPackageSegment Segment;
+	bool bExternalResource = false;
+
+	if (TryGetPackagePathFromChunkId(ChunkId, Path, Segment, bExternalResource) == false)
+	{
+		return false;
+	}
+	
+	TUniquePtr<IAsyncReadFileHandle> FileHandle = bExternalResource
+		? ResourceMgr.OpenAsyncReadExternalResource(EPackageExternalResource::WorkspaceDomainFile, Path.GetPackageName()).Handle
+		: ResourceMgr.OpenAsyncReadPackage(Path, Segment).Handle;
+
+	if (FileHandle.IsValid() == false)
+	{
+		return false;
+	}
+	
+	PendingRequests.Add(Request, MoveTemp(FileHandle), [this, Request](IAsyncReadFileHandle& FileHandle)
+	{
+		FAsyncFileCallBack Callback = [this, Request](bool bWasCancelled, IAsyncReadRequest* FileReadRequest)
+		{
+			if (uint8* Data = FileReadRequest->GetReadResults())
+			{
+				const bool bUserSuppliedMemory = Request->Options.GetTargetVa() != nullptr;
+
+				FIoBuffer Buffer = bUserSuppliedMemory
+					? FIoBuffer(FIoBuffer::Wrap, Data, Request->Options.GetSize())
+					: FIoBuffer(FIoBuffer::AssumeOwnership, Data, Request->Options.GetSize());
+
+				Request->SetResult(Buffer);
+			}
+			else if (bWasCancelled == false)
+			{
+				Request->SetFailed();
+			}
+			
+			CompletedRequests.Enqueue(Request);
+			BackendContext->WakeUpDispatcherThreadDelegate.Execute();
+		};
+
+		const FIoReadOptions& Options = Request->Options;
+		const EAsyncIOPriorityAndFlags AsyncIoPriority = ConvertToAsyncIoPriority(Request->Priority);
+
+		return FileHandle.ReadRequest(
+			Options.GetOffset(),
+			Options.GetSize(),
+			AsyncIoPriority,
+			&Callback,
+			(uint8*)Options.GetTargetVa());
+	});
+
+	return true;
+}
+
+void FPackageResourceIoBackend::CancelIoRequest(FIoRequestImpl* Request)
+{
+	PendingRequests.Cancel(Request);
+}
+
+bool FPackageResourceIoBackend::DoesChunkExist(const FIoChunkId& ChunkId) const
+{
+	FPackagePath Path;
+	EPackageSegment Segment;
+	bool bExternalResource = false;
+
+	if (TryGetPackagePathFromChunkId(ChunkId, Path, Segment, bExternalResource) == false)
+	{
+		return false;
+	}
+	
+	return ResourceMgr.DoesPackageExist(Path, Segment);
+
+	return false;
+}
+
+TIoStatusOr<uint64> FPackageResourceIoBackend::GetSizeForChunk(const FIoChunkId& ChunkId) const
+{
+	FPackagePath Path;
+	EPackageSegment Segment;
+	bool bExternalResource = false;
+
+	if (TryGetPackagePathFromChunkId(ChunkId, Path, Segment, bExternalResource) == false)
+	{
+		return FIoStatus(EIoErrorCode::NotFound);
+	}
+	
+	if (int64 FileSize = ResourceMgr.FileSize(Path, Segment); FileSize > 0)
+	{
+		return static_cast<uint64>(FileSize);
+	}
+
+	return FIoStatus(EIoErrorCode::NotFound);
+}
+
+FIoRequestImpl* FPackageResourceIoBackend::GetCompletedRequests()
+{
+	FIoRequestImpl* Requests = CompletedRequests.Dequeue();
+
+	for (FIoRequestImpl* It = Requests; It != nullptr; It = It->NextRequest)
+	{
+		PendingRequests.Remove(It);
+	}
+
+	return Requests;
+}
+
+TIoStatusOr<FIoMappedRegion> FPackageResourceIoBackend::OpenMapped(const FIoChunkId& ChunkId, const FIoReadOptions& Options)
+{
+	FPackagePath Path;
+	EPackageSegment Segment;
+	bool bExternalResource = false;
+
+	if (TryGetPackagePathFromChunkId(ChunkId, Path, Segment, bExternalResource) == false)
+	{
+		return FIoStatus(EIoErrorCode::NotFound);
+	}
+
+	TUniquePtr<IMappedFileHandle> FileHandle(ResourceMgr.OpenMappedHandleToPackage(Path));
+	
+	if (FileHandle.IsValid())
+	{
+		TUniquePtr<IMappedFileRegion> MappedRegion(
+			FileHandle->MapRegion(static_cast<int64>(Options.GetOffset()), static_cast<int64>(Options.GetSize())));
+
+		if (MappedRegion.IsValid())
+		{
+			return FIoMappedRegion { FileHandle.Release(), MappedRegion.Release() };
+		}
+	}
+
+	return FIoStatus(EIoErrorCode::NotFound);
+}
+
+TSharedRef<IIoDispatcherBackend> MakePackageResourceIoDispatcherBackend(IPackageResourceManager& Mgr)
+{
+	return MakeShared<FPackageResourceIoBackend>(Mgr);
+}
+
+} // namespace UE
diff --git a/Engine/Source/Runtime/CoreUObject/Private/UObject/PackageResourceIoDispatcherBackend.h b/Engine/Source/Runtime/CoreUObject/Private/UObject/PackageResourceIoDispatcherBackend.h
new file mode 100644
--- /dev/null
+++ b/Engine/Source/Runtime/CoreUObject/Private/UObject/PackageResourceIoDispatcherBackend.h
@@ -0,0 +1,21 @@
+// Copyright Epic Games, Inc. All Rights Reserved.
+
+#pragma once
+
+#include "IO/IoDispatcherBackend.h"
+#include "Misc/PackageSegment.h"
+#include "Templates/SharedPointer.h"
+
+class FIoChunkId;
+class IPackageResourceManager;
+
+namespace UE
+{
+
+COREUOBJECT_API TSharedRef<IIoDispatcherBackend> MakePackageResourceIoDispatcherBackend(IPackageResourceManager& Mgr);
+
+COREUOBJECT_API FIoChunkId CreatePackageResourceChunkId(const FName& PackageName, EPackageSegment Segment, bool bExternalResource = false);
+
+COREUOBJECT_API bool TryGetPackageNameFromChunkId(const FIoChunkId& ChunkId, FName& OutPackageName, EPackageSegment& OutSegment, bool& bOutExternal); 
+
+} // namespace UE
diff --git a/Engine/Source/Runtime/CoreUObject/Private/UObject/PackageResourceManager.cpp b/Engine/Source/Runtime/CoreUObject/Private/UObject/PackageResourceManager.cpp
--- a/Engine/Source/Runtime/CoreUObject/Private/UObject/PackageResourceManager.cpp
+++ b/Engine/Source/Runtime/CoreUObject/Private/UObject/PackageResourceManager.cpp
@@ -10,6 +10,7 @@
 #include "Misc/ScopeLock.h"
 #include "Serialization/Archive.h"
 #include "UObject/PackageResourceManagerFile.h"
+#include "UObject/PackageResourceIoDispatcherBackend.h"
 
 DEFINE_LOG_CATEGORY(LogPackageResourceManager);
 
@@ -44,6 +45,12 @@ void IPackageResourceManager::Initialize()
 	{
 		GPackageResourceManager = MakePackageResourceManagerFile();
 	}
+
+	if (FIoDispatcher::IsInitialized())
+	{
+		TSharedRef<IIoDispatcherBackend> Backend = UE::MakePackageResourceIoDispatcherBackend(*GPackageResourceManager);
+		FIoDispatcher::Get().Mount(Backend);
+	}
 }
 
 void IPackageResourceManager::Shutdown()
@@ -300,4 +307,4 @@ bool IsEditorDomainEnabled()
 	}();
 	return bEnabled;
 }
-#endif
\ No newline at end of file
+#endif
diff --git a/Engine/Source/Runtime/CoreUObject/Public/Serialization/BulkData.h b/Engine/Source/Runtime/CoreUObject/Public/Serialization/BulkData.h
--- a/Engine/Source/Runtime/CoreUObject/Public/Serialization/BulkData.h
+++ b/Engine/Source/Runtime/CoreUObject/Public/Serialization/BulkData.h
@@ -372,10 +372,11 @@ public:
 #endif
 	}
 
-	explicit FBulkMetaData(EBulkDataFlags Flags)
+	explicit FBulkMetaData(EBulkDataFlags Flags, EMetaFlags MetaFlags)
 		: FBulkMetaData()
 	{
 		SetFlags(Flags);
+		SetMetaFlags(MetaFlags);
 	}
 
 	inline int64 GetSize() const
@@ -496,51 +497,8 @@ private:
 #endif // USE_RUNTIME_BULKDATA
 };
 
-/**
- * Chunk identifier when reading non-inlined bulk data. Either
- * a package path or a package ID depending on the package loader.
- */
-class FBulkDataChunkId
-{
-public:
-	COREUOBJECT_API FBulkDataChunkId();
-
-	COREUOBJECT_API FBulkDataChunkId(const FBulkDataChunkId& Other);
-
-	COREUOBJECT_API FBulkDataChunkId(FBulkDataChunkId&&) = default;
-
-	COREUOBJECT_API ~FBulkDataChunkId();
-
-	COREUOBJECT_API FBulkDataChunkId& operator=(const FBulkDataChunkId& Other);
-
-	COREUOBJECT_API FBulkDataChunkId& operator=(FBulkDataChunkId&&) = default;
-	
-	COREUOBJECT_API bool operator==(const FBulkDataChunkId& Other) const;
-	
-	bool IsValid() const { return Impl.IsValid(); }
-
-	FPackageId GetPackageId() const;
-	
-	COREUOBJECT_API const FPackagePath& GetPackagePath() const;
-
-	COREUOBJECT_API FIoFilenameHash GetIoFilenameHash(EBulkDataFlags BulkDataFlags) const;
-
-	FString ToDebugString() const;
-
-	static FBulkDataChunkId FromPackagePath(const FPackagePath& PackagePath);
-
-	static FBulkDataChunkId FromPackageId(const FPackageId& PackageId);
-
-private:
-	struct FImpl;
-
-	FBulkDataChunkId(TPimplPtr<FImpl>&& InImpl);
-
-	TPimplPtr<FImpl> Impl;
-};
-
 /** Returns an I/O chunk ID to be used when loading from I/O store. */
-FIoChunkId CreateBulkDataIoChunkId(const FBulkMetaData& BulkMeta, const FPackageId& PackageId);
+FIoChunkId CreateBulkDataIoChunkId(const FBulkMetaData& BulkMeta, UObject* Owner);
 
 /** Returns the package segment when loading from the package resource manager. */
 COREUOBJECT_API EPackageSegment GetPackageSegmentFromFlags(const FBulkMetaData& BulkMeta);
@@ -828,18 +786,25 @@ public:
 
 	/** Returns the PackagePath this bulkdata resides in */
 	UE_DEPRECATED(5.1, "Deprecated")
-	const FPackagePath& GetPackagePath() const { return BulkChunkId.GetPackagePath(); }
+	const FPackagePath& GetPackagePath() const
+	{ 
+		static FPackagePath Empty;
+		return Empty;
+	}
 
 	/** Returns which segment of its PackagePath this bulkdata resides in */
 	UE_DEPRECATED(5.1, "Deprecated")
-	EPackageSegment GetPackageSegment() const { return UE::BulkData::Private::GetPackageSegmentFromFlags(BulkMeta); }
+	EPackageSegment GetPackageSegment() const
+	{ 
+		return EPackageSegment::Header;
+	}
 
 	/** 
 	 * Returns the io filename hash associated with this bulk data.
 	 *
 	 * @return Hash or INVALID_IO_FILENAME_HASH if invalid.
 	 **/
-	FIoFilenameHash GetIoFilenameHash() const { return BulkChunkId.GetIoFilenameHash(BulkMeta.GetFlags()); }
+	FIoFilenameHash GetIoFilenameHash() const { return MakeIoFilenameHash(BulkChunkId); }
 	
 	/** Returns a FIoChunkId for the bulkdata payload, this will be invalid if the bulkdata is not stored in the IoStore */
 	FIoChunkId CreateChunkId() const;
@@ -1138,9 +1103,9 @@ private:
 		Member variables.
 	-----------------------------------------------------------------------------*/
 	
-	FAllocatedPtr								DataAllocation;
-	UE::BulkData::Private::FBulkMetaData		BulkMeta;
-	UE::BulkData::Private::FBulkDataChunkId		BulkChunkId;
+	FAllocatedPtr							DataAllocation;
+	UE::BulkData::Private::FBulkMetaData	BulkMeta;
+	FIoChunkId								BulkChunkId = FIoChunkId::InvalidChunkId;
 	
 protected:
 #if WITH_EDITOR
@@ -1418,8 +1383,7 @@ public:
 	void Reset();
 
 protected:
-	friend class FChunkBatchRequest;
-	friend class FFileSystemBatchRequest;
+	friend class FBulkDataBatchRequest;
 	
 	FBulkDataRequest(IHandle* InHandle);
 	int32 GetRefCount() const;
@@ -1441,7 +1405,7 @@ using FBulkDataBatchReadRequest = FBulkDataRequest;
 class COREUOBJECT_API FBulkDataBatchRequest : public FBulkDataRequest
 {
 public:
-	class IHandle;
+	class FBatchHandle;
 
 private:
 	class COREUOBJECT_API FBuilder
@@ -1452,15 +1416,16 @@ private:
 		FBuilder& operator=(const FBuilder&) = delete;
 
 	protected:
-		FBuilder();
 		explicit FBuilder(int32 MaxCount);
-		FBulkDataBatchRequest::IHandle& GetBatch(const FBulkData& BulkData);
+		FBulkDataBatchRequest::FBatchHandle& GetBatch();
 		EStatus IssueBatch(FBulkDataBatchRequest* OutRequest, FCompletionCallback&& Callback);
 
-		int32 BatchMax = -1;
 		int32 BatchCount = 0;
 		int32 NumLoaded = 0;
-		TRefCountPtr<FBulkDataBatchRequest::IHandle> Batch;
+
+	private:
+		int32 BatchMax = -1;
+		TRefCountPtr<FBulkDataBatchRequest::FBatchHandle> Batch;
 	};
 
 public:	
diff --git a/Engine/Source/Runtime/CoreUObject/Tests/Serialization/BulkDataTests.cpp b/Engine/Source/Runtime/CoreUObject/Tests/Serialization/BulkDataTests.cpp
--- a/Engine/Source/Runtime/CoreUObject/Tests/Serialization/BulkDataTests.cpp
+++ b/Engine/Source/Runtime/CoreUObject/Tests/Serialization/BulkDataTests.cpp
@@ -50,8 +50,25 @@ bool TestBulkDataFlags(FBulkData& BulkData, uint32 Flags, T&& IsSet)
 	return bOk;
 }
 
+class FIoDispatcherTestScope
+{
+public:
+	FIoDispatcherTestScope()
+	{
+		FIoDispatcher::Initialize();
+		FIoDispatcher::InitializePostSettings();
+	}
+	
+	~FIoDispatcherTestScope()
+	{
+		FIoDispatcher::Shutdown();
+	}
+};
+
 TEST_CASE("CoreUObject::Serialization::FBulkData::Basic", "[CoreUObject][Serialization]")
 {
+	FIoDispatcherTestScope _;
+
 	SECTION("Default construction")
 	{
 		FBulkData BulkData;
@@ -220,6 +237,8 @@ TEST_CASE("CoreUObject::Serialization::FBulkData::Basic", "[CoreUObject][Seriali
 
 TEST_CASE("CoreUObject::Serialization::FBulkData::Serialize", "[CoreUObject][Serialization]")
 {
+	FIoDispatcherTestScope _;
+
 	SECTION("Serialize to memory archive")
 	{
 		const int64 ExpectedPayloadSize = 128;
