ZenLoader in editor: Packages can't be missing if they're found in memory

diff --git a/Engine/Source/Runtime/CoreUObject/Private/Serialization/AsyncLoading2.cpp b/Engine/Source/Runtime/CoreUObject/Private/Serialization/AsyncLoading2.cpp
--- a/Engine/Source/Runtime/CoreUObject/Private/Serialization/AsyncLoading2.cpp
+++ b/Engine/Source/Runtime/CoreUObject/Private/Serialization/AsyncLoading2.cpp
@@ -4115,9 +4115,12 @@ void FAsyncPackage2::ImportPackagesRecursiveInner(FAsyncLoadingThreadState2& Thr
 		}
 		else if (ImportedPackageStatus == EPackageStoreEntryStatus::Missing)
 		{
-			UE_ASYNC_PACKAGE_LOG(Warning, Desc, TEXT("ImportPackages: SkipPackage"),
-				TEXT("Skipping non mounted imported package with id '0x%llX'"), ImportedPackageId.Value());
-			ImportedPackageRef.SetIsMissingPackage();
+			if (!ImportedPackageRef.GetPackage()) // If we found a package it's not actually missing but we can't load it anyway
+			{
+				UE_ASYNC_PACKAGE_LOG(Warning, Desc, TEXT("ImportPackages: SkipPackage"),
+					TEXT("Skipping non mounted imported package with id '0x%llX'"), ImportedPackageId.Value());
+				ImportedPackageRef.SetIsMissingPackage();
+			}
 			Data.ImportedAsyncPackages[ImportedPackageIndex++] = nullptr;
 			continue;
 		}
