Disable relocation reference fix up when loading a package for diff

diff --git a/Engine/Source/Runtime/CoreUObject/Private/UObject/LinkerLoad.cpp b/Engine/Source/Runtime/CoreUObject/Private/UObject/LinkerLoad.cpp
--- a/Engine/Source/Runtime/CoreUObject/Private/UObject/LinkerLoad.cpp
+++ b/Engine/Source/Runtime/CoreUObject/Private/UObject/LinkerLoad.cpp
@@ -2084,17 +2084,19 @@ FLinkerLoad::ELinkerStatus FLinkerLoad::PopulateRelocationContext()
 		FString PackageNameToLoad = GetPackagePath().GetPackageName();
 		// Check the summmary versioning alongside, PackageName used to be FolderName which wasn't actively used and got deprecated
 		bool bRelocated = PackageNameToLoad != Summary.PackageName && Summary.GetFileVersionUE() >= EUnrealEngineObjectUE5Version::ADD_SOFTOBJECTPATH_LIST;
-		if (bRelocated)
+		// Do not consider a package relocated if it's being loaded for Diff
+		if (bRelocated && (LoadFlags & LOAD_ForDiff) == 0)
 		{
 			// Only Generate relative fix up if the file was actually moved not just renamed
 			FStringView PackagePathToLoad = FPathViews::GetPath(PackageNameToLoad);
 			FStringView OriginalPackagePath = FPathViews::GetPath(Summary.PackageName);
-			//@todo 5.2: switch to use the string view version
-			//FStringView PackageMount = FPathViews::GetMountPointNameFromPath(OriginalPackagePath, nullptr, false);
-			FString PackageMount = FPackageName::GetPackageMountPoint(FString(OriginalPackagePath.GetData(), OriginalPackagePath.Len()), false).ToString();
 
 			if (PackagePathToLoad != OriginalPackagePath)
 			{
+				//@todo 5.2: switch to use the string view version
+				//FStringView PackageMount = FPathViews::GetMountPointNameFromPath(OriginalPackagePath, nullptr, false);
+				FString PackageMount = FPackageName::GetPackageMountPoint(FString(OriginalPackagePath.GetData(), OriginalPackagePath.Len()), false).ToString();
+
 				for (const FObjectImport& Import : ImportMap)
 				{
 					if (!Import.OuterIndex.IsNull())
@@ -3240,7 +3242,8 @@ bool FLinkerLoad::VerifyImportInner(const int32 ImportIndex, FString& WarningSuf
 		FUObjectSerializeContext* SerializeContext = GetSerializeContext();
 
 		// Resolve the package name for the import, potentially remapping it, if instancing
-		FName PackageToLoad = !Import.HasPackageName() ? Import.ObjectName : Import.GetPackageName();
+		FName OriginalPackageToLoad = !Import.HasPackageName() ? Import.ObjectName : Import.GetPackageName();
+		FName PackageToLoad = InstancingContext.RelocatePackage(OriginalPackageToLoad);
 
 		FName PackageToLoadInto = InstancingContext.RemapPackage(PackageToLoad);
 #if WITH_EDITOR
@@ -3330,14 +3333,12 @@ bool FLinkerLoad::VerifyImportInner(const int32 ImportIndex, FString& WarningSuf
 		// to be linked to any other package's ImportMaps
 		if (!Package || Package->HasAnyPackageFlags(PKG_Compiling))
 		{
-			/*
 			FName RelocatedName;
 			if (InstancingContext.RelocatedPackageMapping.RemoveAndCopyValue(OriginalPackageToLoad, RelocatedName))
 			{
 				UE_ASSET_LOG(LogLinker, Warning, PackagePath, TEXT("VerifyImport: Failed to load package for import object '%s' through using relative path '%s'. Using orginal location as fallback."), *GetImportFullName(ImportIndex), *RelocatedName.ToString());
 				return nullptr;
 			}
-			*/
 
 			if (!FLinkerLoad::IsKnownMissingPackage(PackageToLoad))
 			{
