ZenLoader: Fixed crash in PostLoadInstances for failed packages

diff --git a/Engine/Source/Runtime/CoreUObject/Private/Serialization/AsyncLoading2.cpp b/Engine/Source/Runtime/CoreUObject/Private/Serialization/AsyncLoading2.cpp
--- a/Engine/Source/Runtime/CoreUObject/Private/Serialization/AsyncLoading2.cpp
+++ b/Engine/Source/Runtime/CoreUObject/Private/Serialization/AsyncLoading2.cpp
@@ -7480,6 +7480,10 @@ EAsyncPackageState::Type FAsyncPackage2::ProcessExternalReads(EExternalReadActio
 
 EAsyncPackageState::Type FAsyncPackage2::PostLoadInstances(FAsyncLoadingThreadState2& ThreadState)
 {
+	if (bLoadHasFailed)
+	{
+		return EAsyncPackageState::Complete;
+	}
 	const int32 ExportCount = Data.Exports.Num();
 	while (PostLoadInstanceIndex < ExportCount && !ThreadState.IsTimeLimitExceeded(TEXT("PostLoadInstances")))
 	{
