AsyncLoading2: Optimize loading related memory by releasing big allocations in dynamic data structures after big bursts of loading.
Saves ~3 MB after going from a peak of preloading ~17 000 async packages to in game streaming with a peak of ~700 async packages:
~2.2 MB in AsyncLoading (AsyncLoading2.cpp, StreamableManager.cpp)
~340 KB in FileSystem (IoDispatcherFileBackend.cpp)
~430 KB in Untagged (UObjectBase.cpp)

#rnx

diff --git a/Engine/Source/Runtime/CoreUObject/Private/Serialization/AsyncLoading2.cpp b/Engine/Source/Runtime/CoreUObject/Private/Serialization/AsyncLoading2.cpp
--- a/Engine/Source/Runtime/CoreUObject/Private/Serialization/AsyncLoading2.cpp
+++ b/Engine/Source/Runtime/CoreUObject/Private/Serialization/AsyncLoading2.cpp
@@ -63,6 +63,7 @@
 #include "HAL/LowLevelMemStats.h"
 #include "HAL/IPlatformFileOpenLogWrapper.h"
 #include "Modules/ModuleManager.h"
+#include "Containers/MpscQueue.h"
 #include "Containers/SpscQueue.h"
 #include "Misc/PathViews.h"
 #include "UObject/LinkerLoad.h"
@@ -1851,24 +1852,6 @@ struct FAsyncLoadingThreadState2
 
 	}
 
-	bool HasDeferredFrees() const
-	{
-		return DeferredFreeArcs.Num() > 0;
-	}
-
-	void ProcessDeferredFrees()
-	{
-		if (DeferredFreeArcs.Num() > 0)
-		{
-			TRACE_CPUPROFILER_EVENT_SCOPE(ProcessDeferredFrees);
-			for (TTuple<FEventLoadNode2**, uint32>& DeferredFreeArc : DeferredFreeArcs)
-			{
-				GraphAllocator.FreeArcs(DeferredFreeArc.Get<0>(), DeferredFreeArc.Get<1>());
-			}
-			DeferredFreeArcs.Reset();
-		}
-	}
-
 	void SetTimeLimit(bool bInUseTimeLimit, double InTimeLimit)
 	{
 		bUseTimeLimit = bInUseTimeLimit;
@@ -1908,7 +1891,6 @@ struct FAsyncLoadingThreadState2
 	}
 
 	FAsyncLoadEventGraphAllocator& GraphAllocator;
-	TArray<TTuple<FEventLoadNode2**, uint32>> DeferredFreeArcs;
 	TArray<FEventLoadNode2*> NodesToFire;
 	FEventLoadNode2* CurrentlyExecutingEventNode = nullptr;
 	TArray<FAsyncLoadingSyncLoadContext*> SyncLoadContextStack;
@@ -2570,6 +2552,7 @@ public:
 
 private:
 	/** Thread to run the worker FRunnable on */
+	static constexpr int32 DefaultAsyncPackagesReserveCount = 512;
 	FRunnableThread* Thread;
 	TAtomic<bool> bStopRequested { false };
 	TAtomic<bool> bSuspendRequested { false };
@@ -2595,7 +2578,7 @@ private:
 	TSet<FWeakObjectPtr> LoadedAssets;
 #endif
 	/** [ASYNC/GAME THREAD] Packages to be deleted from async thread */
-	TSpscQueue<FAsyncPackage2*> DeferredDeletePackages;
+	TMpscQueue<FAsyncPackage2*> DeferredDeletePackages;
 	
 	struct FFailedPackageRequest
 	{
@@ -2610,7 +2593,7 @@ private:
 	/** Packages in active loading with GetAsyncPackageId() as key */
 	TMap<FPackageId, FAsyncPackage2*> AsyncPackageLookup;
 
-	TSpscQueue<FAsyncPackage2*> ExternalReadQueue;
+	TMpscQueue<FAsyncPackage2*> ExternalReadQueue;
 	TAtomic<int32> PendingIoRequestsCounter{ 0 };
 	
 	/** List of all pending package requests */
@@ -2644,7 +2627,7 @@ private:
 	FPackageStore& PackageStore;
 	FLoadedPackageStore LoadedPackageStore;
 	FGlobalImportStore GlobalImportStore;
-	TSpscQueue<FPackageRequest> PackageRequestQueue;
+	TMpscQueue<FPackageRequest> PackageRequestQueue;
 	TArray<FAsyncPackage2*> PendingPackages;
 
 	/** [GAME THREAD] Initial load pending CDOs */
@@ -2885,6 +2868,10 @@ public:
 			PendingRequests.Remove(ID);
 			TRACE_LOADTIME_END_REQUEST(ID);
 		}
+		if (PendingRequests.IsEmpty())
+		{
+			PendingRequests.Empty(DefaultAsyncPackagesReserveCount);
+		}
 	}
 
 	void AddPendingCDOs(FAsyncPackage2* Package, TArray<UClass*, TInlineAllocator<8>>& Classes)
@@ -3136,6 +3123,10 @@ private:
 		{
 			RequestIdToPackageMap.Remove(RequestId);
 		}
+		if (RequestIdToPackageMap.IsEmpty())
+		{
+			RequestIdToPackageMap.Empty(DefaultAsyncPackagesReserveCount);
+		}
 		delete Package;
 		--PackagesWithRemainingWorkCounter;
 		TRACE_COUNTER_SET(AsyncLoadingPackagesWithRemainingWork, PackagesWithRemainingWorkCounter);
@@ -3613,7 +3604,7 @@ void FEventLoadNode2::ProcessDependencies(FAsyncLoadingThreadState2& ThreadState
 				ThreadState.NodesToFire.Push(Dependent);
 			}
 		}
-		ThreadState.DeferredFreeArcs.Add(MakeTuple(MultipleDependents, DependenciesCapacity));
+		Package->GetGraphAllocator().FreeArcs(MultipleDependents, DependenciesCapacity);
 	}
 	if (ThreadState.bShouldFireNodes)
 	{
@@ -6159,8 +6150,6 @@ EAsyncPackageState::Type FAsyncLoadingThread2::ProcessAsyncLoadingFromGameThread
 				break;
 			}
 
-			ThreadState.ProcessDeferredFrees();
-
 			if (!DeferredDeletePackages.IsEmpty())
 			{
 				FAsyncPackage2* Package = nullptr;
@@ -6315,6 +6304,10 @@ EAsyncPackageState::Type FAsyncLoadingThread2::ProcessLoadedPackagesFromGameThre
 			{
 				FScopeLock LockAsyncPackages(&AsyncPackagesCritical);
 				AsyncPackageLookup.Remove(Package->Desc.UPackageId);
+				if (AsyncPackageLookup.IsEmpty())
+				{
+					AsyncPackageLookup.Empty(DefaultAsyncPackagesReserveCount);
+				}
 				if (!Package->bLoadHasFailed)
 				{
 #if WITH_EDITOR
@@ -6436,10 +6429,6 @@ EAsyncPackageState::Type FAsyncLoadingThread2::ProcessLoadedPackagesFromGameThre
 
 		// We're not done until all packages have been deleted
 		Result = CompletedPackages.Num() ? EAsyncPackageState::PendingImports  : EAsyncPackageState::Complete;
-		if (Result == EAsyncPackageState::Complete && ThreadState.HasDeferredFrees())
-		{
-			ThreadState.ProcessDeferredFrees();
-		}
 	}
 
 	return Result;
@@ -6525,6 +6514,10 @@ FAsyncLoadingThread2::FAsyncLoadingThread2(FIoDispatcher& InIoDispatcher, IAsync
 
 	EventQueue.SetZenaphore(&AltZenaphore);
 
+	AsyncPackageLookup.Reserve(DefaultAsyncPackagesReserveCount);
+	PendingPackages.Reserve(DefaultAsyncPackagesReserveCount);
+	RequestIdToPackageMap.Reserve(DefaultAsyncPackagesReserveCount);
+
 	EventSpecs.AddDefaulted(EEventLoadNode2::Package_NumPhases + EEventLoadNode2::ExportBundle_NumPhases);
 	EventSpecs[EEventLoadNode2::Package_ProcessSummary] = { &FAsyncPackage2::Event_ProcessPackageSummary, &EventQueue, false };
 	EventSpecs[EEventLoadNode2::Package_SetupDependencies] = { &FAsyncPackage2::Event_SetupDependencies, &EventQueue, false };
@@ -6757,12 +6750,6 @@ uint32 FAsyncLoadingThread2::Run()
 						}
 					}
 
-					if (ThreadState.HasDeferredFrees())
-					{
-						ThreadState.ProcessDeferredFrees();
-						continue;
-					}
-
 					if (!DeferredDeletePackages.IsEmpty())
 					{
 						FAsyncPackage2* Package = nullptr;
diff --git a/Engine/Source/Runtime/CoreUObject/Private/UObject/UObjectBase.cpp b/Engine/Source/Runtime/CoreUObject/Private/UObject/UObjectBase.cpp
--- a/Engine/Source/Runtime/CoreUObject/Private/UObject/UObjectBase.cpp
+++ b/Engine/Source/Runtime/CoreUObject/Private/UObject/UObjectBase.cpp
@@ -931,6 +931,11 @@ void ProcessNewlyLoadedUObjects(FName Package, bool bCanProcessNewlyLoadedObject
 	StructRegistry.EmptyRegistrations();
 	ClassRegistry.EmptyRegistrations();
 
+	if (TMap<UObjectBase*, FPendingRegistrantInfo>& PendingRegistrants = FPendingRegistrantInfo::GetMap(); PendingRegistrants.IsEmpty())
+	{
+		PendingRegistrants.Empty();
+	}
+
 	if (bNewUObjects && !GIsInitialLoad)
 	{
 		UClass::AssembleReferenceTokenStreams();
diff --git a/Engine/Source/Runtime/Engine/Private/StreamableManager.cpp b/Engine/Source/Runtime/Engine/Private/StreamableManager.cpp
--- a/Engine/Source/Runtime/Engine/Private/StreamableManager.cpp
+++ b/Engine/Source/Runtime/Engine/Private/StreamableManager.cpp
@@ -87,7 +87,7 @@ public:
 			}
 			if (!DelegatesForHandle->Head)
 			{
-				PendingDelegatesByHandle.Remove(AssociatedHandle);
+				RemovePendingDelegateInternal(AssociatedHandle);
 			}
 		}
 
@@ -143,7 +143,7 @@ public:
 					}
 					if (!DelegatesForHandle->Head)
 					{
-						PendingDelegatesByHandle.Remove(CurrentNode->RelatedHandle);
+						RemovePendingDelegateInternal(CurrentNode->RelatedHandle);
 					}
 				}
 				CurrentNode = NextNode;
@@ -291,6 +291,17 @@ private:
 		}
 	};
 
+	inline void RemovePendingDelegateInternal(TSharedPtr<FStreamableHandle> Handle)
+	{
+		// requires DataLock by caller
+		PendingDelegatesByHandle.Remove(Handle);
+		if (PendingDelegatesByHandle.IsEmpty())
+		{
+			// release potentially big allocation after huge batches of streaming requests
+			PendingDelegatesByHandle.Empty(128);
+		}
+	}
+
 	FPendingDelegateList PendingDelegates;
 	TMap<TSharedPtr<FStreamableHandle>, FPendingDelegateList> PendingDelegatesByHandle;
 
diff --git a/Engine/Source/Runtime/PakFile/Private/IoDispatcherFileBackend.cpp b/Engine/Source/Runtime/PakFile/Private/IoDispatcherFileBackend.cpp
--- a/Engine/Source/Runtime/PakFile/Private/IoDispatcherFileBackend.cpp
+++ b/Engine/Source/Runtime/PakFile/Private/IoDispatcherFileBackend.cpp
@@ -1003,11 +1003,15 @@ FFileIoStoreReadRequest* FFileIoStoreRequestTracker::FindOrAddRawBlock(FFileIoSt
 	return Result;
 }
 
-void FFileIoStoreRequestTracker::RemoveRawBlock(const FFileIoStoreReadRequest* RawBlock)
+void FFileIoStoreRequestTracker::RemoveRawBlock(const FFileIoStoreReadRequest* RawBlock, bool bRemoveFromCancel)
 {
-	if (!RawBlock->bCancelled)
+	if (!RawBlock->bCancelled || bRemoveFromCancel)
 	{
 		RawBlocksMap.Remove(RawBlock->Key);
+		if (RawBlocksMap.IsEmpty())
+		{
+			RawBlocksMap.Empty(128);
+		}
 	}
 }
 
@@ -1047,11 +1051,15 @@ void FFileIoStoreRequestTracker::AddReadRequestsToResolvedRequest(const FFileIoS
 	}
 }
 
-void FFileIoStoreRequestTracker::RemoveCompressedBlock(const FFileIoStoreCompressedBlock* CompressedBlock)
+void FFileIoStoreRequestTracker::RemoveCompressedBlock(const FFileIoStoreCompressedBlock* CompressedBlock, bool bRemoveFromCancel)
 {
-	if (!CompressedBlock->bCancelled)
+	if (!CompressedBlock->bCancelled || bRemoveFromCancel)
 	{
 		CompressedBlocksMap.Remove(CompressedBlock->Key);
+		if (CompressedBlocksMap.IsEmpty())
+		{
+			CompressedBlocksMap.Empty(512);
+		}
 	}
 }
 
@@ -1101,16 +1109,16 @@ bool FFileIoStoreRequestTracker::CancelIoRequest(FFileIoStoreResolvedRequest& Re
 			if (bCancelCompressedBlock)
 			{
 				CompressedBlock->bCancelled = true;
-				CompressedBlocksMap.Remove(CompressedBlock->Key);
+				RemoveCompressedBlock(CompressedBlock, /*bRemoveFromCancel*/ true);
 			}
 		}
 		if (bCancelReadRequest)
 		{
+			ReadRequest.bCancelled = true;
 			if (!ReadRequest.ImmediateScatter.Request)
 			{
-				RawBlocksMap.Remove(ReadRequest.Key);
+				RemoveRawBlock(&ReadRequest, /*bRemoveFromCancel*/ true);
 			}
-			ReadRequest.bCancelled = true;
 #if DO_CHECK
 			for (FFileIoStoreCompressedBlock* CompressedBlock : ReadRequest.CompressedBlocks)
 			{
@@ -2371,4 +2379,4 @@ void FFileIoStoreStats::OnBufferAllocated()
 #endif
 }
 
-#endif
\ No newline at end of file
+#endif
diff --git a/Engine/Source/Runtime/PakFile/Private/IoDispatcherFileBackend.h b/Engine/Source/Runtime/PakFile/Private/IoDispatcherFileBackend.h
--- a/Engine/Source/Runtime/PakFile/Private/IoDispatcherFileBackend.h
+++ b/Engine/Source/Runtime/PakFile/Private/IoDispatcherFileBackend.h
@@ -80,9 +80,9 @@ public:
 	~FFileIoStoreRequestTracker();
 
 	FFileIoStoreCompressedBlock* FindOrAddCompressedBlock(FFileIoStoreBlockKey Key, bool& bOutWasAdded);
-	void RemoveCompressedBlock(const FFileIoStoreCompressedBlock* CompressedBlock);
+	void RemoveCompressedBlock(const FFileIoStoreCompressedBlock* CompressedBlock, bool bRemoveFromCancel = false);
 	FFileIoStoreReadRequest* FindOrAddRawBlock(FFileIoStoreBlockKey Key, bool& bOutWasAdded);
-	void RemoveRawBlock(const FFileIoStoreReadRequest* RawBlock);
+	void RemoveRawBlock(const FFileIoStoreReadRequest* RawBlock, bool bRemoveFromCancel = false);
 	void AddReadRequestsToResolvedRequest(FFileIoStoreCompressedBlock* CompressedBlock, FFileIoStoreResolvedRequest& ResolvedRequest);
 	void AddReadRequestsToResolvedRequest(const FFileIoStoreReadRequestList& Requests, FFileIoStoreResolvedRequest& ResolvedRequest);
 	bool CancelIoRequest(FFileIoStoreResolvedRequest& ResolvedRequest);
