ZenLoader: Flush the package delete queue before starting a GC and when removing a leaked package from the store

diff --git a/Engine/Source/Runtime/CoreUObject/Private/Serialization/AsyncLoading2.cpp b/Engine/Source/Runtime/CoreUObject/Private/Serialization/AsyncLoading2.cpp
--- a/Engine/Source/Runtime/CoreUObject/Private/Serialization/AsyncLoading2.cpp
+++ b/Engine/Source/Runtime/CoreUObject/Private/Serialization/AsyncLoading2.cpp
@@ -2894,8 +2894,26 @@ private:
 	}
 #endif
 
+	bool ProcessDeferredDeletePackagesQueue(int32 MaxCount = MAX_int32)
+	{
+		bool bDidSomething = false;
+		if (!DeferredDeletePackages.IsEmpty())
+		{
+			FAsyncPackage2* Package = nullptr;
+			int32 Count = 0;
+			while (++Count <= MaxCount && DeferredDeletePackages.Dequeue(Package))
+			{
+				DeleteAsyncPackage(Package);
+				bDidSomething = true;
+			}
+		}
+		return bDidSomething;
+	}
+
 	void OnLeakedPackageRename(UPackage* Package);
 
+	void OnPreGarbageCollect();
+
 	void RemoveUnreachableObjects(FUnreachableObjects& ObjectsToRemove);
 
 	bool ProcessPendingCDOs()
@@ -6150,11 +6168,8 @@ EAsyncPackageState::Type FAsyncLoadingThread2::ProcessAsyncLoadingFromGameThread
 				break;
 			}
 
-			if (!DeferredDeletePackages.IsEmpty())
+			if (ProcessDeferredDeletePackagesQueue(1))
 			{
-				FAsyncPackage2* Package = nullptr;
-				DeferredDeletePackages.Dequeue(Package);
-				DeleteAsyncPackage(Package);
 				bDidSomething = true;
 				break;
 			}
@@ -6533,6 +6548,7 @@ FAsyncLoadingThread2::FAsyncLoadingThread2(FIoDispatcher& InIoDispatcher, IAsync
 	AsyncLoadingTickCounter = 0;
 
 	FCoreUObjectInternalDelegates::GetOnLeakedPackageRenameDelegate().AddRaw(this, &FAsyncLoadingThread2::OnLeakedPackageRename);
+	FCoreUObjectDelegates::GetPreGarbageCollectDelegate().AddRaw(this, &FAsyncLoadingThread2::OnPreGarbageCollect);
 
 	FAsyncLoadingThreadState2::TlsSlot = FPlatformTLS::AllocTlsSlot();
 	GameThreadState = MakeUnique<FAsyncLoadingThreadState2>(GraphAllocator, IoDispatcher);
@@ -6555,7 +6571,6 @@ FAsyncLoadingThread2::~FAsyncLoadingThread2()
 void FAsyncLoadingThread2::ShutdownLoading()
 {
 	FCoreUObjectDelegates::GetPreGarbageCollectDelegate().RemoveAll(this);
-	FCoreUObjectDelegates::GetPostGarbageCollect().RemoveAll(this);
 	FCoreUObjectInternalDelegates::GetOnLeakedPackageRenameDelegate().RemoveAll(this);
 
 	delete Thread;
@@ -6750,14 +6765,8 @@ uint32 FAsyncLoadingThread2::Run()
 						}
 					}
 
-					if (!DeferredDeletePackages.IsEmpty())
+					if (ProcessDeferredDeletePackagesQueue(100))
 					{
-						FAsyncPackage2* Package = nullptr;
-						int32 Count = 0;
-						while (++Count <= 100 && DeferredDeletePackages.Dequeue(Package))
-						{
-							DeleteAsyncPackage(Package);
-						}
 						continue;
 					}
 				} // release FGCScopeGuard
@@ -7000,6 +7009,14 @@ void FAsyncLoadingThread2::OnLeakedPackageRename(UPackage* Package)
 		return;
 	}
 
+	UE_LOG(LogStreaming, Display, TEXT("Renaming leaked package %s (0x%llX)"), *Package->GetName(), Package->GetPackageId().ValueForDebugging());
+
+	// Flush the delete queue to clear references to the renamed package, if there are still references after this we're not done with loading the package and it shouldn't have been renamed
+	// If this function is called from within a GC the ALT will be inactive and not touching the deferred delete queue or the unreachable objects array
+	// If this function is called outside of a GC (CheckAndHandleStaleWorldObjectReferences) the call to FlushAsyncLoading above will empty the deferred delete queue but the ALT can still be processing
+	// the unreachable objects array so we need to protect that with a lock
+	ProcessDeferredDeletePackagesQueue();
+
 	FScopeLock UnreachableObjectsLock(&UnreachableObjectsCritical);
 
 	// unreachable objects from last GC should typically have been processed already,
@@ -7014,6 +7031,8 @@ void FAsyncLoadingThread2::OnLeakedPackageRename(UPackage* Package)
 		return;
 	}
 
+	UE_CLOG(PackageRef->GetRefCount() != 0, LogStreaming, Fatal, TEXT("Renaming leaked package %s (0x%llX) that's still referenced by active loading (RefCount=%d)"), *Package->GetName(), Package->GetPackageId().ValueForDebugging(), PackageRef->GetRefCount());
+
 	// Code such as LoadMap or LevelStreaming is about to rename a loaded package which was detected as leaking so that we can load another copy of it.
 	// We should not have any loading happening at present, so we can remove these objects from our stores 
 	TArray<FUObjectItem*> LeakedObjectItems;
@@ -7031,6 +7050,12 @@ void FAsyncLoadingThread2::OnLeakedPackageRename(UPackage* Package)
 	Package->SetCanBeImportedFlag(false);
 }
 
+void FAsyncLoadingThread2::OnPreGarbageCollect()
+{
+	// Flush the delete queue so that we don't prevent packages from being garbage collected if we're done with them
+	ProcessDeferredDeletePackagesQueue();
+}
+
 void FAsyncLoadingThread2::RemoveUnreachableObjects(FUnreachableObjects& ObjectsToRemove)
 {
 	if (ObjectsToRemove.Num() == 0)
@@ -7047,15 +7072,15 @@ void FAsyncLoadingThread2::RemoveUnreachableObjects(FUnreachableObjects& Objects
 
 	const double StartTime = FPlatformTime::Seconds();
 
+	GlobalImportStore.RemovePublicExports(ObjectsToRemove);
+	LoadedPackageStore.RemovePackages(ObjectsToRemove);
+	ObjectsToRemove.Reset();
+
 	const int32 NewLoadedPackageCount = LoadedPackageStore.NumTracked();
 	const int32 NewPublicExportCount = GlobalImportStore.GetStoredPublicExportsCount();
 	const int32 RemovedLoadedPackageCount = OldLoadedPackageCount - NewLoadedPackageCount;
 	const int32 RemovedPublicExportCount = OldPublicExportCount - NewPublicExportCount;
 
-	GlobalImportStore.RemovePublicExports(ObjectsToRemove);
-	LoadedPackageStore.RemovePackages(ObjectsToRemove);
-	ObjectsToRemove.Reset();
-
 	const double StopTime = FPlatformTime::Seconds();
 	UE_LOG(LogStreaming, Display,
 		TEXT("%.3f ms for processing %d objects in RemoveUnreachableObjects(Queued=%d, Async=%d). ")
