Added missing header from CL# 22208719.

diff --git a/Engine/Source/Runtime/CoreUObject/Public/UObject/UObjectGlobals.h b/Engine/Source/Runtime/CoreUObject/Public/UObject/UObjectGlobals.h
--- a/Engine/Source/Runtime/CoreUObject/Public/UObject/UObjectGlobals.h
+++ b/Engine/Source/Runtime/CoreUObject/Public/UObject/UObjectGlobals.h
@@ -16,6 +16,7 @@
 #include "Containers/StringFwd.h"
 #include "Containers/StringView.h"
 #include "Containers/UnrealString.h"
+#include "Containers/VersePathFwd.h"
 #include "CoreGlobals.h"
 #include "CoreMinimal.h"
 #include "CoreTypes.h"
@@ -374,6 +375,19 @@ COREUOBJECT_API UObject* StaticFindObjectChecked( UClass* Class, UObject* InOute
 /** Internal version of StaticFindObject that will not assert on GIsSavingPackage or IsGarbageCollectingAndLockingUObjectHashTables() */
 COREUOBJECT_API UObject* StaticFindObjectSafe( UClass* Class, UObject* InOuter, const TCHAR* Name, bool ExactClass=false );
 
+#if UE_USE_VERSE_PATHS
+/**
+ * Tries to find an object in memory, using a Verse path.
+ *
+ * @param	VersePath		The path to the object to find.
+ * @param	ObjectPath		FName pair representing the outer package object and the inner top level object (asset)
+ * @param	ExactClass		Whether to require an exact match with the passed in class
+ *
+ * @return	Returns a pointer to the found object or nullptr if none could be found
+ */
+COREUOBJECT_API UObject* StaticFindObject(UClass* Class, const UE::Core::FVersePath& VersePath);
+#endif
+
 /**
  * Tries to find an object in memory. This version uses FTopLevelAssetPath to find the object.
  * Version of StaticFindObject that will not assert on GIsSavingPackage or IsGarbageCollectingAndLockingUObjectHashTables()
