Rework AsyncPurge logic to use std atomics and fix TSAN false positive during garbage collection

diff --git a/Engine/Source/Runtime/CoreUObject/Private/UObject/GarbageCollection.cpp b/Engine/Source/Runtime/CoreUObject/Private/UObject/GarbageCollection.cpp
--- a/Engine/Source/Runtime/CoreUObject/Private/UObject/GarbageCollection.cpp
+++ b/Engine/Source/Runtime/CoreUObject/Private/UObject/GarbageCollection.cpp
@@ -34,6 +34,8 @@
 #include "UObject/FieldPathProperty.h"
 #include "UObject/GarbageCollectionHistory.h"
 
+#include <atomic>
+
 /*-----------------------------------------------------------------------------
    Garbage collection.
 -----------------------------------------------------------------------------*/
@@ -375,7 +377,7 @@ class FAsyncPurge : public FRunnable
 	/** Current index into the global unreachable objects array (GUnreachableObjects) of the object being destroyed */
 	int32 ObjCurrentPurgeObjectIndex;
 	/** Number of objects deferred to the game thread to destroy */
-	int32 NumObjectsToDestroyOnGameThread;
+	std::atomic<int32> NumObjectsToDestroyOnGameThread;
 	/** Number of objectsalready destroyed on the game thread */
 	int32 NumObjectsDestroyedOnGameThread;
 	/** Current index into the global unreachable objects array (GUnreachableObjects) of the object being destroyed on the game thread */
@@ -394,6 +396,11 @@ class FAsyncPurge : public FRunnable
 		int32 ProcessedObjectsCount = 0;
 		bool bFinishedDestroyingObjects = true;
 
+		// Avoid fetch_add synchronization for each object sent to the game-thread and instead load the value
+		// once into a local variable and replace by a simple store to publish the value to the other thread.
+		// This is safe because we know only this thread is going to modify it.
+		int32 LocalNumObjectsToDestroyOnGameThread = NumObjectsToDestroyOnGameThread.load(std::memory_order_acquire);
+		
 		while (ObjCurrentPurgeObjectIndex < GUnreachableObjects.Num())
 		{
 			FUObjectItem* ObjectItem = GUnreachableObjects[ObjCurrentPurgeObjectIndex];
@@ -412,8 +419,7 @@ class FAsyncPurge : public FRunnable
 			}
 			else
 			{
-				FPlatformMisc::MemoryBarrier();
-				++NumObjectsToDestroyOnGameThread;
+				NumObjectsToDestroyOnGameThread.store(++LocalNumObjectsToDestroyOnGameThread, std::memory_order_release);
 			}
 			++ProcessedObjectsCount;
 			++ObjectsDestroyedSinceLastMarkPhase;
@@ -445,7 +451,7 @@ class FAsyncPurge : public FRunnable
 		GUObjectArray.LockInternalArray();
 
 		// Cache the number of objects to destroy locally. The number may grow later but that's ok, we'll catch up to it in the next tick
-		const int32 LocalNumObjectsToDestroyOnGameThread = NumObjectsToDestroyOnGameThread;
+		const int32 LocalNumObjectsToDestroyOnGameThread = NumObjectsToDestroyOnGameThread.load(std::memory_order_acquire);
 
 		while (NumObjectsDestroyedOnGameThread < LocalNumObjectsToDestroyOnGameThread && ObjCurrentPurgeObjectIndexOnGameThread < GUnreachableObjects.Num())
 		{
