Fixes to get IWYU working/compiling on windows.

diff --git a/Engine/Source/Programs/UnrealBuildTool/Modes/GenerateClangDatabase.cs b/Engine/Source/Programs/UnrealBuildTool/Modes/GenerateClangDatabase.cs
--- a/Engine/Source/Programs/UnrealBuildTool/Modes/GenerateClangDatabase.cs
+++ b/Engine/Source/Programs/UnrealBuildTool/Modes/GenerateClangDatabase.cs
@@ -139,7 +139,7 @@ namespace UnrealBuildTool
 				foreach (TargetDescriptor TargetDescriptor in TargetDescriptors)
 				{
 					// Disable PCHs and unity builds for the target
-					TargetDescriptor.AdditionalArguments = TargetDescriptor.AdditionalArguments.Append(new string[] { "-NoPCH", "-DisableUnity" });
+					TargetDescriptor.AdditionalArguments = TargetDescriptor.AdditionalArguments.Append(new string[] { "-NoPCH", "-DisableUnity", "-Compiler=Clang" });
 
 					// Create a makefile for the target
 					UEBuildTarget Target = UEBuildTarget.Create(TargetDescriptor, BuildConfiguration.bSkipRulesCompile, BuildConfiguration.bForceRulesCompile, BuildConfiguration.bUsePrecompiled, Logger);
diff --git a/Engine/Source/Runtime/Core/Public/Math/UnrealMathSSE.h b/Engine/Source/Runtime/Core/Public/Math/UnrealMathSSE.h
--- a/Engine/Source/Runtime/Core/Public/Math/UnrealMathSSE.h
+++ b/Engine/Source/Runtime/Core/Public/Math/UnrealMathSSE.h
@@ -391,6 +391,11 @@ FORCEINLINE VectorRegister2Int64 MakeVectorRegisterInt64(int64 X, int64 Y) {
 	return _mm_set_epi64x(Y, X);
 }
 
+PRAGMA_DISABLE_UNSAFE_TYPECAST_WARNINGS
+#if defined(PRAGMA_DISABLE_MISSING_BRACES_WARNINGS)
+PRAGMA_DISABLE_MISSING_BRACES_WARNINGS
+#endif
+
 /**
 * constexpr 4xint32 vector constant creation that bypasses SIMD intrinsic setter.
 *
@@ -418,6 +423,11 @@ FORCEINLINE constexpr VectorRegister4Float MakeVectorRegisterFloatConstant(float
 	return VectorRegister4Float { X, Y, Z, W };
 }
 
+#if defined(PRAGMA_ENABLE_MISSING_BRACES_WARNINGS)
+PRAGMA_ENABLE_MISSING_BRACES_WARNINGS
+#endif
+PRAGMA_RESTORE_UNSAFE_TYPECAST_WARNINGS
+
 FORCEINLINE constexpr VectorRegister2Double MakeVectorRegister2DoubleConstant(double X, double Y)
 {
 	return VectorRegister2Double { X, Y };
diff --git a/Engine/Source/Runtime/Core/Public/Templates/UnrealTemplate.h b/Engine/Source/Runtime/Core/Public/Templates/UnrealTemplate.h
--- a/Engine/Source/Runtime/Core/Public/Templates/UnrealTemplate.h
+++ b/Engine/Source/Runtime/Core/Public/Templates/UnrealTemplate.h
@@ -217,14 +217,10 @@ OutType FloatCastChecked(InType In, InType Precision)
 #define UE_ARRAY_COUNT( array ) (sizeof(UEArrayCountHelper(array)) - 1)
 
 // Offset of a struct member.
-#ifndef UNREAL_CODE_ANALYZER
-// UCA uses clang on Windows. According to C++11 standard, (which in this case clang follows and msvc doesn't)
-// forbids using reinterpret_cast in constant expressions. msvc uses reinterpret_cast in offsetof macro,
-// while clang uses compiler intrinsic. Calling static_assert(STRUCT_OFFSET(x, y) == SomeValue) causes compiler
-// error when using clang on Windows (while including windows headers).
-#define STRUCT_OFFSET( struc, member )	offsetof(struc, member)
-#else
+#ifdef __clang__
 #define STRUCT_OFFSET( struc, member )	__builtin_offsetof(struc, member)
+#else
+#define STRUCT_OFFSET( struc, member )	offsetof(struc, member)
 #endif
 
 #if PLATFORM_VTABLE_AT_END_OF_CLASS
diff --git a/Engine/Source/Runtime/CoreUObject/Private/UObject/ObjectHandle.cpp b/Engine/Source/Runtime/CoreUObject/Private/UObject/ObjectHandle.cpp
--- a/Engine/Source/Runtime/CoreUObject/Private/UObject/ObjectHandle.cpp
+++ b/Engine/Source/Runtime/CoreUObject/Private/UObject/ObjectHandle.cpp
@@ -149,9 +149,9 @@ namespace ObjectHandle_Private
 		FRWLock Lock;
 	};
 
-	static_assert(offsetof(FObjectHandlePackageData, PackageName) == offsetof(FObjectHandlePackageDebugData, PackageName), "FObjectHandlePackageData and FObjectHandlePackageDebugData must match in position of PackageNameField.");
-	static_assert(offsetof(FObjectHandlePackageData, ObjectPaths) == offsetof(FObjectHandlePackageDebugData, ObjectPaths), "FObjectHandlePackageData and FObjectHandlePackageDebugData must match in position of ObjectPaths.");
-	static_assert(offsetof(FObjectHandlePackageData, DataClassDescriptors) == offsetof(FObjectHandlePackageDebugData, DataClassDescriptors), "FObjectHandlePackageData and FObjectHandlePackageDebugData must match in position of DataClassDescriptors.");
+	static_assert(STRUCT_OFFSET(FObjectHandlePackageData, PackageName) == STRUCT_OFFSET(FObjectHandlePackageDebugData, PackageName), "FObjectHandlePackageData and FObjectHandlePackageDebugData must match in position of PackageNameField.");
+	static_assert(STRUCT_OFFSET(FObjectHandlePackageData, ObjectPaths) == STRUCT_OFFSET(FObjectHandlePackageDebugData, ObjectPaths), "FObjectHandlePackageData and FObjectHandlePackageDebugData must match in position of ObjectPaths.");
+	static_assert(STRUCT_OFFSET(FObjectHandlePackageData, DataClassDescriptors) == STRUCT_OFFSET(FObjectHandlePackageDebugData, DataClassDescriptors), "FObjectHandlePackageData and FObjectHandlePackageDebugData must match in position of DataClassDescriptors.");
 	static_assert(sizeof(FObjectHandlePackageData) == sizeof(FObjectHandlePackageDebugData), "FObjectHandlePackageData and FObjectHandlePackageDebugData must match in size.");
 
 	struct FObjectHandleIndex
