Remove check in RemoveUnreachableObjects pending proper fix.
Also allow ALT2_VERIFY_UNREACHABLE_OBJECTS to be used without DO_CHECK.

diff --git a/Engine/Source/Runtime/CoreUObject/Private/Serialization/AsyncLoading2.cpp b/Engine/Source/Runtime/CoreUObject/Private/Serialization/AsyncLoading2.cpp
--- a/Engine/Source/Runtime/CoreUObject/Private/Serialization/AsyncLoading2.cpp
+++ b/Engine/Source/Runtime/CoreUObject/Private/Serialization/AsyncLoading2.cpp
@@ -695,7 +695,7 @@ private:
 			}
 		}
 
-#if DO_CHECK
+#if ALT2_VERIFY_UNREACHABLE_OBJECTS
 		void VerifyAllObjectsRemoved()
 		{
 			for (int32 ObjectIndex : GetValues())
@@ -873,7 +873,7 @@ public:
 		Package->AtomicallyClearInternalFlags(EInternalObjectFlags::LoaderImport);
 	}
 
-#if DO_CHECK
+#if ALT2_VERIFY_UNREACHABLE_OBJECTS
 	void VerifyAllPublicExportsRemoved()
 	{
 		PublicExportMap.VerifyAllObjectsRemoved();
@@ -1081,7 +1081,7 @@ public:
 					UObject* GCObject = Item.DebugObject;
 
 					UObject* ExistingObject = FindPublicExportObjectUnchecked(PublicExportKey);
-					checkf(ExistingObject,
+					UE_CLOG(!ExistingObject, LogStreaming, Fatal,
 						TEXT("The serialized GC object '%s' with flags (ObjectFlags=%x, InternalObjectFlags=%x) and id 0x%llX:0x%llX is missing in ImportStore. ")
 						TEXT("Reason unknown. Double delete? Bug or hash collision?"),
 						*GCObject->GetFullName(),
@@ -1090,7 +1090,7 @@ public:
 						PublicExportKey.GetPackageId().Value(),
 						PublicExportKey.GetExportHash());
 
-					checkf(ExistingObject && ExistingObject == GCObject,
+					UE_CLOG(!ExistingObject || ExistingObject != GCObject, LogStreaming, Fatal,
 						TEXT("The serialized GC Object '%s' with flags (ObjectFlags=%x, InternalObjectFlags=%x) and id 0x%llX:0x%llX is not matching the object '%s' in ImportStore. "),
 						TEXT("Reason unknown. Overwritten after it was added? Bug or hash collision?"),
 						*GCObject->GetFullName(),
@@ -1109,8 +1109,18 @@ public:
 					PackageRef = LoadedPackageStore.FindPackageRef(PackageId);
 				}
 
-				check(PackageRef);
-				PackageRef->RemovePublicExport(PublicExportKey.GetExportHash());
+				if(PackageRef)
+				{
+					PackageRef->RemovePublicExport(PublicExportKey.GetExportHash());
+				}
+#if !NO_LOGGING
+				else
+				{
+					const FUnreachableObject& UnreachableObject = ObjectsToRemove[i];
+					UE_LOG(LogStreaming, Error, TEXT("Attempted to remove unreachable export %s with no package %llx in the store."),
+						*UnreachableObject.ObjectName.ToString(), PackageId.ValueForDebugging());
+				}
+#endif
 			}
 		}
 	}
