Only propagate instancing context to other instancing dependencies

diff --git a/Engine/Source/Runtime/CoreUObject/Private/Serialization/AsyncLoading.cpp b/Engine/Source/Runtime/CoreUObject/Private/Serialization/AsyncLoading.cpp
--- a/Engine/Source/Runtime/CoreUObject/Private/Serialization/AsyncLoading.cpp
+++ b/Engine/Source/Runtime/CoreUObject/Private/Serialization/AsyncLoading.cpp
@@ -2217,7 +2217,11 @@ EAsyncPackageState::Type FAsyncPackage::LoadImports_Event()
 				const FAsyncPackageDesc Info(INDEX_NONE, ImportPackageFName, ImportPackagePath);
 				PendingPackage = new FAsyncPackage(AsyncLoadingThread, Info, EDLBootNotificationManager);
 				PendingPackage->Desc.Priority = Desc.Priority;
-				PendingPackage->Desc.SetInstancingContext(Linker->GetInstancingContext());
+				// if we are also instancing the dependency, propagate the instancing context
+				if (ImportPackageToLoad != ImportPackageFName)
+				{
+					PendingPackage->Desc.SetInstancingContext(Linker->GetInstancingContext());
+				}
 				if (bIsPrestreamRequest)
 				{
 					UE_LOG(LogStreaming, Display, TEXT("%s is prestreaming %s"), *Desc.PackagePath.GetDebugName(), *ImportPackagePath.GetDebugName());
@@ -6574,7 +6578,7 @@ EAsyncPackageState::Type FAsyncPackage::LoadImports()
 				if (!bLoadedFromIoStore)
 				{
 					UE_LOG(LogStreaming, Verbose, TEXT("FAsyncPackage::LoadImports for %s: Loading %s"), *Desc.PackagePath.GetDebugName(), *ImportPackageName);
-					AddImportDependency(ImportPackageFName, ImportToLoad, InstancingContext);
+					AddImportDependency(ImportPackageFName, ImportToLoad, ImportToLoad != ImportPackageFName ? InstancingContext : FLinkerInstancingContext());
 				}
 			}
 			else
@@ -6884,7 +6888,6 @@ EAsyncPackageState::Type FAsyncPackage::PostLoadObjects()
 	}
 
 	const bool bAsyncPostLoadEnabled = FAsyncLoadingThreadSettings::Get().bAsyncPostLoadEnabled;
-	const bool bIsMultithreaded = AsyncLoadingThread.IsMultithreaded();
 
 	TOptional<FScopedSlowTask> SlowTask;
 	if (IsInGameThread())
diff --git a/Engine/Source/Runtime/CoreUObject/Private/UObject/Linker.cpp b/Engine/Source/Runtime/CoreUObject/Private/UObject/Linker.cpp
--- a/Engine/Source/Runtime/CoreUObject/Private/UObject/Linker.cpp
+++ b/Engine/Source/Runtime/CoreUObject/Private/UObject/Linker.cpp
@@ -882,7 +882,7 @@ FLinkerLoad* GetPackageLinker
 
 	// Create new linker.
 	TRefCountPtr<FUObjectSerializeContext> LoadContext(InExistingContext ? InExistingContext : FUObjectThreadContext::Get().GetSerializeContext());
-	FLinkerLoad* Result = FLinkerLoad::CreateLinker(LoadContext, TargetPackage, PackagePath, LoadFlags, InReaderOverride, ImportLinker ? &ImportLinker->GetInstancingContext() : InstancingContext);
+	FLinkerLoad* Result = FLinkerLoad::CreateLinker(LoadContext, TargetPackage, PackagePath, LoadFlags, InReaderOverride, InstancingContext);
 
 	if (!Result && CreatedPackage)
 	{
diff --git a/Engine/Source/Runtime/CoreUObject/Private/UObject/LinkerLoad.cpp b/Engine/Source/Runtime/CoreUObject/Private/UObject/LinkerLoad.cpp
--- a/Engine/Source/Runtime/CoreUObject/Private/UObject/LinkerLoad.cpp
+++ b/Engine/Source/Runtime/CoreUObject/Private/UObject/LinkerLoad.cpp
@@ -3303,13 +3303,16 @@ bool FLinkerLoad::VerifyImportInner(const int32 ImportIndex, FString& WarningSuf
 			// are doing an instanced load (loading the data of package A on disk to package B in memory)
 			// hence we create a package with a unique instance name provided by the instancing context
 			// In the case of a non instanced load `PackageToLoad` and `PackageToLoadInto` will be the same and we won't be providing a package to load into since `Package` will be null.
+			// If we are going through and instanced load we are also propagating the instancing context
+			const FLinkerInstancingContext* LocalInstancingContext = nullptr;
 			if (PackageToLoad != PackageToLoadInto)
 			{
 				Package = CreatePackage(*PackageToLoadInto.ToString());
+				LocalInstancingContext = &GetInstancingContext();
 			}
 			FPackagePath PackagePathToLoad = FPackagePath::FromPackageNameChecked(PackageToLoad);
 			Package = LoadPackageInternal(Package, PackagePathToLoad, InternalLoadFlags | LOAD_IsVerifying, this, nullptr /* InReaderOverride */,
-				nullptr /* InstancingContext */, nullptr /* DiffPackagePath */);
+				LocalInstancingContext, nullptr /* DiffPackagePath */);
 		}
 #if WITH_IOSTORE_IN_EDITOR
 		if (Package && Package->HasAnyPackageFlags(PKG_Cooked) && Package->GetPackageId().IsValid())
