Fix bad merge.

diff --git a/Engine/Source/Runtime/CoreUObject/Private/Serialization/AsyncLoading2.cpp b/Engine/Source/Runtime/CoreUObject/Private/Serialization/AsyncLoading2.cpp
--- a/Engine/Source/Runtime/CoreUObject/Private/Serialization/AsyncLoading2.cpp
+++ b/Engine/Source/Runtime/CoreUObject/Private/Serialization/AsyncLoading2.cpp
@@ -1116,9 +1116,8 @@ public:
 #if !NO_LOGGING
 				else
 				{
-					const FUnreachableObject& UnreachableObject = ObjectsToRemove[i];
 					UE_LOG(LogStreaming, Error, TEXT("Attempted to remove unreachable export %s with no package %llx in the store."),
-						*UnreachableObject.ObjectName.ToString(), PackageId.ValueForDebugging());
+						*Item.ObjectName.ToString(), PackageId.ValueForDebugging());
 				}
 #endif
 			}
