#UE AssembleTokenStreams optimization. When ProcessNewlyLoadedUObjects, only assemble token streams for the new classes that showed up duing the process

diff --git a/Engine/Source/Runtime/CoreUObject/Private/UObject/UObjectBase.cpp b/Engine/Source/Runtime/CoreUObject/Private/UObject/UObjectBase.cpp
--- a/Engine/Source/Runtime/CoreUObject/Private/UObject/UObjectBase.cpp
+++ b/Engine/Source/Runtime/CoreUObject/Private/UObject/UObjectBase.cpp
@@ -752,7 +752,7 @@ void UClassReplaceReloadClasses()
 /**
  * Load any outstanding compiled in default properties
  */
-static void UObjectLoadAllCompiledInDefaultProperties()
+static void UObjectLoadAllCompiledInDefaultProperties(TArray<UClass*>& OutAllNewClasses)
 {
 	TRACE_LOADTIME_REQUEST_GROUP_SCOPE(TEXT("UObjectLoadAllCompiledInDefaultProperties"));
 
@@ -766,7 +766,7 @@ static void UObjectLoadAllCompiledInDefaultProperties()
 		TArray<UClass*> NewClasses;
 		TArray<UClass*> NewClassesInCoreUObject;
 		TArray<UClass*> NewClassesInEngine;
-		ClassRegistry.DoPendingOuterRegistrations(true, [&NewClasses, &NewClassesInCoreUObject, &NewClassesInEngine](const TCHAR* PackageName, UClass& Class) -> void
+		ClassRegistry.DoPendingOuterRegistrations(true, [&OutAllNewClasses, &NewClasses, &NewClassesInCoreUObject, &NewClassesInEngine](const TCHAR* PackageName, UClass& Class) -> void
 			{
 				UE_LOG(LogUObjectBootstrap, Verbose, TEXT("UObjectLoadAllCompiledInDefaultProperties After Registrant %s %s"), PackageName, *Class.GetName());
 
@@ -782,6 +782,8 @@ static void UObjectLoadAllCompiledInDefaultProperties()
 				{
 					NewClasses.Add(&Class);
 				}
+
+				OutAllNewClasses.Add(&Class);
 			}); 
 
 		auto NotifyClassFinishedRegistrationEvents = [](TArray<UClass*>& Classes)
@@ -899,6 +901,7 @@ void ProcessNewlyLoadedUObjects(FName Package, bool bCanProcessNewlyLoadedObject
 	UClassRegisterAllCompiledInClasses();
 
 	bool bNewUObjects = false;
+	TArray<UClass*> AllNewClasses;
 	while (GFirstPendingRegistrant ||
 		ClassRegistry.HasPendingRegistrations() ||
 		StructRegistry.HasPendingRegistrations() ||
@@ -910,7 +913,7 @@ void ProcessNewlyLoadedUObjects(FName Package, bool bCanProcessNewlyLoadedObject
 
 		FCoreUObjectDelegates::CompiledInUObjectsRegisteredDelegate.Broadcast(Package);
 
-		UObjectLoadAllCompiledInDefaultProperties();
+		UObjectLoadAllCompiledInDefaultProperties(AllNewClasses);
 	}
 
 #if WITH_RELOAD
@@ -938,7 +941,14 @@ void ProcessNewlyLoadedUObjects(FName Package, bool bCanProcessNewlyLoadedObject
 
 	if (bNewUObjects && !GIsInitialLoad)
 	{
-		UClass::AssembleReferenceTokenStreams();
+		for (UClass* Class : AllNewClasses)
+		{
+			// Assemble reference token stream for garbage collection/ RTGC.
+			if (!Class->HasAnyFlags(RF_ClassDefaultObject) && !Class->HasAnyClassFlags(CLASS_TokenStreamAssembled))
+			{
+				Class->AssembleReferenceTokenStream();
+			}
+		}
 	}
 }
 
