#UE Disable relocation for now

diff --git a/Engine/Source/Runtime/CoreUObject/Private/UObject/LinkerLoad.cpp b/Engine/Source/Runtime/CoreUObject/Private/UObject/LinkerLoad.cpp
--- a/Engine/Source/Runtime/CoreUObject/Private/UObject/LinkerLoad.cpp
+++ b/Engine/Source/Runtime/CoreUObject/Private/UObject/LinkerLoad.cpp
@@ -2083,7 +2083,7 @@ FLinkerLoad::ELinkerStatus FLinkerLoad::PopulateRelocationContext()
 		// Validate if the package was moved and we want to generate fix up for references
 		FString PackageNameToLoad = GetPackagePath().GetPackageName();
 		// Check the summmary versioning alongside, PackageName used to be FolderName which wasn't actively used and got deprecated
-		bool bRelocated = PackageNameToLoad != Summary.PackageName && Summary.GetFileVersionUE() >= EUnrealEngineObjectUE5Version::ADD_SOFTOBJECTPATH_LIST;
+		bool bRelocated = false;// PackageNameToLoad != Summary.PackageName && Summary.GetFileVersionUE() >= EUnrealEngineObjectUE5Version::ADD_SOFTOBJECTPATH_LIST;
 		// Do not consider a package relocated if it's being loaded for Diff
 		if (bRelocated && (LoadFlags & LOAD_ForDiff) == 0)
 		{
