Don't return an in-memory instance as the import object inside a WITH_EDITOR context if it's still pending a deferred load finalization.

diff --git a/Engine/Source/Runtime/CoreUObject/Private/UObject/LinkerLoad.cpp b/Engine/Source/Runtime/CoreUObject/Private/UObject/LinkerLoad.cpp
--- a/Engine/Source/Runtime/CoreUObject/Private/UObject/LinkerLoad.cpp
+++ b/Engine/Source/Runtime/CoreUObject/Private/UObject/LinkerLoad.cpp
@@ -5401,12 +5401,19 @@ UObject* FLinkerLoad::CreateImport( int32 Index )
 					}
 
 					if( FindObject )
-					{		
-						// Associate import and indicate that we associated an import for later cleanup.
-						Import.XObject = FindObject;
-						check(CurrentLoadContext);
-						CurrentLoadContext->IncrementImportCount();
-						FLinkerManager::Get().AddLoaderWithNewImports(this);
+					{
+#if USE_CIRCULAR_DEPENDENCY_LOAD_DEFERRING
+						// Don't use the object if it's still waiting on some part of a deferred load.
+						const FLinkerLoad* ObjLinker = FindObject->GetLinker();
+						if (!ObjLinker || !ObjLinker->IsBlueprintFinalizationPending())
+#endif	// USE_CIRCULAR_DEPENDENCY_LOAD_DEFERRING
+						{
+							// Associate import and indicate that we associated an import for later cleanup.
+							Import.XObject = FindObject;
+							check(CurrentLoadContext);
+							CurrentLoadContext->IncrementImportCount();
+							FLinkerManager::Get().AddLoaderWithNewImports(this);
+						}
 					}
 				}
 			}
