Move warning to Log on : Short type name \"%s\" provided for TryFindType. Please convert it to a path name since we would need to resaved all previous Level Instance and map done by Quixel forAncient Game project.

diff --git a/Engine/Source/Runtime/CoreUObject/Private/UObject/Class.cpp b/Engine/Source/Runtime/CoreUObject/Private/UObject/Class.cpp
--- a/Engine/Source/Runtime/CoreUObject/Private/UObject/Class.cpp
+++ b/Engine/Source/Runtime/CoreUObject/Private/UObject/Class.cpp
@@ -6060,7 +6060,7 @@ UField* UClass::TryFindTypeSlow(UClass* TypeClass, const FString& InPathNameOrSh
 			
 			FoundType = (UField*)StaticFindFirstObject(TypeClass, *InPathNameOrShortName, InOptions | EFindFirstObjectOptions::EnsureIfAmbiguous | EFindFirstObjectOptions::NativeFirst, ELogVerbosity::Error, TEXT("TryFindType"));
 
-			UE_LOG(LogClass, Warning, TEXT("Short type name \"%s\" provided for TryFindType. Please convert it to a path name (suggested: \"%s\"). Callstack:\r\n\r\n%s"), *InPathNameOrShortName, *GetPathNameSafe(FoundType), Callstack.ToString());
+			UE_LOG(LogClass, Log, TEXT("Short type name \"%s\" provided for TryFindType. Please convert it to a path name (suggested: \"%s\"). Callstack:\r\n\r\n%s"), *InPathNameOrShortName, *GetPathNameSafe(FoundType), Callstack.ToString());
 		}
 	}
 	return FoundType;
