ZenLoader in editor: Fixed incorrect handling of loading a package after it  has been deleted from the disk

diff --git a/Engine/Source/Runtime/CoreUObject/Private/Serialization/AsyncLoading2.cpp b/Engine/Source/Runtime/CoreUObject/Private/Serialization/AsyncLoading2.cpp
--- a/Engine/Source/Runtime/CoreUObject/Private/Serialization/AsyncLoading2.cpp
+++ b/Engine/Source/Runtime/CoreUObject/Private/Serialization/AsyncLoading2.cpp
@@ -3405,6 +3405,11 @@ bool FAsyncLoadingThread2::CreateAsyncPackagesFromQueue(FAsyncLoadingThreadState
 				bIsZenPackage = false;
 				PackageStatus = EPackageStoreEntryStatus::Ok;
 			}
+			if (bIsZenPackage && PackageStatus == EPackageStoreEntryStatus::Ok && !PackageEntry.UncookedPackageName.IsNone())
+			{
+				// This means that the package was added as an uncooked package to the store at one point and was later deleted from disk
+				PackageStatus = EPackageStoreEntryStatus::Missing;
+			}
 #endif
 
 			// Fixup CustomName to handle any input string that can be converted to a long package name.
@@ -4046,6 +4051,11 @@ void FAsyncPackage2::ImportPackagesRecursiveInner(FAsyncLoadingThreadState2& Thr
 			bIsZenPackageImport = false;
 			ImportedPackageStatus = EPackageStoreEntryStatus::Ok;
 		}
+		if (bIsZenPackageImport && ImportedPackageStatus == EPackageStoreEntryStatus::Ok && !ImportedPackageEntry.UncookedPackageName.IsNone())
+		{
+			// This means that the package was added as an uncooked package to the store at one point and was later deleted from disk
+			ImportedPackageStatus = EPackageStoreEntryStatus::Missing;
+		}
 #else
 		constexpr bool bIsZenPackage = true;
 		constexpr bool bIsZenPackageImport = true;
