Fixed a bug where editor bulkdata could fail to provide the correct payload (or load a corrupted payload) after applying undo operations to a recently saved package.

#rnx

- The repro in question was to edit an asset with editor bulkdata (sound wave), save, undo, then try to access the editor bulkdata payload which would fail.
- The problem is when the asset is edited we store the state of the bulkdata, then when we save the package that bulkdata is detached from its file on disk and is no longer allowed to try loading from it. When the undo is applied the flag stating that the bulkdata is detached was removed and the original files offset was set back in the bulkdata. So not only did it think it should load the payload from disk, but it would now load it from the offset of the file from before it was saved.
- This problem comes from us storing the transient flags in the same enum as the non-transient. The transient flags should not be serialized to the transaction buffer nor should they be restored from there.
- The fix for now is to record the transient flags before serializing from the transaction buffer then reapplying those flags after the fact. In 5.2 we should bite the bullet and try splitting the flags properly to prevent future issues/

diff --git a/Engine/Source/Runtime/CoreUObject/Private/Serialization/EditorBulkData.cpp b/Engine/Source/Runtime/CoreUObject/Private/Serialization/EditorBulkData.cpp
--- a/Engine/Source/Runtime/CoreUObject/Private/Serialization/EditorBulkData.cpp
+++ b/Engine/Source/Runtime/CoreUObject/Private/Serialization/EditorBulkData.cpp
@@ -775,9 +775,13 @@ void FEditorBulkData::Serialize(FArchive& Ar, UObject* Owner, bool bAllowRegiste
 
 		if (bNeedsTransaction)
 		{
+			EFlags FlagsToPreserve = EFlags::None;
+
 			if (Ar.IsLoading())
 			{
 				Unregister();
+
+				FlagsToPreserve = Flags & TransientFlags;
 			}
 
 			Ar << Flags;
@@ -810,6 +814,8 @@ void FEditorBulkData::Serialize(FArchive& Ar, UObject* Owner, bool bAllowRegiste
 			}
 			else
 			{
+				Flags |= FlagsToPreserve;
+
 				if (PayloadSize > 0 && !BulkDataId.IsValid())
 				{
 					BulkDataId = FGuid::NewGuid();
@@ -819,9 +825,12 @@ void FEditorBulkData::Serialize(FArchive& Ar, UObject* Owner, bool bAllowRegiste
 				if (bPayloadInArchive)
 				{
 					SerializeData(Ar, CompressedPayload, Flags);
+					Payload = CompressedPayload.Decompress();
+				}
+				else
+				{
+					UnloadData();
 				}
-				
-				Payload = CompressedPayload.Decompress();	
 
 				Register(Owner, TEXT("Serialize/Transacting"), false /* bAllowUpdateId */);
 			}
