AsyncLoading2: Merge LoadedPackageStore and GlobalImportStore and make it more tolerant to renaming of loaded objects.
Handle loaded objects being moved to another (transient) package by cleaning up all public export state when removing a loaded package.
Handle reloading of a package with internally renamed objects by allowing StoreGlobalObject to replace the previous public export.
Note: Renaming of loaded packages is still only supported via the delegate GetOnLeakedPackageRenameDelegate.
Compile all debug verification in all builds and make it possible to toggle with s.VerifyUnreachableObject=0/1 and s.VerifyObjectLoadFlags=0/1, while enabling it by default it debug and development builds.

#rnx

diff --git a/Engine/Source/Runtime/CoreUObject/Private/Serialization/AsyncLoading2.cpp b/Engine/Source/Runtime/CoreUObject/Private/Serialization/AsyncLoading2.cpp
--- a/Engine/Source/Runtime/CoreUObject/Private/Serialization/AsyncLoading2.cpp
+++ b/Engine/Source/Runtime/CoreUObject/Private/Serialization/AsyncLoading2.cpp
@@ -183,14 +183,6 @@ void FindAllRuntimeScriptPackages(TArray<UPackage*>& OutPackages)
 	}, /*bIncludeDerivedClasses*/false);
 }
 
-#ifndef ALT2_VERIFY_ASYNC_FLAGS
-#define ALT2_VERIFY_ASYNC_FLAGS DO_CHECK && !(WITH_EDITOR)
-#endif
-
-#ifndef ALT2_VERIFY_UNREACHABLE_OBJECTS
-#define ALT2_VERIFY_UNREACHABLE_OBJECTS DO_CHECK
-#endif
-
 #ifndef ALT2_VERIFY_LINKERLOAD_MATCHES_IMPORTSTORE
 #define ALT2_VERIFY_LINKERLOAD_MATCHES_IMPORTSTORE 0
 #endif
@@ -308,12 +300,33 @@ if ((Condition)) \
 #define UE_ASYNC_PACKAGE_CLOG_VERBOSE(Condition, Verbosity, PackageDesc, LogDesc, Format, ...)
 #endif
 
-static bool GGRemoveUnreachableObjectsFromGCNotifyOnGT = false;
-static FAutoConsoleVariableRef CVarGRemoveUnreachableObjectsFromGCNotifyOnGT(
-	TEXT("s.GRemoveUnreachableObjectsFromGCNotifyOnGT"),
-	GGRemoveUnreachableObjectsFromGCNotifyOnGT,
-	TEXT("Force running removal of unreachable objects from the Garbage Collection Notify callback on the Game Thread. "
-		"This also enables extra verification in debug and development builds (slow)."),
+static bool GRemoveUnreachableObjectsOnGT = false;
+static FAutoConsoleVariableRef CVarGRemoveUnreachableObjectsOnGT(
+	TEXT("s.RemoveUnreachableObjectsOnGT"),
+	GRemoveUnreachableObjectsOnGT,
+	TEXT("Remove unreachable objects from GlobalImportStore on the GT from the GC callback NotifyUnreachableObjects (slow)."),
+	ECVF_Default
+	);
+
+#if UE_BUILD_DEBUG || (UE_BUILD_DEVELOPMENT && !WITH_EDITOR)
+static bool GVerifyUnreachableObjects = true;
+static bool GVerifyObjectLoadFlags = true;
+#else
+static bool GVerifyUnreachableObjects = false;
+static bool GVerifyObjectLoadFlags = false;
+#endif
+
+static FAutoConsoleVariableRef CVarGVerifyUnreachableObjects(
+	TEXT("s.VerifyUnreachableObjects"),
+	GVerifyUnreachableObjects,
+	TEXT("Run GlobalImportStore verifications for unreachable objects from the GC callback NotifyUnreachableObjects (slow)."),
+	ECVF_Default
+	);
+
+static FAutoConsoleVariableRef CVarGVerifyObjectLoadFlags(
+	TEXT("s.VerifyObjectLoadFlags"),
+	GVerifyObjectLoadFlags,
+	TEXT("Run AsyncFlags verifications for all objects when finished loading from the GC callback NotifyUnreachableObjects (slow)."),
 	ECVF_Default
 	);
 
@@ -526,9 +539,6 @@ struct FUnreachableObject
 	FPackageId PackageId;
 	int32 ObjectIndex = -1;
 	FName ObjectName;
-#if ALT2_VERIFY_UNREACHABLE_OBJECTS
-	UObject* DebugObject = nullptr;
-#endif
 };
 
 using FUnreachableObjects = TArray<FUnreachableObject>;
@@ -536,7 +546,7 @@ using FUnreachableObjects = TArray<FUnreachableObject>;
 class FLoadedPackageRef
 {
 private:
-	friend class FLoadedPackageStore;
+	friend class FGlobalImportStore;
 
 	class FPublicExportMap
 	{
@@ -695,17 +705,6 @@ private:
 			}
 		}
 
-#if ALT2_VERIFY_UNREACHABLE_OBJECTS
-		void VerifyAllObjectsRemoved()
-		{
-			for (int32 ObjectIndex : GetValues())
-			{
-				check(ObjectIndex < 0);
-			}
-		}
-#endif
-
-	private:
 		TArrayView<uint64> GetKeys()
 		{
 			if (Count == 1)
@@ -730,6 +729,7 @@ private:
 			}
 		}
 
+	private:
 		union
 		{
 			uint8* Allocation = nullptr;
@@ -824,6 +824,11 @@ public:
 		bHasFailed = true;
 	}
 
+	TArrayView<int32> GetPublicExportObjectIndices()
+	{
+		return PublicExportMap.GetValues();
+	}
+
 	void ReserveSpaceForPublicExports(int32 PublicExportCount)
 	{
 		PublicExportMap.Grow(PublicExportCount);
@@ -872,33 +877,47 @@ public:
 		checkf(Package->HasAnyInternalFlags(EInternalObjectFlags::LoaderImport), TEXT("%s"), *Package->GetFullName());
 		Package->AtomicallyClearInternalFlags(EInternalObjectFlags::LoaderImport);
 	}
-
-#if ALT2_VERIFY_UNREACHABLE_OBJECTS
-	void VerifyAllPublicExportsRemoved()
-	{
-		PublicExportMap.VerifyAllObjectsRemoved();
-	}
-#endif
 };
 
-class FLoadedPackageStore
+class FGlobalImportStore
 {
 private:
 	// Packages in active loading or completely loaded packages, with Desc.UPackageId as key.
 	// Does not track temp packages with custom UPackage names, since they are never imorted by other packages.
 	TMap<FPackageId, FLoadedPackageRef> Packages;
+	// All native script objects (structs, enums, classes, CDOs and their subobjects) from the initial load phase
+	TMap<FPackageObjectIndex, UObject*> ScriptObjects;
+	// All currently loaded public export objects from any loaded package
+	TMap<int32, FPublicExportKey> ObjectIndexToPublicExport;
 
 public:
-	FLoadedPackageStore()
+	FGlobalImportStore()
 	{
 		Packages.Reserve(32768);
+		ScriptObjects.Reserve(32768);
+		ObjectIndexToPublicExport.Reserve(32768);
 	}
 
-	int32 NumTracked() const
+	int32 GetStoredPackagesCount() const
 	{
 		return Packages.Num();
 	}
 
+	int32 GetStoredScriptObjectsCount() const
+	{
+		return ScriptObjects.Num();
+	}
+
+	uint32 GetStoredScriptObjectsAllocatedSize() const
+	{
+		return ScriptObjects.GetAllocatedSize();
+	}
+
+	int32 GetStoredPublicExportsCount() const
+	{
+		return ObjectIndexToPublicExport.Num();
+	}
+
 	inline FLoadedPackageRef* FindPackageRef(FPackageId PackageId)
 	{
 		return Packages.Find(PackageId);
@@ -980,9 +999,9 @@ public:
 		}
 	}
 
-#if ALT2_VERIFY_ASYNC_FLAGS
 	void VerifyLoadedPackages()
 	{
+		TRACE_CPUPROFILER_EVENT_SCOPE(VerifyLoadedPackages);
 		for (TPair<FPackageId, FLoadedPackageRef>& Pair : Packages)
 		{
 			FPackageId& PackageId = Pair.Key;
@@ -990,11 +1009,18 @@ public:
 			ensureMsgf(Ref.GetRefCount() == 0,
 				TEXT("PackageId '0x%llX' with ref count %d should not have a ref count now")
 				TEXT(", or this check is incorrectly reached during active loading."),
-				PackageId.Value(),
+				PackageId.ValueForDebugging(),
 				Ref.GetRefCount());
 		}
 	}
-#endif
+
+	void RemovePackage(UPackage* Package)
+	{
+		TRACE_CPUPROFILER_EVENT_SCOPE(RemovePackage);
+		VerifyPackageForRemoval(Package);
+		FPackageId PackageId = Package->GetPackageId();
+		RemovePackage(PackageId, Package->GetFName());
+	}
 
 	void RemovePackages(const FUnreachableObjects& ObjectsToRemove)
 	{
@@ -1004,60 +1030,36 @@ public:
 			const FPackageId& PackageId = Item.PackageId;
 			if (PackageId.IsValid())
 			{
-				UE_ASYNC_PACKAGEID_DEBUG(PackageId);
-
-#if ALT2_VERIFY_UNREACHABLE_OBJECTS
-				FLoadedPackageRef PackageRef;
-				bool bRemoved = Packages.RemoveAndCopyValue(PackageId, PackageRef);
-				if (bRemoved)
-				{
-					if (PackageRef.RefCount > 0)
-					{
-						FString PackageName = Item.ObjectName.ToString();
-						UE_LOG(LogStreaming, Error,
-							TEXT("RemovePackage: %s (0x%llX) - ")
-							TEXT("Package destroyed while still being referenced, RefCount %d > 0."),
-							*PackageName,
-							PackageId.Value(),
-							PackageRef.RefCount);
-						checkf(false, TEXT("Package %s destroyed with RefCount"), *PackageName);
-					}
-					PackageRef.VerifyAllPublicExportsRemoved();
-				}
-#endif
-				Packages.Remove(PackageId);
+				RemovePackage(PackageId, Item.ObjectName);
 			}
 		}
 	}
-};
-
-class FGlobalImportStore
-{
-private:
-	FLoadedPackageStore& LoadedPackageStore;
-	TMap<FPackageObjectIndex, UObject*> ScriptObjects;
-	TMap<int32, FPublicExportKey> ObjectIndexToPublicExport;
-
-public:
-	FGlobalImportStore(FLoadedPackageStore& InLoadedPackageStore)
-		: LoadedPackageStore(InLoadedPackageStore)
-	{
-		ObjectIndexToPublicExport.Reserve(32768);
-	}
-
-	int32 GetStoredScriptObjectsCount() const
-	{
-		return ScriptObjects.Num();
-	}
 
-	uint32 GetStoredScriptObjectsAllocatedSize() const
+	void RemovePackage(FPackageId PackageId, FName PackageName)
 	{
-		return ScriptObjects.GetAllocatedSize();
-	}
+		UE_ASYNC_PACKAGEID_DEBUG(PackageId);
 
-	int32 GetStoredPublicExportsCount() const
-	{
-		return ObjectIndexToPublicExport.Num();
+		FLoadedPackageRef PackageRef;
+		bool bRemoved = Packages.RemoveAndCopyValue(PackageId, PackageRef);
+		if (bRemoved)
+		{
+			if (PackageRef.GetRefCount() > 0)
+			{
+				UE_LOG(LogStreaming, Fatal,
+					TEXT("FGlobalImportStore::RemovePackage: %s (0x%llX) - ")
+					TEXT("Package removed while still being referenced, RefCount %d > 0."),
+					*PackageName.ToString(),
+					PackageId.ValueForDebugging(),
+					PackageRef.GetRefCount());
+			}
+			for (int32 ObjectIndex : PackageRef.GetPublicExportObjectIndices())
+			{
+				if (ObjectIndex >= 0)
+				{
+					ObjectIndexToPublicExport.Remove(ObjectIndex);
+				}
+			}
+		}
 	}
 
 	void RemovePublicExports(const FUnreachableObjects& ObjectsToRemove)
@@ -1066,7 +1068,6 @@ public:
 
 		FPackageId LastPackageId;
 		FLoadedPackageRef* PackageRef = nullptr;
-
 		for (const FUnreachableObject& Item : ObjectsToRemove)
 		{
 			int32 ObjectIndex = Item.ObjectIndex;
@@ -1075,58 +1076,121 @@ public:
 			FPublicExportKey PublicExportKey;
 			if (ObjectIndexToPublicExport.RemoveAndCopyValue(ObjectIndex, PublicExportKey))
 			{
-#if ALT2_VERIFY_UNREACHABLE_OBJECTS
-				if (GGRemoveUnreachableObjectsFromGCNotifyOnGT)
-				{
-					UObject* GCObject = Item.DebugObject;
-
-					UObject* ExistingObject = FindPublicExportObjectUnchecked(PublicExportKey);
-					UE_CLOG(!ExistingObject, LogStreaming, Fatal,
-						TEXT("The serialized GC object '%s' with flags (ObjectFlags=%x, InternalObjectFlags=%x) and id 0x%llX:0x%llX is missing in ImportStore. ")
-						TEXT("Reason unknown. Double delete? Bug or hash collision?"),
-						*GCObject->GetFullName(),
-						GCObject->GetFlags(),
-						GCObject->GetInternalFlags(),
-						PublicExportKey.GetPackageId().Value(),
-						PublicExportKey.GetExportHash());
-
-					UE_CLOG(!ExistingObject || ExistingObject != GCObject, LogStreaming, Fatal,
-						TEXT("The serialized GC Object '%s' with flags (ObjectFlags=%x, InternalObjectFlags=%x) and id 0x%llX:0x%llX is not matching the object '%s' in ImportStore. "),
-						TEXT("Reason unknown. Overwritten after it was added? Bug or hash collision?"),
-						*GCObject->GetFullName(),
-						GCObject->GetFlags(),
-						GCObject->GetInternalFlags(),
-						PublicExportKey.GetPackageId().Value(),
-						PublicExportKey.GetExportHash(),
-						*ExistingObject->GetFullName());
-				}
-#endif
-
 				FPackageId PackageId = PublicExportKey.GetPackageId();
 				if (PackageId != LastPackageId)
 				{
+					UE_ASYNC_PACKAGEID_DEBUG(PackageId);
 					LastPackageId = PackageId;
-					PackageRef = LoadedPackageStore.FindPackageRef(PackageId);
+					PackageRef = FindPackageRef(PackageId);
 				}
-
-				if(PackageRef)
+				if (PackageRef)
 				{
 					PackageRef->RemovePublicExport(PublicExportKey.GetExportHash());
 				}
-#if !NO_LOGGING
-				else
-				{
-					UE_LOG(LogStreaming, Error, TEXT("Attempted to remove unreachable export %s with no package %llx in the store."),
-						*Item.ObjectName.ToString(), PackageId.ValueForDebugging());
-				}
-#endif
+			}
+		}
+	}
+
+	void VerifyObjectForRemoval(UObject* GCObject)
+	{
+		int32 ObjectIndex = GUObjectArray.ObjectToIndex(GCObject);
+		FPublicExportKey* PublicExportKey = ObjectIndexToPublicExport.Find(ObjectIndex);
+		if (PublicExportKey)
+		{
+			FPackageId PackageId = PublicExportKey->GetPackageId();
+			FLoadedPackageRef* PackageRef = FindPackageRef(PackageId);
+			if (PackageRef)
+			{
+				UObject* ExistingObject = FindPublicExportObjectUnchecked(*PublicExportKey);
+				UE_CLOG(!ExistingObject, LogStreaming, Fatal,
+					TEXT("FGlobalImportStore::VerifyObjectForRemoval: The loaded public export object '%s' with flags (ObjectFlags=%x, InternalObjectFlags=%x) and id 0x%llX:0x%llX is missing in GlobalImportStore. ")
+					TEXT("Reason unknown. Double delete? Bug or hash collision?"),
+					*GCObject->GetFullName(),
+					GCObject->GetFlags(),
+					GCObject->GetInternalFlags(),
+					PackageId.ValueForDebugging(),
+					PublicExportKey->GetExportHash());
+
+				UE_CLOG(ExistingObject != GCObject, LogStreaming, Fatal,
+					TEXT("FGlobalImportStore::VerifyObjectForRemoval: The loaded public export object '%s' with flags (ObjectFlags=%x, InternalObjectFlags=%x) and id 0x%llX:0x%llX is not matching the object '%s' in GlobalImportStore. "),
+					TEXT("Reason unknown. Overwritten after it was added? Bug or hash collision?"),
+					*GCObject->GetFullName(),
+					GCObject->GetFlags(),
+					GCObject->GetInternalFlags(),
+					PackageId.ValueForDebugging(),
+					PublicExportKey->GetExportHash(),
+					*ExistingObject->GetFullName());
+			}
+			else
+			{
+				UE_LOG(LogStreaming, Warning,
+					TEXT("FGlobalImportStore::VerifyObjectForRemoval: The package for the serialized GC object '%s' with flags (ObjectFlags=%x, InternalObjectFlags=%x) and id 0x%llX:0x%llX is missing in GlobalImportStore. ")
+					TEXT("Most likely this object has been moved into this package after it was loaded, while the original package is still around."),
+					*GCObject->GetFullName(),
+					GCObject->GetFlags(),
+					GCObject->GetInternalFlags(),
+					PackageId.ValueForDebugging(),
+					PublicExportKey->GetExportHash());
+			}
+		}
+	}
+
+	void VerifyPackageForRemoval(UPackage* Package)
+	{
+		check(Package && Package->CanBeImported());
+
+		FPackageId PackageId = Package->GetPackageId();
+		FLoadedPackageRef* PackageRef = FindPackageRef(PackageId);
+
+		UE_CLOG(!PackageRef, LogStreaming, Fatal,
+			TEXT("FGlobalImportStore::VerifyPackageForRemoval: %s (0x%llX) - ")
+			TEXT("A loaded package that can be imported by other packages is missing in GlobalImportStore."),
+			*Package->GetName(),
+			PackageId.ValueForDebugging());
+
+		UE_CLOG(PackageRef->GetRefCount() > 0, LogStreaming, Fatal,
+			TEXT("FGlobalImportStore::VerifyPackageForRemoval: %s (0x%llX) - ")
+			TEXT("Package removed while still being referenced, RefCount %d > 0."),
+			*Package->GetName(),
+			PackageId.ValueForDebugging(),
+			PackageRef->GetRefCount());
+
+		for (int32 ObjectIndex : PackageRef->GetPublicExportObjectIndices())
+		{
+			if (ObjectIndex >= 0)
+			{
+				UObject* Object = static_cast<UObject*>(GUObjectArray.IndexToObject(ObjectIndex)->Object);
+				ensureMsgf(!Object->HasAnyInternalFlags(EInternalObjectFlags::LoaderImport),
+						TEXT("FGlobalImportStore::VerifyPackageForRemoval: The loaded public export object '%s' with flags (ObjectFlags=%x, InternalObjectFlags=%x) and id 0x%llX:0x%llX is probably still referenced by the loader."),
+						*Object->GetFullName(),
+						Object->GetFlags(),
+						Object->GetInternalFlags());
+
+				FPublicExportKey* PublicExportKey = ObjectIndexToPublicExport.Find(ObjectIndex);
+				UE_CLOG(!PublicExportKey, LogStreaming, Fatal,
+					TEXT("FGlobalImportStore::VerifyPackageForRemoval: %s (0x%llX) - ")
+					TEXT("The loaded public export object '%s' is missing in GlobalImportStore."),
+					*Package->GetName(),
+					PackageId.ValueForDebugging(),
+					*Object->GetFullName());
+
+				FPackageId ObjectPackageId = PublicExportKey->GetPackageId();
+				UE_CLOG(ObjectPackageId != PackageId, LogStreaming, Fatal,
+					TEXT("FGlobalImportStore::VerifyPackageForRemoval: %s (0x%llX) - ")
+					TEXT("The loaded public export object '%s' has a mismatching package id 0x%llX in GlobalImportStore."),
+					*Package->GetName(),
+					PackageId.ValueForDebugging(),
+					*Object->GetFullName(),
+					ObjectPackageId.ValueForDebugging());
+
+				VerifyObjectForRemoval(Object);
 			}
 		}
 	}
 
 	inline UObject* FindPublicExportObjectUnchecked(const FPublicExportKey& Key)
 	{
-		FLoadedPackageRef* PackageRef = LoadedPackageStore.FindPackageRef(Key.GetPackageId());
+		FLoadedPackageRef* PackageRef = FindPackageRef(Key.GetPackageId());
 		if (!PackageRef)
 		{
 			return nullptr;
@@ -1155,34 +1219,39 @@ public:
 		check(ExportHash != 0);
 		int32 ObjectIndex = GUObjectArray.ObjectToIndex(Object);
 		FPublicExportKey Key = FPublicExportKey::MakeKey(PackageId, ExportHash);
-#if DO_CHECK
+
+		UObject* ExistingObject = FindPublicExportObjectUnchecked(Key);
+		if (ExistingObject && ExistingObject != Object)
 		{
-			UObject* ExistingObject = FindPublicExportObjectUnchecked(Key);
-			if (ExistingObject)
-			{
-				checkf(ExistingObject == Object,
-					TEXT("The constructed serialized object '%s' with index %d and id 0x%llX:0x%llX collides with the object '%s' in ImportStore. ")
-					TEXT("Reason unknown. Bug or hash collision?"),
-					Object ? *Object->GetFullName() : TEXT("null"),
-					ObjectIndex,
-					Key.GetPackageId().Value(), Key.GetExportHash(),
-					ExistingObject ? *ExistingObject->GetFullName() : TEXT("null"));
-			}
+			int32 ExistingObjectIndex = GUObjectArray.ObjectToIndex(ExistingObject);
 
-			FPublicExportKey* ExistingKey = ObjectIndexToPublicExport.Find(ObjectIndex);
-			if (ExistingKey)
-			{
-				checkf(*ExistingKey == Key,
-					TEXT("The constructed serialized object '%s' with index %d and id 0x%llX:0x%llX collides with the object '%s' in ImportStore. ")
-					TEXT("Reason unknown. Bug or hash collision?"),
-					Object ? *Object->GetFullName() : TEXT("null"),
-					ObjectIndex,
-					Key.GetPackageId().Value(), Key.GetExportHash(),
-					ExistingObject ? *ExistingObject->GetFullName() : TEXT("null"));
-			}
+			UE_LOG(LogStreaming, Warning,
+				TEXT("FGlobalImportStore::StoreGlobalObject: The constructed public export object '%s' with index %d and id 0x%llX:0x%llX collides with object '%s' (ObjectFlags=%X, InternalObjectFlags=%x) with index %d in GlobalImportStore. ")
+				TEXT("The existing object will be replaced since it was most likely renamed within the package after it was loaded."),
+				Object ? *Object->GetFullName() : TEXT("null"),
+				ObjectIndex,
+				Key.GetPackageId().Value(), Key.GetExportHash(),
+				*ExistingObject->GetFullName(),
+				ExistingObject->GetFlags(),
+				ExistingObject->GetInternalFlags(),
+				ExistingObjectIndex);
+
+			ExistingObject->AtomicallyClearInternalFlags(EInternalObjectFlags::LoaderImport);
+			ObjectIndexToPublicExport.Remove(ExistingObjectIndex);
 		}
-#endif
-		FLoadedPackageRef& PackageRef = LoadedPackageStore.FindPackageRefChecked(Key.GetPackageId());
+
+		FPublicExportKey* ExistingKey = ObjectIndexToPublicExport.Find(ObjectIndex);
+		if (ExistingKey)
+		{
+			UE_CLOG(*ExistingKey != Key, LogStreaming, Fatal,
+				TEXT("FGlobalImportStore::StoreGlobalObject: The constructed public export object '%s' with index %d and id 0x%llX:0x%llX already exists in GlobalImportStore but with a different key 0x%llX:0x%llX."),
+				Object ? *Object->GetFullName() : TEXT("null"),
+				ObjectIndex,
+				Key.GetPackageId().Value(), Key.GetExportHash(),
+				ExistingKey->GetPackageId().Value(), ExistingKey->GetExportHash());
+		}
+
+		FLoadedPackageRef& PackageRef = FindPackageRefChecked(Key.GetPackageId());
 		PackageRef.StorePublicExport(ExportHash, Object);
 		ObjectIndexToPublicExport.Add(ObjectIndex, Key);
 	}
@@ -1257,11 +1326,9 @@ struct FAsyncPackageLinkerLoadHeaderData
 struct FPackageImportStore
 {
 	FGlobalImportStore& GlobalImportStore;
-	FLoadedPackageStore& LoadedPackageStore;
 
-	FPackageImportStore(FGlobalImportStore& InGlobalImportStore, FLoadedPackageStore& InLoadedPackageStore)
+	FPackageImportStore(FGlobalImportStore& InGlobalImportStore)
 		: GlobalImportStore(InGlobalImportStore)
-		, LoadedPackageStore(InLoadedPackageStore)
 	{
 	}
 
@@ -1353,14 +1420,14 @@ struct FPackageImportStore
 public:
 	void AddImportedPackageReference(FPackageId ImportedPackageId, FName PackageNameIfKnown)
 	{
-		LoadedPackageStore.AddPackageRef(ImportedPackageId, PackageNameIfKnown);
+		GlobalImportStore.AddPackageRef(ImportedPackageId, PackageNameIfKnown);
 	}
 
 	void AddPackageReference(const FAsyncPackageDesc2& Desc)
 	{
 		if (Desc.bCanBeImported)
 		{
-			FLoadedPackageRef& PackageRef = LoadedPackageStore.AddPackageRef(Desc.UPackageId, Desc.UPackageName);
+			FLoadedPackageRef& PackageRef = GlobalImportStore.AddPackageRef(Desc.UPackageId, Desc.UPackageName);
 			PackageRef.ClearErrorFlags();
 		}
 	}
@@ -1369,7 +1436,7 @@ public:
 	{
 		for (const FPackageId& ImportedPackageId : ImportedPackageIds)
 		{
-			LoadedPackageStore.ReleasePackageRef(ImportedPackageId, Desc.UPackageId);
+			GlobalImportStore.ReleasePackageRef(ImportedPackageId, Desc.UPackageId);
 		}
 	}
 
@@ -1377,7 +1444,7 @@ public:
 	{
 		if (Desc.bCanBeImported)
 		{
-			LoadedPackageStore.ReleasePackageRef(Desc.UPackageId);
+			GlobalImportStore.ReleasePackageRef(Desc.UPackageId);
 		}
 	}
 };
@@ -2634,7 +2701,6 @@ private:
 	IAsyncPackageLoader* UncookedPackageLoader;
 
 	FPackageStore& PackageStore;
-	FLoadedPackageStore LoadedPackageStore;
 	FGlobalImportStore GlobalImportStore;
 	TMpscQueue<FPackageRequest> PackageRequestQueue;
 	TArray<FAsyncPackage2*> PendingPackages;
@@ -2923,6 +2989,8 @@ private:
 
 	void OnPreGarbageCollect();
 
+	void CollectUnreachableObjects(TArrayView<FUObjectItem*> UnreachableObjectItems, FUnreachableObjects& OutUnreachableObjects);
+
 	void RemoveUnreachableObjects(FUnreachableObjects& ObjectsToRemove);
 
 	bool ProcessPendingCDOs()
@@ -4061,7 +4129,7 @@ void FAsyncPackage2::ImportPackagesRecursiveInner(FAsyncLoadingThreadState2& Thr
 		constexpr bool bIsZenPackageImport = true;
 #endif
 
-		FLoadedPackageRef& ImportedPackageRef = ImportStore.LoadedPackageStore.FindPackageRefChecked(ImportedPackageId);
+		FLoadedPackageRef& ImportedPackageRef = AsyncLoadingThread.GlobalImportStore.FindPackageRefChecked(ImportedPackageId);
 #if WITH_EDITOR
 		if (AsyncLoadingThread.UncookedPackageLoader && ImportedPackageStatus == EPackageStoreEntryStatus::Ok && !ImportedPackageEntry.UncookedPackageName.IsNone())
 		{
@@ -4605,7 +4673,7 @@ EEventLoadNodeExecutionResult FAsyncPackage2::ProcessLinkerLoadPackageSummary(FA
 	HeaderData.ImportedPublicExportHashes = LinkerLoadState->LinkerLoadHeaderData.ImportedPublicExportHashes;
 	HeaderData.ImportMap = LinkerLoadState->LinkerLoadHeaderData.ImportMap;
 
-	FLoadedPackageRef& PackageRef = ImportStore.LoadedPackageStore.FindPackageRefChecked(Desc.UPackageId);
+	FLoadedPackageRef& PackageRef = AsyncLoadingThread.GlobalImportStore.FindPackageRefChecked(Desc.UPackageId);
 	if (!bLoadHasFailed)
 	{
 		PackageRef.ReserveSpaceForPublicExports(LinkerLoadState->LinkerLoadHeaderData.PublicExportHashes.Num());
@@ -4855,7 +4923,7 @@ EEventLoadNodeExecutionResult FAsyncPackage2::Event_ProcessPackageSummary(FAsync
 	{
 		if (Package->Desc.bCanBeImported)
 		{
-			FLoadedPackageRef& PackageRef = Package->ImportStore.LoadedPackageStore.FindPackageRefChecked(Package->Desc.UPackageId);
+			FLoadedPackageRef& PackageRef = Package->AsyncLoadingThread.GlobalImportStore.FindPackageRefChecked(Package->Desc.UPackageId);
 			PackageRef.SetHasFailed();
 		}
 	}
@@ -4893,7 +4961,7 @@ EEventLoadNodeExecutionResult FAsyncPackage2::Event_ProcessPackageSummary(FAsync
 				}
 			}
 #endif
-			FLoadedPackageRef& PackageRef = Package->ImportStore.LoadedPackageStore.FindPackageRefChecked(Package->Desc.UPackageId);
+			FLoadedPackageRef& PackageRef = Package->AsyncLoadingThread.GlobalImportStore.FindPackageRefChecked(Package->Desc.UPackageId);
 			if (PublicExportsCount)
 			{
 				PackageRef.ReserveSpaceForPublicExports(PublicExportsCount);
@@ -5549,7 +5617,7 @@ EEventLoadNodeExecutionResult FAsyncPackage2::Event_ExportsDone(FAsyncLoadingThr
 
 	if (!Package->bLoadHasFailed && Package->Desc.bCanBeImported)
 	{
-		FLoadedPackageRef& PackageRef = Package->AsyncLoadingThread.LoadedPackageStore.FindPackageRefChecked(Package->Desc.UPackageId);
+		FLoadedPackageRef& PackageRef = Package->AsyncLoadingThread.GlobalImportStore.FindPackageRefChecked(Package->Desc.UPackageId);
 		PackageRef.SetAllPublicExportsLoaded();
 	}
 
@@ -6539,7 +6607,6 @@ FAsyncLoadingThread2::FAsyncLoadingThread2(FIoDispatcher& InIoDispatcher, IAsync
 	, IoDispatcher(InIoDispatcher)
 	, UncookedPackageLoader(InUncookedPackageLoader)
 	, PackageStore(FPackageStore::Get())
-	, GlobalImportStore(LoadedPackageStore)
 {
 #if !WITH_EDITOR
 	IsEventDrivenLoaderEnabled(); // make sure the one time init inside runs
@@ -6889,10 +6956,9 @@ float FAsyncLoadingThread2::GetAsyncLoadPercentage(const FName& PackageName)
 	return LoadPercentage;
 }
 
-#if ALT2_VERIFY_ASYNC_FLAGS
-static void VerifyLoadFlagsWhenFinishedLoading()
+static void VerifyObjectLoadFlagsWhenFinishedLoading()
 {
-	TRACE_CPUPROFILER_EVENT_SCOPE(VerifyLoadFlagsWhenFinishedLoading);
+	TRACE_CPUPROFILER_EVENT_SCOPE(VerifyObjectLoadFlagsWhenFinishedLoading);
 
 	const EInternalObjectFlags AsyncFlags =
 		EInternalObjectFlags::Async | EInternalObjectFlags::AsyncLoading;
@@ -6901,7 +6967,8 @@ static void VerifyLoadFlagsWhenFinishedLoading()
 		EObjectFlags::RF_NeedLoad | EObjectFlags::RF_WillBeLoaded |
 		EObjectFlags::RF_NeedPostLoad | RF_NeedPostLoadSubobjects;
 
-	for (int32 ObjectIndex = 0; ObjectIndex < GUObjectArray.GetObjectArrayNum(); ++ObjectIndex)
+	ParallelFor(TEXT("VerifyObjectLoadFlagsDebugTask"), GUObjectArray.GetObjectArrayNum(), 512,
+		[AsyncFlags,LoadIntermediateFlags](int32 ObjectIndex)
 	{
 		FUObjectItem* ObjectItem = &GUObjectArray.GetObjectItemArrayUnsafe()[ObjectIndex];
 		if (UObject* Obj = static_cast<UObject*>(ObjectItem->Object))
@@ -6947,80 +7014,91 @@ static void VerifyLoadFlagsWhenFinishedLoading()
 					InternalFlags);
 			}
 		}
-	}
+	});
 	UE_LOG(LogStreaming, Log, TEXT("Verified load flags when finished active loading."));
 }
-#endif
 
-FORCENOINLINE static void CollectUnreachableObjects(
+void FAsyncLoadingThread2::CollectUnreachableObjects(
 	TArrayView<FUObjectItem*> UnreachableObjectItems,
 	FUnreachableObjects& OutUnreachableObjects)
 {
-	TRACE_CPUPROFILER_EVENT_SCOPE(CollectUnreachableObjects);
+	check(IsInGameThread());
 
 	OutUnreachableObjects.SetNum(UnreachableObjectItems.Num());
 
-	ParallelFor(TEXT("CollectUnreachableObjects"), UnreachableObjectItems.Num(), 2048, [&UnreachableObjectItems, &OutUnreachableObjects](int32 Index)
 	{
-		UObject* Object = static_cast<UObject*>(UnreachableObjectItems[Index]->Object);
-
-		FUnreachableObject& Item = OutUnreachableObjects[Index];
-		Item.ObjectIndex = GUObjectArray.ObjectToIndex(Object);
-		Item.ObjectName = Object->GetFName();
-#if ALT2_VERIFY_UNREACHABLE_OBJECTS
-		if (GGRemoveUnreachableObjectsFromGCNotifyOnGT)
-		{
-			Item.DebugObject = Object;
-		}
-#endif
-		if (!Object->GetOuter())
+		TRACE_CPUPROFILER_EVENT_SCOPE(CollectUnreachableObjects);
+		ParallelFor(TEXT("CollectUnreachableObjectsTask"), UnreachableObjectItems.Num(), 2048,
+			[&UnreachableObjectItems, &OutUnreachableObjects](int32 Index)
 		{
-			UPackage* Package = static_cast<UPackage*>(Object);
-			if (Package->bCanBeImported)
+			UObject* Object = static_cast<UObject*>(UnreachableObjectItems[Index]->Object);
+
+			FUnreachableObject& Item = OutUnreachableObjects[Index];
+			Item.ObjectIndex = GUObjectArray.ObjectToIndex(Object);
+			Item.ObjectName = Object->GetFName();
+
+			if (!Object->GetOuter())
 			{
-				Item.PackageId = Package->GetPackageId();
+				UPackage* Package = static_cast<UPackage*>(Object);
+				if (Package->bCanBeImported)
+				{
+					Item.PackageId = Package->GetPackageId();
+				}
 			}
-		}
 
 #if ALT2_ENABLE_LINKERLOAD_SUPPORT
-		// Clear garbage objects from linker export tables
-		// Normally done from UObject::BeginDestroy but we need to do it already here
-		if (FLinkerLoad* ObjectLinker = Object->GetLinker())
-		{
-#if WITH_EDITORONLY_DATA
-			// Make sure the linker entry stays as 'bExportLoadFailed' if the entry was marked as such, 
-			// Not doing this will cause crashes when resolving circular imports!
-			const int32 CachedLinkerIndex = Object->GetLinkerIndex();
-			bool bLinkerEntryWasInvalid = false;
-			if (ObjectLinker->ExportMap.IsValidIndex(CachedLinkerIndex))
+			// Clear garbage objects from linker export tables
+			// Normally done from UObject::BeginDestroy but we need to do it already here
+			if (FLinkerLoad* ObjectLinker = Object->GetLinker())
 			{
-				FObjectExport& ObjExport = ObjectLinker->ExportMap[CachedLinkerIndex];
-				bLinkerEntryWasInvalid = ObjExport.bExportLoadFailed;
-			}
+#if WITH_EDITORONLY_DATA
+				// Make sure the linker entry stays as 'bExportLoadFailed' if the entry was marked as such, 
+				// Not doing this will cause crashes when resolving circular imports!
+				const int32 CachedLinkerIndex = Object->GetLinkerIndex();
+				bool bLinkerEntryWasInvalid = false;
+				if (ObjectLinker->ExportMap.IsValidIndex(CachedLinkerIndex))
+				{
+					FObjectExport& ObjExport = ObjectLinker->ExportMap[CachedLinkerIndex];
+					bLinkerEntryWasInvalid = ObjExport.bExportLoadFailed;
+				}
 #endif // WITH_EDITORONLY_DATA
-			Object->SetLinker(NULL, INDEX_NONE);
+				Object->SetLinker(NULL, INDEX_NONE);
 #if WITH_EDITORONLY_DATA
-			if (bLinkerEntryWasInvalid)
-			{
-				FObjectExport& ObjExport = ObjectLinker->ExportMap[CachedLinkerIndex];
-				ObjExport.bExportLoadFailed = true;
-			}
+				if (bLinkerEntryWasInvalid)
+				{
+					FObjectExport& ObjExport = ObjectLinker->ExportMap[CachedLinkerIndex];
+					ObjExport.bExportLoadFailed = true;
+				}
 #endif // WITH_EDITORONLY_DATA
-		}
+			}
 #endif // ALT2_ENABLE_LINKERLOAD_SUPPORT
-	});
+		});
+	}
+
+	if (GVerifyUnreachableObjects)
+	{
+		TRACE_CPUPROFILER_EVENT_SCOPE(VerifyUnreachableObjects);
+		ParallelFor(TEXT("VerifyUnreachableObjectsDebugTask"), UnreachableObjectItems.Num(), 512,
+			[this, &UnreachableObjectItems](int32 Index)
+		{
+			UObject* Object = static_cast<UObject*>(UnreachableObjectItems[Index]->Object);
+			if (!Object->GetOuter())
+			{
+				UPackage* Package = static_cast<UPackage*>(Object);
+				if (Package->bCanBeImported)
+				{
+					GlobalImportStore.VerifyPackageForRemoval(Package);
+				}
+				GlobalImportStore.VerifyObjectForRemoval(Object);
+			}
+		});
+	}
 }
 
 void FAsyncLoadingThread2::OnLeakedPackageRename(UPackage* Package)
 {
 	TRACE_CPUPROFILER_EVENT_SCOPE(OnLeakedPackageRename);
 	check(IsInGameThread());
-	
-	if (!FGCCSyncObject::Get().IsGCLocked())
-	{
-		// Flush async loading so that we can be sure nothing is modifying the stores, and that nothing is depending on this package
-		FlushAsyncLoading(); 
-	}
 
 	// We don't care about levelstreaming /Temp/ packages that are never imported by other packages
 	if (!Package->CanBeImported())
@@ -7028,41 +7106,33 @@ void FAsyncLoadingThread2::OnLeakedPackageRename(UPackage* Package)
 		return;
 	}
 
+	// Code such as LoadMap or LevelStreaming is about to rename a loaded package which was detected as leaking so that we can load another copy of it.
 	UE_LOG(LogStreaming, Display, TEXT("Renaming leaked package %s (0x%llX)"), *Package->GetName(), Package->GetPackageId().ValueForDebugging());
 
-	// Flush the delete queue to clear references to the renamed package, if there are still references after this we're not done with loading the package and it shouldn't have been renamed
-	// If this function is called from within a GC the ALT will be inactive and not touching the deferred delete queue or the unreachable objects array
-	// If this function is called outside of a GC (CheckAndHandleStaleWorldObjectReferences) the call to FlushAsyncLoading above will empty the deferred delete queue but the ALT can still be processing
-	// the unreachable objects array so we need to protect that with a lock
-	ProcessDeferredDeletePackagesQueue();
+	if (!FGCCSyncObject::Get().IsGCLocked())
+	{
+		// If this function is called outside of a GC (CheckAndHandleStaleWorldObjectReferences), then flush async loading so that we can be sure nothing is modifying the stores,
+		// and that nothing is depending on this package. This will also empty the deferred delete queue.
+		FlushAsyncLoading(); 
+	}
+	else
+	{
+		// If this function is called from within a GC the ALT will be inactive and not touching the deferred delete queue.
+		// Flush the delete queue to clear references to the renamed package, if there are still references after this we're not done with loading the package and it shouldn't have been renamed.
+		ProcessDeferredDeletePackagesQueue();
+	}
 
+	// We need to protect UnreachableObjects since the GC will wake up the ALT from NotifyUnreachableObjects to run RemoveUnreachableObjects on the ALT.
+	// FlushAsyncLoading and ProcessDeferredDeletePackagesQueue will not empty the UnreachableObjects array that was just filled by GC.
 	FScopeLock UnreachableObjectsLock(&UnreachableObjectsCritical);
 
 	// unreachable objects from last GC should typically have been processed already,
 	// if not handle them here before processing new ones
 	RemoveUnreachableObjects(UnreachableObjects);
 
-	// If a package that can be imported was leaked and renamed,
-	// then it must exist in the loaded package store at this point since it is normally only trimmed during GC.
-	FLoadedPackageRef* PackageRef = LoadedPackageStore.FindPackageRef(Package->GetPackageId());
-	if (!ensureAlwaysMsgf(PackageRef, TEXT("Package %s (0x%llX)"), *Package->GetName(), Package->GetPackageId().ValueForDebugging())) 
-	{
-		return;
-	}
-
-	UE_CLOG(PackageRef->GetRefCount() != 0, LogStreaming, Fatal, TEXT("Renaming leaked package %s (0x%llX) that's still referenced by active loading (RefCount=%d)"), *Package->GetName(), Package->GetPackageId().ValueForDebugging(), PackageRef->GetRefCount());
-
-	// Code such as LoadMap or LevelStreaming is about to rename a loaded package which was detected as leaking so that we can load another copy of it.
-	// We should not have any loading happening at present, so we can remove these objects from our stores 
-	TArray<FUObjectItem*> LeakedObjectItems;
-	LeakedObjectItems.Emplace(GUObjectArray.ObjectToObjectItem(Package));
-	ForEachObjectWithOuter(Package, [&LeakedObjectItems](UObject* Obj) {
-		LeakedObjectItems.Emplace(GUObjectArray.ObjectToObjectItem(Obj));
-	});
-
-	FUnreachableObjects LeakedUnreachableObjects;
-	CollectUnreachableObjects(LeakedObjectItems, LeakedUnreachableObjects);
-	RemoveUnreachableObjects(LeakedUnreachableObjects);
+	// If a package that can be imported was leaked and renamed, then go ahead and remove it.
+	// Note that it must exist in the global import store at this point since it is normally only trimmed during GC.
+	GlobalImportStore.RemovePackage(Package);
 
 	// Clear the CanBeImportedFlag so that this package is only removed once,
 	// else we would try and remove it again during GC, which would instead remove the reloaded package if it exists
@@ -7086,16 +7156,16 @@ void FAsyncLoadingThread2::RemoveUnreachableObjects(FUnreachableObjects& Objects
 
 	const int32 ObjectCount = ObjectsToRemove.Num();
 
-	const int32 OldLoadedPackageCount = LoadedPackageStore.NumTracked();
+	const int32 OldLoadedPackageCount = GlobalImportStore.GetStoredPackagesCount();
 	const int32 OldPublicExportCount = GlobalImportStore.GetStoredPublicExportsCount();
 
 	const double StartTime = FPlatformTime::Seconds();
 
+	GlobalImportStore.RemovePackages(ObjectsToRemove);
 	GlobalImportStore.RemovePublicExports(ObjectsToRemove);
-	LoadedPackageStore.RemovePackages(ObjectsToRemove);
 	ObjectsToRemove.Reset();
 
-	const int32 NewLoadedPackageCount = LoadedPackageStore.NumTracked();
+	const int32 NewLoadedPackageCount = GlobalImportStore.GetStoredPackagesCount();
 	const int32 NewPublicExportCount = GlobalImportStore.GetStoredPublicExportsCount();
 	const int32 RemovedLoadedPackageCount = OldLoadedPackageCount - NewLoadedPackageCount;
 	const int32 RemovedPublicExportCount = OldPublicExportCount - NewPublicExportCount;
@@ -7128,15 +7198,13 @@ void FAsyncLoadingThread2::NotifyUnreachableObjects(const TArrayView<FUObjectIte
 
 	CollectUnreachableObjects(UnreachableObjectItems, UnreachableObjects);
 
-#if ALT2_VERIFY_ASYNC_FLAGS
-	if (!IsAsyncLoading())
+	if (GVerifyObjectLoadFlags && !IsAsyncLoading())
 	{
-		LoadedPackageStore.VerifyLoadedPackages();
-		VerifyLoadFlagsWhenFinishedLoading();
+		GlobalImportStore.VerifyLoadedPackages();
+		VerifyObjectLoadFlagsWhenFinishedLoading();
 	}
-#endif
 
-	if (GGRemoveUnreachableObjectsFromGCNotifyOnGT)
+	if (GRemoveUnreachableObjectsOnGT)
 	{
 		RemoveUnreachableObjects(UnreachableObjects);
 	}
@@ -7212,7 +7280,7 @@ FAsyncPackage2::FAsyncPackage2(
 : Desc(InDesc)
 , AsyncLoadingThread(InAsyncLoadingThread)
 , GraphAllocator(InGraphAllocator)
-, ImportStore(AsyncLoadingThread.GlobalImportStore, AsyncLoadingThread.LoadedPackageStore)
+, ImportStore(AsyncLoadingThread.GlobalImportStore)
 {
 	TRACE_LOADTIME_NEW_ASYNC_PACKAGE(this);
 	AddRequestID(Desc.RequestID);
@@ -7407,7 +7475,7 @@ void FAsyncPackage2::CreateUPackage()
 {
 	check(!LinkerRoot);
 
-	// temp packages are never stored or found in loaded package store
+	// temp packages are never stored or found in global import store
 	FLoadedPackageRef* PackageRef = nullptr;
 
 	// Try to find existing package or create it if not already present.
@@ -7416,7 +7484,7 @@ void FAsyncPackage2::CreateUPackage()
 		TRACE_CPUPROFILER_EVENT_SCOPE(UPackageFind);
 		if (Desc.bCanBeImported)
 		{
-			PackageRef = ImportStore.LoadedPackageStore.FindPackageRef(Desc.UPackageId);
+			PackageRef = AsyncLoadingThread.GlobalImportStore.FindPackageRef(Desc.UPackageId);
 			UE_ASYNC_PACKAGE_CLOG(!PackageRef, Fatal, Desc, TEXT("CreateUPackage"), TEXT("Package has been destroyed by GC."));
 			LinkerRoot = PackageRef->GetPackage();
 #if DO_CHECK
diff --git a/Engine/Source/Runtime/CoreUObject/Public/Serialization/AsyncLoading2.h b/Engine/Source/Runtime/CoreUObject/Public/Serialization/AsyncLoading2.h
--- a/Engine/Source/Runtime/CoreUObject/Public/Serialization/AsyncLoading2.h
+++ b/Engine/Source/Runtime/CoreUObject/Public/Serialization/AsyncLoading2.h
@@ -189,6 +189,12 @@ public:
 			GetPackageId() == Other.GetPackageId();
 	}
 
+	inline bool operator!=(const FPublicExportKey& Other) const
+	{
+		return GetExportHash() != Other.GetExportHash() ||
+			GetPackageId() != Other.GetPackageId();
+	}
+
 	inline friend uint32 GetTypeHash(const FPublicExportKey& In)
 	{
 		return HashCombine(GetTypeHash(In.GetPackageId()), GetTypeHash(In.GetExportHash()));
