Avoid rehashing editor bulkdata payloads when saving to a transaction buffer.

#rnx

- Converting the in memory payload (FSharedBuffer) to a non-compressed FCompressedBuffer will cause the payload to be hashed. Given that we are not compressing the buffer and it will just be loaded and converted back to FSharedBuffer, this hash is pointless.
- This first started when we changed editor bulkdata to serialize the payload as a FCompressedBuffer during normal serialization, since the transaction serialization was using the same ::SerializeData method as FSharedBuffer does not have native FArchive serialization support.
- We can add our own serialization support for FSharedBuffer and then not convert it to FCompressedBuffer to avoid the whole problem.
- Profiling showed that this changed reduces the time spent serializing a 100MiB payload from 100ms to 30ms (average)

diff --git a/Engine/Source/Runtime/CoreUObject/Private/Serialization/EditorBulkData.cpp b/Engine/Source/Runtime/CoreUObject/Private/Serialization/EditorBulkData.cpp
--- a/Engine/Source/Runtime/CoreUObject/Private/Serialization/EditorBulkData.cpp
+++ b/Engine/Source/Runtime/CoreUObject/Private/Serialization/EditorBulkData.cpp
@@ -297,6 +297,31 @@ void UpdateArchiveData(FArchive& Ar, int64 DataPosition, DataType& Data)
 	Ar.Seek(OriginalPosition);
 }
 
+FArchive& operator<<(FArchive& Ar, FSharedBuffer& Buffer)
+{
+	if (Ar.IsLoading())
+	{
+		int64 BufferLength;
+		Ar << BufferLength;
+
+		FUniqueBuffer MutableBuffer = FUniqueBuffer::Alloc(BufferLength);
+
+		Ar.Serialize(MutableBuffer.GetData(), BufferLength);
+
+		Buffer = MutableBuffer.MoveToShared();
+	}
+	else if (Ar.IsSaving())
+	{
+		int64 BufferLength = (int64)Buffer.GetSize();
+		Ar << BufferLength;
+
+		// Need to remove const due to FArchive API
+		Ar.Serialize(const_cast<void*>(Buffer.GetData()), BufferLength);
+	}
+
+	return Ar;
+}
+
 /** Utility for accessing IVirtualizationSourceControlUtilities from the modular feature system. */
 UE::Virtualization::Experimental::IVirtualizationSourceControlUtilities* GetSourceControlInterface()
 {
@@ -808,8 +833,7 @@ void FEditorBulkData::Serialize(FArchive& Ar, UObject* Owner, bool bAllowRegiste
 			{
 				if (bPayloadInArchive)
 				{
-					FCompressedBuffer CompressedPayload = FCompressedBuffer::Compress(Payload, ECompressedBufferCompressor::NotSet, ECompressedBufferCompressionLevel::None);
-					SerializeData(Ar, CompressedPayload, Flags);
+					Ar << Payload;
 				}
 			}
 			else
@@ -821,11 +845,9 @@ void FEditorBulkData::Serialize(FArchive& Ar, UObject* Owner, bool bAllowRegiste
 					BulkDataId = FGuid::NewGuid();
 				}
 
-				FCompressedBuffer CompressedPayload;
 				if (bPayloadInArchive)
 				{
-					SerializeData(Ar, CompressedPayload, Flags);
-					Payload = CompressedPayload.Decompress();
+					Ar << Payload;
 				}
 				else
 				{
