[CoreUObject] fix for TObjectPtr equals operator with multiple inheritance

resolved pointers where comparing uintptr_t instead of the raw pointers which allow for inheritance

diff --git a/Engine/Source/Runtime/CoreUObject/Public/UObject/ObjectPtr.h b/Engine/Source/Runtime/CoreUObject/Public/UObject/ObjectPtr.h
--- a/Engine/Source/Runtime/CoreUObject/Public/UObject/ObjectPtr.h
+++ b/Engine/Source/Runtime/CoreUObject/Public/UObject/ObjectPtr.h
@@ -431,6 +431,7 @@ public:
 
 private:
 	FORCEINLINE T* GetNoReadNoCheck() const { return (T*)(ResolveObjectHandleNoReadNoCheck(ObjectPtr.GetHandleRef())); }
+	FORCEINLINE T* GetNoResolveNoCheck() const { return (T*)(ReadObjectHandlePointerNoCheck(ObjectPtr.GetHandleRef())); }
 
 	// @TODO: OBJPTR: There is a risk of a gap in access tracking here.  The caller may get a mutable pointer, write to it, then
 	//			read from it.  That last read would happen without an access being recorded.  Not sure if there is a good way
@@ -537,12 +538,12 @@ namespace ObjectPtr_Private
 	}
 
 	template <typename T, typename U>
-	FORCEINLINE bool IsObjectPtrEqualToRawPtrOfRelatedType(const TObjectPtr<T>& Ptr, const U* Other)
+	bool IsObjectPtrEqualToRawPtrOfRelatedType(const TObjectPtr<T>& Ptr, const U* Other)
 	{
 #if UE_WITH_OBJECT_HANDLE_LATE_RESOLVE
 		if (Ptr.IsResolved())
 		{
-			return Ptr.GetHandle().PointerOrRef == uintptr_t(Other);
+			return Ptr.GetNoResolveNoCheck() == Other;
 		}
 		else if (!Other) //avoids resolving if Other is null
 		{
@@ -561,7 +562,7 @@ namespace ObjectPtr_Private
 			decltype(CoerceToPointer<T>(std::declval<U>()) == std::declval<const T*>())
 		>* = nullptr
 	>
-	FORCEINLINE bool IsObjectPtrEqual(const TObjectPtr<T>& Ptr, U&& Other)
+	bool IsObjectPtrEqual(const TObjectPtr<T>& Ptr, U&& Other)
 	{
 		// This function deliberately avoids the tracking code path as we are only doing
 		// a shallow pointer comparison.
diff --git a/Engine/Source/Runtime/CoreUObject/Tests/ObjectPtrTest.cpp b/Engine/Source/Runtime/CoreUObject/Tests/ObjectPtrTest.cpp
--- a/Engine/Source/Runtime/CoreUObject/Tests/ObjectPtrTest.cpp
+++ b/Engine/Source/Runtime/CoreUObject/Tests/ObjectPtrTest.cpp
@@ -149,7 +149,6 @@ public:
 protected:
 };
 
-
 TEST_CASE_METHOD(FObjectPtrTestBase, "CoreUObject::TObjectPtr::Null Behavior", "[CoreUObject][ObjectPtr]")
 {
 	TObjectPtr<UObject> NullObjectPtr(nullptr);
@@ -358,6 +357,62 @@ TEST_CASE_METHOD(FObjectPtrTestBase, "CoreUObject::TObjectPtr::Long Path", "[Cor
 // 	return true;
 // }
 
+
+//bunch of class to reproduce multiple inheritance issues
+class FTestBaseClass
+{
+public:
+	virtual ~FTestBaseClass() = default;
+	virtual void VirtFunc() { };
+};
+
+class UMiddleClass : public UObject, public FTestBaseClass
+{
+	DECLARE_CLASS_INTRINSIC(UMiddleClass, UObject, CLASS_MatchedSerializers, TEXT("/Script/CoreUObject"))
+
+public:
+	virtual void VirtFunc() override { };
+};
+
+
+class FAnotherBaseClass
+{
+public:
+	virtual ~FAnotherBaseClass() = default;
+	virtual void AnotherVirtFunc() { };
+};
+
+
+class UDerrivedClass : public UMiddleClass, public FAnotherBaseClass
+{
+	DECLARE_CLASS_INTRINSIC(UDerrivedClass, UMiddleClass, CLASS_MatchedSerializers, TEXT("/Script/CoreUObject"))
+public:
+	virtual void AnotherVirtFunc() override { };
+};
+
+IMPLEMENT_CORE_INTRINSIC_CLASS(UMiddleClass, UObject,
+	{
+	}
+);
+
+IMPLEMENT_CORE_INTRINSIC_CLASS(UDerrivedClass, UMiddleClass,
+	{
+	}
+);
+
+TEST_CASE("CoreUObject::TObjectPtr::TestEquals")
+{
+	const FName TestPackageName(TEXT("/Engine/Test/TestEquals/Transient"));
+	UPackage* TestPackage = NewObject<UPackage>(nullptr, TestPackageName, RF_Transient);
+	TestPackage->AddToRoot();
+	UDerrivedClass* Obj = NewObject<UDerrivedClass>(TestPackage, TEXT("DefaultSerializeObject"));
+
+	FTestBaseClass* BasePtr = Obj;
+	TObjectPtr<UDerrivedClass> ObjPtr(Obj);
+
+	CHECK(BasePtr == ObjPtr);
+}
+
 #endif
 
 class UForwardDeclaredObjDerived: public UObject {};
