[CoreUObject] compile fix for Linux.
moved test classes outside of namespace as that was causing problems for the compiler

diff --git a/Engine/Source/Runtime/CoreUObject/Tests/ObjectPtrTest.cpp b/Engine/Source/Runtime/CoreUObject/Tests/ObjectPtrTest.cpp
--- a/Engine/Source/Runtime/CoreUObject/Tests/ObjectPtrTest.cpp
+++ b/Engine/Source/Runtime/CoreUObject/Tests/ObjectPtrTest.cpp
@@ -358,7 +358,18 @@ TEST_CASE_METHOD(FObjectPtrTestBase, "CoreUObject::TObjectPtr::Long Path", "[Cor
 // }
 
 
-//bunch of class to reproduce multiple inheritance issues
+#endif
+
+class UForwardDeclaredObjDerived: public UObject {};
+class FForwardDeclaredNotObjDerived {};
+
+}
+
+#if WITH_LOW_LEVEL_TESTS
+
+//bunch of class to reproduce multiple inheritance issues.
+//need to live outside of the namespace be IMPLEMENT_CORE_INTRINSIC_CLASS don't like namespaces
+
 class FTestBaseClass
 {
 public:
@@ -413,9 +424,4 @@ TEST_CASE("CoreUObject::TObjectPtr::TestEquals")
 	CHECK(BasePtr == ObjPtr);
 }
 
-#endif
-
-class UForwardDeclaredObjDerived: public UObject {};
-class FForwardDeclaredNotObjDerived {};
-
-}
+#endif 
