Fix for cotf server crashing and client exiting immediately when Quick Launching Lyra with cotf enabled.

#rnx

diff --git a/Engine/Source/Runtime/CoreUObject/Private/Serialization/AsyncLoading2.cpp b/Engine/Source/Runtime/CoreUObject/Private/Serialization/AsyncLoading2.cpp
--- a/Engine/Source/Runtime/CoreUObject/Private/Serialization/AsyncLoading2.cpp
+++ b/Engine/Source/Runtime/CoreUObject/Private/Serialization/AsyncLoading2.cpp
@@ -3399,12 +3399,27 @@ bool FAsyncLoadingThread2::CreateAsyncPackagesFromQueue(FAsyncLoadingThreadState
 	{
 		FAsyncPackage2* PendingPackage = *It;
 		FPackageStoreEntry PackageEntry;
-		if (EPackageStoreEntryStatus::Ok == PackageStore.GetPackageStoreEntry(PendingPackage->Desc.PackageIdToLoad, PackageEntry))
+		EPackageStoreEntryStatus PendingPackageStatus = PackageStore.GetPackageStoreEntry(PendingPackage->Desc.PackageIdToLoad, PackageEntry);
+		if (PendingPackageStatus == EPackageStoreEntryStatus::Ok)
 		{
 			InitializeAsyncPackageFromPackageStore(ThreadState, IoBatch, PendingPackage, PackageEntry, TArrayView<FName>());
 			PendingPackage->StartLoading(IoBatch);
 			It.RemoveCurrent();
 		}
+		else if (PendingPackageStatus == EPackageStoreEntryStatus::Missing)
+		{
+			// Initialize package with a fake package store entry
+			FPackageStoreEntry FakePackageEntry;
+			FakePackageEntry.ExportInfo.ExportCount = 1;
+			FakePackageEntry.ExportInfo.ExportBundleCount = 1;
+			InitializeAsyncPackageFromPackageStore(ThreadState, IoBatch, PendingPackage, FakePackageEntry, TArrayView<FName>());
+			// Simulate StartLoading() getting back a failed IoRequest and let it go through all package states
+			PendingPackage->AsyncPackageLoadingState = EAsyncPackageLoadingState2::WaitingForIo;
+			PendingPackage->bLoadHasFailed = true;
+			PendingPackage->GetPackageNode(EEventLoadNode2::Package_ProcessSummary).ReleaseBarrier();
+			// Remove from PendingPackages
+			It.RemoveCurrent();
+		}
 	}
 	for (;;)
 	{
@@ -5032,6 +5047,7 @@ EEventLoadNodeExecutionResult FAsyncPackage2::Event_ProcessPackageSummary(FAsync
 EEventLoadNodeExecutionResult FAsyncPackage2::Event_SetupDependencies(FAsyncLoadingThreadState2& ThreadState, FAsyncPackage2* Package, int32)
 {
 	TRACE_CPUPROFILER_EVENT_SCOPE(Event_SetupDependencies);
+	UE_ASYNC_PACKAGE_DEBUG(Package->Desc);
 	check(Package->AsyncPackageLoadingState == EAsyncPackageLoadingState2::SetupDependencies);
 
 #if ALT2_ENABLE_LINKERLOAD_SUPPORT
diff --git a/Engine/Source/Runtime/Engine/Private/Materials/MaterialShared.cpp b/Engine/Source/Runtime/Engine/Private/Materials/MaterialShared.cpp
--- a/Engine/Source/Runtime/Engine/Private/Materials/MaterialShared.cpp
+++ b/Engine/Source/Runtime/Engine/Private/Materials/MaterialShared.cpp
@@ -2864,13 +2864,18 @@ void FMaterial::CacheGivenTypes(EShaderPlatform Platform, const TArray<const FVe
 		return;
 	}
 
+	if (bGameThreadShaderMapIsComplete)
+	{
+		UE_LOG(LogMaterial, Verbose, TEXT("Cache given types for a material resource %s with a complete ShaderMap"), *GetFriendlyName());
+		return;
+	}
+
 	if (GameThreadShaderMap)
 	{
 		TRACE_CPUPROFILER_EVENT_SCOPE(FMaterial::CacheGivenTypes);
 		check(IsInGameThread());
 		checkf(ShaderTypes.Num() == VFTypes.Num(), TEXT("The size of the shader type array and vertex factory type array must match."));
 		checkf(PipelineTypes.Num() == ShaderTypes.Num(), TEXT("The size of the pipeline type array and shader type array must match.  Pass in null entries if pipelines are not used."));
-		checkf(GameThreadShaderMap, TEXT("Shader map is not initialized.  Please call CacheShaders first."));
 		checkf(GetGameThreadCompilingShaderMapId() != 0, TEXT("Material is not prepared to compile yet.  Please call CacheShaders first."));
 
 		TArray<FShaderCommonCompileJobPtr> CompileJobs;
