ZenLoader: Fixed incorrect early out when determining if a package should be included in the active sync load context or not

diff --git a/Engine/Source/Runtime/CoreUObject/Private/Serialization/AsyncLoading2.cpp b/Engine/Source/Runtime/CoreUObject/Private/Serialization/AsyncLoading2.cpp
--- a/Engine/Source/Runtime/CoreUObject/Private/Serialization/AsyncLoading2.cpp
+++ b/Engine/Source/Runtime/CoreUObject/Private/Serialization/AsyncLoading2.cpp
@@ -3369,7 +3369,7 @@ void FAsyncLoadingThread2::IncludePackageInSyncLoadContext(uint64 ContextId, FAs
 		return;
 	}
 
-	if (Package->AsyncPackageLoadingState >= EAsyncPackageLoadingState2::DeferredPostLoadDone)
+	if (Package->AsyncPackageLoadingState > EAsyncPackageLoadingState2::DeferredPostLoadDone)
 	{
 		return;
 	}
