DerivedData: Added FDerivedData::IsCooked

#rnx

diff --git a/Engine/Source/Runtime/CoreUObject/Public/Serialization/DerivedData.h b/Engine/Source/Runtime/CoreUObject/Public/Serialization/DerivedData.h
--- a/Engine/Source/Runtime/CoreUObject/Public/Serialization/DerivedData.h
+++ b/Engine/Source/Runtime/CoreUObject/Public/Serialization/DerivedData.h
@@ -102,6 +102,9 @@ public:
 	/** Returns true if this is a non-null reference. */
 	inline bool HasData() const { return !IsNull(); }
 
+	/** Returns true if this is a cooked reference. */
+	inline bool IsCooked() const { return !!CookedData; }
+
 	/** Returns the flags, which are mainly relevant for staged references. */
 	inline EDerivedDataFlags GetFlags() const { return CookedData.Flags; }
 
