DerivedData: Serialize as bulk data until derived data serialization is implemented

#rnx

diff --git a/Engine/Source/Runtime/CoreUObject/Private/Serialization/DerivedData.cpp b/Engine/Source/Runtime/CoreUObject/Private/Serialization/DerivedData.cpp
--- a/Engine/Source/Runtime/CoreUObject/Private/Serialization/DerivedData.cpp
+++ b/Engine/Source/Runtime/CoreUObject/Private/Serialization/DerivedData.cpp
@@ -25,6 +25,16 @@
 #include "UObject/LinkerSave.h"
 #endif // WITH_EDITORONLY_DATA
 
+
+///////////////////////////////////////////////////////////////////////////////////////////////////////////////////////
+// TEMPORARY SERIALIZATION AS BULK DATA
+#if WITH_EDITORONLY_DATA
+#include "Experimental/Async/LazyEvent.h"
+#include "Serialization/BulkData.h"
+#endif // WITH_EDITORONLY_DATA
+// TEMPORARY SERIALIZATION AS BULK DATA
+///////////////////////////////////////////////////////////////////////////////////////////////////////////////////////
+
 static_assert(sizeof(UE::FDerivedData) == 32 + sizeof(void*) * WITH_EDITORONLY_DATA);
 
 namespace UE::DerivedData::Private { class FIoResponseDispatcher; }
@@ -134,6 +144,12 @@ public:
 		::Visit(Visitor, Data);
 	}
 
+	///////////////////////////////////////////////////////////////////////////////////////////////////////////////////
+	// TEMPORARY SERIALIZATION AS BULK DATA
+	FByteBulkData BulkData;
+	// TEMPORARY SERIALIZATION AS BULK DATA
+	///////////////////////////////////////////////////////////////////////////////////////////////////////////////////
+
 private:
 	FSharedString Name;
 	TVariant<FCompositeBufferWithHash, FCompressedBuffer, FCacheKeyWithId> Data;
@@ -286,12 +302,74 @@ void FDerivedData::Serialize(FArchive& Ar, UObject* Owner)
 	}
 
 #if WITH_EDITORONLY_DATA
-	if (EditorData)
+	const EDerivedDataFlags Flags = GetFlags();
+	if (EditorData && EnumHasAnyFlags(Flags, EDerivedDataFlags::Required | EDerivedDataFlags::Optional))
 	{
 		checkf(Ar.IsSaving() && Ar.IsCooking(),
 			TEXT("FEditorData for FDerivedData only supports saving to cooked packages. %s"), *WriteToString<256>(*this));
 		FLinkerSave* Linker = Cast<FLinkerSave>(Ar.GetLinker());
 		checkf(Linker, TEXT("Serializing FDerivedData requires a linker. %s"), *WriteToString<256>(*this));
+
+		///////////////////////////////////////////////////////////////////////////////////////////////////////////////
+		// TEMPORARY SERIALIZATION AS BULK DATA
+
+		FCookedData LocalCookedData = Linker->AddDerivedData(*this).CookedData;
+		LocalCookedData.ChunkId[8] = 0;
+		LocalCookedData.ChunkId[9] = 0;
+		LocalCookedData.ChunkId[10] = 0;
+		LocalCookedData.ChunkId[11] = uint8(EIoChunkType::BulkData);
+
+		uint32 BulkDataFlags = BULKDATA_PayloadAtEndOfFile | BULKDATA_Size64Bit | BULKDATA_NoOffsetFixUp;
+		// Required AND Optional is the equivalent of DuplicateNonOptional, which is implemented by writing appending
+		// the bulk data twice, but that may have different offsets, and this stores one offset. Only save as Required.
+		if (EnumHasAnyFlags(GetFlags(), EDerivedDataFlags::Optional) && !EnumHasAnyFlags(Flags, EDerivedDataFlags::Required))
+		{
+			BulkDataFlags |= BULKDATA_OptionalPayload;
+			LocalCookedData.ChunkId[11] = uint8(EIoChunkType::OptionalBulkData);
+		}
+		// Bulk Data does not support memory-mapped optional data. Optional takes priority over memory-mapped here.
+		else if (EnumHasAnyFlags(Flags, EDerivedDataFlags::MemoryMapped))
+		{
+			BulkDataFlags |= BULKDATA_MemoryMappedPayload;
+			LocalCookedData.ChunkId[11] = uint8(EIoChunkType::MemoryMappedBulkData);
+		}
+
+		// Blocking read to populate the bulk data to be serialized later by the linker.
+		{
+			FLazyEvent BatchComplete(EEventMode::ManualReset);
+			FDerivedDataIoBatch Batch;
+			FDerivedDataIoRequest Request = Batch.Read(*this);
+			FDerivedDataIoResponse Response;
+			Batch.Dispatch(Response, [&BatchComplete] { BatchComplete.Trigger(); });
+			BatchComplete.Wait();
+
+			const FSharedBuffer Data = Response.GetData(Request);
+			checkf(Data, TEXT("FDerivedData failed to save because data was not fetched. %s"), *WriteToString<256>(*this));
+
+			FByteBulkData& BulkData = EditorData->BulkData;
+			BulkData.Lock(LOCK_READ_WRITE);
+			MakeMemoryView(BulkData.Realloc(Data.GetSize()), Data.GetSize()).CopyFrom(Data);
+			BulkData.Unlock();
+			BulkData.SetBulkDataFlags(BulkDataFlags | BULKDATA_Force_NOT_InlinePayload);
+		}
+
+		const int64 Offset = Ar.Tell();
+
+		FLinkerSave::FBulkDataStorageInfo& BulkDataInfo = Linker->BulkDataToAppend.AddZeroed_GetRef();
+		BulkDataInfo.BulkDataOffsetInFilePos = Offset;
+		// Size will be written as 64 bits because of the flag BULKDATA_Size64Bit.
+		BulkDataInfo.BulkDataSizeOnDiskPos = Offset + sizeof(uint64);
+		// Use the offset of the offset because SaveBulkData writes flags before offset and size anyway.
+		BulkDataInfo.BulkDataFlagsPos = Offset;
+		BulkDataInfo.BulkDataFlags = BulkDataFlags;
+		BulkDataInfo.BulkDataFileRegionType = EFileRegionType::None;
+		BulkDataInfo.BulkData = &EditorData->BulkData;
+
+		return LocalCookedData.Serialize(Ar);
+
+		// TEMPORARY SERIALIZATION AS BULK DATA
+		///////////////////////////////////////////////////////////////////////////////////////////////////////////////
+
 		return Linker->AddDerivedData(*this).CookedData.Serialize(Ar);
 	}
 #endif
