Print script warning caller name.
Printing full stack trace could quickly bloat log, but we can print only top level caller actor to narrow down what BP is responsible for the warning.
Main use case is to catch BP which causes Script Msg: Divide by zero: Divide_VectorFloat and other UKismetMathLibrary errors.

#tests replay run

diff --git a/Engine/Source/Runtime/CoreUObject/Private/UObject/ScriptCore.cpp b/Engine/Source/Runtime/CoreUObject/Private/UObject/ScriptCore.cpp
--- a/Engine/Source/Runtime/CoreUObject/Private/UObject/ScriptCore.cpp
+++ b/Engine/Source/Runtime/CoreUObject/Private/UObject/ScriptCore.cpp
@@ -631,17 +631,29 @@ void FFrame::KismetExecutionMessage(const TCHAR* Message, ELogVerbosity::Type Ve
 	// Tracking down some places that display warnings but no message..
 	ensureAlways(Verbosity > ELogVerbosity::Warning || FCString::Strlen(Message) > 0);
 
-#if DO_BLUEPRINT_GUARD
+#if !UE_BUILD_SHIPPING
 	// Show the stack for fatal/error, and on warning if that option is enabled
-	if (Verbosity <= ELogVerbosity::Error || (ShowKismetScriptStackOnWarnings() && Verbosity == ELogVerbosity::Warning))
+	auto PopulateStackString = [](TStringBuilder<4096>& ScriptStack, ELogVerbosity::Type Verbosity)
 	{
-		ScriptStack = TEXT("Script call stack:\n");
-		GetScriptCallstack(ScriptStack);
-	}
-	else if (Verbosity == ELogVerbosity::Warning)
+#if DO_BLUEPRINT_GUARD
+		if (Verbosity <= ELogVerbosity::Error || ShowKismetScriptStackOnWarnings())
+		{
+			ScriptStack = TEXT("Script call stack:\n");
+			GetScriptCallstack(ScriptStack);
+			return;
+		}
+#endif
+
+		if (const FFrame* CurrentFrame = GetThreadLocalTopStackFrame())
+		{
+			ScriptStack = TEXT("Script Msg called by: ");
+			ScriptStack << CurrentFrame->Object->GetFullName();
+			return;
+		}
+	};
+	if (Verbosity <= ELogVerbosity::Warning)
 	{
-		ScriptStack = TEXT("Last function called:\n");
-		GetScriptCallstack(ScriptStack, false, true);
+		PopulateStackString(ScriptStack, Verbosity);
 	}
 #endif
 
