Fixed up const/constexpr variables in header files to be "inline"

This is required for c++ modules/header units in order to decide if constexpr/const value should be available to outside module

diff --git a/Engine/Source/Runtime/ApplicationCore/Public/GenericPlatform/GenericApplication.h b/Engine/Source/Runtime/ApplicationCore/Public/GenericPlatform/GenericApplication.h
--- a/Engine/Source/Runtime/ApplicationCore/Public/GenericPlatform/GenericApplication.h
+++ b/Engine/Source/Runtime/ApplicationCore/Public/GenericPlatform/GenericApplication.h
@@ -33,19 +33,19 @@ namespace EModifierKey
 	typedef uint8 Type;
 
 	/** No key. */
-	const Type None	= 0;
+	inline const Type None	= 0;
 
 	/** Ctrl key (Command key on Mac, Control key on Windows). */
-	const Type Control = 1 << 0;
+	inline const Type Control = 1 << 0;
 
 	/** Alt key. */
-	const Type Alt = 1 << 1;
+	inline const Type Alt = 1 << 1;
 
 	/** Shift key. */
-	const Type Shift = 1 << 2;
+	inline const Type Shift = 1 << 2;
 
 	/** Cmd key (Control key on Mac, Win key on Windows) */
-	const Type Command = 1 << 3;
+	inline const Type Command = 1 << 3;
 
 	FORCEINLINE EModifierKey::Type FromBools(const bool bControl, const bool bAlt, const bool bShift, const bool bCommand)
 	{
diff --git a/Engine/Source/Runtime/Core/Public/Async/TaskTrace.h b/Engine/Source/Runtime/Core/Public/Async/TaskTrace.h
--- a/Engine/Source/Runtime/Core/Public/Async/TaskTrace.h
+++ b/Engine/Source/Runtime/Core/Public/Async/TaskTrace.h
@@ -30,9 +30,9 @@ namespace TaskTrace
 
 	using FId = uint64;
 
-	const FId InvalidId = ~FId(0);
+	inline const FId InvalidId = ~FId(0);
 
-	static constexpr uint32 TaskTraceVersion = 1;
+	inline constexpr uint32 TaskTraceVersion = 1;
 
 	FId CORE_API GenerateTaskId();
 
diff --git a/Engine/Source/Runtime/Core/Public/Math/VectorRegister.h b/Engine/Source/Runtime/Core/Public/Math/VectorRegister.h
--- a/Engine/Source/Runtime/Core/Public/Math/VectorRegister.h
+++ b/Engine/Source/Runtime/Core/Public/Math/VectorRegister.h
@@ -101,7 +101,7 @@ using TPersistentVectorRegisterType = std::conditional_t<std::is_same_v<T, float
 
 
 /** Vector that represents (1/255,1/255,1/255,1/255) */
-static constexpr VectorRegister VECTOR_INV_255 = MakeVectorRegisterDoubleConstant(1.0 / 255, 1.0 / 255, 1.0 / 255, 1.0 / 255);
+inline constexpr VectorRegister VECTOR_INV_255 = MakeVectorRegisterDoubleConstant(1.0 / 255, 1.0 / 255, 1.0 / 255, 1.0 / 255);
 
 // Old names for comparison functions, kept for compatibility.
 #define VectorMask_LT( Vec1, Vec2 )			VectorCompareLT(Vec1, Vec2)
@@ -119,7 +119,7 @@ static constexpr VectorRegister VECTOR_INV_255 = MakeVectorRegisterDoubleConstan
 
 namespace GlobalVectorConstants
 {
-	static constexpr VectorRegister AnimWeightThreshold = MakeVectorRegisterConstant(ZERO_ANIMWEIGHT_THRESH_DOUBLE, ZERO_ANIMWEIGHT_THRESH_DOUBLE, ZERO_ANIMWEIGHT_THRESH_DOUBLE, ZERO_ANIMWEIGHT_THRESH_DOUBLE);
-	static constexpr VectorRegister RotationSignificantThreshold = MakeVectorRegisterConstant(1.0 - UE_DELTA*UE_DELTA, 1.0 - UE_DELTA*UE_DELTA, 1.0 - UE_DELTA*UE_DELTA, 1.0 - UE_DELTA*UE_DELTA);
+	inline constexpr VectorRegister AnimWeightThreshold = MakeVectorRegisterConstant(ZERO_ANIMWEIGHT_THRESH_DOUBLE, ZERO_ANIMWEIGHT_THRESH_DOUBLE, ZERO_ANIMWEIGHT_THRESH_DOUBLE, ZERO_ANIMWEIGHT_THRESH_DOUBLE);
+	inline constexpr VectorRegister RotationSignificantThreshold = MakeVectorRegisterConstant(1.0 - UE_DELTA*UE_DELTA, 1.0 - UE_DELTA*UE_DELTA, 1.0 - UE_DELTA*UE_DELTA, 1.0 - UE_DELTA*UE_DELTA);
 }
 
diff --git a/Engine/Source/Runtime/Core/Public/Misc/Optional.h b/Engine/Source/Runtime/Core/Public/Misc/Optional.h
--- a/Engine/Source/Runtime/Core/Public/Misc/Optional.h
+++ b/Engine/Source/Runtime/Core/Public/Misc/Optional.h
@@ -12,7 +12,7 @@ struct FNullOpt
 {
 	explicit constexpr FNullOpt(int) {}
 };
-constexpr FNullOpt NullOpt{0};
+inline constexpr FNullOpt NullOpt{0};
 
 /**
  * When we have an optional value IsSet() returns true, and GetValue() is meaningful.
diff --git a/Engine/Source/Runtime/Core/Public/Misc/OutputDevice.h b/Engine/Source/Runtime/Core/Public/Misc/OutputDevice.h
--- a/Engine/Source/Runtime/Core/Public/Misc/OutputDevice.h
+++ b/Engine/Source/Runtime/Core/Public/Misc/OutputDevice.h
@@ -52,26 +52,26 @@ An empty string reverts to the normal gray on black.
 // putting them in a namespace to protect against future name conflicts
 namespace OutputDeviceColor
 {
-	const TCHAR* const COLOR_BLACK = TEXT("0000");
-
-	const TCHAR* const COLOR_DARK_RED = TEXT("1000");
-	const TCHAR* const COLOR_DARK_GREEN = TEXT("0100");
-	const TCHAR* const COLOR_DARK_BLUE = TEXT("0010");
-	const TCHAR* const COLOR_DARK_YELLOW = TEXT("1100");
-	const TCHAR* const COLOR_DARK_CYAN = TEXT("0110");
-	const TCHAR* const COLOR_DARK_PURPLE = TEXT("1010");
-	const TCHAR* const COLOR_DARK_WHITE = TEXT("1110");
-	const TCHAR* const COLOR_GRAY = COLOR_DARK_WHITE;
-
-	const TCHAR* const COLOR_RED = TEXT("1001");
-	const TCHAR* const COLOR_GREEN = TEXT("0101");
-	const TCHAR* const COLOR_BLUE = TEXT("0011");
-	const TCHAR* const COLOR_YELLOW = TEXT("1101");
-	const TCHAR* const COLOR_CYAN = TEXT("0111");
-	const TCHAR* const COLOR_PURPLE = TEXT("1011");
-	const TCHAR* const COLOR_WHITE = TEXT("1111");
-
-	const TCHAR* const COLOR_NONE = TEXT("");
+	inline const TCHAR* const COLOR_BLACK = TEXT("0000");
+
+	inline const TCHAR* const COLOR_DARK_RED = TEXT("1000");
+	inline const TCHAR* const COLOR_DARK_GREEN = TEXT("0100");
+	inline const TCHAR* const COLOR_DARK_BLUE = TEXT("0010");
+	inline const TCHAR* const COLOR_DARK_YELLOW = TEXT("1100");
+	inline const TCHAR* const COLOR_DARK_CYAN = TEXT("0110");
+	inline const TCHAR* const COLOR_DARK_PURPLE = TEXT("1010");
+	inline const TCHAR* const COLOR_DARK_WHITE = TEXT("1110");
+	inline const TCHAR* const COLOR_GRAY = COLOR_DARK_WHITE;
+
+	inline const TCHAR* const COLOR_RED = TEXT("1001");
+	inline const TCHAR* const COLOR_GREEN = TEXT("0101");
+	inline const TCHAR* const COLOR_BLUE = TEXT("0011");
+	inline const TCHAR* const COLOR_YELLOW = TEXT("1101");
+	inline const TCHAR* const COLOR_CYAN = TEXT("0111");
+	inline const TCHAR* const COLOR_PURPLE = TEXT("1011");
+	inline const TCHAR* const COLOR_WHITE = TEXT("1111");
+
+	inline const TCHAR* const COLOR_NONE = TEXT("");
 }
 
 using namespace OutputDeviceColor;
diff --git a/Engine/Source/Runtime/Core/Public/Misc/Timespan.h b/Engine/Source/Runtime/Core/Public/Misc/Timespan.h
--- a/Engine/Source/Runtime/Core/Public/Misc/Timespan.h
+++ b/Engine/Source/Runtime/Core/Public/Misc/Timespan.h
@@ -20,37 +20,37 @@ template <typename ElementType> struct TIntervalTraits;
 namespace ETimespan
 {
 	/** The maximum number of ticks that can be represented in FTimespan. */
-	constexpr int64 MaxTicks = 9223372036854775807;
+	inline constexpr int64 MaxTicks = 9223372036854775807;
 
 	/** The minimum number of ticks that can be represented in FTimespan. */
-	constexpr int64 MinTicks = -9223372036854775807 - 1;
+	inline constexpr int64 MinTicks = -9223372036854775807 - 1;
 
 	/** The number of nanoseconds per tick. */
-	constexpr int64 NanosecondsPerTick = 100;
+	inline constexpr int64 NanosecondsPerTick = 100;
 
 	/** The number of timespan ticks per day. */
-	constexpr int64 TicksPerDay = 864000000000;
+	inline constexpr int64 TicksPerDay = 864000000000;
 
 	/** The number of timespan ticks per hour. */
-	constexpr int64 TicksPerHour = 36000000000;
+	inline constexpr int64 TicksPerHour = 36000000000;
 
 	/** The number of timespan ticks per microsecond. */
-	constexpr int64 TicksPerMicrosecond = 10;
+	inline constexpr int64 TicksPerMicrosecond = 10;
 
 	/** The number of timespan ticks per millisecond. */
-	constexpr int64 TicksPerMillisecond = 10000;
+	inline constexpr int64 TicksPerMillisecond = 10000;
 
 	/** The number of timespan ticks per minute. */
-	constexpr int64 TicksPerMinute = 600000000;
+	inline constexpr int64 TicksPerMinute = 600000000;
 
 	/** The number of timespan ticks per second. */
-	constexpr int64 TicksPerSecond = 10000000;
+	inline constexpr int64 TicksPerSecond = 10000000;
 
 	/** The number of timespan ticks per week. */
-	constexpr int64 TicksPerWeek = 6048000000000;
+	inline constexpr int64 TicksPerWeek = 6048000000000;
 
 	/** The number of timespan ticks per year (365 days, not accounting for leap years). */
-	constexpr int64 TicksPerYear = 365 * TicksPerDay;
+	inline constexpr int64 TicksPerYear = 365 * TicksPerDay;
 }
 
 
diff --git a/Engine/Source/Runtime/Core/Public/ProfilingDebugging/TagTrace.h b/Engine/Source/Runtime/Core/Public/ProfilingDebugging/TagTrace.h
--- a/Engine/Source/Runtime/Core/Public/ProfilingDebugging/TagTrace.h
+++ b/Engine/Source/Runtime/Core/Public/ProfilingDebugging/TagTrace.h
@@ -81,7 +81,7 @@ private:
 };
 
 ////////////////////////////////////////////////////////////////////////////////
-constexpr int32 TRACE_TAG = 257;
+inline constexpr int32 TRACE_TAG = 257;
 
 ////////////////////////////////////////////////////////////////////////////////
 #define UE_MEMSCOPE(InTag)				FMemScope PREPROCESSOR_JOIN(MemScope,__LINE__)(InTag);
diff --git a/Engine/Source/Runtime/CoreUObject/Public/Misc/PackagePath.h b/Engine/Source/Runtime/CoreUObject/Public/Misc/PackagePath.h
--- a/Engine/Source/Runtime/CoreUObject/Public/Misc/PackagePath.h
+++ b/Engine/Source/Runtime/CoreUObject/Public/Misc/PackagePath.h
@@ -67,7 +67,7 @@ enum class EPackageExtension : uint8
 	BulkDataMemoryMapped,
 	PayloadSidecar,
 };
-constexpr int EPackageExtensionCount = static_cast<int>(EPackageExtension::PayloadSidecar) + 1;
+inline constexpr int EPackageExtensionCount = static_cast<int>(EPackageExtension::PayloadSidecar) + 1;
 COREUOBJECT_API const TCHAR* LexToString(EPackageExtension PackageExtension);
 
 
diff --git a/Engine/Source/Runtime/CoreUObject/Public/Serialization/BulkData.h b/Engine/Source/Runtime/CoreUObject/Public/Serialization/BulkData.h
--- a/Engine/Source/Runtime/CoreUObject/Public/Serialization/BulkData.h
+++ b/Engine/Source/Runtime/CoreUObject/Public/Serialization/BulkData.h
@@ -64,7 +64,7 @@ class IMappedFileRegion;
 
 /** A loose hash value that can be created from either a filenames or a FIoChunkId */
 using FIoFilenameHash = uint32;
-const FIoFilenameHash INVALID_IO_FILENAME_HASH = 0;
+inline const FIoFilenameHash INVALID_IO_FILENAME_HASH = 0;
 /** Helpers to create the hash from a filename. Returns IOFILENAMEHASH_NONE if and only if the filename is empty. */
 COREUOBJECT_API FIoFilenameHash MakeIoFilenameHash(const FString& Filename);
 /** Helpers to create the hash from a FPackagePath. Returns IOFILENAMEHASH_NONE if and only if the PackagePath is empty. */
diff --git a/Engine/Source/Runtime/CoreUObject/Public/UObject/ObjectMacros.h b/Engine/Source/Runtime/CoreUObject/Public/UObject/ObjectMacros.h
--- a/Engine/Source/Runtime/CoreUObject/Public/UObject/ObjectMacros.h
+++ b/Engine/Source/Runtime/CoreUObject/Public/UObject/ObjectMacros.h
@@ -558,9 +558,9 @@ enum EObjectFlags
 };
 
 PRAGMA_DISABLE_DEPRECATION_WARNINGS
-constexpr EObjectFlags RF_InternalPendingKill = RF_PendingKill;
-constexpr EObjectFlags RF_InternalGarbage = RF_Garbage;
-constexpr EObjectFlags RF_InternalMirroredFlags = (EObjectFlags)(RF_PendingKill | RF_Garbage);
+inline constexpr EObjectFlags RF_InternalPendingKill = RF_PendingKill;
+inline constexpr EObjectFlags RF_InternalGarbage = RF_Garbage;
+inline constexpr EObjectFlags RF_InternalMirroredFlags = (EObjectFlags)(RF_PendingKill | RF_Garbage);
 PRAGMA_ENABLE_DEPRECATION_WARNINGS
 
 /** Mask for all object flags */
@@ -2149,12 +2149,12 @@ namespace UE4
 {
 	using ELoadConfigPropagationFlags UE_DEPRECATED(5.0, "Use UE namespace instead of UE4") = UE::ELoadConfigPropagationFlags;
 
-	UE_DEPRECATED(5.0, "Use UE namespace instead of UE4") static constexpr UE::ELoadConfigPropagationFlags LCPF_None							= UE::ELoadConfigPropagationFlags::LCPF_None;
-	UE_DEPRECATED(5.0, "Use UE namespace instead of UE4") static constexpr UE::ELoadConfigPropagationFlags LCPF_ReadParentSections				= UE::ELoadConfigPropagationFlags::LCPF_ReadParentSections;
-	UE_DEPRECATED(5.0, "Use UE namespace instead of UE4") static constexpr UE::ELoadConfigPropagationFlags LCPF_PropagateToChildDefaultObjects	= UE::ELoadConfigPropagationFlags::LCPF_PropagateToChildDefaultObjects;
-	UE_DEPRECATED(5.0, "Use UE namespace instead of UE4") static constexpr UE::ELoadConfigPropagationFlags LCPF_PropagateToInstances			= UE::ELoadConfigPropagationFlags::LCPF_PropagateToInstances;
-	UE_DEPRECATED(5.0, "Use UE namespace instead of UE4") static constexpr UE::ELoadConfigPropagationFlags LCPF_ReloadingConfigData				= UE::ELoadConfigPropagationFlags::LCPF_ReloadingConfigData;
-	UE_DEPRECATED(5.0, "Use UE namespace instead of UE4") static constexpr UE::ELoadConfigPropagationFlags LCPF_PersistentFlags					= UE::ELoadConfigPropagationFlags::LCPF_PersistentFlags;
+	UE_DEPRECATED(5.0, "Use UE namespace instead of UE4") inline constexpr UE::ELoadConfigPropagationFlags LCPF_None							= UE::ELoadConfigPropagationFlags::LCPF_None;
+	UE_DEPRECATED(5.0, "Use UE namespace instead of UE4") inline constexpr UE::ELoadConfigPropagationFlags LCPF_ReadParentSections				= UE::ELoadConfigPropagationFlags::LCPF_ReadParentSections;
+	UE_DEPRECATED(5.0, "Use UE namespace instead of UE4") inline constexpr UE::ELoadConfigPropagationFlags LCPF_PropagateToChildDefaultObjects	= UE::ELoadConfigPropagationFlags::LCPF_PropagateToChildDefaultObjects;
+	UE_DEPRECATED(5.0, "Use UE namespace instead of UE4") inline constexpr UE::ELoadConfigPropagationFlags LCPF_PropagateToInstances			= UE::ELoadConfigPropagationFlags::LCPF_PropagateToInstances;
+	UE_DEPRECATED(5.0, "Use UE namespace instead of UE4") inline constexpr UE::ELoadConfigPropagationFlags LCPF_ReloadingConfigData				= UE::ELoadConfigPropagationFlags::LCPF_ReloadingConfigData;
+	UE_DEPRECATED(5.0, "Use UE namespace instead of UE4") inline constexpr UE::ELoadConfigPropagationFlags LCPF_PersistentFlags					= UE::ELoadConfigPropagationFlags::LCPF_PersistentFlags;
 }
 
 
diff --git a/Engine/Source/Runtime/CoreUObject/Public/UObject/UObjectGlobals.h b/Engine/Source/Runtime/CoreUObject/Public/UObject/UObjectGlobals.h
--- a/Engine/Source/Runtime/CoreUObject/Public/UObject/UObjectGlobals.h
+++ b/Engine/Source/Runtime/CoreUObject/Public/UObject/UObjectGlobals.h
@@ -3041,7 +3041,7 @@ namespace UECodeGen_Private
 	ENUM_CLASS_FLAGS(EPropertyGenFlags)
 
 	// Value which masks out the type of combined EPropertyGenFlags.
-	constexpr EPropertyGenFlags PropertyTypeMask = (EPropertyGenFlags)0x3F;
+	inline constexpr EPropertyGenFlags PropertyTypeMask = (EPropertyGenFlags)0x3F;
 
 #if WITH_METADATA
 	struct FMetaDataPairParam
diff --git a/Engine/Source/Runtime/CoreUObject/Public/UObject/UnrealType.h b/Engine/Source/Runtime/CoreUObject/Public/UObject/UnrealType.h
--- a/Engine/Source/Runtime/CoreUObject/Public/UObject/UnrealType.h
+++ b/Engine/Source/Runtime/CoreUObject/Public/UObject/UnrealType.h
@@ -6355,23 +6355,23 @@ namespace EPropertyChangeType
 	typedef uint32 Type;
 
 	//default value.  Add new enums to add new functionality.
-	const Type Unspecified		= 1 << 0;
+	inline const Type Unspecified		= 1 << 0;
 	//Array Add
-	const Type ArrayAdd			= 1 << 1;
+	inline const Type ArrayAdd			= 1 << 1;
 	//Array Remove
-	const Type ArrayRemove		= 1 << 2;
+	inline const Type ArrayRemove		= 1 << 2;
 	//Array Clear
-	const Type ArrayClear		= 1 << 3;
+	inline const Type ArrayClear		= 1 << 3;
 	//Value Set
-	const Type ValueSet			= 1 << 4;
+	inline const Type ValueSet			= 1 << 4;
 	//Duplicate
-	const Type Duplicate		= 1 << 5;
+	inline const Type Duplicate			= 1 << 5;
 	//Interactive, e.g. dragging a slider. Will be followed by a ValueSet when finished.
-	const Type Interactive		= 1 << 6;
+	inline const Type Interactive		= 1 << 6;
 	//Redirected.  Used when property references are updated due to content hot-reloading, or an asset being replaced during asset deletion (aka, asset consolidation).
-	const Type Redirected		= 1 << 7;
+	inline const Type Redirected		= 1 << 7;
 	// Array Item Moved Within the Array
-	const Type ArrayMove		= 1 << 8;
+	inline const Type ArrayMove			= 1 << 8;
 };
 
 /**
diff --git a/Engine/Source/Runtime/CoreUObject/Public/UObject/WeakObjectPtr.h b/Engine/Source/Runtime/CoreUObject/Public/UObject/WeakObjectPtr.h
--- a/Engine/Source/Runtime/CoreUObject/Public/UObject/WeakObjectPtr.h
+++ b/Engine/Source/Runtime/CoreUObject/Public/UObject/WeakObjectPtr.h
@@ -31,9 +31,9 @@ namespace UE::Core::Private
 {
 	/** Specifies the ObjectIndex used for invalid object pointers. */
 #if UE_WEAKOBJECTPTR_ZEROINIT_FIX
-	static constexpr int32 InvalidWeakObjectIndex = 0;
+	inline constexpr int32 InvalidWeakObjectIndex = 0;
 #else
-	static constexpr int32 InvalidWeakObjectIndex = INDEX_NONE;
+	inline constexpr int32 InvalidWeakObjectIndex = INDEX_NONE;
 #endif
 }
 
diff --git a/Engine/Source/Runtime/Engine/Classes/Engine/Channel.h b/Engine/Source/Runtime/Engine/Classes/Engine/Channel.h
--- a/Engine/Source/Runtime/Engine/Classes/Engine/Channel.h
+++ b/Engine/Source/Runtime/Engine/Classes/Engine/Channel.h
@@ -16,7 +16,7 @@ class FOutBunch;
 class UNetConnection;
 
 // Constant for all buffers that are reading from the network
-const int MAX_STRING_SERIALIZE_SIZE	= NAME_SIZE;
+inline const int MAX_STRING_SERIALIZE_SIZE	= NAME_SIZE;
 
 /**
  * Enumerates channel types.
diff --git a/Engine/Source/Runtime/Engine/Classes/Sound/SoundWave.h b/Engine/Source/Runtime/Engine/Classes/Sound/SoundWave.h
--- a/Engine/Source/Runtime/Engine/Classes/Sound/SoundWave.h
+++ b/Engine/Source/Runtime/Engine/Classes/Sound/SoundWave.h
@@ -60,7 +60,7 @@ enum class ESoundWavePrecacheState
 	Done
 };
 
-constexpr uint64 InvalidAudioStreamCacheLookupID = TNumericLimits<uint64>::Max();
+inline constexpr uint64 InvalidAudioStreamCacheLookupID = TNumericLimits<uint64>::Max();
 
 /**
  * A chunk of streamed audio.
diff --git a/Engine/Source/Runtime/Engine/Public/Scalability.h b/Engine/Source/Runtime/Engine/Public/Scalability.h
--- a/Engine/Source/Runtime/Engine/Public/Scalability.h
+++ b/Engine/Source/Runtime/Engine/Public/Scalability.h
@@ -16,7 +16,7 @@ namespace Scalability
 		ERelativeToMax,
 	};
 
-	const int32 DefaultQualityLevel = 3;
+	inline const int32 DefaultQualityLevel = 3;
 
 	/**
 	 * Structure for holding the state of the engine scalability groups
diff --git a/Engine/Source/Runtime/Experimental/ChaosCore/Public/Chaos/Real.h b/Engine/Source/Runtime/Experimental/ChaosCore/Public/Chaos/Real.h
--- a/Engine/Source/Runtime/Experimental/ChaosCore/Public/Chaos/Real.h
+++ b/Engine/Source/Runtime/Experimental/ChaosCore/Public/Chaos/Real.h
@@ -25,5 +25,5 @@ namespace Chaos
 	* ISPC optimization supports float and double, this allows classes that uses ISPC to branch to the right implementation 
 	* without having to check the actual underlying type of FReal
 	*/
-	constexpr bool bRealTypeCompatibleWithISPC = (TAreTypesEqual<FReal, float>::Value == true) || (TAreTypesEqual<FReal, double>::Value == true);
+	inline constexpr bool bRealTypeCompatibleWithISPC = (TAreTypesEqual<FReal, float>::Value == true) || (TAreTypesEqual<FReal, double>::Value == true);
 }
diff --git a/Engine/Source/Runtime/Net/Core/Public/Net/Core/Analytics/NetAnalytics.h b/Engine/Source/Runtime/Net/Core/Public/Net/Core/Analytics/NetAnalytics.h
--- a/Engine/Source/Runtime/Net/Core/Public/Net/Core/Analytics/NetAnalytics.h
+++ b/Engine/Source/Runtime/Net/Core/Public/Net/Core/Analytics/NetAnalytics.h
@@ -53,7 +53,7 @@ class IAnalyticsProvider;
 #define NET_ANALYTICS_MULTITHREADING 0
 
 
-constexpr const ESPMode NetAnalyticsThreadSafety = ESPMode::ThreadSafe;
+inline constexpr const ESPMode NetAnalyticsThreadSafety = ESPMode::ThreadSafe;
 
 
 // Forward Declarations
diff --git a/Engine/Source/Runtime/Net/Core/Public/Net/Core/Trace/NetTraceConfig.h b/Engine/Source/Runtime/Net/Core/Public/Net/Core/Trace/NetTraceConfig.h
--- a/Engine/Source/Runtime/Net/Core/Public/Net/Core/Trace/NetTraceConfig.h
+++ b/Engine/Source/Runtime/Net/Core/Public/Net/Core/Trace/NetTraceConfig.h
@@ -26,7 +26,7 @@ namespace ENetTraceVerbosity
 	};
 }
 
-const uint32 NetTraceInvalidGameInstanceId = ~0U;
+inline const uint32 NetTraceInvalidGameInstanceId = ~0U;
 
 #if UE_NET_TRACE_ENABLED
 #	ifndef UE_NET_TRACE_COMPILETIME_VERBOSITY
diff --git a/Engine/Source/Runtime/PhysicsCore/Public/PhysicsInterfaceTypesCore.h b/Engine/Source/Runtime/PhysicsCore/Public/PhysicsInterfaceTypesCore.h
--- a/Engine/Source/Runtime/PhysicsCore/Public/PhysicsInterfaceTypesCore.h
+++ b/Engine/Source/Runtime/PhysicsCore/Public/PhysicsInterfaceTypesCore.h
@@ -161,7 +161,7 @@ static bool CalcMeshNegScaleCompensation(const FVector& InScale3D, FTransform& O
 // TODO: Fixup types, these are more or less the PhysX types renamed as a temporary solution.
 // Probably should move to different header as well.
 
-const uint32 AggregateMaxSize = 128;
+inline const uint32 AggregateMaxSize = 128;
 
 class UPhysicalMaterial;
 class UPrimitiveComponent;
diff --git a/Engine/Source/Runtime/RHI/Public/RHI.h b/Engine/Source/Runtime/RHI/Public/RHI.h
--- a/Engine/Source/Runtime/RHI/Public/RHI.h
+++ b/Engine/Source/Runtime/RHI/Public/RHI.h
@@ -492,10 +492,10 @@ extern RHI_API bool GRHISupportsMultithreadedShaderCreation;
 extern RHI_API bool GRHISupportsMultithreadedResources;
 
 /** The minimum Z value in clip space for the RHI. This is a constant value to always match D3D clip-space. */
-constexpr float GMinClipZ = 0.0f;
+inline constexpr float GMinClipZ = 0.0f;
 
 /** The sign to apply to the Y axis of projection matrices. This is a constant value to always match D3D clip-space. */
-constexpr float GProjectionSignY = 1.0f;
+inline constexpr float GProjectionSignY = 1.0f;
 
 /** Does this RHI need to wait for deletion of resources due to ref counting. */
 extern RHI_API bool GRHINeedsExtraDeletionLatency;
diff --git a/Engine/Source/Runtime/RenderCore/Public/ShaderCompilerCore.h b/Engine/Source/Runtime/RenderCore/Public/ShaderCompilerCore.h
--- a/Engine/Source/Runtime/RenderCore/Public/ShaderCompilerCore.h
+++ b/Engine/Source/Runtime/RenderCore/Public/ShaderCompilerCore.h
@@ -18,13 +18,13 @@
 class Error;
 
 // this is for the protocol, not the data, bump if FShaderCompilerInput or ProcessInputFromArchive changes.
-const int32 ShaderCompileWorkerInputVersion = 17;
+inline const int32 ShaderCompileWorkerInputVersion = 17;
 // this is for the protocol, not the data, bump if FShaderCompilerOutput or WriteToOutputArchive changes.
-const int32 ShaderCompileWorkerOutputVersion = 8;
+inline const int32 ShaderCompileWorkerOutputVersion = 8;
 // this is for the protocol, not the data.
-const int32 ShaderCompileWorkerSingleJobHeader = 'S';
+inline const int32 ShaderCompileWorkerSingleJobHeader = 'S';
 // this is for the protocol, not the data.
-const int32 ShaderCompileWorkerPipelineJobHeader = 'P';
+inline const int32 ShaderCompileWorkerPipelineJobHeader = 'P';
 
 /** Returns true if shader symbols should be kept for a given platform. */
 extern RENDERCORE_API bool ShouldGenerateShaderSymbols(FName ShaderFormat);
diff --git a/Engine/Source/Runtime/Renderer/Public/GlobalDistanceFieldConstants.h b/Engine/Source/Runtime/Renderer/Public/GlobalDistanceFieldConstants.h
--- a/Engine/Source/Runtime/Renderer/Public/GlobalDistanceFieldConstants.h
+++ b/Engine/Source/Runtime/Renderer/Public/GlobalDistanceFieldConstants.h
@@ -5,5 +5,5 @@
 namespace GlobalDistanceField
 {
 	/** Must match global distance field shaders. */
-	const int32 MaxClipmaps = 6;
+	inline const int32 MaxClipmaps = 6;
 }
diff --git a/Engine/Source/Runtime/Slate/Public/Framework/MultiBox/MultiBox.h b/Engine/Source/Runtime/Slate/Public/Framework/MultiBox/MultiBox.h
--- a/Engine/Source/Runtime/Slate/Public/Framework/MultiBox/MultiBox.h
+++ b/Engine/Source/Runtime/Slate/Public/Framework/MultiBox/MultiBox.h
@@ -31,24 +31,24 @@ class FToolBarComboButtonBlock;
 namespace MultiBoxConstants
 {	
 	/** The time that a mouse should be hovered over a sub-menu before it automatically opens */
-	const float SubMenuOpenTime = 0.0f;
+	inline const float SubMenuOpenTime = 0.0f;
 
 	/** When a sub-menu is already open (for at least SubMenuClobberMinLifetime), the time that a mouse
 	    should be hovered over another sub-menu entry before dismissing the first menu and opening
 		the new one; this doesn't apply to short-lived sub-menus, see below */
-	const float SubMenuClobberTime = 0.5f;
+	inline const float SubMenuClobberTime = 0.5f;
 
 	/** The time that a sub-menu needs to remain open in order for the SubMenuClobberTime to apply;
 	    menus open shorter than this min lifetime will be instantly dismissed */
-	const float SubMenuClobberMinLifetime = 0.5f;
+	inline const float SubMenuClobberMinLifetime = 0.5f;
 
-	//const FName MenuItemFont = "MenuItem.Font";
-	//const FName MenuBindingFont = "MenuItem.BindingFont";
+	//inline const FName MenuItemFont = "MenuItem.Font";
+	//inline const FName MenuBindingFont = "MenuItem.BindingFont";
 
-	//const FName EditableTextFont = "MenuItem.Font";
+	//inline const FName EditableTextFont = "MenuItem.Font";
 
 	/** Minimum widget of an editable text box within a multi-box */
-	const float EditableTextMinWidth = 30.0f;
+	inline const float EditableTextMinWidth = 30.0f;
 }
 
 /**
diff --git a/Engine/Source/Runtime/SlateCore/Public/Styling/CoreStyle.h b/Engine/Source/Runtime/SlateCore/Public/Styling/CoreStyle.h
--- a/Engine/Source/Runtime/SlateCore/Public/Styling/CoreStyle.h
+++ b/Engine/Source/Runtime/SlateCore/Public/Styling/CoreStyle.h
@@ -82,36 +82,36 @@ namespace CoreStyleConstants
 {
 	// Note, these sizes are in Slate Units.
 	// Slate Units do NOT have to map to pixels.
-	const FVector2D Icon5x16(5.0f, 16.0f);
-	const FVector2D Icon8x4(8.0f, 4.0f);
-	const FVector2D Icon16x4(16.0f, 4.0f);
-	const FVector2D Icon8x8(8.0f, 8.0f);
-	const FVector2D Icon10x10(10.0f, 10.0f);
-	const FVector2D Icon12x12(12.0f, 12.0f);
-	const FVector2D Icon12x16(12.0f, 16.0f);
-	const FVector2D Icon14x14(14.0f, 14.0f);
-	const FVector2D Icon16x16(16.0f, 16.0f);
-	const FVector2D Icon18x18(18.0f, 18.0f);
-	const FVector2D Icon20x20(20.0f, 20.0f);
-	const FVector2D Icon22x22(22.0f, 22.0f);
-	const FVector2D Icon24x24(24.0f, 24.0f);
-	const FVector2D Icon25x25(25.0f, 25.0f);
-	const FVector2D Icon32x32(32.0f, 32.0f);
-	const FVector2D Icon40x40(40.0f, 40.0f);
-	const FVector2D Icon64x64(64.0f, 64.0f);
-	const FVector2D Icon36x24(36.0f, 24.0f);
-	const FVector2D Icon128x128(128.0f, 128.0f);
+	inline const FVector2D Icon5x16(5.0f, 16.0f);
+	inline const FVector2D Icon8x4(8.0f, 4.0f);
+	inline const FVector2D Icon16x4(16.0f, 4.0f);
+	inline const FVector2D Icon8x8(8.0f, 8.0f);
+	inline const FVector2D Icon10x10(10.0f, 10.0f);
+	inline const FVector2D Icon12x12(12.0f, 12.0f);
+	inline const FVector2D Icon12x16(12.0f, 16.0f);
+	inline const FVector2D Icon14x14(14.0f, 14.0f);
+	inline const FVector2D Icon16x16(16.0f, 16.0f);
+	inline const FVector2D Icon18x18(18.0f, 18.0f);
+	inline const FVector2D Icon20x20(20.0f, 20.0f);
+	inline const FVector2D Icon22x22(22.0f, 22.0f);
+	inline const FVector2D Icon24x24(24.0f, 24.0f);
+	inline const FVector2D Icon25x25(25.0f, 25.0f);
+	inline const FVector2D Icon32x32(32.0f, 32.0f);
+	inline const FVector2D Icon40x40(40.0f, 40.0f);
+	inline const FVector2D Icon64x64(64.0f, 64.0f);
+	inline const FVector2D Icon36x24(36.0f, 24.0f);
+	inline const FVector2D Icon128x128(128.0f, 128.0f);
 
 	// Common Margins
-	const FMargin DefaultMargins(8.f, 4.f);
+	inline const FMargin DefaultMargins(8.f, 4.f);
 	// Buttons already have a built in (4., 2.) padding - adding to that a little
-	const FMargin ButtonMargins(12.f, 1.5f, 12.f, 1.5f);
+	inline const FMargin ButtonMargins(12.f, 1.5f, 12.f, 1.5f);
 
-	const FMargin PressedButtonMargins(12.f, 2.5f, 12.f, 0.5f);
-	const FMargin ToggleButtonMargins(16.0f, 2.0f);
-	const FMargin ComboButtonMargin(8.f, 1.f, 8.f, 1.f);
-	const FMargin PressedComboButtonMargin(8.f, 2.f, 8.f, 0.f);
+	inline const FMargin PressedButtonMargins(12.f, 2.5f, 12.f, 0.5f);
+	inline const FMargin ToggleButtonMargins(16.0f, 2.0f);
+	inline const FMargin ComboButtonMargin(8.f, 1.f, 8.f, 1.f);
+	inline const FMargin PressedComboButtonMargin(8.f, 2.f, 8.f, 0.f);
 
-	const float InputFocusRadius = 4.f;
-	const float InputFocusThickness = 1.0f;
+	inline const float InputFocusRadius = 4.f;
+	inline const float InputFocusThickness = 1.0f;
 }
diff --git a/Engine/Source/Runtime/TypedElementFramework/Public/Elements/Framework/TypedElementLimits.h b/Engine/Source/Runtime/TypedElementFramework/Public/Elements/Framework/TypedElementLimits.h
--- a/Engine/Source/Runtime/TypedElementFramework/Public/Elements/Framework/TypedElementLimits.h
+++ b/Engine/Source/Runtime/TypedElementFramework/Public/Elements/Framework/TypedElementLimits.h
@@ -14,17 +14,17 @@
  * @note Limited to a combined 32-bits so that they can be used directly within render targets, though could be made 64-bits if the 
  *       editor used 64-bit render targets (this would also require 64-bit container support in TTypedElementInternalDataStore).
  */
-constexpr SIZE_T TypedHandleTypeIdBits = 8;
-constexpr SIZE_T TypedHandleElementIdBits = 24;
+inline constexpr SIZE_T TypedHandleTypeIdBits = 8;
+inline constexpr SIZE_T TypedHandleElementIdBits = 24;
 
-constexpr SIZE_T TypedHandleTypeIdBytes = TypedHandleTypeIdBits >> 3;
-constexpr SIZE_T TypedHandleElementIdBytes = TypedHandleElementIdBits >> 3;
+inline constexpr SIZE_T TypedHandleTypeIdBytes = TypedHandleTypeIdBits >> 3;
+inline constexpr SIZE_T TypedHandleElementIdBytes = TypedHandleElementIdBits >> 3;
 
-constexpr SIZE_T TypedHandleMaxTypeId = ((SIZE_T)1 << TypedHandleTypeIdBits) - 1;
-constexpr SIZE_T TypedHandleMaxElementId = ((SIZE_T)1 << TypedHandleElementIdBits) - 1;
+inline constexpr SIZE_T TypedHandleMaxTypeId = ((SIZE_T)1 << TypedHandleTypeIdBits) - 1;
+inline constexpr SIZE_T TypedHandleMaxElementId = ((SIZE_T)1 << TypedHandleElementIdBits) - 1;
 
-constexpr SIZE_T TypedHandleRefTrackingSkipCount = 1;
-constexpr SIZE_T TypedHandleRefTrackingDepth = 31;
+inline constexpr SIZE_T TypedHandleRefTrackingSkipCount = 1;
+inline constexpr SIZE_T TypedHandleRefTrackingDepth = 31;
 
 using FTypedHandleTypeId = uint8;
 using FTypedHandleElementId = int32;
