Cap the number of reference gathering when printing invalid references during save since they can be pathologically expensive

diff --git a/Engine/Source/Runtime/CoreUObject/Private/UObject/SavePackage/SavePackageUtilities.cpp b/Engine/Source/Runtime/CoreUObject/Private/UObject/SavePackage/SavePackageUtilities.cpp
--- a/Engine/Source/Runtime/CoreUObject/Private/UObject/SavePackage/SavePackageUtilities.cpp
+++ b/Engine/Source/Runtime/CoreUObject/Private/UObject/SavePackage/SavePackageUtilities.cpp
@@ -368,33 +368,36 @@ void FindMostLikelyCulprit(const TArray<UObject*>& BadObjects, UObject*& MostLik
 			continue;
 		}
 
-		UE_LOG(LogSavePackage, Warning, TEXT("\r\nReferencers of %s:"), *Obj->GetFullName());
-
-		FReferencerInformationList Refs;
-
-		if (IsReferenced(Obj, RF_Public, EInternalObjectFlags::Native, true, &Refs))
+		// Do not run more than a certain number of ref gather if there's a lot of illegal references
+		const int32 MaxNumberOfRefGather = 5;
+		if (BadObjIndex < MaxNumberOfRefGather)
 		{
-			for (int32 i = 0; i < Refs.ExternalReferences.Num(); i++)
+			UE_LOG(LogSavePackage, Warning, TEXT("\r\nReferencers of %s:"), *Obj->GetFullName());
+			FReferencerInformationList Refs;
+			if (IsReferenced(Obj, RF_Public, EInternalObjectFlags::Native, true, &Refs))
 			{
-				UObject* RefObj = Refs.ExternalReferences[i].Referencer;
-				if (IsObjectIncluded(RefObj))
+				for (int32 i = 0; i < Refs.ExternalReferences.Num(); i++)
 				{
-					if (RefObj->GetFName() == NAME_PersistentLevel || RefObj->GetClass()->GetFName() == NAME_World)
-					{
-						// these types of references should be ignored
-						continue;
-					}
-
-					UE_LOG(LogSavePackage, Warning, TEXT("\t%s (%i refs)"), *RefObj->GetFullName(), Refs.ExternalReferences[i].TotalReferences);
-					for (int32 j = 0; j < Refs.ExternalReferences[i].ReferencingProperties.Num(); j++)
+					UObject* RefObj = Refs.ExternalReferences[i].Referencer;
+					if (IsObjectIncluded(RefObj))
 					{
-						const FProperty* Prop = Refs.ExternalReferences[i].ReferencingProperties[j];
-						UE_LOG(LogSavePackage, Warning, TEXT("\t\t%i) %s"), j, *Prop->GetFullName());
-						ReferencedCulpritReferencer = Prop;
+						if (RefObj->GetFName() == NAME_PersistentLevel || RefObj->GetClass()->GetFName() == NAME_World)
+						{
+							// these types of references should be ignored
+							continue;
+						}
+
+						UE_LOG(LogSavePackage, Warning, TEXT("\t%s (%i refs)"), *RefObj->GetFullName(), Refs.ExternalReferences[i].TotalReferences);
+						for (int32 j = 0; j < Refs.ExternalReferences[i].ReferencingProperties.Num(); j++)
+						{
+							const FProperty* Prop = Refs.ExternalReferences[i].ReferencingProperties[j];
+							UE_LOG(LogSavePackage, Warning, TEXT("\t\t%i) %s"), j, *Prop->GetFullName());
+							ReferencedCulpritReferencer = Prop;
+						}
+
+						// Later ReferencedCulprits are higher priority than earlier culprits. TODO: Not sure if this is an intentional behavior or if they choice was arbitrary.
+						ReferencedCulprit = Obj;
 					}
-
-					// Later ReferencedCulprits are higher priority than earlier culprits. TODO: Not sure if this is an intentional behavior or if they choice was arbitrary.
-					ReferencedCulprit = Obj;
 				}
 			}
 		}
