Ensure bCooked doesn't get mistakenly unset during reference collection in cooked versions of the editor.

diff --git a/Engine/Source/Runtime/CoreUObject/Private/UObject/Class.cpp b/Engine/Source/Runtime/CoreUObject/Private/UObject/Class.cpp
--- a/Engine/Source/Runtime/CoreUObject/Private/UObject/Class.cpp
+++ b/Engine/Source/Runtime/CoreUObject/Private/UObject/Class.cpp
@@ -5045,12 +5045,13 @@ void UClass::Serialize( FArchive& Ar )
 	{
 		if (Ar.IsSaving())
 		{
-			bCooked = Ar.IsCooking();
+			bool bCookedAsBool = bCooked || Ar.IsCooking();
+			Ar << bCookedAsBool;
 		}
-		bool bCookedAsBool = bCooked;
-		Ar << bCookedAsBool;
 		if (Ar.IsLoading())
 		{
+			bool bCookedAsBool = false;
+			Ar << bCookedAsBool;
 			bCooked = bCookedAsBool;
 		}
 	}
