Fix crash when serializing(saving) FSkeletalMeshImportData cause by code not returning the proper UE version for the bulk data.

#rnx

diff --git a/Engine/Source/Runtime/CoreUObject/Private/Serialization/BulkData.cpp b/Engine/Source/Runtime/CoreUObject/Private/Serialization/BulkData.cpp
--- a/Engine/Source/Runtime/CoreUObject/Private/Serialization/BulkData.cpp
+++ b/Engine/Source/Runtime/CoreUObject/Private/Serialization/BulkData.cpp
@@ -1577,24 +1577,22 @@ void FBulkData::GetBulkDataVersions(FArchive& InlineArchive, FPackageFileVersion
 	EPackageSegment Segment;
 	bool bExternal;
 
-	if (UE::TryGetPackageNameFromChunkId(BulkChunkId, PackageName, Segment, bExternal) == false)
+	if (UE::TryGetPackageNameFromChunkId(BulkChunkId, PackageName, Segment, bExternal))
 	{
-		return;
-	}
-
-	IPackageResourceManager& ResourceMgr = IPackageResourceManager::Get();
-
-	if (TUniquePtr<FArchive> Ar = ResourceMgr.OpenReadExternalResource(EPackageExternalResource::WorkspaceDomainFile, PackageName.ToString()))
-	{
-		FPackageFileSummary Summary;
-		*Ar << Summary;
+		IPackageResourceManager& ResourceMgr = IPackageResourceManager::Get();
 
-		if (Ar->IsError() == false && Summary.Tag == PACKAGE_FILE_TAG)
+		if (TUniquePtr<FArchive> Ar = ResourceMgr.OpenReadExternalResource(EPackageExternalResource::WorkspaceDomainFile, PackageName.ToString()))
 		{
-			OutUEVersion = Summary.GetFileVersionUE();
-			OutLicenseeUEVersion = Summary.GetFileVersionLicenseeUE();
-			OutCustomVersions = Summary.GetCustomVersionContainer();
-			return;
+			FPackageFileSummary Summary;
+			*Ar << Summary;
+
+			if (Ar->IsError() == false && Summary.Tag == PACKAGE_FILE_TAG)
+			{
+				OutUEVersion = Summary.GetFileVersionUE();
+				OutLicenseeUEVersion = Summary.GetFileVersionLicenseeUE();
+				OutCustomVersions = Summary.GetCustomVersionContainer();
+				return;
+			}
 		}
 	}
 
