[CoreUObject]
* StructOnScope - Removed include of Package
* GeneratedCppInclude - Removed include of Package

Both of these files include Class.h which in turn include Package so no need to add backwards compatibility. Will add backwards compatibility in Class.h when Package.h is removed there (require a bit more work)

diff --git a/Engine/Source/Runtime/CoreUObject/Private/UObject/StructOnScope.cpp b/Engine/Source/Runtime/CoreUObject/Private/UObject/StructOnScope.cpp
new file mode 100644
--- /dev/null
+++ b/Engine/Source/Runtime/CoreUObject/Private/UObject/StructOnScope.cpp
@@ -0,0 +1,51 @@
+// Copyright Epic Games, Inc. All Rights Reserved.
+
+#include "UObject/StructOnScope.h"
+#include "UObject/Package.h"
+
+FStructOnScope::FStructOnScope()
+	: SampleStructMemory(nullptr)
+	, OwnsMemory(false)
+{
+}
+
+FStructOnScope::FStructOnScope(const UStruct* InScriptStruct)
+	: ScriptStruct(InScriptStruct)
+	, SampleStructMemory(nullptr)
+	, OwnsMemory(false)
+{
+	Initialize(); //-V1053
+}
+
+FStructOnScope::FStructOnScope(const UStruct* InScriptStruct, uint8* InData)
+	: ScriptStruct(InScriptStruct)
+	, SampleStructMemory(InData)
+	, OwnsMemory(false)
+{
+}
+
+FStructOnScope::FStructOnScope(FStructOnScope&& InOther)
+{
+	ScriptStruct = InOther.ScriptStruct;
+	SampleStructMemory = InOther.SampleStructMemory;
+	OwnsMemory = InOther.OwnsMemory;
+
+	InOther.OwnsMemory = false;
+	InOther.Reset();
+}
+
+FStructOnScope::~FStructOnScope()
+{
+	Destroy(); //-V1053
+}
+
+
+UPackage* FStructOnScope::GetPackage() const
+{
+	return Package.Get();
+}
+
+void FStructOnScope::SetPackage(UPackage* InPackage)
+{
+	Package = InPackage;
+}
diff --git a/Engine/Source/Runtime/CoreUObject/Public/UObject/GeneratedCppIncludes.h b/Engine/Source/Runtime/CoreUObject/Public/UObject/GeneratedCppIncludes.h
--- a/Engine/Source/Runtime/CoreUObject/Public/UObject/GeneratedCppIncludes.h
+++ b/Engine/Source/Runtime/CoreUObject/Public/UObject/GeneratedCppIncludes.h
@@ -7,7 +7,6 @@
 #include "UObject/UObjectGlobals.h"
 #include "UObject/CoreNative.h"
 #include "UObject/Class.h"
-#include "UObject/Package.h"
 #include "UObject/MetaData.h"
 #include "UObject/UnrealType.h"
 #include "UObject/EnumProperty.h"
diff --git a/Engine/Source/Runtime/CoreUObject/Public/UObject/StructOnScope.h b/Engine/Source/Runtime/CoreUObject/Public/UObject/StructOnScope.h
--- a/Engine/Source/Runtime/CoreUObject/Public/UObject/StructOnScope.h
+++ b/Engine/Source/Runtime/CoreUObject/Public/UObject/StructOnScope.h
@@ -5,7 +5,6 @@
 #include "CoreMinimal.h"
 #include "UObject/Class.h"
 #include "UObject/WeakObjectPtr.h"
-#include "UObject/Package.h"
 #include "Templates/Casts.h"
 #include "Templates/RemoveReference.h"
 
@@ -30,36 +29,10 @@ protected:
 
 public:
 
-	FStructOnScope()
-		: SampleStructMemory(nullptr)
-		, OwnsMemory(false)
-	{
-	}
-
-	FStructOnScope(const UStruct* InScriptStruct)
-		: ScriptStruct(InScriptStruct)
-		, SampleStructMemory(nullptr)
-		, OwnsMemory(false)
-	{
-		Initialize(); //-V1053
-	}
-
-	FStructOnScope(const UStruct* InScriptStruct, uint8* InData)
-		: ScriptStruct(InScriptStruct)
-		, SampleStructMemory(InData)
-		, OwnsMemory(false)
-	{
-	}
-
-	FStructOnScope(FStructOnScope&& InOther)
-	{
-		ScriptStruct = InOther.ScriptStruct;
-		SampleStructMemory = InOther.SampleStructMemory;
-		OwnsMemory = InOther.OwnsMemory;
-
-		InOther.OwnsMemory = false;
-		InOther.Reset();
-	}
+	COREUOBJECT_API FStructOnScope();
+	COREUOBJECT_API FStructOnScope(const UStruct* InScriptStruct);
+	COREUOBJECT_API FStructOnScope(const UStruct* InScriptStruct, uint8* InData);
+	COREUOBJECT_API FStructOnScope(FStructOnScope&& InOther);
 
 	FStructOnScope& operator=(FStructOnScope&& InOther)
 	{
@@ -100,15 +73,8 @@ public:
 		return ScriptStruct.Get();
 	}
 
-	virtual UPackage* GetPackage() const
-	{
-		return Package.Get();
-	}
-
-	virtual void SetPackage(UPackage* InPackage)
-	{
-		Package = InPackage;
-	}
+	COREUOBJECT_API virtual UPackage* GetPackage() const;
+	COREUOBJECT_API virtual void SetPackage(UPackage* InPackage);
 
 	virtual bool IsValid() const
 	{
@@ -147,10 +113,7 @@ public:
 		OwnsMemory = false;
 	}
 
-	virtual ~FStructOnScope()
-	{
-		Destroy(); //-V1053
-	}
+	COREUOBJECT_API virtual ~FStructOnScope();
 
 	/** Re-initializes the scope with a specified UStruct */
 	void Initialize(TWeakObjectPtr<const UStruct> InScriptStruct)
