Adding FPackageName::GetContentPathForPackageRoot to be able to get the filesystem path (e.g. D:/Project/Engine/Content) of a given package root (e.g. /Engine/)

diff --git a/Engine/Source/Runtime/CoreUObject/Private/Misc/PackageName.cpp b/Engine/Source/Runtime/CoreUObject/Private/Misc/PackageName.cpp
--- a/Engine/Source/Runtime/CoreUObject/Private/Misc/PackageName.cpp
+++ b/Engine/Source/Runtime/CoreUObject/Private/Misc/PackageName.cpp
@@ -1243,6 +1243,22 @@ FName FPackageName::GetPackageMountPoint(const FString& InPackagePath, bool InWi
 	return FName();
 }
 
+FString FPackageName::GetContentPathForPackageRoot(FStringView InMountPoint)
+{
+	FLongPackagePathsSingleton& Paths = FLongPackagePathsSingleton::Get();
+
+	FReadScopeLock ScopeLock(ContentMountPointCriticalSection);
+	for (const FPathPair& RootToPathEntry : Paths.ContentRootToPath)
+	{
+		if (RootToPathEntry.RootPath == InMountPoint)
+		{
+			return RootToPathEntry.AbsolutePath;
+		}
+	}
+
+	return FString();
+}
+
 bool FPackageName::TryConvertToMountedPathComponents(FStringView InFilePathOrPackageName,
 	FStringBuilderBase& OutMountPointPackageName, FStringBuilderBase& OutMountPointFilePath,
 	FStringBuilderBase& OutRelPath, FStringBuilderBase& OutObjectName, EPackageExtension& OutExtension,
diff --git a/Engine/Source/Runtime/CoreUObject/Public/Misc/PackageName.h b/Engine/Source/Runtime/CoreUObject/Public/Misc/PackageName.h
--- a/Engine/Source/Runtime/CoreUObject/Public/Misc/PackageName.h
+++ b/Engine/Source/Runtime/CoreUObject/Public/Misc/PackageName.h
@@ -316,6 +316,13 @@ public:
 	 */
 	static FName GetPackageMountPoint(const FString& InPackagePath, bool InWithoutSlashes = true);
 
+	/**
+	 * Get the path associated with the given package root
+	 * @param InPackageRoot Package root to return the path for, e.g. /Game/, /Engine/, /PluginName/
+	 * @return Filesystem path associated with the provided package root
+	 */
+	static FString GetContentPathForPackageRoot(FStringView InPackageRoot);
+
 	/**
 	 * Checks if the package exists on disk.
 	 * 
