Fixed FPropertyLocalizationDataGatherer failing to gather text from external actor packages

It wasn't finding any root objects in the actor package due to the "Object->GetOuter() == Package" test failing for external actors.

#rnx

diff --git a/Engine/Source/Runtime/CoreUObject/Private/Serialization/PropertyLocalizationDataGathering.cpp b/Engine/Source/Runtime/CoreUObject/Private/Serialization/PropertyLocalizationDataGathering.cpp
--- a/Engine/Source/Runtime/CoreUObject/Private/Serialization/PropertyLocalizationDataGathering.cpp
+++ b/Engine/Source/Runtime/CoreUObject/Private/Serialization/PropertyLocalizationDataGathering.cpp
@@ -30,9 +30,12 @@ FPropertyLocalizationDataGatherer::FPropertyLocalizationDataGatherer(TArray<FGat
 	}, true, RF_Transient, EInternalObjectFlags::Garbage);
 
 	// Iterate over each root object in the package
-	for (const UObject* Object : AllObjectsInPackage)
+	// Note: This calls GetObjectsWithPackage rather than test the AllObjectsInPackage array, as external actors don't pass a "Object->GetOuter() == Package" test when trying to query for root objects
 	{
-		if (Object->GetOuter() == Package)
+		TArray<UObject*> RootObjectsInPackage;
+		GetObjectsWithPackage(Package, RootObjectsInPackage, false, RF_Transient, EInternalObjectFlags::Garbage);
+
+		for (const UObject* Object : RootObjectsInPackage)
 		{
 			GatherLocalizationDataFromObjectWithCallbacks(Object, EPropertyLocalizationGathererTextFlags::None);
 		}
