Core: Array Text import not working for ()

diff --git a/Engine/Source/Runtime/CoreUObject/Private/UObject/PropertyArray.cpp b/Engine/Source/Runtime/CoreUObject/Private/UObject/PropertyArray.cpp
--- a/Engine/Source/Runtime/CoreUObject/Private/UObject/PropertyArray.cpp
+++ b/Engine/Source/Runtime/CoreUObject/Private/UObject/PropertyArray.cpp
@@ -641,6 +641,16 @@ const TCHAR* FArrayProperty::ImportTextInnerItem( const TCHAR* Buffer, const FPr
 
 	SkipWhitespace(Buffer);
 
+	if (*Buffer == TCHAR(')'))
+	{
+		// We didn't find any items. Clear the array again
+		if (ArrayHelper)
+		{
+			ArrayHelper->EmptyValues();
+		}
+		return Buffer;
+	}
+
 	int32 Index = 0;
 	while (*Buffer != TCHAR(')'))
 	{
