Add Ability to provide Default PackageFlags to each mount point during CreatePackage

diff --git a/Engine/Source/Runtime/CoreUObject/Private/UObject/UObjectGlobals.cpp b/Engine/Source/Runtime/CoreUObject/Private/UObject/UObjectGlobals.cpp
--- a/Engine/Source/Runtime/CoreUObject/Private/UObject/UObjectGlobals.cpp
+++ b/Engine/Source/Runtime/CoreUObject/Private/UObject/UObjectGlobals.cpp
@@ -41,6 +41,7 @@
 #include "Serialization/DuplicatedDataWriter.h"
 #include "Serialization/LoadTimeTracePrivate.h"
 #include "Misc/PackageName.h"
+#include "Misc/PathViews.h"
 #include "UObject/LinkerLoad.h"
 #include "Blueprint/BlueprintSupport.h"
 #include "Misc/SecureHash.h"
@@ -902,6 +903,59 @@ UPackage* FindPackage( UObject* InOuter, const TCHAR* PackageName )
 	return Result;
 }
 
+#if WITH_EDITOR
+struct FCreatePackageDefaultFlagsMap
+{
+	uint32 Find(FStringView MountPoint)
+	{
+		FReadScopeLock ReadScopeLock(Lock);
+		EPackageFlags* DefaultPackageFlags = MountPointToDefaultPackageFlags.FindByHash(GetTypeHash(MountPoint), MountPoint);
+
+		return DefaultPackageFlags ? *DefaultPackageFlags : EPackageFlags::PKG_None;
+	}
+
+	void Add(const TMap<FString, EPackageFlags>& InMountPointToDefaultPackageFlags)
+	{
+		if (InMountPointToDefaultPackageFlags.IsEmpty())
+		{
+			return;
+		}
+
+		FWriteScopeLock WriteScopeLock(Lock);
+		MountPointToDefaultPackageFlags.Append(InMountPointToDefaultPackageFlags);
+	}
+
+	void Remove(const TArrayView<FString>& InMountPoints)
+	{
+		if (InMountPoints.IsEmpty())
+		{
+			return;
+		}
+
+		FWriteScopeLock WriteScopeLock(Lock);
+		for (const FString& MountPoint : InMountPoints)
+		{
+			MountPointToDefaultPackageFlags.Remove(MountPoint);
+		}
+	}
+private:
+	FRWLock Lock;
+	TMap<FString, EPackageFlags> MountPointToDefaultPackageFlags;
+};
+
+static FCreatePackageDefaultFlagsMap GCreatePackageDefaultFlagsMap;
+
+void SetMountPointDefaultPackageFlags(const TMap<FString, EPackageFlags>& InMountPointToDefaultPackageFlags)
+{
+	GCreatePackageDefaultFlagsMap.Add(InMountPointToDefaultPackageFlags);
+}
+
+void RemoveMountPointDefaultPackageFlags(const TArrayView<FString> InMountPoints)
+{
+	GCreatePackageDefaultFlagsMap.Remove(InMountPoints);
+}
+#endif //if WITH_EDITOR
+
 UPackage* CreatePackage(UObject* InOuter, const TCHAR* PackageName)
 {
 	return CreatePackage(PackageName);
@@ -956,6 +1010,11 @@ UPackage* CreatePackage(const TCHAR* PackageName )
 			else
 			{
 				Result = NewObject<UPackage>(nullptr, NewPackageName, RF_Public);
+#if WITH_EDITOR
+				FStringView PackageMountPoint = FPathViews::GetMountPointNameFromPath(InName);
+				uint32 DefaultPackageFlags = GCreatePackageDefaultFlagsMap.Find(PackageMountPoint);
+				Result->SetPackageFlags(DefaultPackageFlags);
+#endif
 			}
 		}
 	}
diff --git a/Engine/Source/Runtime/CoreUObject/Public/UObject/UObjectGlobals.h b/Engine/Source/Runtime/CoreUObject/Public/UObject/UObjectGlobals.h
--- a/Engine/Source/Runtime/CoreUObject/Public/UObject/UObjectGlobals.h
+++ b/Engine/Source/Runtime/CoreUObject/Public/UObject/UObjectGlobals.h
@@ -928,6 +928,20 @@ COREUOBJECT_API void EndLoad(FUObjectSerializeContext* LoadContext);
  */
 COREUOBJECT_API UPackage* FindPackage(UObject* InOuter, const TCHAR* PackageName);
 
+#if WITH_EDITOR
+/**
+* Sets default PackageFlags for new packages made in specified mount points during CreatePackage
+* @param InMountPointToDefaultPackageFlags A map of MountPoints to PackageFlags. Used to provide new Packages in each mount point with the associated DefaultFlags
+*/
+COREUOBJECT_API void SetMountPointDefaultPackageFlags(const TMap<FString, EPackageFlags>& InMountPointToDefaultPackageFlags);
+
+/**
+* Removes the provided list of mount points from the MountPointToDefaultPackageFlags map
+* @param InMountPoints A list of mount points that will be removed and no longer have DefaultPackageFlags associated with them
+*/
+COREUOBJECT_API void RemoveMountPointDefaultPackageFlags(const TArrayView<FString> InMountPoints);
+#endif
+
 /**
  * Find an existing package by name or create it if it doesn't exist
  * @param InOuter		The Outer object to search inside (unused)
