PR #9698: Adding missing script macros for weak object properties.
#rnx

diff --git a/Engine/Source/Runtime/CoreUObject/Public/UObject/ScriptMacros.h b/Engine/Source/Runtime/CoreUObject/Public/UObject/ScriptMacros.h
--- a/Engine/Source/Runtime/CoreUObject/Public/UObject/ScriptMacros.h
+++ b/Engine/Source/Runtime/CoreUObject/Public/UObject/ScriptMacros.h
@@ -83,6 +83,9 @@ enum {MAX_VARIABLE_SIZE = 0x0FFF };
 #define P_GET_TINTERFACE(ObjectType,ParamName)		PARAM_PASSED_BY_VAL(ParamName, FInterfaceProperty, TScriptInterface<ObjectType>)
 #define P_GET_TINTERFACE_REF(ObjectType,ParamName)	PARAM_PASSED_BY_REF(ParamName, FInterfaceProperty, TScriptInterface<ObjectType>)
 
+#define P_GET_WEAKOBJECT(ObjectType,ParamName)		PARAM_PASSED_BY_VAL(ParamName, FWeakObjectProperty, ObjectType)
+#define P_GET_WEAKOBJECT_REF(ObjectType,ParamName)	PARAM_PASSED_BY_REF(ParamName, FWeakObjectProperty, ObjectType)
+
 #define P_GET_SOFTOBJECT(ObjectType,ParamName)		PARAM_PASSED_BY_VAL(ParamName, FSoftObjectProperty, ObjectType)
 #define P_GET_SOFTOBJECT_REF(ObjectType,ParamName)	PARAM_PASSED_BY_REF(ParamName, FSoftObjectProperty, ObjectType)
 
