Disallow TObjectPtr use on complete non-UObject derived types.

diff --git a/Engine/Source/Runtime/CoreUObject/Public/UObject/ObjectPtr.h b/Engine/Source/Runtime/CoreUObject/Public/UObject/ObjectPtr.h
--- a/Engine/Source/Runtime/CoreUObject/Public/UObject/ObjectPtr.h
+++ b/Engine/Source/Runtime/CoreUObject/Public/UObject/ObjectPtr.h
@@ -319,6 +319,12 @@ namespace ObjectPtr_Private
 		// a shallow pointer comparison.
 		return IsObjectPtrEqualToRawPtrOfRelatedType<T>(Ptr, ObjectPtr_Private::CoerceToPointer<T>(Other));
 	}
+
+	template <typename T, int = sizeof(T)>
+	char (&ResolveTypeIsComplete(int))[2];
+
+	template <typename T>
+	char (&ResolveTypeIsComplete(...))[1];
 };
 
 /**
@@ -336,6 +342,13 @@ namespace ObjectPtr_Private
 template <typename T>
 struct TObjectPtr
 {
+	// TObjectPtr should only be used on types T that are EITHER:
+	// - incomplete (ie: forward declared and we have not seen their definition yet)
+	// - complete and derived from UObject
+	// This means that the following are invalid and must fail to compile:
+	// - TObjectPtr<int>
+	// - TObjectPtr<IInterface>
+	static_assert(std::disjunction<std::negation<std::bool_constant<sizeof(ObjectPtr_Private::ResolveTypeIsComplete<T>(1)) == 2>>, std::is_base_of<UObject, T>>::value, "TObjectPtr<T> can only be used with types derived from UObject");
 public:
 	using ElementType = T;
 
