Using undo after landscape editing should no longer cause a mismatch in heightmap data being stored, preventing the asset from being saved.

- I needed to remove the 5.0 optimization where we tried to avoid loading the payload into memory so that it could be serialized to the undo stack. This logic would not work if the package were saved later and we were therefor unable to access the payload on disk.
- The landscape bug was not caused by this (as saving landscape actors wipes the undo stack) but caused by the inability to unload the payload properly. If an undo restore was run from a state where the payload was in memory to a state where the payload was not in memory, then that payload would not unload properly.So the bulkdata would end up with members for one payload, but the data of a different payload.

- We now always load the payload into memory if it is not there already and decompress into shared buffer format. In theory we could support two paths here one for FSharedBuffer when the payload is in memory and one for FCompressedBuffer when it is loaded from disk but it doesn't seem worth the effort. If it comes up as a problem in a profile then we'd be better off trying to add a way to create a non-compressed FCompressedBuffer from a FIoHash/FSharedBuffer pair rather than rehashing the FSharedBuffer.
- This also helps make the code clearer as we can remove the bPayloadInArchive value that we used to save.
- Once the undo stack transaction is completed we will try to unload the payload if we can. This is done if we are saving or loading and will only unload the payload if we know we can safely load from disk. Which means that the payload has not been edited from the version on disk and that the owning package has not been re-saved. This is to conform to the editor bulkdata policy that we should evict the payload from memory when ever possible to keep over all memory use down. It will however mean that repeated editing of an asset will cause repeated disk loads. This logic could be changed if it is profiled as an issue.

diff --git a/Engine/Source/Runtime/CoreUObject/Private/Serialization/EditorBulkData.cpp b/Engine/Source/Runtime/CoreUObject/Private/Serialization/EditorBulkData.cpp
--- a/Engine/Source/Runtime/CoreUObject/Private/Serialization/EditorBulkData.cpp
+++ b/Engine/Source/Runtime/CoreUObject/Private/Serialization/EditorBulkData.cpp
@@ -823,18 +823,19 @@ void FEditorBulkData::Serialize(FArchive& Ar, UObject* Owner, bool bAllowRegiste
 			// in memory or some other caching system.
 			// Serializing full 8k texture payloads to memory on each metadata change will empty
 			// the undo stack very quickly.
-			
-			// Note that we will only serialize the payload if it is in memory. Otherwise we can
-			// continue to load the payload as needed from disk or pull from the virtualization system
-			bool bPayloadInArchive = Ar.IsSaving() ? !Payload.IsNull() : false;
-			Ar << bPayloadInArchive;
 
 			if (Ar.IsSaving())
 			{
-				if (bPayloadInArchive)
+				if (Payload.IsNull())
 				{
-					Ar << Payload;
+					// We need to serialize in FSharedBuffer form, or otherwise we'd need to support
+					// multiple code paths here. Technically a bit wasteful but for general use it 
+					// shouldn't be noticable. This will make it easier to do the real perf wins
+					// in the future.
+					Payload = GetDataInternal().Decompress();
 				}
+
+				Ar << Payload;
 			}
 			else
 			{
@@ -845,17 +846,17 @@ void FEditorBulkData::Serialize(FArchive& Ar, UObject* Owner, bool bAllowRegiste
 					BulkDataId = FGuid::NewGuid();
 				}
 
-				if (bPayloadInArchive)
-				{
-					Ar << Payload;
-				}
-				else
-				{
-					UnloadData();
-				}
+				Ar << Payload;
 
 				Register(Owner, TEXT("Serialize/Transacting"), false /* bAllowUpdateId */);
 			}
+
+			// Try to unload the payload if possible, usually because we loaded it during the transaction in the
+			// first place and we don't want to keep it in memory anymore.
+			// This does mean if the owning asset is frequently edited we will be reloading the payload off disk
+			// a lot. But in practice this didn't show up as too much of a problem. If someone has found this to 
+			// be a perf issue, then remove the call to ::UnloadData and trade memory cost for perf gain.
+			UnloadData();
 		}
 	}
 	else if (Ar.IsPersistent() && !Ar.IsObjectReferenceCollector() && !Ar.ShouldSkipBulkData())
