Add The Input Device trigger enum to the IInputInterface and as a noexport blueprint type. This will allow you to specify which trigger should recieve effects like resistence, feedback, vibration, etc.

#rnx

diff --git a/Engine/Source/Runtime/ApplicationCore/Public/GenericPlatform/IInputInterface.h b/Engine/Source/Runtime/ApplicationCore/Public/GenericPlatform/IInputInterface.h
--- a/Engine/Source/Runtime/ApplicationCore/Public/GenericPlatform/IInputInterface.h
+++ b/Engine/Source/Runtime/ApplicationCore/Public/GenericPlatform/IInputInterface.h
@@ -89,6 +89,20 @@ struct FHapticFeedbackValues
 	}
 };
 
+/**
+ * Represents input device triggers that are available
+ *
+ * NOTE: Make sure to keep this type in sync with the reflected version in NoExportTypes.h!
+ */
+enum class EInputDeviceTriggerMask : uint8
+{
+	None		= 0x00,
+	Left		= 0x01,
+	Right		= 0x02,
+	All			= Left | Right
+};
+ENUM_CLASS_FLAGS(EInputDeviceTriggerMask)
+
 struct FInputDeviceProperty
 {
 	FInputDeviceProperty(FName InName)
diff --git a/Engine/Source/Runtime/CoreUObject/Public/UObject/NoExportTypes.h b/Engine/Source/Runtime/CoreUObject/Public/UObject/NoExportTypes.h
--- a/Engine/Source/Runtime/CoreUObject/Public/UObject/NoExportTypes.h
+++ b/Engine/Source/Runtime/CoreUObject/Public/UObject/NoExportTypes.h
@@ -55,6 +55,7 @@
 #include "AssetRegistry/AssetData.h"
 
 #include "../../../ApplicationCore/Public/GenericPlatform/ICursor.h"
+#include "../../../ApplicationCore/Public/GenericPlatform/IInputInterface.h"
 
 #endif
 
@@ -2431,6 +2432,20 @@ enum class EInputDeviceConnectionState : uint8
 	Connected
 };
 
+/**
+ * Represents input device triggers that are available
+ *
+ * NOTE: Make sure to keep this type in sync with the reflected version in IInputInterface.h!
+ */
+UENUM(BlueprintType)
+enum class EInputDeviceTriggerMask : uint8
+{
+	None		= 0x00,
+	Left		= 0x01,
+	Right		= 0x02,
+	All			= Left | Right
+};
+
 /**
  * Data about an input device's current state
  */
