[CoreUObject] changed to check for null in `TObjectPtr.GetPathName()`

diff --git a/Engine/Source/Runtime/CoreUObject/Private/UObject/ObjectPtr.cpp b/Engine/Source/Runtime/CoreUObject/Private/UObject/ObjectPtr.cpp
--- a/Engine/Source/Runtime/CoreUObject/Private/UObject/ObjectPtr.cpp
+++ b/Engine/Source/Runtime/CoreUObject/Private/UObject/ObjectPtr.cpp
@@ -6,7 +6,7 @@
 
 FString FObjectPtr::GetPathName() const
 {
-	if (IsResolved())
+	if (IsResolved() && !IsObjectHandleNull(Handle))
 	{
 		return Get()->GetPathName();
 	}
diff --git a/Engine/Source/Runtime/CoreUObject/Public/UObject/ObjectPtr.h b/Engine/Source/Runtime/CoreUObject/Public/UObject/ObjectPtr.h
--- a/Engine/Source/Runtime/CoreUObject/Public/UObject/ObjectPtr.h
+++ b/Engine/Source/Runtime/CoreUObject/Public/UObject/ObjectPtr.h
@@ -166,7 +166,8 @@ public:
 #else
 	FString GetPathName() const
 	{
-		return Get()->GetPathName();
+		const UObject* ResolvedObject = Get();
+		return ResolvedObject ? ResolvedObject->GetPathName() : TEXT("None");
 	}
 
 	FName GetFName() const
diff --git a/Engine/Source/Runtime/CoreUObject/Tests/ObjectPtrTest.cpp b/Engine/Source/Runtime/CoreUObject/Tests/ObjectPtrTest.cpp
--- a/Engine/Source/Runtime/CoreUObject/Tests/ObjectPtrTest.cpp
+++ b/Engine/Source/Runtime/CoreUObject/Tests/ObjectPtrTest.cpp
@@ -318,6 +318,48 @@ TEST_CASE_METHOD(FObjectPtrTestBase, "CoreUObject::TObjectPtr::Long Path", "[Cor
 	TEST_EQUAL(TEXT("Resolved path from FObjectPathId should have TestObject4 at element 3"), ResolvedNames[3], TestObject4->GetFName());
 }
 
+TEST_CASE("CoreUObject::TObjectPtr::GetPathName", "[CoreUObject][ObjectPtr]")
+{
+	const FName TestPackage1Name(TEXT("/Engine/Test/FObjectPtrTestLongPath/Transient"));
+	UPackage* TestPackage1 = NewObject<UPackage>(nullptr, TestPackage1Name, RF_Transient);
+	TestPackage1->AddToRoot();
+	UObject* TestObject1 = NewObject<UTestDummyObject>(TestPackage1, TEXT("TestObject1"));
+
+#if UE_WITH_OBJECT_HANDLE_TRACKING
+	int ResolveCount = 0;
+	auto ResolveDelegate = FObjectHandleReferenceResolvedDelegate::CreateLambda([&ResolveCount](const FObjectRef& SourceRef, UPackage* ObjectPackage, UObject* Object)
+		{
+			++ResolveCount;
+		});
+	auto Handle = AddObjectHandleReferenceResolvedCallback(ResolveDelegate);
+#endif
+
+	ON_SCOPE_EXIT
+	{
+		TestPackage1->RemoveFromRoot();
+#if UE_WITH_OBJECT_HANDLE_TRACKING
+		RemoveObjectHandleReferenceResolvedCallback(Handle);
+#endif
+	};
+
+	TObjectPtr<UObject> Ptr = nullptr;
+	CHECK(TEXT("None") == Ptr.GetPathName());
+
+	Ptr = TestObject1;
+	CHECK(TestObject1->GetPathName() == Ptr.GetPathName());
+
+#if UE_WITH_OBJECT_HANDLE_LATE_RESOLVE
+	FObjectPtr ObjPtr(MakePackedObjectRef(TestObject1));
+	Ptr = *reinterpret_cast<TObjectPtr<UObject>*>(&ObjPtr);
+	REQUIRE(!Ptr.IsResolved());
+	CHECK(TestObject1->GetPathName() == Ptr.GetPathName());
+#endif 
+
+#if UE_WITH_OBJECT_HANDLE_TRACKING
+	REQUIRE(ResolveCount == 0);
+#endif
+}
+
 // @TODO: OBJPTR: We should have a test that ensures that lazy loading of an object with an external package is handled correctly.
 //				  This should also include external packages in the outer chain of the target object.
 // IMPLEMENT_CUSTOM_SIMPLE_AUTOMATION_TEST(FObjectPtrTestExternalPackages, FObjectPtrTestBase, TEXT(TEST_NAME_ROOT ".ExternalPackages"), ObjectPtrTestFlags)
