Fixed compile error in BulkData.h caused by integration

diff --git a/Engine/Source/Runtime/CoreUObject/Public/Serialization/BulkData.h b/Engine/Source/Runtime/CoreUObject/Public/Serialization/BulkData.h
--- a/Engine/Source/Runtime/CoreUObject/Public/Serialization/BulkData.h
+++ b/Engine/Source/Runtime/CoreUObject/Public/Serialization/BulkData.h
@@ -14,6 +14,7 @@
 #include "CoreTypes.h"
 #include "GenericPlatform/GenericPlatformFile.h"
 #include "HAL/MemoryBase.h"
+#include "IO/IoChunkId.h"
 #include "IO/IoDispatcherPriority.h"
 #include "IO/PackageId.h"
 #include "Math/NumericLimits.h"
