#UE Do not cook editor only soft object paths. Fix a regression that came from serializing soft object paths before processing the properties that reference them

diff --git a/Engine/Source/Runtime/CoreUObject/Private/UObject/LinkerLoad.cpp b/Engine/Source/Runtime/CoreUObject/Private/UObject/LinkerLoad.cpp
--- a/Engine/Source/Runtime/CoreUObject/Private/UObject/LinkerLoad.cpp
+++ b/Engine/Source/Runtime/CoreUObject/Private/UObject/LinkerLoad.cpp
@@ -1698,6 +1698,15 @@ FLinkerLoad::ELinkerStatus FLinkerLoad::SerializeSoftObjectPathList()
 		Seek(Summary.SoftObjectPathsOffset);
 	}
 
+#if WITH_EDITOR
+	FName PackageName, PropertyName;
+	ESoftObjectPathCollectType CollectType = ESoftObjectPathCollectType::AlwaysCollect;
+	ESoftObjectPathSerializeType SerializeType = ESoftObjectPathSerializeType::AlwaysSerialize;
+	FSoftObjectPathThreadContext& ThreadContext = FSoftObjectPathThreadContext::Get();
+	ThreadContext.GetSerializationOptions(PackageName, PropertyName, CollectType, SerializeType, this);
+	FSoftObjectPathSerializationScope SerializationScope(PackageName, PropertyName, ESoftObjectPathCollectType::NeverCollect, SerializeType);
+#endif // WITH_EDITOR
+
 	FStructuredArchive::FStream Stream = StructuredArchiveRootRecord->EnterStream(TEXT("SoftObjectPathList"));
 	while (SoftObjectPathListIndex < Summary.SoftObjectPathsCount && !IsTimeLimitExceeded(TEXT("serializing soft object path list"), 100))
 	{
@@ -5853,6 +5862,10 @@ FArchive& FLinkerLoad::operator<<(FSoftObjectPath& Value)
 		if (SoftObjectPathList.IsValidIndex(SoftObjectPathIndex))
 		{
 			Value = SoftObjectPathList[SoftObjectPathIndex];
+
+#if WITH_EDITOR
+			Value.PostLoadPath(this);
+#endif // WITH_EDITOR
 		}
 		else
 		{
