BulkDataStreaming.cpp missing include Package.h

diff --git a/Engine/Source/Runtime/CoreUObject/Private/Serialization/BulkDataStreaming.cpp b/Engine/Source/Runtime/CoreUObject/Private/Serialization/BulkDataStreaming.cpp
--- a/Engine/Source/Runtime/CoreUObject/Private/Serialization/BulkDataStreaming.cpp
+++ b/Engine/Source/Runtime/CoreUObject/Private/Serialization/BulkDataStreaming.cpp
@@ -11,6 +11,7 @@
 #include "Serialization/MemoryReader.h"
 #include "Templates/RefCounting.h"
 #include "UObject/LinkerLoad.h"
+#include "UObject/Package.h"
 #include "UObject/PackageResourceIoDispatcherBackend.h"
 
 //////////////////////////////////////////////////////////////////////////////
