Updated UObject::SaveConfig to not flush all the configs but instead only flush the updated config.

Before: 837ms
After: 46ms

diff --git a/Engine/Source/Runtime/CoreUObject/Private/UObject/Obj.cpp b/Engine/Source/Runtime/CoreUObject/Private/UObject/Obj.cpp
--- a/Engine/Source/Runtime/CoreUObject/Private/UObject/Obj.cpp
+++ b/Engine/Source/Runtime/CoreUObject/Private/UObject/Obj.cpp
@@ -2883,7 +2883,7 @@ void UObject::SaveConfig( uint64 Flags, const TCHAR* InFilename, FConfigCacheIni
 	// only write out the config file if this is GConfig
 	if (Config == GConfig)
 	{
-		Config->Flush( 0 );
+		Config->Flush( 0, Filename );
 	}
 }
 
