AsyncLoading2: Fix infinite wait for io issue caused by decreasing PendingIoRequestsCounter too late in cl 22264831.
Improve handling of ProcessPendingCDOs by relying on the SyncLoadContext instead of a combination of PendingIoRequestsCounter and ReferencerRequestId.
Only create CDOs for packages in the current sync load context when there is one.
Add and release barriers for all export bundle nodes, not only the first one.

diff --git a/Engine/Source/Runtime/CoreUObject/Private/Serialization/AsyncLoading2.cpp b/Engine/Source/Runtime/CoreUObject/Private/Serialization/AsyncLoading2.cpp
--- a/Engine/Source/Runtime/CoreUObject/Private/Serialization/AsyncLoading2.cpp
+++ b/Engine/Source/Runtime/CoreUObject/Private/Serialization/AsyncLoading2.cpp
@@ -460,8 +460,6 @@ struct FAsyncPackageDesc2
 {
 	// A unique request id for each external call to LoadPackage
 	int32 RequestID;
-	// The request id of the referencer propagated down the import chain from the most recent load request
-	int32 ReferencerRequestId;
 	// Package priority
 	int32 Priority;
 	/** The flags that should be applied to the package */
@@ -491,7 +489,6 @@ struct FAsyncPackageDesc2
 	{
 		return FAsyncPackageDesc2
 		{
-			Request.RequestId,
 			Request.RequestId,
 			Request.Priority,
 			Request.PackageFlags,
@@ -520,7 +517,6 @@ struct FAsyncPackageDesc2
 		return FAsyncPackageDesc2
 		{
 			INDEX_NONE,
-			ImportingPackageDesc.ReferencerRequestId,
 			ImportingPackageDesc.Priority,
 			PKG_None,
 #if WITH_EDITOR
@@ -1766,7 +1762,7 @@ public:
 		return Package;
 	}
 
-	int32 ReferencerRequestId() const;
+	uint64 GetSyncLoadContextId() const;
 
 private:
 	friend class FAsyncLoadEventQueue2;
@@ -2301,9 +2297,9 @@ struct FAsyncPackage2
 	/** Adds new request ID to the existing package */
 	void AddRequestID(int32 Id);
 
-	int32 ReferencerRequestId() const
+	uint64 GetSyncLoadContextId() const
 	{
-		return Desc.ReferencerRequestId;
+		return SyncLoadContextId;
 	}
 
 	/**
@@ -2952,11 +2948,15 @@ public:
 
 	void AddPendingCDOs(FAsyncPackage2* Package, TArray<UClass*, TInlineAllocator<8>>& Classes)
 	{
-		FEventLoadNode2& FirstBundleNode = Package->GetExportBundleNode(ExportBundle_Process, 0);
-		FirstBundleNode.AddBarrier(Classes.Num());
 		for (UClass* Class : Classes)
 		{
-			PendingCDOs.FindOrAdd(Class).Add(&FirstBundleNode);
+			TArray<FEventLoadNode2*>& Nodes = PendingCDOs.FindOrAdd(Class);
+			for (int32 ExportBundleIndex = 0; ExportBundleIndex < Package->Data.ExportBundleCount; ++ExportBundleIndex)
+			{
+				FEventLoadNode2& BundleNode = Package->GetExportBundleNode(ExportBundle_Process, ExportBundleIndex);
+				BundleNode.AddBarrier();
+				Nodes.Add(&BundleNode);
+			}
 		}
 	}
 
@@ -2994,13 +2994,13 @@ private:
 
 	void RemoveUnreachableObjects(FUnreachableObjects& ObjectsToRemove);
 
-	bool ProcessPendingCDOs()
+	bool ProcessPendingCDOs(FAsyncLoadingThreadState2& ThreadState)
 	{
 		TRACE_CPUPROFILER_EVENT_SCOPE(ProcessPendingCDOs);
 
 		bool bDidSomething = false;
 		UClass* Class = nullptr;
-		int32 MaxRequestId = -1;
+		const uint64 SyncLoadContextId = ThreadState.SyncLoadContextStack.Num() > 0 ? ThreadState.SyncLoadContextStack.Top()->ContextId : 0;
 		for (TMap<UClass*, TArray<FEventLoadNode2*>>::TIterator It = PendingCDOs.CreateIterator(); It; ++It)
 		{
 			UClass* CurrentClass = It.Key();
@@ -3026,23 +3026,28 @@ private:
 			const TArray<FEventLoadNode2*>& Nodes = It.Value();
 			for (const FEventLoadNode2* Node : Nodes)
 			{
-				int32 RequestId = Node->ReferencerRequestId();
-				if (RequestId > MaxRequestId)
+				int32 NodeContextId = Node->GetSyncLoadContextId();
+				if (NodeContextId >= SyncLoadContextId)
 				{
-					MaxRequestId = RequestId;
 					Class = CurrentClass;
+					break;
 				}
 			}
+
+			if (Class != nullptr)
+			{
+				break;
+			}
 		}
 
-		if (ensure(Class))
+		if (Class)
 		{
 			TArray<FEventLoadNode2*> Nodes;
 			PendingCDOs.RemoveAndCopyValue(Class, Nodes);
 
 			UE_LOG(LogStreaming, Log,
-				TEXT("ProcessPendingCDOs: Creating CDO for '%s' for request id %d, releasing %d nodes. %d CDOs remaining."),
-				*Class->GetFullName(), MaxRequestId, Nodes.Num(), PendingCDOs.Num());
+				TEXT("ProcessPendingCDOs: Creating CDO for '%s' for SyncLoadContextId %llu, releasing %d nodes. %d CDOs remaining."),
+				*Class->GetFullName(), SyncLoadContextId, Nodes.Num(), PendingCDOs.Num());
 			PendingCDOsRecursiveStack.Push(Class);
 			UObject* CDO = Class->GetDefaultObject(/*bCreateIfNeeded*/ true);
 			verify(PendingCDOsRecursiveStack.Pop() == Class);
@@ -3057,20 +3062,6 @@ private:
 
 			bDidSomething = true;
 		}
-		else
-		{
-			for (TMap<UClass*, TArray<FEventLoadNode2*>>::TIterator It = PendingCDOs.CreateIterator(); It; ++It)
-			{
-				UClass* CurrentClass = It.Key();
-				check(CurrentClass != nullptr);
-
-				const TArray<FEventLoadNode2*>& Nodes = It.Value();
-				UE_LOG(LogStreaming, Warning,
-					TEXT("ProcessPendingCDOs: '%s' with %d nodes could not be processed from this stack."),
-					*CurrentClass->GetFullName(), Nodes.Num());
-			}
-		}
-
 		return bDidSomething;
 	}
 
@@ -3353,7 +3344,6 @@ FAsyncPackage2* FAsyncLoadingThread2::FindOrInsertPackage(FAsyncPackageDesc2& De
 			{
 				UpdatePackagePriority(Package, Desc.Priority);
 			}
-			Package->Desc.ReferencerRequestId = Desc.ReferencerRequestId;
 		}
 		if (PackageLoadedDelegate.IsValid())
 		{
@@ -3733,9 +3723,9 @@ void FEventLoadNode2::ProcessDependencies(FAsyncLoadingThreadState2& ThreadState
 	}
 }
 
-int32 FEventLoadNode2::ReferencerRequestId() const
+uint64 FEventLoadNode2::GetSyncLoadContextId() const
 {
-	return Package->ReferencerRequestId();
+	return Package->GetSyncLoadContextId();
 }
 
 FAsyncLoadEventQueue2::FAsyncLoadEventQueue2()
@@ -4435,10 +4425,9 @@ void FAsyncPackage2::StartLoading(FIoBatch& IoBatch)
 						TEXT("Failed reading optional chunk for package: %s"), *Result.Status().ToString());
 					bLoadHasFailed = true;
 				}
-				FAsyncLoadingThread2* LocalAsyncLoadingThread = &AsyncLoadingThread;
-				GetPackageNode(EEventLoadNode2::Package_ProcessSummary).ReleaseBarrier();
-				int32 LocalPendingIoRequestsCounter = LocalAsyncLoadingThread->PendingIoRequestsCounter.DecrementExchange() - 1;
+				int32 LocalPendingIoRequestsCounter = AsyncLoadingThread.PendingIoRequestsCounter.DecrementExchange() - 1;
 				TRACE_COUNTER_SET(AsyncLoadingPendingIoRequests, LocalPendingIoRequestsCounter);
+				GetPackageNode(EEventLoadNode2::Package_ProcessSummary).ReleaseBarrier();
 			});
 	}
 #endif
@@ -4462,10 +4451,9 @@ void FAsyncPackage2::StartLoading(FIoBatch& IoBatch)
 					TEXT("Failed reading chunk for package: %s"), *Result.Status().ToString());
 				bLoadHasFailed = true;
 			}
-			FAsyncLoadingThread2* LocalAsyncLoadingThread = &AsyncLoadingThread;
-			GetPackageNode(EEventLoadNode2::Package_ProcessSummary).ReleaseBarrier();
-			int32 LocalPendingIoRequestsCounter = LocalAsyncLoadingThread->PendingIoRequestsCounter.DecrementExchange() - 1;
+			int32 LocalPendingIoRequestsCounter = AsyncLoadingThread.PendingIoRequestsCounter.DecrementExchange() - 1;
 			TRACE_COUNTER_SET(AsyncLoadingPendingIoRequests, LocalPendingIoRequestsCounter);
+			GetPackageNode(EEventLoadNode2::Package_ProcessSummary).ReleaseBarrier();
 		});
 
 	if (!Data.ShaderMapHashes.IsEmpty())
@@ -4480,10 +4468,9 @@ void FAsyncPackage2::StartLoading(FIoBatch& IoBatch)
 				[this, GraphEvent](TIoStatusOr<FIoBuffer> Result)
 				{
 					GraphEvent->DispatchSubsequents();
-					FAsyncLoadingThread2* LocalAsyncLoadingThread = &AsyncLoadingThread;
-					GetPackageNode(Package_ExportsSerialized).ReleaseBarrier();
-					int32 LocalPendingIoRequestsCounter = LocalAsyncLoadingThread->PendingIoRequestsCounter.DecrementExchange() - 1;
+					int32 LocalPendingIoRequestsCounter = AsyncLoadingThread.PendingIoRequestsCounter.DecrementExchange() - 1;
 					TRACE_COUNTER_SET(AsyncLoadingPendingIoRequests, LocalPendingIoRequestsCounter);
+					GetPackageNode(Package_ExportsSerialized).ReleaseBarrier();
 				});
 		};
 		FCoreDelegates::PreloadPackageShaderMaps.ExecuteIfBound(Data.ShaderMapHashes, ReadShaderMapFunc);
@@ -6580,7 +6567,6 @@ EAsyncPackageState::Type FAsyncLoadingThread2::TickAsyncLoadingFromGameThread(FA
 		ThreadState.SetTimeLimit(bUseTimeLimit, TimeLimit);
 
 		const bool bIsMultithreaded = FAsyncLoadingThread2::IsMultithreaded();
-		bool bMayProcessPendingCDOs = PendingCDOs.Num() > 0 && (PendingIoRequestsCounter.Load() == 0);
 		double TickStartTime = FPlatformTime::Seconds();
 
 		{
@@ -6596,14 +6582,9 @@ EAsyncPackageState::Type FAsyncLoadingThread2::TickAsyncLoadingFromGameThread(FA
 
 		if (Result != EAsyncPackageState::TimeOut)
 		{
-			if (!bDidSomething && bMayProcessPendingCDOs)
+			if (!bDidSomething && PendingCDOs.Num() > 0)
 			{
-				// Only process/create pending CDOs when absolutely required to unblock the current FlushAsyncLoading stack during initial load,
-				// i.e. when we are not making any progress and we are sure that there are no outstanding IO requests.
-				if (PendingIoRequestsCounter.Load() == 0)
-				{
-					bDidSomething = ProcessPendingCDOs();
-				}
+				bDidSomething = ProcessPendingCDOs(ThreadState);
 			}
 
 			// Flush deferred messages
