FAssetData - You can now pass in EResolveClass::Yes or EResolveClass::No, into GetClass and IsInstanceOf on FAssetData in order to ask it to load classes that it can't find.

diff --git a/Engine/Source/Runtime/CoreUObject/Private/AssetRegistry/AssetData.cpp b/Engine/Source/Runtime/CoreUObject/Private/AssetRegistry/AssetData.cpp
--- a/Engine/Source/Runtime/CoreUObject/Private/AssetRegistry/AssetData.cpp
+++ b/Engine/Source/Runtime/CoreUObject/Private/AssetRegistry/AssetData.cpp
@@ -390,7 +390,7 @@ bool FAssetData::IsTopLevelAsset(UObject* Object)
 	return Outer->IsA<UPackage>();
 }
 
-UClass* FAssetData::GetClass() const
+UClass* FAssetData::GetClass(EResolveClass ResolveClass) const
 {
 	if ( !IsValid() )
 	{
@@ -408,6 +408,13 @@ UClass* FAssetData::GetClass() const
 			FoundClass = FindObject<UClass>(nullptr, *NewPath);
 		}
 	}
+
+	// If they decided to load the class if unresolved, then lets load it.
+	if (!FoundClass && ResolveClass == EResolveClass::Yes)
+	{
+		FoundClass = LoadObject<UClass>(nullptr, *AssetClassPath.ToString());
+	}
+	
 	return FoundClass;
 }
 
diff --git a/Engine/Source/Runtime/CoreUObject/Public/AssetRegistry/AssetData.h b/Engine/Source/Runtime/CoreUObject/Public/AssetRegistry/AssetData.h
--- a/Engine/Source/Runtime/CoreUObject/Public/AssetRegistry/AssetData.h
+++ b/Engine/Source/Runtime/CoreUObject/Public/AssetRegistry/AssetData.h
@@ -132,6 +132,16 @@ namespace UE::AssetRegistry::Private
 	COREUOBJECT_API void ConcatenateOuterPathAndObjectName(FStringBuilderBase& Builder, FName OuterPath, FName ObjectName);
 }
 
+/**
+ * Should we try and resolve the class when you call GetClass on FAssetData?  If so then it might potentially load several
+ * other assets used by the class (Example: A UBlueprintGeneratedClass, might link in multiple other assets).
+ */
+enum class EResolveClass : uint8
+{
+	No,
+	Yes,
+};
+
 /** 
  * A struct to hold important information about an assets found by the Asset Registry
  * This struct is transient and should never be serialized
@@ -424,20 +434,26 @@ PRAGMA_ENABLE_DEPRECATION_WARNINGS
 		return Object && IsRedirectorClassName(Object->GetClass()->GetClassPathName());
 	}
 
-	/** Returns the class UClass if it is loaded. */
-	COREUOBJECT_API UClass* GetClass() const;
+	/**
+	 * Returns the class UClass if it is loaded.  If you choose to override ResolveClass, to EResolveClass::Yes
+	 * it will attempt to load the class if unloaded.
+	 */
+	COREUOBJECT_API UClass* GetClass(EResolveClass ResolveClass = EResolveClass::No) const;
 	
-	/** Returns whether the Asset's class is equal to or a child class of the given class. Returns false if the Asset's class can not be loaded. */
-	bool IsInstanceOf(const UClass* BaseClass) const
+	/**
+	 * Returns whether the Asset's class is equal to or a child class of the given class. Returns false if the Asset's class can not be loaded.
+	 * If you choose to override ResolveClass, to EResolveClass::Yes it will attempt to load the class if unloaded.
+	 */
+	bool IsInstanceOf(const UClass* BaseClass, EResolveClass ResolveClass = EResolveClass::No) const
 	{
-		UClass* ClassPointer = GetClass();
+		UClass* ClassPointer = GetClass(ResolveClass);
 		return ClassPointer && ClassPointer->IsChildOf(BaseClass);
 	}
 
 	template<typename BaseClass>
-	bool IsInstanceOf() const
+	bool IsInstanceOf(EResolveClass ResolveClass = EResolveClass::No) const
 	{
-		return IsInstanceOf(BaseClass::StaticClass());
+		return IsInstanceOf(BaseClass::StaticClass(), ResolveClass);
 	}
 
 	/** Append the object path to the given string builder. */
