Editor - Introducing a new helper utility class, FWarnIfAssetsLoadedInScope exists in CoreUObject so that it can be used in lots of places, but it will only work in Editor builds, so it is marked as such.  It allows you to throw a scope around some code to detect unexpected asset loads during that scope and then emit them as errors in the editor.  It wont warn anyone that doesn't have -WarnIfAssetsLoaded in their commandline.  It's a great little debugging utility when you're trying to prevent unexpected hitches from being introduced into the editor from unnessesary loads.

diff --git a/Engine/Source/Runtime/CoreUObject/Private/Misc/WarnIfAssetsLoadedInScope.cpp b/Engine/Source/Runtime/CoreUObject/Private/Misc/WarnIfAssetsLoadedInScope.cpp
new file mode 100644
--- /dev/null
+++ b/Engine/Source/Runtime/CoreUObject/Private/Misc/WarnIfAssetsLoadedInScope.cpp
@@ -0,0 +1,90 @@
+// Copyright Epic Games, Inc. All Rights Reserved.
+
+#include "Misc/WarnIfAssetsLoadedInScope.h"
+
+#include "Algo/AnyOf.h"
+#include "AssetRegistry/AssetData.h"
+#include "GenericPlatform/GenericPlatformStackWalk.h"
+#include "Logging/MessageLog.h"
+#include "HAL/PlatformStackWalk.h"
+#include "UObject/UObjectGlobals.h"
+#include "Misc/CommandLine.h"
+
+#define LOCTEXT_NAMESPACE "FWarnIfAssetsLoadedInScope"
+
+#if WITH_EDITOR
+
+FWarnIfAssetsLoadedInScope::FWarnIfAssetsLoadedInScope()
+{
+    AssetsLoadedHandle = FCoreUObjectDelegates::OnAssetLoaded.AddRaw(this, &FWarnIfAssetsLoadedInScope::OnAssetLoaded);
+}
+
+FWarnIfAssetsLoadedInScope::FWarnIfAssetsLoadedInScope(TConstArrayView<FAssetData> InSpecificAssets)
+    : SpecificAssets(InSpecificAssets)
+{
+    AssetsLoadedHandle = FCoreUObjectDelegates::OnAssetLoaded.AddRaw(this, &FWarnIfAssetsLoadedInScope::OnAssetLoaded);
+}
+
+FWarnIfAssetsLoadedInScope::~FWarnIfAssetsLoadedInScope()
+{
+    FCoreUObjectDelegates::OnAssetLoaded.Remove(AssetsLoadedHandle);
+
+    FMessageLog AssetLoadsLog("EditorErrors");
+
+    for (const FLoadEvent& LoadEvent : ObjectsLoaded)
+    {
+        TSharedRef<FTokenizedMessage> ErrorMessage = AssetLoadsLog.Error();
+        ErrorMessage->AddToken(FAssetNameToken::Create(LoadEvent.ObjectPath.ToString(),FText::FromString(LoadEvent.ObjectPath.ToString())));
+        ErrorMessage->AddToken(FTextToken::Create(LOCTEXT("DontLoadAssetsNow", "was loaded")));
+
+        for (const FProgramCounterSymbolInfo& CallInfo : LoadEvent.Callstack)
+        {
+        	ANSICHAR HumanReadableString[1024] = { '\0' };
+        	if (FPlatformStackWalk::SymbolInfoToHumanReadableString(CallInfo, HumanReadableString, sizeof(HumanReadableString)))
+        	{
+        		ErrorMessage->AddToken(FTextToken::Create(FText::FromString(FString(HumanReadableString) + TEXT("\n"))));	
+        	}
+        }
+    }
+
+	if (ObjectsLoaded.Num() > 0)
+	{
+        AssetLoadsLog.Notify(LOCTEXT("AssetsLoadedUnexpectedly", "Asset Unexpectedly Loaded!"));
+    }
+}
+
+void FWarnIfAssetsLoadedInScope::OnAssetLoaded(UObject* InObject)
+{
+	// If the handle isn't valid we can avoid all this, as there's no way anything was captured.
+	if (!AssetsLoadedHandle.IsValid())
+	{
+		return;
+	}
+	
+	if (SpecificAssets.Num() == 0 || Algo::AnyOf(SpecificAssets, [InObject](const FAssetData& Asset){ return Asset.IsUAsset(InObject); }))
+	{
+		FLoadEvent Event;
+		Event.ObjectPath = FSoftObjectPath(InObject);
+		Event.Callstack = FPlatformStackWalk::GetStack(4, 15);
+		ObjectsLoaded.Add(MoveTemp(Event));
+	}
+}
+
+bool FWarnIfAssetsLoadedInScope::IsEnabled() const
+{
+	auto WarnIfAssetsLoadedEnabled = []()
+	{
+		if (FParse::Param(FCommandLine::Get(), TEXT("WarnIfAssetsLoaded")))
+		{
+			return true;
+		}
+		return false;
+	};
+	static const bool bWarnIfAssetsLoadedEnabled = WarnIfAssetsLoadedEnabled();
+	
+	return bWarnIfAssetsLoadedEnabled;
+}
+
+#endif
+
+#undef LOCTEXT_NAMESPACE
\ No newline at end of file
diff --git a/Engine/Source/Runtime/CoreUObject/Public/Misc/WarnIfAssetsLoadedInScope.h b/Engine/Source/Runtime/CoreUObject/Public/Misc/WarnIfAssetsLoadedInScope.h
new file mode 100644
--- /dev/null
+++ b/Engine/Source/Runtime/CoreUObject/Public/Misc/WarnIfAssetsLoadedInScope.h
@@ -0,0 +1,39 @@
+// Copyright Epic Games, Inc. All Rights Reserved.
+
+#pragma once
+
+#include "CoreMinimal.h"
+#include "UObject/SoftObjectPath.h"
+
+class UObject;
+struct FProgramCounterSymbolInfo;
+struct FAssetData;
+
+#if WITH_EDITOR
+
+class COREUOBJECT_API FWarnIfAssetsLoadedInScope : public FNoncopyable
+{
+public:
+    FWarnIfAssetsLoadedInScope();
+	FWarnIfAssetsLoadedInScope(TConstArrayView<FAssetData> InSpecificAssets);
+
+    ~FWarnIfAssetsLoadedInScope();
+	
+private:
+	bool IsEnabled() const;
+    void OnAssetLoaded(UObject* InObject);
+
+private:
+    FDelegateHandle AssetsLoadedHandle;
+    TConstArrayView<FAssetData> SpecificAssets;
+
+    struct FLoadEvent
+    {
+        FSoftObjectPath ObjectPath;
+        TArray<FProgramCounterSymbolInfo> Callstack;
+    };
+
+    TArray<FLoadEvent> ObjectsLoaded;
+};
+
+#endif
\ No newline at end of file
