AsyncLoading2: Fix infinite wait for io issue part 2 (see 23028119 for part 1)
Make sure to always leave the waiting for io state if there are no more io requests to wait for.

#rnx

diff --git a/Engine/Source/Runtime/CoreUObject/Private/Serialization/AsyncLoading2.cpp b/Engine/Source/Runtime/CoreUObject/Private/Serialization/AsyncLoading2.cpp
--- a/Engine/Source/Runtime/CoreUObject/Private/Serialization/AsyncLoading2.cpp
+++ b/Engine/Source/Runtime/CoreUObject/Private/Serialization/AsyncLoading2.cpp
@@ -4428,6 +4428,10 @@ void FAsyncPackage2::StartLoading(FIoBatch& IoBatch)
 				int32 LocalPendingIoRequestsCounter = AsyncLoadingThread.PendingIoRequestsCounter.DecrementExchange() - 1;
 				TRACE_COUNTER_SET(AsyncLoadingPendingIoRequests, LocalPendingIoRequestsCounter);
 				GetPackageNode(EEventLoadNode2::Package_ProcessSummary).ReleaseBarrier();
+				if (LocalPendingIoRequestsCounter == 0)
+				{
+					AsyncLoadingThread.AltZenaphore.NotifyOne();
+				}
 			});
 	}
 #endif
@@ -4454,6 +4458,10 @@ void FAsyncPackage2::StartLoading(FIoBatch& IoBatch)
 			int32 LocalPendingIoRequestsCounter = AsyncLoadingThread.PendingIoRequestsCounter.DecrementExchange() - 1;
 			TRACE_COUNTER_SET(AsyncLoadingPendingIoRequests, LocalPendingIoRequestsCounter);
 			GetPackageNode(EEventLoadNode2::Package_ProcessSummary).ReleaseBarrier();
+			if (LocalPendingIoRequestsCounter == 0)
+			{
+				AsyncLoadingThread.AltZenaphore.NotifyOne();
+			}
 		});
 
 	if (!Data.ShaderMapHashes.IsEmpty())
@@ -4471,6 +4479,10 @@ void FAsyncPackage2::StartLoading(FIoBatch& IoBatch)
 					int32 LocalPendingIoRequestsCounter = AsyncLoadingThread.PendingIoRequestsCounter.DecrementExchange() - 1;
 					TRACE_COUNTER_SET(AsyncLoadingPendingIoRequests, LocalPendingIoRequestsCounter);
 					GetPackageNode(Package_ExportsSerialized).ReleaseBarrier();
+					if (LocalPendingIoRequestsCounter == 0)
+					{
+						AsyncLoadingThread.AltZenaphore.NotifyOne();
+					}
 				});
 		};
 		FCoreDelegates::PreloadPackageShaderMaps.ExecuteIfBound(Data.ShaderMapHashes, ReadShaderMapFunc);
