Fixed crash when loading maps using ZenLoader.

diff --git a/Engine/Source/Runtime/CoreUObject/Private/Serialization/AsyncLoading2.cpp b/Engine/Source/Runtime/CoreUObject/Private/Serialization/AsyncLoading2.cpp
--- a/Engine/Source/Runtime/CoreUObject/Private/Serialization/AsyncLoading2.cpp
+++ b/Engine/Source/Runtime/CoreUObject/Private/Serialization/AsyncLoading2.cpp
@@ -4364,6 +4364,10 @@ void FAsyncPackage2::CreateLinker(const FLinkerInstancingContext* InstancingCont
 		Linker = FLinkerLoad::CreateLinkerAsync(LoadContext, LinkerRoot, Desc.PackagePathToLoad, LinkerFlags, InstancingContext, TFunction<void()>([]() {}));
 #endif
 	}
+	else
+	{
+		Linker->LoadFlags |= LOAD_Async | LOAD_NoVerify | LOAD_SkipLoadImportedPackages;
+	}
 	check(Linker);
 	check(Linker->LinkerRoot == LinkerRoot);
 	check(!Linker->AsyncRoot);
