ZenLoader in editor: Improve handling of sync loads of packages with circular dependencies.
Add additional early outs from ProcessLinkerLoadPackageImports when a package has already been processed in a deeper circular dependency recursion stack.
Relax SyncLoadContextId constaints since when a package has been added to PackagesExcludedFromStateDependencyChecks due to a circular package dependency it may have been pulled into a higher sync load context id than the currently executing one.

#rnx

diff --git a/Engine/Source/Runtime/CoreUObject/Private/Serialization/AsyncLoading2.cpp b/Engine/Source/Runtime/CoreUObject/Private/Serialization/AsyncLoading2.cpp
--- a/Engine/Source/Runtime/CoreUObject/Private/Serialization/AsyncLoading2.cpp
+++ b/Engine/Source/Runtime/CoreUObject/Private/Serialization/AsyncLoading2.cpp
@@ -3355,7 +3355,7 @@ FAsyncPackage2* FAsyncLoadingThread2::FindOrInsertPackage(FAsyncPackageDesc2& De
 
 void FAsyncLoadingThread2::IncludePackageInSyncLoadContext(uint64 ContextId, FAsyncPackage2* Package)
 {
-	if (Package->SyncLoadContextId == ContextId)
+	if (Package->SyncLoadContextId >= ContextId)
 	{
 		return;
 	}
@@ -3368,7 +3368,7 @@ void FAsyncLoadingThread2::IncludePackageInSyncLoadContext(uint64 ContextId, FAs
 	Package->SyncLoadContextId = ContextId;
 	for (FAsyncPackage2* ImportedPackage : Package->Data.ImportedAsyncPackages)
 	{
-		if (ImportedPackage && ImportedPackage->SyncLoadContextId != ContextId)
+		if (ImportedPackage && ImportedPackage->SyncLoadContextId < ContextId)
 		{
 			IncludePackageInSyncLoadContext(ContextId, ImportedPackage);
 		}
@@ -3839,7 +3839,7 @@ bool FAsyncLoadEventQueue2::ExecuteSyncLoadEvents(FAsyncLoadingThreadState2& Thr
 
 	auto ShouldExecuteNode = [&SyncLoadContext](FEventLoadNode2* Node) -> bool
 	{
-		return Node->Package->SyncLoadContextId == SyncLoadContext.ContextId;
+		return Node->Package->SyncLoadContextId >= SyncLoadContext.ContextId;
 	};
 
 	bool bDidSomething = false;
@@ -4782,6 +4782,9 @@ bool FAsyncPackage2::ProcessLinkerLoadPackageImports(FAsyncLoadingThreadState2&
 				{
 					return false;
 				}
+			}
+			if (ImportedPackage->AsyncPackageLoadingState == EAsyncPackageLoadingState2::ProcessExportBundles)
+			{
 				if (!ImportedPackage->ProcessLinkerLoadPackageExports(ThreadState))
 				{
 					return false;
@@ -4791,6 +4794,11 @@ bool FAsyncPackage2::ProcessLinkerLoadPackageImports(FAsyncLoadingThreadState2&
 		++LinkerLoadState->ProcessingImportedPackageIndex;
 	}
 
+	if (AsyncPackageLoadingState > EAsyncPackageLoadingState2::ProcessExportBundles)
+	{
+		return true;
+	}
+
 	const int32 ImportCount = LinkerLoadState->Linker->ImportMap.Num();
 	while (LinkerLoadState->CreateImportIndex < ImportCount)
 	{
