EditorBulkData with virtualized payloads will no longer attempt to pull those payloads when creating undo states.

#rnx

- This bug was introduced with the fix for UE-169114.
- A virtualized payload can always be accessed via the virtualization system so we have no reason to serialize it during a transaction.
-- To prevent this we simply don't try to load the payload into memory if it is virtualized.
- While investigating this I noticed an additional minor issue that the << overload added for FSharedBuffer was treating a 0 length valid FSharedBuffer and a null FSharedBuffer as the same thing which is not exactly the case.
-- A 0 length buffer will return false for FSharedBuffer::IsNull and will allocate an owner structure internally. A null buffer will not allocate an owner and will return true for FSharedBuffer::IsNull
-- This is a minor distinction which was unlikely to be causing problems but should be fixed for technical accuracy.
-- We now write INDEX_NONE for the buffer length if the buffer is null and when reading we simply reset the buffer if we see that the length is INDEX_NONE rather than trying to create a valid buffer.

diff --git a/Engine/Source/Runtime/CoreUObject/Private/Serialization/EditorBulkData.cpp b/Engine/Source/Runtime/CoreUObject/Private/Serialization/EditorBulkData.cpp
--- a/Engine/Source/Runtime/CoreUObject/Private/Serialization/EditorBulkData.cpp
+++ b/Engine/Source/Runtime/CoreUObject/Private/Serialization/EditorBulkData.cpp
@@ -299,24 +299,43 @@ void UpdateArchiveData(FArchive& Ar, int64 DataPosition, DataType& Data)
 
 FArchive& operator<<(FArchive& Ar, FSharedBuffer& Buffer)
 {
+	// Note that there is a difference between a null FSharedBuffer and a zero length
+	// shared buffer and they are not interchangeable! If we have a null buffer we 
+	// write an invalid value for the buffer length so that we know to set the buffer
+	// to null when loaded and not create a valid zero length buffer instead.
+
 	if (Ar.IsLoading())
 	{
 		int64 BufferLength;
 		Ar << BufferLength;
 
-		FUniqueBuffer MutableBuffer = FUniqueBuffer::Alloc(BufferLength);
-
-		Ar.Serialize(MutableBuffer.GetData(), BufferLength);
+		if (BufferLength >= 0)
+		{
+			FUniqueBuffer MutableBuffer = FUniqueBuffer::Alloc(BufferLength);
+			Ar.Serialize(MutableBuffer.GetData(), BufferLength);
 
-		Buffer = MutableBuffer.MoveToShared();
+			Buffer = MutableBuffer.MoveToShared();
+		}
+		else
+		{
+			Buffer.Reset();
+		}
 	}
 	else if (Ar.IsSaving())
 	{
-		int64 BufferLength = (int64)Buffer.GetSize();
-		Ar << BufferLength;
+		if (!Buffer.IsNull())
+		{
+			int64 BufferLength = (int64)Buffer.GetSize();
+			Ar << BufferLength;
 
-		// Need to remove const due to FArchive API
-		Ar.Serialize(const_cast<void*>(Buffer.GetData()), BufferLength);
+			// Need to remove const due to FArchive API
+			Ar.Serialize(const_cast<void*>(Buffer.GetData()), BufferLength);
+		}
+		else
+		{
+			int64 InvalidLength = INDEX_NONE;
+			Ar << InvalidLength;
+		}
 	}
 
 	return Ar;
@@ -826,11 +845,11 @@ void FEditorBulkData::Serialize(FArchive& Ar, UObject* Owner, bool bAllowRegiste
 
 			if (Ar.IsSaving())
 			{
-				if (Payload.IsNull())
+				if (Payload.IsNull() && !IsDataVirtualized())
 				{
 					// We need to serialize in FSharedBuffer form, or otherwise we'd need to support
 					// multiple code paths here. Technically a bit wasteful but for general use it 
-					// shouldn't be noticable. This will make it easier to do the real perf wins
+					// shouldn't be noticeable. This will make it easier to do the real perf wins
 					// in the future.
 					Payload = GetDataInternal().Decompress();
 				}
