Fix regression from CL 21826272 where GetAssetPathRedirection was no longer redirecting based on AssetPath causing some failed redirections of SoftObjectPaths when saving new maps

diff --git a/Engine/Source/Runtime/CoreUObject/Private/Misc/RedirectCollector.cpp b/Engine/Source/Runtime/CoreUObject/Private/Misc/RedirectCollector.cpp
--- a/Engine/Source/Runtime/CoreUObject/Private/Misc/RedirectCollector.cpp
+++ b/Engine/Source/Runtime/CoreUObject/Private/Misc/RedirectCollector.cpp
@@ -296,22 +296,27 @@ FSoftObjectPath FRedirectCollector::GetAssetPathRedirection(const FSoftObjectPat
 
 	while (!CurrentPath.IsNull())
 	{
-		SeenPaths.Add(CurrentPath);
-		FSoftObjectPath NewPath = ObjectPathRedirectionMap.FindRef(CurrentPath);
-
-		if (!NewPath.IsNull())
+		if (!ensureMsgf(!SeenPaths.Contains(CurrentPath), TEXT("Found circular redirect from %s to %s! Returning None instead"), *OriginalPath.ToString(), *CurrentPath.ToString()))
 		{
-			if (!ensureMsgf(!SeenPaths.Contains(NewPath), TEXT("Found circular redirect from %s to %s! Returning None instead"), *CurrentPath.ToString(), *NewPath.ToString()))
+			UE_LOG(LogRedirectors, Error, TEXT("Logging redirection chain: "));
+			for (const FSoftObjectPath& Entry : SeenPaths)
 			{
-				return FSoftObjectPath();
+				UE_LOG(LogRedirectors, Error, TEXT(" %s"), *Entry.ToString());
 			}
+			return FSoftObjectPath();
+		}
+		SeenPaths.Add(CurrentPath);
 
-			// Continue trying to follow chain
-			CurrentPath = NewPath;
+		if (FSoftObjectPath* NewPath = ObjectPathRedirectionMap.Find(CurrentPath))
+		{
+			CurrentPath = *NewPath;
+		}
+		else if ((NewPath = ObjectPathRedirectionMap.Find(CurrentPath.GetWithoutSubPath())))
+		{
+			CurrentPath = FSoftObjectPath(NewPath->GetAssetPath(), CurrentPath.GetSubPathString());
 		}
 		else
 		{
-			// No more redirections
 			break;
 		}
 	}
