Add the option to note if the editor bulkdata is cooked or uncooked when logging about VA pull failures

#rnx

- Recently we have had several incidents where there has been confusion if a VA payload pull that failed was initiated by a uncooked editor bulkdata or a cooked one. Adding the option to postfix the log message with this info for selected projects will make debugging this sort of thing easier in the future.
- Enabled by changing '[EditorBulkData]LogCookedStatus=' to true in the Editor ini file.
-- As this is only useful for projects that allow cooked content in the editor, if not then it is just noise so the post fix is disabled by default.
- Note that the check in FEditorBulkData::PullData should also be removed but that has been added as its own work item in JIRA.

diff --git a/Engine/Config/BaseEditor.ini b/Engine/Config/BaseEditor.ini
--- a/Engine/Config/BaseEditor.ini
+++ b/Engine/Config/BaseEditor.ini
@@ -616,4 +616,7 @@ bCanModifyPluginsFromBrowser=true
 
 [/Script/GameFeaturesEditor.GameFeaturesEditorSettings]
 +PluginTemplates=(Path=(Path="../../Plugins/Experimental/GameFeatures/Templates/GameFeaturePluginContentOnly"),Label=NSLOCTEXT("[/Script/GameFeaturesEditor]", "E17DB7E94A2C0F5968A52E89959864C0", "Game Feature (Content Only)"),Description=NSLOCTEXT("[/Script/GameFeaturesEditor]", "CB5136B64AFAFBA5EDEA078FCAC7BCC9", "Create a new Game Feature Plugin."),DefaultGameFeatureDataClass=None,DefaultGameFeatureDataName="")
-+PluginTemplates=(Path=(Path="../../Plugins/Experimental/GameFeatures/Templates/GameFeaturePluginWithCode"),Label=NSLOCTEXT("[/Script/GameFeaturesEditor]", "90D6FD2A455176CEC42DD79A23AEC3B0", "Game Feature (with C++)"),Description=NSLOCTEXT("[/Script/GameFeaturesEditor]", "2B00A0E7415C8CEAA6DA609A141237F0", "Create a new Game Feature Plugin with a minimal amount of code."),DefaultGameFeatureDataClass=None,DefaultGameFeatureDataName="")
\ No newline at end of file
++PluginTemplates=(Path=(Path="../../Plugins/Experimental/GameFeatures/Templates/GameFeaturePluginWithCode"),Label=NSLOCTEXT("[/Script/GameFeaturesEditor]", "90D6FD2A455176CEC42DD79A23AEC3B0", "Game Feature (with C++)"),Description=NSLOCTEXT("[/Script/GameFeaturesEditor]", "2B00A0E7415C8CEAA6DA609A141237F0", "Create a new Game Feature Plugin with a minimal amount of code."),DefaultGameFeatureDataClass=None,DefaultGameFeatureDataName="")
+
+[EditorBulkData]
+LogCookedStatus=false
diff --git a/Engine/Source/Runtime/CoreUObject/Private/Serialization/EditorBulkData.cpp b/Engine/Source/Runtime/CoreUObject/Private/Serialization/EditorBulkData.cpp
--- a/Engine/Source/Runtime/CoreUObject/Private/Serialization/EditorBulkData.cpp
+++ b/Engine/Source/Runtime/CoreUObject/Private/Serialization/EditorBulkData.cpp
@@ -92,6 +92,25 @@ static bool ShouldSaveToPackageSidecar()
 	return ConfigSetting.bEnabled;
 }
 
+/** 
+ * Returns true if the project wants to log if an editor bulkdata is cooked or uncooked annd returns
+ * false if the project does not care. The engine ini file defaults to false 
+ */
+static bool ShouldLogCookedStatus()
+{
+	static const struct FShouldLogCookedStatus
+	{
+		bool bEnabled = false;
+
+		FShouldLogCookedStatus()
+		{
+			GConfig->GetBool(TEXT("EditorBulkData"), TEXT("LogCookedStatus"), bEnabled, GEditorIni);
+		}
+	} ConfigSetting;
+
+	return ConfigSetting.bEnabled;
+}
+
 #if UE_ENABLE_VIRTUALIZATION_TOGGLE
 bool ShouldAllowVirtualizationOptOut()
 {
@@ -1547,6 +1566,20 @@ FCompressedBuffer FEditorBulkData::PullData() const
 {
 	TRACE_CPUPROFILER_EVENT_SCOPE(FEditorBulkData::PullData);
 
+	// Utility lambda used to generate a postfix for VA failure messages based on if the project
+	// wants to know the cooked status of the editor bulkdata.
+	auto PostFixStatus = [](EFlags InFlags)
+		{
+			if (ShouldLogCookedStatus())
+			{
+				return EnumHasAnyFlags(InFlags, EFlags::IsCooked) ? TEXT(" [Cooked]") : TEXT(" [Uncooked]");
+			}
+			else
+			{
+				return TEXT("");
+			}
+		};
+
 	UE::Virtualization::IVirtualizationSystem& System = UE::Virtualization::IVirtualizationSystem::Get();
 	if (System.IsEnabled())
 	{
@@ -1557,13 +1590,18 @@ FCompressedBuffer FEditorBulkData::PullData() const
 			PayloadSize,
 			PulledPayload.GetRawSize());
 
-		UE_CLOG(PulledPayload.IsNull(), LogSerialization, Error, TEXT("Failed to pull payload '%s'"), *LexToString(PayloadContentId));
+		UE_CLOG(PulledPayload.IsNull(), LogSerialization, Error, TEXT("Failed to pull payload '%s'%s"),
+			*LexToString(PayloadContentId),
+			PostFixStatus(Flags));
 
 		return PulledPayload;
 	}
 	else
 	{
-		UE_LOG(LogSerialization, Error, TEXT("Cannot pull payload '%s' as the virtualization system is disabled"), *LexToString(PayloadContentId));
+		UE_LOG(LogSerialization, Error, TEXT("Failed to pull payload '%s' as the virtualization system is disabled%s"),
+			*LexToString(PayloadContentId), 
+			PostFixStatus(Flags));
+
 		return FCompressedBuffer();
 	}	
 }
