Fix log spam for renaming leaked package
#tests client

diff --git a/Engine/Source/Runtime/CoreUObject/Private/Serialization/AsyncLoading2.cpp b/Engine/Source/Runtime/CoreUObject/Private/Serialization/AsyncLoading2.cpp
--- a/Engine/Source/Runtime/CoreUObject/Private/Serialization/AsyncLoading2.cpp
+++ b/Engine/Source/Runtime/CoreUObject/Private/Serialization/AsyncLoading2.cpp
@@ -7130,7 +7130,7 @@ void FAsyncLoadingThread2::OnLeakedPackageRename(UPackage* Package)
 	}
 
 	// Code such as LoadMap or LevelStreaming is about to rename a loaded package which was detected as leaking so that we can load another copy of it.
-	UE_LOG(LogStreaming, Display, TEXT("Renaming leaked package %s (0x%llX)"), *Package->GetName(), Package->GetPackageId().ValueForDebugging());
+	UE_LOG(LogStreaming, Verbose, TEXT("Renaming leaked package %s (0x%llX)"), *Package->GetName(), Package->GetPackageId().ValueForDebugging());
 
 	if (!FGCCSyncObject::Get().IsGCLocked())
 	{
