#UE Delete FArchiveObjectCrc32NonEditorProperties, and use a plain old FArchiveObjectCrc32 with ArIsFilterEditorOnly set to true to test for cook non-determinism. The logic in FArchiveObjectCrc32NonEditorProperties did not properly ignore FNames serialized in editor-only properties

diff --git a/Engine/Source/Runtime/CoreUObject/Private/UObject/SavePackage/SavePackageUtilities.cpp b/Engine/Source/Runtime/CoreUObject/Private/UObject/SavePackage/SavePackageUtilities.cpp
--- a/Engine/Source/Runtime/CoreUObject/Private/UObject/SavePackage/SavePackageUtilities.cpp
+++ b/Engine/Source/Runtime/CoreUObject/Private/UObject/SavePackage/SavePackageUtilities.cpp
@@ -159,20 +159,6 @@ void FSavePackageStats::MergeStats(const TMap<FName, FArchiveDiffStats>& ToMerge
 
 #endif
 
-#if WITH_EDITORONLY_DATA
-
-void FArchiveObjectCrc32NonEditorProperties::Serialize(void* Data, int64 Length)
-{
-	int32 NewEditorOnlyProp = EditorOnlyProp + this->IsEditorOnlyPropertyOnTheStack();
-	TGuardValue<int32> Guard(EditorOnlyProp, NewEditorOnlyProp);
-	if (NewEditorOnlyProp == 0)
-	{
-		Super::Serialize(Data, Length);
-	}
-}
-
-#endif
-
 static FThreadSafeCounter OutstandingAsyncWrites;
 
 
diff --git a/Engine/Source/Runtime/CoreUObject/Private/UObject/SavePackage/SavePackageUtilities.h b/Engine/Source/Runtime/CoreUObject/Private/UObject/SavePackage/SavePackageUtilities.h
--- a/Engine/Source/Runtime/CoreUObject/Private/UObject/SavePackage/SavePackageUtilities.h
+++ b/Engine/Source/Runtime/CoreUObject/Private/UObject/SavePackage/SavePackageUtilities.h
@@ -343,39 +343,6 @@ private:
 	friend FEDLCookChecker;
 };
 
-#if WITH_EDITORONLY_DATA
-
-/**
- * Archive to calculate a checksum on an object's serialized data stream, but only of its non-editor properties.
- */
-class FArchiveObjectCrc32NonEditorProperties : public FArchiveObjectCrc32
-{
-	using Super = FArchiveObjectCrc32;
-
-public:
-	FArchiveObjectCrc32NonEditorProperties()
-		: EditorOnlyProp(0)
-	{
-	}
-
-	virtual FString GetArchiveName() const
-	{
-		return TEXT("FArchiveObjectCrc32NonEditorProperties");
-	}
-
-	virtual void Serialize(void* Data, int64 Length);
-private:
-	int32 EditorOnlyProp;
-};
-
-#else
-
-class COREUOBJECT_API FArchiveObjectCrc32NonEditorProperties : public FArchiveObjectCrc32
-{
-};
-
-#endif
-
 // Utility functions used by both UPackage::Save and/or UPackage::Save2
 namespace SavePackageUtilities
 {
diff --git a/Engine/Source/Runtime/CoreUObject/Private/UObject/SavePackage2.cpp b/Engine/Source/Runtime/CoreUObject/Private/UObject/SavePackage2.cpp
--- a/Engine/Source/Runtime/CoreUObject/Private/UObject/SavePackage2.cpp
+++ b/Engine/Source/Runtime/CoreUObject/Private/UObject/SavePackage2.cpp
@@ -273,7 +273,8 @@ ESavePackageResult RoutePresave(FSaveContext& SaveContext)
 		{
 			if (SaveContext.IsCooking() && Object->HasAnyFlags(RF_ClassDefaultObject | RF_ArchetypeObject))
 			{
-				FArchiveObjectCrc32NonEditorProperties CrcArchive;
+				FArchiveObjectCrc32 CrcArchive;
+				CrcArchive.ArIsFilterEditorOnly = true;
 				int32 Before = CrcArchive.Crc32(Object);
 				UE::SavePackageUtilities::CallPreSave(Object, SaveContext.GetObjectSaveContext());
 				int32 After = CrcArchive.Crc32(Object);
