Bulk Data - use correct package ID

#rnx

diff --git a/Engine/Source/Runtime/CoreUObject/Private/Serialization/BulkData.cpp b/Engine/Source/Runtime/CoreUObject/Private/Serialization/BulkData.cpp
--- a/Engine/Source/Runtime/CoreUObject/Private/Serialization/BulkData.cpp
+++ b/Engine/Source/Runtime/CoreUObject/Private/Serialization/BulkData.cpp
@@ -1228,6 +1228,7 @@ void FBulkData::Serialize(FArchive& Ar, UObject* Owner, bool bAttemptFileMapping
 					check(IsInSeparateFile());
 					check(Package != nullptr);
 					BulkMeta.SetFlags(EBulkDataFlags(BulkMeta.GetFlags() | BULKDATA_UsesIoDispatcher));
+					BulkChunkId = CreateBulkDataIoChunkId(BulkMeta, Owner);
 				}
 			}
 			else
@@ -1238,6 +1239,11 @@ void FBulkData::Serialize(FArchive& Ar, UObject* Owner, bool bAttemptFileMapping
 					{
 						LinkerLoad = FLinkerLoad::FindExistingLinkerForPackage(Package);
 					}
+
+					if (LinkerLoad != nullptr)
+					{
+						BulkChunkId = CreateBulkDataIoChunkId(BulkMeta, Owner);
+					}
 #if WITH_EDITOR
 					Linker = LinkerLoad;
 					if (bCanLazyLoad)
@@ -1250,8 +1256,6 @@ void FBulkData::Serialize(FArchive& Ar, UObject* Owner, bool bAttemptFileMapping
 				}
 			}
 
-			BulkChunkId = CreateBulkDataIoChunkId(BulkMeta, Owner);
-
 			if (IsInlined())
 			{
 				checkf(bAttemptFileMapping == false || bIsUsingIoDispatcher == false, TEXT("Trying to memory map inline bulk data '%s' which is not supported when loading from I/O store"), *Owner->GetPackage()->GetFName().ToString());
diff --git a/Engine/Source/Runtime/CoreUObject/Private/Serialization/BulkDataStreaming.cpp b/Engine/Source/Runtime/CoreUObject/Private/Serialization/BulkDataStreaming.cpp
--- a/Engine/Source/Runtime/CoreUObject/Private/Serialization/BulkDataStreaming.cpp
+++ b/Engine/Source/Runtime/CoreUObject/Private/Serialization/BulkDataStreaming.cpp
@@ -191,7 +191,7 @@ FIoChunkId CreateBulkDataIoChunkId(const FBulkMetaData& BulkMeta, UObject* Owner
 
 		const uint16 ChunkIndex = EnumHasAnyFlags(BulkMeta.GetMetaFlags(), FBulkMetaData::EMetaFlags::OptionalPackage) ? 1 : 0;
 		
-		return CreateIoChunkId(Pkg->GetPackageId().Value(), ChunkIndex, ChunkType);
+		return CreateIoChunkId(Pkg->GetPackageIdToLoad().Value(), ChunkIndex, ChunkType);
 	}
 	else
 	{
