ZenLoader: Fixed (rare) crash introduced in 23059228, a package can be deleted as soon as the last barrier is released
#localization none
#tests client

diff --git a/Engine/Source/Runtime/CoreUObject/Private/Serialization/AsyncLoading2.cpp b/Engine/Source/Runtime/CoreUObject/Private/Serialization/AsyncLoading2.cpp
--- a/Engine/Source/Runtime/CoreUObject/Private/Serialization/AsyncLoading2.cpp
+++ b/Engine/Source/Runtime/CoreUObject/Private/Serialization/AsyncLoading2.cpp
@@ -4431,10 +4431,11 @@ void FAsyncPackage2::StartLoading(FIoBatch& IoBatch)
 				}
 				int32 LocalPendingIoRequestsCounter = AsyncLoadingThread.PendingIoRequestsCounter.DecrementExchange() - 1;
 				TRACE_COUNTER_SET(AsyncLoadingPendingIoRequests, LocalPendingIoRequestsCounter);
+				FAsyncLoadingThread2& LocalAsyncLoadingThread = AsyncLoadingThread;
 				GetPackageNode(EEventLoadNode2::Package_ProcessSummary).ReleaseBarrier();
 				if (LocalPendingIoRequestsCounter == 0)
 				{
-					AsyncLoadingThread.AltZenaphore.NotifyOne();
+					LocalAsyncLoadingThread.AltZenaphore.NotifyOne();
 				}
 			});
 	}
@@ -4461,10 +4462,11 @@ void FAsyncPackage2::StartLoading(FIoBatch& IoBatch)
 			}
 			int32 LocalPendingIoRequestsCounter = AsyncLoadingThread.PendingIoRequestsCounter.DecrementExchange() - 1;
 			TRACE_COUNTER_SET(AsyncLoadingPendingIoRequests, LocalPendingIoRequestsCounter);
+			FAsyncLoadingThread2& LocalAsyncLoadingThread = AsyncLoadingThread;
 			GetPackageNode(EEventLoadNode2::Package_ProcessSummary).ReleaseBarrier();
 			if (LocalPendingIoRequestsCounter == 0)
 			{
-				AsyncLoadingThread.AltZenaphore.NotifyOne();
+				LocalAsyncLoadingThread.AltZenaphore.NotifyOne();
 			}
 		});
 
@@ -4482,10 +4484,11 @@ void FAsyncPackage2::StartLoading(FIoBatch& IoBatch)
 					GraphEvent->DispatchSubsequents();
 					int32 LocalPendingIoRequestsCounter = AsyncLoadingThread.PendingIoRequestsCounter.DecrementExchange() - 1;
 					TRACE_COUNTER_SET(AsyncLoadingPendingIoRequests, LocalPendingIoRequestsCounter);
+					FAsyncLoadingThread2& LocalAsyncLoadingThread = AsyncLoadingThread;
 					GetPackageNode(Package_ExportsSerialized).ReleaseBarrier();
 					if (LocalPendingIoRequestsCounter == 0)
 					{
-						AsyncLoadingThread.AltZenaphore.NotifyOne();
+						LocalAsyncLoadingThread.AltZenaphore.NotifyOne();
 					}
 				});
 		};
