Fixed struct alignment not being used by FStructOnScope

diff --git a/Engine/Source/Runtime/CoreUObject/Public/UObject/StructOnScope.h b/Engine/Source/Runtime/CoreUObject/Public/UObject/StructOnScope.h
--- a/Engine/Source/Runtime/CoreUObject/Public/UObject/StructOnScope.h
+++ b/Engine/Source/Runtime/CoreUObject/Public/UObject/StructOnScope.h
@@ -21,7 +21,7 @@ protected:
 	{
 		if (const UStruct* ScriptStructPtr = ScriptStruct.Get())
 		{
-			SampleStructMemory = (uint8*)FMemory::Malloc(ScriptStructPtr->GetStructureSize() ? ScriptStructPtr->GetStructureSize() : 1);
+			SampleStructMemory = (uint8*)FMemory::Malloc(ScriptStructPtr->GetStructureSize() ? ScriptStructPtr->GetStructureSize() : 1, ScriptStructPtr->GetMinAlignment());
 			ScriptStructPtr->InitializeStruct(SampleStructMemory);
 			OwnsMemory = true;
 		}
