removing `PRAGMA_DISABLE_OPTIMIZATION`

diff --git a/Engine/Source/Runtime/CoreUObject/Private/Serialization/AsyncLoading2.cpp b/Engine/Source/Runtime/CoreUObject/Private/Serialization/AsyncLoading2.cpp
--- a/Engine/Source/Runtime/CoreUObject/Private/Serialization/AsyncLoading2.cpp
+++ b/Engine/Source/Runtime/CoreUObject/Private/Serialization/AsyncLoading2.cpp
@@ -72,10 +72,6 @@
 
 #include <atomic>
 
-#if UE_BUILD_DEVELOPMENT || UE_BUILD_DEBUG
-PRAGMA_DISABLE_OPTIMIZATION
-#endif
-
 CSV_DECLARE_CATEGORY_MODULE_EXTERN(CORE_API, FileIO);
 CSV_DEFINE_STAT(FileIO, FrameCompletedExportBundleLoadsKB);
 
@@ -7929,8 +7925,4 @@ EAsyncPackageState::Type FAsyncLoadingThread2::ProcessLoadingUntilCompleteFromGa
 IAsyncPackageLoader* MakeAsyncPackageLoader2(FIoDispatcher& InIoDispatcher, IAsyncPackageLoader* InUncookedPackageLoader)
 {
 	return new FAsyncLoadingThread2(InIoDispatcher, InUncookedPackageLoader);
-}
-
-#if UE_BUILD_DEVELOPMENT || UE_BUILD_DEBUG
-PRAGMA_ENABLE_OPTIMIZATION
-#endif
+}
\ No newline at end of file
