Enable the garbage collector to find UObjects that might be referenced by FDynamicallyTypedValueType.

diff --git a/Engine/Source/Runtime/CoreUObject/Private/UObject/DynamicallyTypedValue.cpp b/Engine/Source/Runtime/CoreUObject/Private/UObject/DynamicallyTypedValue.cpp
--- a/Engine/Source/Runtime/CoreUObject/Private/UObject/DynamicallyTypedValue.cpp
+++ b/Engine/Source/Runtime/CoreUObject/Private/UObject/DynamicallyTypedValue.cpp
@@ -8,7 +8,7 @@ UE::FDynamicallyTypedValueType& UE::FDynamicallyTypedValue::NullType()
 	{
 		FNullType(): FDynamicallyTypedValueType(0, 0, EContainsReferences::DoesNot) {}
 
-		virtual void MarkReachable() override {}
+		virtual void MarkReachable(FReferenceCollector& Collector) override {}
 		virtual void MarkValueReachable(void* Data, FReferenceCollector& Collector) const override {}
 
 		virtual void InitializeValue(void* Data) const override {}
diff --git a/Engine/Source/Runtime/CoreUObject/Public/UObject/DynamicallyTypedValue.h b/Engine/Source/Runtime/CoreUObject/Public/UObject/DynamicallyTypedValue.h
--- a/Engine/Source/Runtime/CoreUObject/Public/UObject/DynamicallyTypedValue.h
+++ b/Engine/Source/Runtime/CoreUObject/Public/UObject/DynamicallyTypedValue.h
@@ -30,7 +30,7 @@ namespace UE
 		}
 
 		// Marks the type itself as reachable.
-		virtual void MarkReachable() = 0;
+		virtual void MarkReachable(FReferenceCollector& Collector) = 0;
 
 		// Marks a value of the type as reachable.
 		virtual void MarkValueReachable(void* Data, FReferenceCollector& Collector) const = 0;
diff --git a/Engine/Source/Runtime/CoreUObject/Public/UObject/FastReferenceCollector.h b/Engine/Source/Runtime/CoreUObject/Public/UObject/FastReferenceCollector.h
--- a/Engine/Source/Runtime/CoreUObject/Public/UObject/FastReferenceCollector.h
+++ b/Engine/Source/Runtime/CoreUObject/Public/UObject/FastReferenceCollector.h
@@ -1288,7 +1288,7 @@ private:
 					{
 						TokenReturnCount = ReferenceInfo.ReturnCount;
 						UE::FDynamicallyTypedValue* DynamicallyTypedValue = (UE::FDynamicallyTypedValue*)(StackEntryData + ReferenceInfo.Offset);
-						DynamicallyTypedValue->GetType().MarkReachable();
+						DynamicallyTypedValue->GetType().MarkReachable(ReferenceCollector);
 						if (DynamicallyTypedValue->GetType().GetContainsReferences() != UE::FDynamicallyTypedValueType::EContainsReferences::DoesNot)
 						{
 							DynamicallyTypedValue->GetType().MarkValueReachable(DynamicallyTypedValue->GetDataPointer(), ReferenceCollector);
