ZenLoader:
* Always pass along a thread context to ReleaseBarrier
* Update the active sync load context and associated package as early as possible also for the game thread

diff --git a/Engine/Source/Runtime/CoreUObject/Private/Serialization/AsyncLoading2.cpp b/Engine/Source/Runtime/CoreUObject/Private/Serialization/AsyncLoading2.cpp
--- a/Engine/Source/Runtime/CoreUObject/Private/Serialization/AsyncLoading2.cpp
+++ b/Engine/Source/Runtime/CoreUObject/Private/Serialization/AsyncLoading2.cpp
@@ -1750,7 +1750,7 @@ public:
 	void DependsOn(FEventLoadNode2* Other);
 	void AddBarrier();
 	void AddBarrier(int32 Count);
-	void ReleaseBarrier(FAsyncLoadingThreadState2* ThreadState = nullptr);
+	void ReleaseBarrier(FAsyncLoadingThreadState2* ThreadState);
 	EEventLoadNodeExecutionResult Execute(FAsyncLoadingThreadState2& ThreadState);
 
 	FAsyncPackage2* GetPackage() const
@@ -1769,7 +1769,7 @@ private:
 	}
 
 	void ProcessDependencies(FAsyncLoadingThreadState2& ThreadState);
-	void Fire(FAsyncLoadingThreadState2* ThreadState = nullptr);
+	void Fire(FAsyncLoadingThreadState2* ThreadState);
 
 	union
 	{
@@ -2327,7 +2327,7 @@ struct FAsyncPackage2
 	EAsyncPackageState::Type CreateClusters(FAsyncLoadingThreadState2& ThreadState);
 
 	void ImportPackagesRecursive(FAsyncLoadingThreadState2& ThreadState, FIoBatch& IoBatch, FPackageStore& PackageStore);
-	void StartLoading(FIoBatch& IoBatch);
+	void StartLoading(FAsyncLoadingThreadState2& ThreadState, FIoBatch& IoBatch);
 
 #if ALT2_ENABLE_LINKERLOAD_SUPPORT
 	void InitializeLinkerLoadState(const FLinkerInstancingContext* InstancingContext);
@@ -2597,7 +2597,7 @@ private:
 	 * @return Complete if all dependencies are finished, TimeOut otherwise
 	 */
 	enum EExternalReadAction { ExternalReadAction_Poll, ExternalReadAction_Wait };
-	EAsyncPackageState::Type ProcessExternalReads(EExternalReadAction Action);
+	EAsyncPackageState::Type ProcessExternalReads(FAsyncLoadingThreadState2& ThreadState, EExternalReadAction Action);
 
 	/**
 	* Updates load percentage stat
@@ -2813,9 +2813,9 @@ public:
 		return AsyncPackageLookup.FindRef(PackageId);
 	}
 
-	void UpdatePackagePriority(FAsyncPackage2* Package, int32 NewPriority);
+	void UpdatePackagePriority(FAsyncLoadingThreadState2& ThreadState, FAsyncPackage2* Package, int32 NewPriority);
 
-	FAsyncPackage2* FindOrInsertPackage(FAsyncPackageDesc2& InDesc, bool& bInserted, TUniquePtr<FLoadPackageAsyncDelegate>&& PackageLoadedDelegate = TUniquePtr<FLoadPackageAsyncDelegate>());
+	FAsyncPackage2* FindOrInsertPackage(FAsyncLoadingThreadState2& ThreadState, FAsyncPackageDesc2& InDesc, bool& bInserted, TUniquePtr<FLoadPackageAsyncDelegate>&& PackageLoadedDelegate = TUniquePtr<FLoadPackageAsyncDelegate>());
 	void QueueMissingPackage(FAsyncPackageDesc2& PackageDesc, TUniquePtr<FLoadPackageAsyncDelegate>&& LoadPackageAsyncDelegate);
 
 	/**
@@ -3053,7 +3053,7 @@ private:
 
 			for (FEventLoadNode2* Node : Nodes)
 			{
-				Node->ReleaseBarrier();
+				Node->ReleaseBarrier(&ThreadState);
 			}
 
 			bDidSomething = true;
@@ -3072,7 +3072,7 @@ private:
 	*/
 	EAsyncPackageState::Type ProcessLoadedPackagesFromGameThread(FAsyncLoadingThreadState2& ThreadState, bool& bDidSomething, int32 FlushRequestID = INDEX_NONE);
 
-	void IncludePackageInSyncLoadContext(uint64 ContextId, FAsyncPackage2* Package);
+	void IncludePackageInSyncLoadContextRecursive(FAsyncLoadingThreadState2& ThreadState, uint64 ContextId, FAsyncPackage2* Package);
 	FAsyncLoadingSyncLoadContext* UpdateSyncLoadContext(FAsyncLoadingThreadState2& ThreadState);
 
 	bool CreateAsyncPackagesFromQueue(FAsyncLoadingThreadState2& ThreadState);
@@ -3299,7 +3299,7 @@ void FAsyncLoadingThread2::InitializeLoading()
 	UE_LOG(LogStreaming, Display, TEXT("AsyncLoading2 - Initialized"));
 }
 
-void FAsyncLoadingThread2::UpdatePackagePriority(FAsyncPackage2* Package, int32 NewPriority)
+void FAsyncLoadingThread2::UpdatePackagePriority(FAsyncLoadingThreadState2& ThreadState, FAsyncPackage2* Package, int32 NewPriority)
 {
 	TRACE_CPUPROFILER_EVENT_SCOPE(UpdatePackagePriority);
 	Package->Desc.Priority = NewPriority;
@@ -3312,7 +3312,7 @@ void FAsyncLoadingThread2::UpdatePackagePriority(FAsyncPackage2* Package, int32
 #endif
 }
 
-FAsyncPackage2* FAsyncLoadingThread2::FindOrInsertPackage(FAsyncPackageDesc2& Desc, bool& bInserted, TUniquePtr<FLoadPackageAsyncDelegate>&& PackageLoadedDelegate)
+FAsyncPackage2* FAsyncLoadingThread2::FindOrInsertPackage(FAsyncLoadingThreadState2& ThreadState, FAsyncPackageDesc2& Desc, bool& bInserted, TUniquePtr<FLoadPackageAsyncDelegate>&& PackageLoadedDelegate)
 {
 	TRACE_CPUPROFILER_EVENT_SCOPE(FindOrInsertPackage);
 	FAsyncPackage2* Package = nullptr;
@@ -3338,7 +3338,7 @@ FAsyncPackage2* FAsyncLoadingThread2::FindOrInsertPackage(FAsyncPackageDesc2& De
 			}
 			if (Desc.Priority > Package->Desc.Priority)
 			{
-				UpdatePackagePriority(Package, Desc.Priority);
+				UpdatePackagePriority(ThreadState, Package, Desc.Priority);
 			}
 		}
 		if (PackageLoadedDelegate.IsValid())
@@ -3349,7 +3349,7 @@ FAsyncPackage2* FAsyncLoadingThread2::FindOrInsertPackage(FAsyncPackageDesc2& De
 	return Package;
 }
 
-void FAsyncLoadingThread2::IncludePackageInSyncLoadContext(uint64 ContextId, FAsyncPackage2* Package)
+void FAsyncLoadingThread2::IncludePackageInSyncLoadContextRecursive(FAsyncLoadingThreadState2& ThreadState, uint64 ContextId, FAsyncPackage2* Package)
 {
 	if (Package->SyncLoadContextId >= ContextId)
 	{
@@ -3366,7 +3366,7 @@ void FAsyncLoadingThread2::IncludePackageInSyncLoadContext(uint64 ContextId, FAs
 	{
 		if (ImportedPackage && ImportedPackage->SyncLoadContextId < ContextId)
 		{
-			IncludePackageInSyncLoadContext(ContextId, ImportedPackage);
+			IncludePackageInSyncLoadContextRecursive(ThreadState, ContextId, ImportedPackage);
 		}
 	}
 }
@@ -3390,7 +3390,7 @@ bool FAsyncLoadingThread2::CreateAsyncPackagesFromQueue(FAsyncLoadingThreadState
 		if (PendingPackageStatus == EPackageStoreEntryStatus::Ok)
 		{
 			InitializeAsyncPackageFromPackageStore(ThreadState, IoBatch, PendingPackage, PackageEntry, TArrayView<FName>());
-			PendingPackage->StartLoading(IoBatch);
+			PendingPackage->StartLoading(ThreadState, IoBatch);
 			It.RemoveCurrent();
 		}
 		else if (PendingPackageStatus == EPackageStoreEntryStatus::Missing)
@@ -3403,7 +3403,7 @@ bool FAsyncLoadingThread2::CreateAsyncPackagesFromQueue(FAsyncLoadingThreadState
 			// Simulate StartLoading() getting back a failed IoRequest and let it go through all package states
 			PendingPackage->AsyncPackageLoadingState = EAsyncPackageLoadingState2::WaitingForIo;
 			PendingPackage->bLoadHasFailed = true;
-			PendingPackage->GetPackageNode(EEventLoadNode2::Package_ProcessSummary).ReleaseBarrier();
+			PendingPackage->GetPackageNode(EEventLoadNode2::Package_ProcessSummary).ReleaseBarrier(&ThreadState);
 			// Remove from PendingPackages
 			It.RemoveCurrent();
 		}
@@ -3506,7 +3506,7 @@ bool FAsyncLoadingThread2::CreateAsyncPackagesFromQueue(FAsyncLoadingThreadState
 			else
 			{
 				bool bInserted;
-				FAsyncPackage2* Package = FindOrInsertPackage(PackageDesc, bInserted, MoveTemp(Request.PackageLoadedDelegate));
+				FAsyncPackage2* Package = FindOrInsertPackage(ThreadState, PackageDesc, bInserted, MoveTemp(Request.PackageLoadedDelegate));
 				checkf(Package, TEXT("Failed to find or insert package %s"), *PackageDesc.UPackageName.ToString());
 
 				if (bInserted)
@@ -3536,7 +3536,7 @@ bool FAsyncLoadingThread2::CreateAsyncPackagesFromQueue(FAsyncLoadingThreadState
 						{
 							InitializeAsyncPackageFromPackageStore(ThreadState, IoBatch, Package, PackageEntry, TArrayView<FName>());
 						}
-						Package->StartLoading(IoBatch);
+						Package->StartLoading(ThreadState, IoBatch);
 					}
 				}
 				else
@@ -4225,7 +4225,7 @@ void FAsyncPackage2::ImportPackagesRecursiveInner(FAsyncLoadingThreadState2& Thr
 		else
 		{
 			FAsyncPackageDesc2 PackageDesc = FAsyncPackageDesc2::FromPackageImport(Desc, ImportedPackageUPackageName, ImportedPackageId, ImportedPackageIdToLoad, MoveTemp(ImportedPackagePath));
-			ImportedPackage = AsyncLoadingThread.FindOrInsertPackage(PackageDesc, bInserted);
+			ImportedPackage = AsyncLoadingThread.FindOrInsertPackage(ThreadState, PackageDesc, bInserted);
 		}
 
 		checkf(ImportedPackage, TEXT("Failed to find or insert imported package with id '0x%llX'"), ImportedPackageId.Value());
@@ -4278,7 +4278,7 @@ void FAsyncPackage2::ImportPackagesRecursiveInner(FAsyncLoadingThreadState2& Thr
 				{
 					AsyncLoadingThread.InitializeAsyncPackageFromPackageStore(ThreadState, IoBatch, ImportedPackage, ImportedPackageEntry, TArrayView<FName>());
 				}
-				ImportedPackage->StartLoading(IoBatch);
+				ImportedPackage->StartLoading(ThreadState, IoBatch);
 			}
 		}
 	}
@@ -4322,7 +4322,7 @@ void FAsyncPackage2::ImportPackagesRecursive(FAsyncLoadingThreadState2& ThreadSt
 		{
 			if (ImportedPackage)
 			{
-				AsyncLoadingThread.IncludePackageInSyncLoadContext(SyncLoadContextId, ImportedPackage);
+				AsyncLoadingThread.IncludePackageInSyncLoadContextRecursive(ThreadState, SyncLoadContextId, ImportedPackage);
 			}
 		}
 	}
@@ -4386,7 +4386,7 @@ void FAsyncPackage2::DetachLinker()
 }
 #endif
 
-void FAsyncPackage2::StartLoading(FIoBatch& IoBatch)
+void FAsyncPackage2::StartLoading(FAsyncLoadingThreadState2& ThreadState, FIoBatch& IoBatch)
 {
 	TRACE_CPUPROFILER_EVENT_SCOPE(StartLoading);
 
@@ -4401,7 +4401,7 @@ void FAsyncPackage2::StartLoading(FIoBatch& IoBatch)
 #if ALT2_ENABLE_NEW_ARCHIVE_FOR_LINKERLOAD
 		static_cast<FLinkerLoadArchive2*>(LinkerLoadState->Linker->GetLoader())->BeginRead(&GetPackageNode(EEventLoadNode2::Package_ProcessSummary));
 #else
-		GetPackageNode(EEventLoadNode2::Package_ProcessSummary).ReleaseBarrier();
+		GetPackageNode(EEventLoadNode2::Package_ProcessSummary).ReleaseBarrier(&ThreadState);
 #endif
 		return;
 	}
@@ -4428,7 +4428,7 @@ void FAsyncPackage2::StartLoading(FIoBatch& IoBatch)
 				int32 LocalPendingIoRequestsCounter = AsyncLoadingThread.PendingIoRequestsCounter.DecrementExchange() - 1;
 				TRACE_COUNTER_SET(AsyncLoadingPendingIoRequests, LocalPendingIoRequestsCounter);
 				FAsyncLoadingThread2& LocalAsyncLoadingThread = AsyncLoadingThread;
-				GetPackageNode(EEventLoadNode2::Package_ProcessSummary).ReleaseBarrier();
+				GetPackageNode(EEventLoadNode2::Package_ProcessSummary).ReleaseBarrier(nullptr);
 				if (LocalPendingIoRequestsCounter == 0)
 				{
 					LocalAsyncLoadingThread.AltZenaphore.NotifyOne();
@@ -4459,7 +4459,7 @@ void FAsyncPackage2::StartLoading(FIoBatch& IoBatch)
 			int32 LocalPendingIoRequestsCounter = AsyncLoadingThread.PendingIoRequestsCounter.DecrementExchange() - 1;
 			TRACE_COUNTER_SET(AsyncLoadingPendingIoRequests, LocalPendingIoRequestsCounter);
 			FAsyncLoadingThread2& LocalAsyncLoadingThread = AsyncLoadingThread;
-			GetPackageNode(EEventLoadNode2::Package_ProcessSummary).ReleaseBarrier();
+			GetPackageNode(EEventLoadNode2::Package_ProcessSummary).ReleaseBarrier(nullptr);
 			if (LocalPendingIoRequestsCounter == 0)
 			{
 				LocalAsyncLoadingThread.AltZenaphore.NotifyOne();
@@ -4481,7 +4481,7 @@ void FAsyncPackage2::StartLoading(FIoBatch& IoBatch)
 					int32 LocalPendingIoRequestsCounter = AsyncLoadingThread.PendingIoRequestsCounter.DecrementExchange() - 1;
 					TRACE_COUNTER_SET(AsyncLoadingPendingIoRequests, LocalPendingIoRequestsCounter);
 					FAsyncLoadingThread2& LocalAsyncLoadingThread = AsyncLoadingThread;
-					GetPackageNode(Package_ExportsSerialized).ReleaseBarrier();
+					GetPackageNode(Package_ExportsSerialized).ReleaseBarrier(nullptr);
 					if (LocalPendingIoRequestsCounter == 0)
 					{
 						LocalAsyncLoadingThread.AltZenaphore.NotifyOne();
@@ -4543,10 +4543,10 @@ void FAsyncPackage2::ConditionalProcessLinkerLoadPackageImportsAndExports(FAsync
 	TRACE_CPUPROFILER_EVENT_SCOPE(ConditionalProcessLinkerLoadPackageExports);
 	constexpr bool bExcludeZenPackages = true;
 	WaitForAllDependenciesToReachState(ThreadState, &FAsyncPackage2::AllDependenciesSetupState, EAsyncPackageLoadingState2::ProcessExportBundles, AsyncLoadingThread.ConditionalProcessLinkerLoadPackageExportsTick, bExcludeZenPackages,
-		[](FAsyncPackage2* Package)
+		[&ThreadState](FAsyncPackage2* Package)
 		{
 			Package->AsyncPackageLoadingState = EAsyncPackageLoadingState2::ProcessExportBundles;
-			Package->GetExportBundleNode(ExportBundle_Process, 0).ReleaseBarrier();
+			Package->GetExportBundleNode(ExportBundle_Process, 0).ReleaseBarrier(&ThreadState);
 		});
 }
 
@@ -4725,7 +4725,7 @@ EEventLoadNodeExecutionResult FAsyncPackage2::ProcessLinkerLoadPackageSummary(FA
 	}
 
 	AsyncPackageLoadingState = EAsyncPackageLoadingState2::SetupDependencies;
-	GetPackageNode(Package_SetupDependencies).ReleaseBarrier();
+	GetPackageNode(Package_SetupDependencies).ReleaseBarrier(&ThreadState);
 	return EEventLoadNodeExecutionResult::Complete;
 }
 
@@ -5055,7 +5055,7 @@ EEventLoadNodeExecutionResult FAsyncPackage2::Event_ProcessPackageSummary(FAsync
 	}
 
 	Package->AsyncPackageLoadingState = EAsyncPackageLoadingState2::SetupDependencies;
-	Package->GetPackageNode(Package_SetupDependencies).ReleaseBarrier();
+	Package->GetPackageNode(Package_SetupDependencies).ReleaseBarrier(&ThreadState);
 	return EEventLoadNodeExecutionResult::Complete;
 }
 
@@ -5097,7 +5097,7 @@ EEventLoadNodeExecutionResult FAsyncPackage2::Event_SetupDependencies(FAsyncLoad
 	Package->AsyncPackageLoadingState = EAsyncPackageLoadingState2::ProcessExportBundles;
 	for (int32 ExportBundleIndex = 0; ExportBundleIndex < Package->Data.ExportBundleCount; ++ExportBundleIndex)
 	{
-		Package->GetExportBundleNode(ExportBundle_Process, ExportBundleIndex).ReleaseBarrier();
+		Package->GetExportBundleNode(ExportBundle_Process, ExportBundleIndex).ReleaseBarrier(&ThreadState);
 	}
 	return EEventLoadNodeExecutionResult::Complete;
 }
@@ -5869,11 +5869,11 @@ void FAsyncPackage2::ConditionalBeginPostLoad(FAsyncLoadingThreadState2& ThreadS
 	TRACE_CPUPROFILER_EVENT_SCOPE(ConditionalBeginPostLoad);
 	constexpr bool bExcludeZenPackages = false;
 	WaitForAllDependenciesToReachState(ThreadState, &FAsyncPackage2::AllDependenciesSerializedState, EAsyncPackageLoadingState2::ExportsDone, AsyncLoadingThread.ConditionalBeginPostLoadTick, bExcludeZenPackages,
-		[](FAsyncPackage2* Package)
+		[&ThreadState](FAsyncPackage2* Package)
 		{
 			for (int32 ExportBundleIndex = 0; ExportBundleIndex < Package->Data.ExportBundleCount; ++ExportBundleIndex)
 			{
-				Package->GetExportBundleNode(EEventLoadNode2::ExportBundle_PostLoad, ExportBundleIndex).ReleaseBarrier();
+				Package->GetExportBundleNode(EEventLoadNode2::ExportBundle_PostLoad, ExportBundleIndex).ReleaseBarrier(&ThreadState);
 			}
 		});
 }
@@ -5908,7 +5908,7 @@ EEventLoadNodeExecutionResult FAsyncPackage2::Event_PostLoadExportBundle(FAsyncL
 	{
 		// No async postload for now
 		Package->AsyncPackageLoadingState = EAsyncPackageLoadingState2::DeferredPostLoad;
-		Package->GetExportBundleNode(ExportBundle_DeferredPostLoad, 0).ReleaseBarrier();
+		Package->GetExportBundleNode(ExportBundle_DeferredPostLoad, 0).ReleaseBarrier(&ThreadState);
 		return EEventLoadNodeExecutionResult::Complete;
 	}
 #endif
@@ -6041,7 +6041,7 @@ EEventLoadNodeExecutionResult FAsyncPackage2::Event_PostLoadExportBundle(FAsyncL
 		Package->AsyncPackageLoadingState = EAsyncPackageLoadingState2::DeferredPostLoad;
 		for (int32 ExportBundleIndex = 0; ExportBundleIndex < Package->Data.ExportBundleCount; ++ExportBundleIndex)
 		{
-			Package->GetExportBundleNode(ExportBundle_DeferredPostLoad, ExportBundleIndex).ReleaseBarrier();
+			Package->GetExportBundleNode(ExportBundle_DeferredPostLoad, ExportBundleIndex).ReleaseBarrier(&ThreadState);
 		}
 	}
 
@@ -6184,28 +6184,28 @@ FAsyncLoadingSyncLoadContext* FAsyncLoadingThread2::UpdateSyncLoadContext(FAsync
 		return nullptr;
 	}
 	FAsyncLoadingSyncLoadContext* SyncLoadContext = ThreadState.SyncLoadContextStack.Top();
-	if (ThreadState.bIsAsyncLoadingThread && !ContainsRequestID(SyncLoadContext->RequestId))
+	if (ThreadState.bIsAsyncLoadingThread)
 	{
-		SyncLoadContext->ReleaseRef();
-		ThreadState.SyncLoadContextStack.Pop();
-		if (ThreadState.SyncLoadContextStack.IsEmpty())
+		while (!ContainsRequestID(SyncLoadContext->RequestId))
 		{
-			return nullptr;
+			SyncLoadContext->ReleaseRef();
+			ThreadState.SyncLoadContextStack.Pop();
+			if (ThreadState.SyncLoadContextStack.IsEmpty())
+			{
+				return nullptr;
+			}
+			SyncLoadContext = ThreadState.SyncLoadContextStack.Top();
 		}
-		SyncLoadContext = ThreadState.SyncLoadContextStack.Top();
 	}
 	if (ThreadState.bCanAccessAsyncLoadingThreadData && !SyncLoadContext->bHasFoundRequestedPackage)
 	{
-		if (ThreadState.bIsAsyncLoadingThread)
-		{
-			// Ensure that we've created the package we're waiting for
-			CreateAsyncPackagesFromQueue(ThreadState);
-		}
+		// Ensure that we've created the package we're waiting for
+		CreateAsyncPackagesFromQueue(ThreadState);
 		if (FAsyncPackage2* RequestedPackage = RequestIdToPackageMap.FindRef(SyncLoadContext->RequestId))
 		{
 			SyncLoadContext->bHasFoundRequestedPackage = true;
 			SyncLoadContext->RequestedPackageDebug = RequestedPackage;
-			IncludePackageInSyncLoadContext(SyncLoadContext->ContextId, RequestedPackage);
+			IncludePackageInSyncLoadContextRecursive(ThreadState, SyncLoadContext->ContextId, RequestedPackage);
 		}
 	}
 	return SyncLoadContext;
@@ -6220,7 +6220,6 @@ EAsyncPackageState::Type FAsyncLoadingThread2::ProcessAsyncLoadingFromGameThread
 	FAsyncLoadingTickScope2 InAsyncLoadingTick(*this);
 	uint32 LoopIterations = 0;
 
-	RemoveUnreachableObjects(UnreachableObjects);
 	while (true)
 	{
 		do 
@@ -6256,9 +6255,9 @@ EAsyncPackageState::Type FAsyncLoadingThread2::ProcessAsyncLoadingFromGameThread
 				}
 			}
 
-			if (FAsyncLoadingSyncLoadContext* SyncLoadContext = UpdateSyncLoadContext(ThreadState))
+			if (!ThreadState.SyncLoadContextStack.IsEmpty())
 			{
-				if (EventQueue.ExecuteSyncLoadEvents(ThreadState, *SyncLoadContext))
+				if (EventQueue.ExecuteSyncLoadEvents(ThreadState, *ThreadState.SyncLoadContextStack.Top()))
 				{
 					bDidSomething = true;
 					break;
@@ -6277,7 +6276,7 @@ EAsyncPackageState::Type FAsyncLoadingThread2::ProcessAsyncLoadingFromGameThread
 				FAsyncPackage2* Package = nullptr;
 				ExternalReadQueue.Dequeue(Package);
 
-				EAsyncPackageState::Type Result = Package->ProcessExternalReads(FAsyncPackage2::ExternalReadAction_Wait);
+				EAsyncPackageState::Type Result = Package->ProcessExternalReads(ThreadState, FAsyncPackage2::ExternalReadAction_Wait);
 				check(Result == EAsyncPackageState::Complete);
 
 				bDidSomething = true;
@@ -6327,9 +6326,9 @@ EAsyncPackageState::Type FAsyncLoadingThread2::ProcessLoadedPackagesFromGameThre
 		}
 
 		bool bLocalDidSomething = false;
-		if (FAsyncLoadingSyncLoadContext* SyncLoadContext = UpdateSyncLoadContext(ThreadState))
+		if (!ThreadState.SyncLoadContextStack.IsEmpty())
 		{
-			bLocalDidSomething |= MainThreadEventQueue.ExecuteSyncLoadEvents(ThreadState, *SyncLoadContext);
+			bLocalDidSomething |= MainThreadEventQueue.ExecuteSyncLoadEvents(ThreadState, *ThreadState.SyncLoadContextStack.Top());
 		}
 		else
 		{
@@ -6592,6 +6591,12 @@ EAsyncPackageState::Type FAsyncLoadingThread2::TickAsyncLoadingFromGameThread(FA
 		const bool bIsMultithreaded = FAsyncLoadingThread2::IsMultithreaded();
 		double TickStartTime = FPlatformTime::Seconds();
 
+		if (!bIsMultithreaded)
+		{
+			RemoveUnreachableObjects(UnreachableObjects);
+		}
+		UpdateSyncLoadContext(ThreadState);
+
 		{
 			Result = ProcessLoadedPackagesFromGameThread(ThreadState, bDidSomething, FlushRequestID);
 			double TimeLimitUsedForProcessLoaded = FPlatformTime::Seconds() - TickStartTime;
@@ -6803,7 +6808,7 @@ uint32 FAsyncLoadingThread2::Run()
 						FAsyncPackage2* Package = nullptr;
 						ExternalReadQueue.Dequeue(Package);
 						check(Package);
-						EAsyncPackageState::Type Result = Package->ProcessExternalReads(FAsyncPackage2::ExternalReadAction_Wait);
+						EAsyncPackageState::Type Result = Package->ProcessExternalReads(ThreadState, FAsyncPackage2::ExternalReadAction_Wait);
 						check(Result == EAsyncPackageState::Complete);
 						bShouldWaitForExternalReads = false;
 						continue;
@@ -6864,7 +6869,7 @@ uint32 FAsyncLoadingThread2::Run()
 						{
 							TRACE_CPUPROFILER_EVENT_SCOPE(PollExternalReads);
 							check(*Package);
-							EAsyncPackageState::Type Result = (*Package)->ProcessExternalReads(FAsyncPackage2::ExternalReadAction_Poll);
+							EAsyncPackageState::Type Result = (*Package)->ProcessExternalReads(ThreadState, FAsyncPackage2::ExternalReadAction_Poll);
 							if (Result == EAsyncPackageState::Complete)
 							{
 								ExternalReadQueue.Dequeue();
@@ -7571,7 +7576,7 @@ void FAsyncPackage2::CreateUPackage()
 	}
 }
 
-EAsyncPackageState::Type FAsyncPackage2::ProcessExternalReads(EExternalReadAction Action)
+EAsyncPackageState::Type FAsyncPackage2::ProcessExternalReads(FAsyncLoadingThreadState2& ThreadState, EExternalReadAction Action)
 {
 	check(AsyncPackageLoadingState == EAsyncPackageLoadingState2::WaitingForExternalReads);
 	double WaitTime;
@@ -7595,7 +7600,7 @@ EAsyncPackageState::Type FAsyncPackage2::ProcessExternalReads(EExternalReadActio
 	}
 
 	ExternalReadDependencies.Empty();
-	GetPackageNode(Package_ExportsSerialized).ReleaseBarrier();
+	GetPackageNode(Package_ExportsSerialized).ReleaseBarrier(&ThreadState);
 	return EAsyncPackageState::Complete;
 }
 
@@ -7864,6 +7869,7 @@ void FAsyncLoadingThread2::FlushLoading(int32 RequestId)
 			check(GameThreadState->SyncLoadContextStack.Top() == SyncLoadContext);
 			SyncLoadContext->ReleaseRef();
 			GameThreadState->SyncLoadContextStack.Pop();
+			AltZenaphore.NotifyOne();
 		}
 
 		if (CurrentlyExecutingPackage)
