Improved comments on SerializeBin() and SerializeBinEx() to make it clear that binary serialization of deprecated properties can cause problems.
Made FProperty::ShouldSerializeValue() easier to debug for why a property is skipped.

diff --git a/Engine/Source/Runtime/CoreUObject/Private/UObject/Property.cpp b/Engine/Source/Runtime/CoreUObject/Private/UObject/Property.cpp
--- a/Engine/Source/Runtime/CoreUObject/Private/UObject/Property.cpp
+++ b/Engine/Source/Runtime/CoreUObject/Private/UObject/Property.cpp
@@ -1137,13 +1137,15 @@ bool FProperty::ExportText_Direct
 	return false;
 }
 
-bool FProperty::ShouldSerializeValue( FArchive& Ar ) const
+bool FProperty::ShouldSerializeValue(FArchive& Ar) const
 {
+	// Skip the property if the archive says we should
 	if (Ar.ShouldSkipProperty(this))
 	{
 		return false;
 	}
 
+	// Skip non-SaveGame properties if we're saving game state
 	if (!(PropertyFlags & CPF_SaveGame) && Ar.IsSaveGame())
 	{
 		return false;
@@ -1155,16 +1157,50 @@ bool FProperty::ShouldSerializeValue( FArchive& Ar ) const
 		return true;
 	}
 
-	bool Skip =
-			((PropertyFlags & CPF_Transient) && Ar.IsPersistent() && !Ar.IsSerializingDefaults())
-		||	((PropertyFlags & CPF_DuplicateTransient) && (Ar.GetPortFlags() & PPF_Duplicate))
-		||	((PropertyFlags & CPF_NonPIEDuplicateTransient) && !(Ar.GetPortFlags() & PPF_DuplicateForPIE) && (Ar.GetPortFlags() & PPF_Duplicate))
-		||	((PropertyFlags & CPF_NonTransactional) && Ar.IsTransacting())
-		||	((PropertyFlags & CPF_Deprecated) && !Ar.HasAllPortFlags(PPF_UseDeprecatedProperties) && (Ar.IsSaving() || Ar.IsTransacting() || Ar.WantBinaryPropertySerialization()))
-		||  ((PropertyFlags & CPF_SkipSerialization) && (Ar.WantBinaryPropertySerialization() || !Ar.HasAllPortFlags(PPF_ForceTaggedSerialization)))
-		||  (IsEditorOnlyProperty() && Ar.IsFilterEditorOnly());
+	// Skip properties marked Transient when persisting an object, unless we're saving an archetype
+	if ((PropertyFlags & CPF_Transient) && Ar.IsPersistent() && !Ar.IsSerializingDefaults())
+	{
+		return false;
+	}
+
+	// Skip properties marked DuplicateTransient when duplicating
+	if ((PropertyFlags & CPF_DuplicateTransient) && (Ar.GetPortFlags() & PPF_Duplicate))
+	{
+		return false;
+	}
+
+	// Skip properties marked NonPIEDuplicateTransient when duplicating, but not when we're duplicating for PIE
+	if ((PropertyFlags & CPF_NonPIEDuplicateTransient) && !(Ar.GetPortFlags() & PPF_DuplicateForPIE) && (Ar.GetPortFlags() & PPF_Duplicate))
+	{
+		return false;
+	}
+
+	// Skip properties marked NonTransactional when transacting
+	if ((PropertyFlags & CPF_NonTransactional) && Ar.IsTransacting())
+	{
+		return false;
+	}
+
+	// Skip deprecated properties when saving or transacting, unless the archive has explicitly requested them
+	if ((PropertyFlags & CPF_Deprecated) && !Ar.HasAllPortFlags(PPF_UseDeprecatedProperties) && (Ar.IsSaving() || Ar.IsTransacting() || Ar.WantBinaryPropertySerialization()))
+	{
+		return false;
+	}
+
+	// Skip properties marked SkipSerialization, unless the archive is forcing them
+	if ((PropertyFlags & CPF_SkipSerialization) && (Ar.WantBinaryPropertySerialization() || !Ar.HasAllPortFlags(PPF_ForceTaggedSerialization)))
+	{
+		return false;
+	}
 
-	return !Skip;
+	// Skip editor-only properties when the archive is rejecting them
+	if (IsEditorOnlyProperty() && Ar.IsFilterEditorOnly())
+	{
+		return false;
+	}
+
+	// Otherwise serialize!
+	return true;
 }
 
 
diff --git a/Engine/Source/Runtime/CoreUObject/Public/UObject/Class.h b/Engine/Source/Runtime/CoreUObject/Public/UObject/Class.h
--- a/Engine/Source/Runtime/CoreUObject/Public/UObject/Class.h
+++ b/Engine/Source/Runtime/CoreUObject/Public/UObject/Class.h
@@ -474,22 +474,47 @@ public:
 	/** Creates the field/property links and gets structure ready for use at runtime */
 	virtual void Link(FArchive& Ar, bool bRelinkExistingProperties);
 
-	/** Serializes struct properties, does not handle defaults*/
+	/**
+	 * Serializes struct properties, does not handle defaults.  See SerializeBinEx for handling defaults.
+	 *
+	 * @param	Ar				the archive to use for serialization
+	 * @param	Data			pointer to the location of the beginning of the property data
+	 *
+	 * @note Binary serialization will read and write unstructured data from the archive.  As deprecated
+	 *       properties are read from archives but not written, it is dangerous to call this function on
+	 *       types with deprecated properties, unless the ArWantBinarySerialization flag is set on the
+	 *       archive to force serialization to occur always.
+	 */
 	virtual void SerializeBin(FArchive& Ar, void* Data) const final 
 	{
 		SerializeBin(FStructuredArchiveFromArchive(Ar).GetSlot(), Data);
 	}
 
-	/** Serializes struct properties, does not handle defaults */
+	/**
+	 * Serializes struct properties, does not handle defaults.  See SerializeBinEx for handling defaults.
+	 *
+	 * @param	Slot			The structured archive slot we are serializing to
+	 * @param	Data			pointer to the location of the beginning of the property data
+	 *
+	 * @note Binary serialization will read and write unstructured data from the archive.  As deprecated
+	 *       properties are read from archives but not written, it is dangerous to call this function on
+	 *       types with deprecated properties, unless the ArWantBinarySerialization flag is set on the
+	 *       archive to force serialization to occur always.
+	 */
 	virtual void SerializeBin(FStructuredArchive::FSlot Slot, void* Data) const;
 
 	/**
 	 * Serializes the class properties that reside in Data if they differ from the corresponding values in DefaultData
 	 *
-	 * @param	Ar				the archive to use for serialization
+	 * @param	Slot			The structured archive slot we are serializing to
 	 * @param	Data			pointer to the location of the beginning of the property data
 	 * @param	DefaultData		pointer to the location of the beginning of the data that should be compared against
 	 * @param	DefaultStruct	the struct corresponding to the block of memory located at DefaultData 
+	 *
+	 * @note Binary serialization will read and write unstructured data from the archive.  As deprecated
+	 *       properties are read from archives but not written, it is dangerous to call this function on
+	 *       types with deprecated properties, unless the ArWantBinarySerialization flag is set on the
+	 *       archive to force serialization to occur always.
 	 */
 	void SerializeBinEx( FStructuredArchive::FSlot Slot, void* Data, void const* DefaultData, UStruct* DefaultStruct ) const;
 
