ZenLoader in editor:
* Allow sync loading a package that's already on the stack when postloading an object
* Fixed issue with packages that were already being waited for when issuing a sync load from one of their event nodes
* Clear linker imports when detaching a linker so that we don't end up with invalid import object pointers after a GC
* Fixed infinite loading when a predicate based flush is done from within a sync load context
#preflght 6374e425bf76990b71b34f93

diff --git a/Engine/Source/Runtime/CoreUObject/Private/Serialization/AsyncLoading2.cpp b/Engine/Source/Runtime/CoreUObject/Private/Serialization/AsyncLoading2.cpp
--- a/Engine/Source/Runtime/CoreUObject/Private/Serialization/AsyncLoading2.cpp
+++ b/Engine/Source/Runtime/CoreUObject/Private/Serialization/AsyncLoading2.cpp
@@ -1888,7 +1888,7 @@ public:
 
 	uint64 ContextId;
 	int32 RequestId;
-	FAsyncPackage2* RequestedPackageDebug = nullptr;
+	FAsyncPackage2* RequestedPackage = nullptr;
 	FAsyncPackage2* RequestingPackageDebug = nullptr;
 	bool bHasFoundRequestedPackage = false;
 
@@ -1963,7 +1963,7 @@ struct FAsyncLoadingThreadState2
 	TArray<FEventLoadNode2*> NodesToFire;
 	FEventLoadNode2* CurrentlyExecutingEventNode = nullptr;
 	TArray<FAsyncLoadingSyncLoadContext*> SyncLoadContextStack;
-	TArray<FAsyncPackage2*> PackagesExcludedFromStateDependencyChecks;
+	TArray<FAsyncPackage2*> PackagesOnStack;
 	TSpscQueue<FAsyncLoadingSyncLoadContext*> SyncLoadContextsCreatedOnGameThread;
 	bool bIsAsyncLoadingThread = false;
 	bool bCanAccessAsyncLoadingThreadData = true;
@@ -2429,7 +2429,7 @@ private:
 			}
 			if (PackageToAddState.WaitingForPackage)
 			{
-				PackageToAddState.RemoveFromWaitList(StateMemberPtr, PackageToAddState.WaitingForPackage, PackageToAdd);
+				FAllDependenciesState::RemoveFromWaitList(StateMemberPtr, PackageToAddState.WaitingForPackage, PackageToAdd);
 			}
 
 			check(!PackageToAddState.PrevLink);
@@ -4379,6 +4379,13 @@ void FAsyncPackage2::DetachLinker()
 {
 	if (LinkerLoadState.IsSet() && LinkerLoadState->Linker)
 	{
+		// We're no longer keeping the imports alive so clear them from the linker
+		for (FObjectImport& ObjectImport : LinkerLoadState->Linker->ImportMap)
+		{
+			ObjectImport.XObject = nullptr;
+			ObjectImport.SourceLinker = nullptr;
+			ObjectImport.SourceIndex = INDEX_NONE;
+		}
 		check(LinkerLoadState->Linker->AsyncRoot == this);
 		LinkerLoadState->Linker->AsyncRoot = nullptr;
 		LinkerLoadState->Linker = nullptr;
@@ -5674,7 +5681,7 @@ bool FAsyncPackage2::HaveAllDependenciesReachedStateDebug(FAsyncLoadingThreadSta
 				continue;
 			}
 #endif
-			if (ThreadState.PackagesExcludedFromStateDependencyChecks.Contains(ImportedPackage))
+			if (ThreadState.PackagesOnStack.Contains(ImportedPackage))
 			{
 				continue;
 			}
@@ -5727,11 +5734,6 @@ void FAsyncPackage2::UpdateDependenciesStateRecursive(FAsyncLoadingThreadState2&
 
 	if (FAsyncPackage2* WaitingForPackage = ThisState.WaitingForPackage)
 	{
-		if (ThreadState.PackagesExcludedFromStateDependencyChecks.Contains(WaitingForPackage))
-		{
-			ThisState.WaitingForPackage = nullptr;
-		}
-		else
 		{
 			FAllDependenciesState& WaitingForPackageState = WaitingForPackage->*StateMemberPtr;
 			if (WaitingForPackage->AsyncPackageLoadingState < WaitForPackageState)
@@ -5771,7 +5773,7 @@ void FAsyncPackage2::UpdateDependenciesStateRecursive(FAsyncLoadingThreadState2&
 		}
 #endif
 
-		if (ThreadState.PackagesExcludedFromStateDependencyChecks.Contains(ImportedPackage))
+		if (ThreadState.PackagesOnStack.Contains(ImportedPackage))
 		{
 			continue;
 		}
@@ -5851,6 +5853,8 @@ void FAsyncPackage2::WaitForAllDependenciesToReachState(FAsyncLoadingThreadState
 				check(WaitingPackageState.WaitingForPackage == PackageReadyToProceed);
 				if (WaitingPackage->HaveAllDependenciesReachedState(ThreadState, StateMemberPtr, WaitForPackageState, CurrentTickVariable++, bExcludeZenPackages))
 				{
+					// If the waiting package wasn't ready to proceed it will remove itself from the waitlist (PackageReadyToProceedState.PackagesWaitingForThisHead) when its WaitingForPackage
+					// is replaced by another package
 					FAllDependenciesState::RemoveFromWaitList(StateMemberPtr, PackageReadyToProceed, WaitingPackage);
 					WaitingPackageState.NextLink = FirstPackageReadyToProceed;
 					FirstPackageReadyToProceed = WaitingPackage;
@@ -6197,6 +6201,10 @@ FAsyncLoadingSyncLoadContext* FAsyncLoadingThread2::UpdateSyncLoadContext(FAsync
 			SyncLoadContext = ThreadState.SyncLoadContextStack.Top();
 		}
 	}
+	else if (!ContainsRequestID(SyncLoadContext->RequestId))
+	{
+		return nullptr;
+	}
 	if (ThreadState.bCanAccessAsyncLoadingThreadData && !SyncLoadContext->bHasFoundRequestedPackage)
 	{
 		// Ensure that we've created the package we're waiting for
@@ -6204,10 +6212,18 @@ FAsyncLoadingSyncLoadContext* FAsyncLoadingThread2::UpdateSyncLoadContext(FAsync
 		if (FAsyncPackage2* RequestedPackage = RequestIdToPackageMap.FindRef(SyncLoadContext->RequestId))
 		{
 			SyncLoadContext->bHasFoundRequestedPackage = true;
-			SyncLoadContext->RequestedPackageDebug = RequestedPackage;
+			SyncLoadContext->RequestedPackage = RequestedPackage;
 			IncludePackageInSyncLoadContextRecursive(ThreadState, SyncLoadContext->ContextId, RequestedPackage);
 		}
 	}
+	if (SyncLoadContext->bHasFoundRequestedPackage && ThreadState.PackagesOnStack.Contains(SyncLoadContext->RequestedPackage))
+	{
+		// Flushing a package while it's already being processed on the stack, if we're done preloading we let it pass and remove the request id
+		bool bPreloadIsDone = SyncLoadContext->RequestedPackage->AsyncPackageLoadingState >= EAsyncPackageLoadingState2::DeferredPostLoad;
+		UE_CLOG(!bPreloadIsDone, LogStreaming, Fatal, TEXT("Flushing package %s while it's being preloaded in the same callstack is not permitted"), *SyncLoadContext->RequestedPackage->Desc.UPackageName.ToString());
+		RemovePendingRequests(TArrayView<int32>(&SyncLoadContext->RequestId, 1));
+	}
+
 	return SyncLoadContext;
 }
 
@@ -7807,7 +7823,29 @@ void FAsyncLoadingThread2::FlushLoading(int32 RequestId)
 		{
 			UE_CLOG(RequestId == INDEX_NONE, LogStreaming, Fatal, TEXT("Flushing async loading while creating, serializing or postloading an object is not permitted"));
 			CurrentlyExecutingPackage = GameThreadState->CurrentlyExecutingEventNode->GetPackage();
-			GameThreadState->PackagesExcludedFromStateDependencyChecks.Push(CurrentlyExecutingPackage);
+			GameThreadState->PackagesOnStack.Push(CurrentlyExecutingPackage);
+			// Update the state of any package that is waiting for the currently executing one
+			while (FAsyncPackage2* WaitingPackage = CurrentlyExecutingPackage->AllDependenciesFullyLoadedState.PackagesWaitingForThisHead)
+			{
+				FAsyncPackage2::FAllDependenciesState::RemoveFromWaitList(&FAsyncPackage2::AllDependenciesFullyLoadedState, CurrentlyExecutingPackage, WaitingPackage);
+				WaitingPackage->ConditionalFinishLoading(*GameThreadState);
+			}
+			if (GameThreadState->bCanAccessAsyncLoadingThreadData)
+			{
+				while (FAsyncPackage2* WaitingPackage = CurrentlyExecutingPackage->AllDependenciesSerializedState.PackagesWaitingForThisHead)
+				{
+					FAsyncPackage2::FAllDependenciesState::RemoveFromWaitList(&FAsyncPackage2::AllDependenciesSerializedState, CurrentlyExecutingPackage, WaitingPackage);
+					WaitingPackage->ConditionalBeginPostLoad(*GameThreadState);
+				}
+
+#if ALT2_ENABLE_LINKERLOAD_SUPPORT
+				while (FAsyncPackage2* WaitingPackage = CurrentlyExecutingPackage->AllDependenciesSetupState.PackagesWaitingForThisHead)
+				{
+					FAsyncPackage2::FAllDependenciesState::RemoveFromWaitList(&FAsyncPackage2::AllDependenciesSetupState, CurrentlyExecutingPackage, WaitingPackage);
+					WaitingPackage->ConditionalProcessLinkerLoadPackageImportsAndExports(*GameThreadState);
+				}
+#endif
+			}
 		}
 
 		FAsyncLoadingSyncLoadContext* SyncLoadContext = nullptr;
@@ -7874,8 +7912,8 @@ void FAsyncLoadingThread2::FlushLoading(int32 RequestId)
 
 		if (CurrentlyExecutingPackage)
 		{
-			check(GameThreadState->PackagesExcludedFromStateDependencyChecks.Top() == CurrentlyExecutingPackage);
-			GameThreadState->PackagesExcludedFromStateDependencyChecks.Pop();
+			check(GameThreadState->PackagesOnStack.Top() == CurrentlyExecutingPackage);
+			GameThreadState->PackagesOnStack.Pop();
 		}
 
 		check(RequestId != INDEX_NONE || !IsAsyncLoadingPackages());
@@ -7898,6 +7936,11 @@ EAsyncPackageState::Type FAsyncLoadingThread2::ProcessLoadingUntilCompleteFromGa
 	bool bUseTimeLimit = TimeLimit > 0.0f;
 	double TimeLoadingPackage = 0.0f;
 
+	// If there's an active sync load context we need to supress is for the duration of this call since we've no idea what the CompletionPredicate is waiting for
+	uint64 NoSyncLoadContextId = 0;
+	uint64& SyncLoadContextId = ThreadState.SyncLoadContextStack.IsEmpty() ? NoSyncLoadContextId : ThreadState.SyncLoadContextStack.Top()->ContextId;
+	TGuardValue<uint64> GuardSyncLoadContextId(SyncLoadContextId, 0);
+
 	bool bLoadingComplete = !IsAsyncLoadingPackages() || CompletionPredicate();
 	while (!bLoadingComplete && (!bUseTimeLimit || TimeLimit > 0.0f))
 	{
