Make the WarnIfAssetsLoadedInScope respect the enabled state.

diff --git a/Engine/Source/Runtime/CoreUObject/Private/Misc/WarnIfAssetsLoadedInScope.cpp b/Engine/Source/Runtime/CoreUObject/Private/Misc/WarnIfAssetsLoadedInScope.cpp
--- a/Engine/Source/Runtime/CoreUObject/Private/Misc/WarnIfAssetsLoadedInScope.cpp
+++ b/Engine/Source/Runtime/CoreUObject/Private/Misc/WarnIfAssetsLoadedInScope.cpp
@@ -55,6 +55,12 @@ FWarnIfAssetsLoadedInScope::~FWarnIfAssetsLoadedInScope()
 
 void FWarnIfAssetsLoadedInScope::OnAssetLoaded(UObject* InObject)
 {
+	// Don't warn people if the switch isn't enalbed.
+	if (!IsEnabled())
+	{
+		return;
+	}
+
 	// If the handle isn't valid we can avoid all this, as there's no way anything was captured.
 	if (!AssetsLoadedHandle.IsValid())
 	{
