ZenLoader: Match suspend/resume behaviour of legacy loader by using a counter instead of a boolean

diff --git a/Engine/Source/Runtime/CoreUObject/Private/Serialization/AsyncLoading2.cpp b/Engine/Source/Runtime/CoreUObject/Private/Serialization/AsyncLoading2.cpp
--- a/Engine/Source/Runtime/CoreUObject/Private/Serialization/AsyncLoading2.cpp
+++ b/Engine/Source/Runtime/CoreUObject/Private/Serialization/AsyncLoading2.cpp
@@ -2623,8 +2623,8 @@ private:
 	/** Thread to run the worker FRunnable on */
 	static constexpr int32 DefaultAsyncPackagesReserveCount = 512;
 	FRunnableThread* Thread;
-	TAtomic<bool> bStopRequested { false };
-	TAtomic<bool> bSuspendRequested { false };
+	std::atomic<bool> bStopRequested = false;
+	std::atomic<int32> SuspendRequestedCount = 0;
 	bool bHasRegisteredAllScriptObjects = false;
 	/** [ASYNC/GAME THREAD] true if the async thread is actually started. We don't start it until after we boot because the boot process on the game thread can create objects that are also being created by the loader */
 	bool bThreadStarted = false;
@@ -2788,7 +2788,7 @@ public:
 	/** Returns true if async loading is suspended */
 	inline virtual bool IsAsyncLoadingSuspended() override
 	{
-		return bSuspendRequested;
+		return SuspendRequestedCount.load(std::memory_order_relaxed) > 0;
 	}
 
 	virtual void NotifyConstructedDuringAsyncLoading(UObject* Object, bool bSubObjectThatAlreadyExists) override;
@@ -6771,14 +6771,14 @@ uint32 FAsyncLoadingThread2::Run()
 	};
 	EMainState PreviousState = EMainState::Loading;
 	EMainState CurrentState = EMainState::Loading;
-	while (!bStopRequested)
+	while (!bStopRequested.load(std::memory_order_relaxed))
 	{
 		if (CurrentState == EMainState::Suspended)
 		{
 			// suspended, sleep until loading can be resumed
-			while (!bStopRequested)
+			while (!bStopRequested.load(std::memory_order_relaxed))
 			{
-				if (!bSuspendRequested.Load(EMemoryOrder::SequentiallyConsistent) && !IsGarbageCollectionWaiting())
+				if (SuspendRequestedCount.load(std::memory_order_relaxed) == 0 && !IsGarbageCollectionWaiting())
 				{
 					ThreadResumedEvent->Trigger();
 					CurrentState = EMainState::Loading;
@@ -6801,9 +6801,9 @@ uint32 FAsyncLoadingThread2::Run()
 
 			bool bShouldSuspend = false;
 			bool bShouldWaitForExternalReads = false;
-			while (!bStopRequested)
+			while (!bStopRequested.load(std::memory_order_relaxed))
 			{
-				if (bShouldSuspend || bSuspendRequested.Load(EMemoryOrder::Relaxed) || IsGarbageCollectionWaiting())
+				if (bShouldSuspend || SuspendRequestedCount.load(std::memory_order_relaxed) > 0 || IsGarbageCollectionWaiting())
 				{
 					TRACE_CPUPROFILER_EVENT_SCOPE(SuspendAsyncLoading);
 					ThreadSuspendedEvent->Trigger();
@@ -6836,7 +6836,7 @@ uint32 FAsyncLoadingThread2::Run()
 						if (CreateAsyncPackagesFromQueue(ThreadState))
 						{
 							// Fall through to FAsyncLoadEventQueue2 processing unless we need to suspend
-							if (bSuspendRequested.Load(EMemoryOrder::Relaxed) || IsGarbageCollectionWaiting())
+							if (SuspendRequestedCount.load(std::memory_order_relaxed) > 0 || IsGarbageCollectionWaiting())
 							{
 								bShouldSuspend = true;
 								continue;
@@ -6864,7 +6864,7 @@ uint32 FAsyncLoadingThread2::Run()
 								bPopped = true;
 								bDidSomething = true;
 							}
-							if (bSuspendRequested.Load(EMemoryOrder::Relaxed) || IsGarbageCollectionWaiting())
+							if (SuspendRequestedCount.load(std::memory_order_relaxed) > 0 || IsGarbageCollectionWaiting())
 							{
 								bShouldSuspend = true;
 								bDidSomething = true;
@@ -6944,8 +6944,8 @@ EAsyncPackageState::Type FAsyncLoadingThread2::TickAsyncThreadFromGameThread(FAs
 
 void FAsyncLoadingThread2::Stop()
 {
-	bSuspendRequested = true;
-	bStopRequested = true;
+	SuspendRequestedCount.fetch_add(1);
+	bStopRequested.store(true);
 	AltZenaphore.NotifyAll();
 }
 
@@ -6959,9 +6959,10 @@ void FAsyncLoadingThread2::SuspendLoading()
 {
 	TRACE_CPUPROFILER_EVENT_SCOPE(SuspendLoading);
 	UE_CLOG(!IsInGameThread() || IsInSlateThread(), LogStreaming, Fatal, TEXT("Async loading can only be suspended from the main thread"));
-	if (!bSuspendRequested)
+	int32 OldCount = SuspendRequestedCount.fetch_add(1);
+	if (OldCount == 0)
 	{
-		bSuspendRequested = true;
+		UE_LOG(LogStreaming, Display, TEXT("Suspending async loading"));
 		if (IsMultithreaded())
 		{
 			TRACE_LOADTIME_SUSPEND_ASYNC_LOADING();
@@ -6969,21 +6970,31 @@ void FAsyncLoadingThread2::SuspendLoading()
 			ThreadSuspendedEvent->Wait();
 		}
 	}
+	else
+	{
+		UE_LOG(LogStreaming, Verbose, TEXT("Async loading is already suspended (count: %d)"), OldCount + 1);
+	}
 }
 
 void FAsyncLoadingThread2::ResumeLoading()
 {
 	TRACE_CPUPROFILER_EVENT_SCOPE(ResumeLoading);
-	check(IsInGameThread() && !IsInSlateThread());
-	if (bSuspendRequested)
+	UE_CLOG(!IsInGameThread() || IsInSlateThread(), LogStreaming, Fatal, TEXT("Async loading can only be resumed from the main thread"));
+	int32 OldCount = SuspendRequestedCount.fetch_sub(1);
+	UE_CLOG(OldCount < 1, LogStreaming, Fatal, TEXT("Trying to resume async loading when it's not suspended"));
+	if (OldCount == 1)
 	{
-		bSuspendRequested = false;
+		UE_LOG(LogStreaming, Display, TEXT("Resuming async loading"));
 		if (IsMultithreaded())
 		{
 			ThreadResumedEvent->Wait();
 			TRACE_LOADTIME_RESUME_ASYNC_LOADING();
 		}
 	}
+	else
+	{
+		UE_LOG(LogStreaming, Verbose, TEXT("Async loading is still suspended (count: %d)"), OldCount - 1);
+	}
 }
 
 float FAsyncLoadingThread2::GetAsyncLoadPercentage(const FName& PackageName)
@@ -7804,7 +7815,7 @@ void FAsyncLoadingThread2::FlushLoading(int32 RequestId)
 	if (IsAsyncLoadingPackages())
 	{
 		// Flushing async loading while loading is suspend will result in infinite stall
-		UE_CLOG(bSuspendRequested, LogStreaming, Fatal, TEXT("Cannot Flush Async Loading while async loading is suspended"));
+		UE_CLOG(SuspendRequestedCount.load(std::memory_order_relaxed) > 0, LogStreaming, Fatal, TEXT("Cannot Flush Async Loading while async loading is suspended"));
 
 		SCOPE_CYCLE_COUNTER(STAT_FAsyncPackage_FlushAsyncLoadingGameThread);
 
@@ -7931,7 +7942,7 @@ EAsyncPackageState::Type FAsyncLoadingThread2::ProcessLoadingUntilCompleteFromGa
 	SCOPE_CYCLE_COUNTER(STAT_FAsyncPackage_FlushAsyncLoadingGameThread);
 
 	// Flushing async loading while loading is suspend will result in infinite stall
-	UE_CLOG(bSuspendRequested, LogStreaming, Fatal, TEXT("Cannot Flush Async Loading while async loading is suspended"));
+	UE_CLOG(SuspendRequestedCount.load(std::memory_order_relaxed) > 0, LogStreaming, Fatal, TEXT("Cannot Flush Async Loading while async loading is suspended"));
 
 	bool bUseTimeLimit = TimeLimit > 0.0f;
 	double TimeLoadingPackage = 0.0f;
