Pass an optional default data pointer to FDynamicallyTypedValueType::SerializeValue.
Ensure that FDynamicallyTypedValue::AllocateData at least zero-initializes the data payload in the inline storage case to avoid a static analysis warning.

diff --git a/Engine/Source/Runtime/CoreUObject/Private/UObject/DynamicallyTypedValue.cpp b/Engine/Source/Runtime/CoreUObject/Private/UObject/DynamicallyTypedValue.cpp
--- a/Engine/Source/Runtime/CoreUObject/Private/UObject/DynamicallyTypedValue.cpp
+++ b/Engine/Source/Runtime/CoreUObject/Private/UObject/DynamicallyTypedValue.cpp
@@ -15,7 +15,7 @@ UE::FDynamicallyTypedValueType& UE::FDynamicallyTypedValue::NullType()
 		virtual void InitializeValueFromCopy(void* DestData, const void* SourceData) const override {}
 		virtual void DestroyValue(void* Data) const override {}
 
-		virtual void SerializeValue(FStructuredArchive::FSlot Slot, void* Data) const override {}
+		virtual void SerializeValue(FStructuredArchive::FSlot Slot, void* Data, const void* DefaultData) const override {}
 
 		virtual uint32 GetValueHash(const void* Data) const override { return 0; }
 		virtual bool AreIdentical(const void* DataA, const void* DataB) const override { return true; }
diff --git a/Engine/Source/Runtime/CoreUObject/Public/UObject/DynamicallyTypedValue.h b/Engine/Source/Runtime/CoreUObject/Public/UObject/DynamicallyTypedValue.h
--- a/Engine/Source/Runtime/CoreUObject/Public/UObject/DynamicallyTypedValue.h
+++ b/Engine/Source/Runtime/CoreUObject/Public/UObject/DynamicallyTypedValue.h
@@ -39,7 +39,7 @@ namespace UE
 		virtual void InitializeValueFromCopy(void* DestData, const void* SourceData) const = 0;
 		virtual void DestroyValue(void* Data) const = 0;
 
-		virtual void SerializeValue(FStructuredArchive::FSlot Slot, void* Data) const = 0;
+		virtual void SerializeValue(FStructuredArchive::FSlot Slot, void* Data, const void* DefaultData) const = 0;
 
 		virtual uint32 GetValueHash(const void* Data) const = 0;
 		virtual bool AreIdentical(const void* DataA, const void* DataB) const = 0;
@@ -163,7 +163,13 @@ namespace UE
 		// Allocates heap memory for the value if it uses it.
 		void AllocateData()
 		{
-			if (!IsInline())
+			if (IsInline())
+			{
+				// Ensure that the data is zeroed in the inline case to avoid spurious static analysis
+				// errors about passing a reference to uninitialized data to InitializeValueFromCopy.
+				InlineData = 0;
+			}
+			else
 			{
 				HeapData = FMemory::Malloc(Type->GetNumBytes(), Type->GetMinAlignment());
 			}
