Fixed TIsTSubclassOf<T> for cv-qualified types.

diff --git a/Engine/Source/Runtime/CoreUObject/Public/Templates/SubclassOf.h b/Engine/Source/Runtime/CoreUObject/Public/Templates/SubclassOf.h
--- a/Engine/Source/Runtime/CoreUObject/Public/Templates/SubclassOf.h
+++ b/Engine/Source/Runtime/CoreUObject/Public/Templates/SubclassOf.h
@@ -9,12 +9,20 @@
 
 #include <type_traits>
 
+template <typename T>
+class TSubclassOf;
+
 template <typename T>
 struct TIsTSubclassOf
 {
 	enum { Value = false };
 };
 
+template <typename T> struct TIsTSubclassOf<               TSubclassOf<T>> { enum { Value = true }; };
+template <typename T> struct TIsTSubclassOf<const          TSubclassOf<T>> { enum { Value = true }; };
+template <typename T> struct TIsTSubclassOf<      volatile TSubclassOf<T>> { enum { Value = true }; };
+template <typename T> struct TIsTSubclassOf<const volatile TSubclassOf<T>> { enum { Value = true }; };
+
 /**
  * Template to allow TClassType's to be passed around with type safety 
  */
@@ -160,12 +168,6 @@ private:
 	TClassType* Class;
 };
 
-template <typename T>
-struct TIsTSubclassOf<TSubclassOf<T>>
-{
-	enum { Value = true };
-};
-
 template <typename T>
 FArchive& operator<<(FArchive& Ar, TSubclassOf<T>& SubclassOf)
 {
