Removed TNonNullSubclassOf<T>'s dependency on TSubclassOf<T>::TClassType to avoid needing the full type definition for T.
Fixed coding standard issues.

diff --git a/Engine/Source/Runtime/CoreUObject/Public/Templates/NonNullSubclassOf.h b/Engine/Source/Runtime/CoreUObject/Public/Templates/NonNullSubclassOf.h
--- a/Engine/Source/Runtime/CoreUObject/Public/Templates/NonNullSubclassOf.h
+++ b/Engine/Source/Runtime/CoreUObject/Public/Templates/NonNullSubclassOf.h
@@ -14,46 +14,55 @@ enum class EDefaultConstructNonNullSubclassOf { UnsafeDoNotUse };
 /**
  * Template to allow TClassType's to be passed around with type safety 
  */
-template<class TClass>
-class TNonNullSubclassOf : public TSubclassOf<TClass>
+template <typename T>
+class TNonNullSubclassOf : public TSubclassOf<T>
 {
-public:
-	
-	using TClassType = typename TSubclassOf<TClass>::TClassType;
-	using TBaseType = typename TSubclassOf<TClass>::TBaseType;
+	using Super = TSubclassOf<T>;
 
+public:
 	/** Default Constructor, defaults to null */
-	FORCEINLINE TNonNullSubclassOf(EDefaultConstructNonNullSubclassOf) :
-		TSubclassOf<TClass>(nullptr)
-	{}
+	FORCEINLINE TNonNullSubclassOf(EDefaultConstructNonNullSubclassOf)
+		: Super(nullptr)
+	{
+	}
 
-	/** Constructor that takes a UClass and does a runtime check to make sure this is a compatible class */
-	FORCEINLINE TNonNullSubclassOf(TClassType* From) :
-		TSubclassOf<TClass>(From)
+	/** Constructor that takes a UClass* or FFieldClass* if T is a UObject or a FField type respectively. */
+	template <typename PtrType>
+	FORCEINLINE TNonNullSubclassOf(PtrType* From)
+		: Super(From)
 	{
 		checkf(From, TEXT("Initializing TNonNullSubclassOf with null"));
 	}
 
 	/** Copy Constructor, will only compile if types are compatible */
-	template <class TClassA, class = decltype(ImplicitConv<TClass*>((TClassA*)nullptr))>
-	FORCEINLINE TNonNullSubclassOf(const TSubclassOf<TClassA>& From) :
-		TSubclassOf<TClass>(From)
-	{}
+	template <
+		typename U,
+		decltype(ImplicitConv<T*>((U*)nullptr))* = nullptr
+	>
+	FORCEINLINE TNonNullSubclassOf(const TSubclassOf<U>& From)
+		: Super(From)
+	{
+	}
 
 	/** Assignment operator, will only compile if types are compatible */
-	template <class TClassA, class = decltype(ImplicitConv<TClass*>((TClassA*)nullptr))>
-	FORCEINLINE TNonNullSubclassOf& operator=(const TSubclassOf<TClassA>& From)
+	template
+	<
+		typename U,
+		decltype(ImplicitConv<T*>((U*)nullptr))* = nullptr
+	>
+	FORCEINLINE TNonNullSubclassOf& operator=(const TSubclassOf<U>& From)
 	{
 		checkf(*From, TEXT("Assigning null to TNonNullSubclassOf"));
-		TSubclassOf<TClass>::operator=(From);
+		Super::operator=(From);
 		return *this;
 	}
 	
 	/** Assignment operator from UClass, the type is checked on get not on set */
-	FORCEINLINE TNonNullSubclassOf& operator=(TClassType* From)
+	template <typename PtrType>
+	FORCEINLINE TNonNullSubclassOf& operator=(PtrType* From)
 	{
 		checkf(From, TEXT("Assigning null to TNonNullSubclassOf"));
-		TSubclassOf<TClass>::operator=(From);
+		Super::operator=(From);
 		return *this;
 	}
 };
@@ -62,15 +71,14 @@ public:
 /**
  * Specialization of TOptional for TNonNullSubclassOf value types
  */
-template<typename TClass>
-struct TOptional<TNonNullSubclassOf<TClass>>
+template<typename T>
+struct TOptional<TNonNullSubclassOf<T>>
 {
-public:
-
-	using TClassType = typename TSubclassOf<TClass>::TClassType;
+	using ClassType = std::remove_pointer_t<decltype(std::declval<TNonNullSubclassOf<T>>().Get())>;
 
+public:
 	/** Construct an OptionaType with a valid value. */
-	TOptional(const TNonNullSubclassOf<TClass>& InPointer)
+	TOptional(const TNonNullSubclassOf<T>& InPointer)
 		: Pointer(InPointer)
 	{
 	}
@@ -93,7 +101,7 @@ public:
 		return *this;
 	}
 
-	TOptional& operator=(TClass* InPointer)
+	TOptional& operator=(T* InPointer)
 	{
 		Pointer = InPointer;
 		return *this;
@@ -104,7 +112,7 @@ public:
 		Pointer = nullptr;
 	}
 
-	TClass* Emplace(TClass* InPointer)
+	T* Emplace(T* InPointer)
 	{
 		Pointer = InPointer;
 		return InPointer;
@@ -131,15 +139,28 @@ public:
 	FORCEINLINE explicit operator bool() const { return Pointer != nullptr; }
 
 	/** @return The optional value; undefined when IsSet() returns false. */
-	TClassType* GetValue() const { checkf(IsSet(), TEXT("It is an error to call GetValue() on an unset TOptional. Please either check IsSet() or use Get(DefaultValue) instead.")); return Pointer; }
+	ClassType* GetValue() const
+	{
+		checkf(IsSet(), TEXT("It is an error to call GetValue() on an unset TOptional. Please either check IsSet() or use Get(DefaultValue) instead."));
+		return Pointer;
+	}
 
-	TClassType* operator->() const { return Pointer; }
+	ClassType* operator->() const
+	{
+		return Pointer;
+	}
 
-	TClassType& operator*() const { return *Pointer; }
+	ClassType& operator*() const
+	{
+		return *Pointer;
+	}
 
 	/** @return The optional value when set; DefaultValue otherwise. */
-	TClassType* Get(TClassType* DefaultPointer) const { return IsSet() ? Pointer : DefaultPointer; }
+	ClassType* Get(ClassType* DefaultPointer) const
+	{
+		return IsSet() ? Pointer : DefaultPointer;
+	}
 
 private:
-	TClassType* Pointer;
+	ClassType* Pointer;
 };
