Improved comments and some coding standard issues in TSubclassOf.
Defaulted the special member functions.

diff --git a/Engine/Source/Runtime/CoreUObject/Public/Templates/SubclassOf.h b/Engine/Source/Runtime/CoreUObject/Public/Templates/SubclassOf.h
--- a/Engine/Source/Runtime/CoreUObject/Public/Templates/SubclassOf.h
+++ b/Engine/Source/Runtime/CoreUObject/Public/Templates/SubclassOf.h
@@ -24,35 +24,37 @@ template <typename T> struct TIsTSubclassOf<      volatile TSubclassOf<T>> { enu
 template <typename T> struct TIsTSubclassOf<const volatile TSubclassOf<T>> { enum { Value = true }; };
 
 /**
- * Template to allow TClassType's to be passed around with type safety 
+ * Template to allow TClassTypes to be passed around with type safety 
  */
-template<class TClass>
+template <typename T>
 class TSubclassOf
 {
 public:
-
-	typedef typename TChooseClass<TIsDerivedFrom<TClass, FField>::IsDerived, FFieldClass, UClass>::Result TClassType;
-	typedef typename TChooseClass<TIsDerivedFrom<TClass, FField>::IsDerived, FField, UObject>::Result TBaseType;
+	typedef typename TChooseClass<TIsDerivedFrom<T, FField>::IsDerived, FFieldClass, UClass>::Result TClassType;
+	typedef typename TChooseClass<TIsDerivedFrom<T, FField>::IsDerived, FField, UObject>::Result TBaseType;
 
 private:
-
-	template <class TClassA>
+	template <typename U>
 	friend class TSubclassOf;
 
 public:
-	/** Default Constructor, defaults to null */
-	FORCEINLINE TSubclassOf() :
-		Class(nullptr)
-	{
-	}
+	TSubclassOf() = default;
+	TSubclassOf(TSubclassOf&&) = default;
+	TSubclassOf(const TSubclassOf&) = default;
+	TSubclassOf& operator=(TSubclassOf&&) = default;
+	TSubclassOf& operator=(const TSubclassOf&) = default;
+	~TSubclassOf() = default;
 
 	/** Constructor that takes a UClass and does a runtime check to make sure this is a compatible class */
-	FORCEINLINE TSubclassOf(TClassType* From) :
-		Class(From)
+	FORCEINLINE TSubclassOf(TClassType* From)
+		: Class(From)
 	{
 	}
 
-	/** Constructor that takes a UClass and does a runtime check to make sure this is a compatible class */
+	/**
+	 * Construct from a UClass* or FFieldClass* (or something implicitly convertible to those)
+	 * if T is a UObject or a FField type respectively.
+	 */
 	template <
 		typename U,
 		std::enable_if_t<
@@ -65,29 +67,37 @@ public:
 	{
 	}
 
-	/** Copy Constructor, will only compile if types are compatible */
-	template <class TClassA, class = decltype(ImplicitConv<TClass*>((TClassA*)nullptr))>
-	FORCEINLINE TSubclassOf(const TSubclassOf<TClassA>& From) :
-		Class(*From)
+	/** Construct from another TSubclassOf, only if types are compatible */
+	template <
+		typename OtherT,
+		decltype(ImplicitConv<T*>((OtherT*)nullptr))* = nullptr
+	>
+	FORCEINLINE TSubclassOf(const TSubclassOf<OtherT>& From)
+		: Class(*From)
 	{
 	}
 
-	/** Assignment operator, will only compile if types are compatible */
-	template <class TClassA, class = decltype(ImplicitConv<TClass*>((TClassA*)nullptr))>
-	FORCEINLINE TSubclassOf& operator=(const TSubclassOf<TClassA>& From)
+	/** Assign from another TSubclassOf, only if types are compatible */
+	template <
+		typename OtherT,
+		decltype(ImplicitConv<T*>((OtherT*)nullptr))* = nullptr
+	>
+	FORCEINLINE TSubclassOf& operator=(const TSubclassOf<OtherT>& From)
 	{
 		Class = *From;
 		return *this;
 	}
-	
-	/** Assignment operator from UClass, the type is checked on get not on set */
+
+	/** Assign from a UClass* or FFieldClass*. */
 	FORCEINLINE TSubclassOf& operator=(TClassType* From)
 	{
 		Class = From;
 		return *this;
 	}
-	
-	/** Assignment operator from UClass, the type is checked on get not on set */
+
+	/**
+	 * Assign from a UClass* or FFieldClass* (or something implicitly convertible to those).
+	 */
 	template <
 		typename U,
 		std::enable_if_t<
@@ -100,31 +110,31 @@ public:
 		Class = From;
 		return *this;
 	}
-	
-	/** Dereference back into a UClass, does runtime type checking */
+
+	/** Dereference back into a UClass* or FFieldClass*, does runtime type checking. */
 	FORCEINLINE TClassType* operator*() const
 	{
-		if (!Class || !Class->IsChildOf(TClass::StaticClass()))
+		if (!Class || !Class->IsChildOf(T::StaticClass()))
 		{
 			return nullptr;
 		}
 		return Class;
 	}
 	
-	/** Dereference back into a UClass */
+	/** Dereference back into a UClass* or FFieldClass*, does runtime type checking. */
 	FORCEINLINE TClassType* Get() const
 	{
 		return **this;
 	}
 
-	/** Dereference back into a UClass */
+	/** Dereference back into a UClass* or FFieldClass*, does runtime type checking. */
 	FORCEINLINE TClassType* operator->() const
 	{
 		return **this;
 	}
 
-	/** Implicit conversion to UClass */
-	FORCEINLINE operator TClassType* () const
+	/** Implicit conversion to UClass* or FFieldClass*, does runtime type checking. */
+	FORCEINLINE operator TClassType*() const
 	{
 		return **this;
 	}
@@ -134,15 +144,15 @@ public:
 	 *
 	 * @return the CDO, or null if class is null
 	 */
-	FORCEINLINE TClass* GetDefaultObject() const
+	FORCEINLINE T* GetDefaultObject() const
 	{
 		TBaseType* Result = nullptr;
 		if (Class)
 		{
 			Result = Class->GetDefaultObject();
-			check(Result && Result->IsA(TClass::StaticClass()));
+			check(Result && Result->IsA(T::StaticClass()));
 		}
-		return (TClass*)Result;
+		return (T*)Result;
 	}
 
 	FORCEINLINE void SerializeTSubclassOf(FArchive& Ar)
@@ -165,7 +175,7 @@ public:
 #endif
 
 private:
-	TClassType* Class;
+	TClassType* Class = nullptr;
 };
 
 template <typename T>
