Improve the validation that we run when loading a package trailer off disk to avoid problems if we start to load corrupted data.

#rnx

## Loading the header
- The header tag is now zeroed so that we can be sure if we do end up with a valid tag, that it was loaded.
- Once we have loaded the stored sizes for the trailer we check against the total archive length to make sure that it all fits inside of the archive.
- We now make sure that the number of payloads was successfully loaded off disk and that the resulting lookup table could fit inside of the remaining data in the archive before we try to allocate the table and load it. This will prevent a garbage value from causing us to attempt to make a silly large allocation at this point.
- Each FLookupTableEntry is checked after loading to make sure that if the entry is for a locally stored payload, that the payload would be able to exist inside of the package trailer. Any obviously bad value (negative offset, or a payload length that would go past the end of the file) will be rejected.

### Loading the footer
- When loading the trailer backwards from the end of the archive we first load the footer.
- If we detect that the package tag is wrong, we assume that the archive is corrupted somehow and mark the archive as having an error. In theory the package fill loading will have already detected this earlier in the package loading life cycle but it doesn't hurt to double check.

diff --git a/Engine/Source/Runtime/CoreUObject/Private/UObject/PackageTrailer.cpp b/Engine/Source/Runtime/CoreUObject/Private/UObject/PackageTrailer.cpp
--- a/Engine/Source/Runtime/CoreUObject/Private/UObject/PackageTrailer.cpp
+++ b/Engine/Source/Runtime/CoreUObject/Private/UObject/PackageTrailer.cpp
@@ -81,6 +81,43 @@ Private::FLookupTableEntry* FindEntryInHeader(FPackageTrailer::FHeader& Header,
 	return Entry;
 }
 
+/** 
+ * Utility to check if a FLookupTableEntry seems valid for the trailer.
+ * @param	Entry					The entry we are trying to validate
+ * @param	PayloadDataEndPoint		The end point of the payloads stored locally. Anything
+ *									exceeding this point is not valud.
+ * 
+ * @return							True if it looks like the entry could exist in the 
+  *									current trailer, false if 
+ *									there are obvious problems.
+ */
+bool IsEntryValid(const Private::FLookupTableEntry& Entry, uint64 PayloadDataEndPoint)
+{
+	// If the entry is virtualized then we should assume that it is valid.
+	if (Entry.IsVirtualized())
+	{
+		return true;
+	}
+
+	// If not virtualized then the offset should not be negative
+	if (Entry.OffsetInFile < 0)
+	{
+		return false;
+	}
+
+	// PayloadDataEndPoint only applies if the entry is for a locally stored payload
+	if (Entry.IsLocal())
+	{
+		// A locally stored payload must end before PayloadDataEndPoint
+		if ((uint64)Entry.OffsetInFile + Entry.CompressedSize > PayloadDataEndPoint)
+		{
+			return false;
+		}
+	}
+
+	return true;
+}
+
 } //namespace
 
 namespace Private
@@ -434,10 +471,10 @@ void FPackageTrailerBuilder::RemoveDuplicateEntries()
 	for (auto Iter = LocalEntries.CreateIterator(); Iter; ++Iter)
 	{
 		if (VirtualizedEntries.Contains(Iter.Key()))
-		{	
-			UE_LOG(LogSerialization, Verbose, 
-				TEXT("Replacing localized payload '%s' with the virtualized version when building the package trailer for '%s'"), 
-				*LexToString(Iter.Key()) , 
+		{
+			UE_LOG(LogSerialization, Verbose,
+				TEXT("Replacing localized payload '%s' with the virtualized version when building the package trailer for '%s'"),
+				*LexToString(Iter.Key()),
 				*DebugContext);
 
 			Iter.RemoveCurrent();
@@ -448,9 +485,9 @@ void FPackageTrailerBuilder::RemoveDuplicateEntries()
 	{
 		if (VirtualizedEntries.Contains(Iter.Key()))
 		{
-			UE_LOG(LogSerialization, Verbose, 
-				TEXT("Replacing localized payload '%s' with the virtualized version when building the package trailer for '%s'"), 
-				*LexToString(Iter.Key()), 
+			UE_LOG(LogSerialization, Verbose,
+				TEXT("Replacing localized payload '%s' with the virtualized version when building the package trailer for '%s'"),
+				*LexToString(Iter.Key()),
 				*DebugContext);
 
 			Iter.RemoveCurrent();
@@ -465,7 +502,7 @@ bool FPackageTrailer::IsEnabled()
 		bool bEnabled = true;
 
 		FUsePackageTrailer()
-		{		
+		{
 			GConfig->GetBool(TEXT("Core.System"), TEXT("UsePackageTrailer"), bEnabled, GEngineIni);
 			UE_LOG(LogSerialization, Log, TEXT("UsePackageTrailer: '%s'"), bEnabled ? TEXT("true") : TEXT("false"));
 		}
@@ -488,7 +525,7 @@ bool FPackageTrailer::TryLoadFromPackage(const FPackagePath& PackagePath, FPacka
 	{
 		LogPackageOpenFailureMessage(PackagePath.GetDebugName());
 		return false;
-	}	
+	}
 }
 
 bool FPackageTrailer::TryLoadFromFile(const FString& Path, FPackageTrailer& OutTrailer)
@@ -519,6 +556,7 @@ bool FPackageTrailer::TryLoad(FArchive& Ar)
 
 	TrailerPositionInFile = Ar.Tell();
 
+	Header.Tag = 0; // Make sure we ignore anything previously loaded
 	Ar << Header.Tag;
 
 	// Make sure that we are parsing a valid FPackageTrailer
@@ -528,10 +566,16 @@ bool FPackageTrailer::TryLoad(FArchive& Ar)
 	}
 
 	Ar << Header.Version;
-
 	Ar << Header.HeaderLength;
 	Ar << Header.PayloadsDataLength;
 
+	// The header and the payloads should not exceed the remaining data in the archive
+	if (TrailerPositionInFile + Header.HeaderLength + Header.PayloadsDataLength > (uint64)Ar.TotalSize())
+	{
+		Ar.SetError();
+		return false;
+	}
+
 	EPayloadAccessMode LegacyAccessMode = EPayloadAccessMode::Local;
 	if (Header.Version < (uint32)EPackageTrailerVersion::ACCESS_PER_PAYLOAD)
 	{
@@ -540,6 +584,19 @@ bool FPackageTrailer::TryLoad(FArchive& Ar)
 
 	int32 NumPayloads = 0;
 	Ar << NumPayloads;
+	
+	// Make sure that we serialized a valid value for 'NumPayloads' before we try to allocate memory based off it
+	if (Ar.IsError())
+	{
+		return false; 
+	}
+
+	// Make sure that there is enough data remaining in the archive to load the look up table
+	if (NumPayloads * Private::FLookupTableEntry::SizeOnDisk > Header.HeaderLength)
+	{
+		Ar.SetError();
+		return false;
+	}
 
 	Header.PayloadLookupTable.Reserve(NumPayloads);
 
@@ -548,6 +605,13 @@ bool FPackageTrailer::TryLoad(FArchive& Ar)
 		Private::FLookupTableEntry& Entry = Header.PayloadLookupTable.AddDefaulted_GetRef();
 		Entry.Serialize(Ar, (EPackageTrailerVersion)Header.Version);
 
+		if (!IsEntryValid(Entry, Header.PayloadsDataLength))
+		{
+			// If an entry is invalid then something is likely up with the data so signal an error
+			Ar.SetError();
+			return false;
+		}
+
 		if (Header.Version < (uint32)EPackageTrailerVersion::ACCESS_PER_PAYLOAD)
 		{
 			Entry.AccessMode = Entry.OffsetInFile != INDEX_NONE ? LegacyAccessMode : EPayloadAccessMode::Virtualized;
@@ -575,9 +639,16 @@ bool FPackageTrailer::TryLoadBackwards(FArchive& Ar)
 	Ar << Footer.TrailerLength;
 	Ar << Footer.PackageTag;
 
-	// First check the package tag as this indicates if the file is corrupted or not
+	// First make sure that we were able to serialize everything
+	if (Ar.IsError())
+	{
+		return false;
+	}
+
+	// Then check the package tag as this indicates if the data/file is corrupted or not
 	if (Footer.PackageTag != PACKAGE_FILE_TAG)
 	{
+		Ar.SetError();
 		return false;
 	}
 
@@ -589,6 +660,12 @@ bool FPackageTrailer::TryLoadBackwards(FArchive& Ar)
 
 	Ar.Seek(Ar.Tell() - Footer.TrailerLength);
 
+	// Lastly make sure we were able to see to where we believe the header is before we try to load it
+	if (Ar.IsError())
+	{
+		return false;
+	}
+
 	return TryLoad(Ar);
 }
 
diff --git a/Engine/Source/Runtime/CoreUObject/Public/UObject/PackageTrailer.h b/Engine/Source/Runtime/CoreUObject/Public/UObject/PackageTrailer.h
--- a/Engine/Source/Runtime/CoreUObject/Public/UObject/PackageTrailer.h
+++ b/Engine/Source/Runtime/CoreUObject/Public/UObject/PackageTrailer.h
@@ -425,7 +425,8 @@ public:
 	 * for the trailer.
 	 * 
 	 * @param Ar	The archive to load the trailer from
-	 * @return		True if a valid trailer was found and was able to be loaded, otherwise false.
+	 * @return		True if a valid trailer was found and was able to be loaded, otherwise false. If the trailer was found but failed
+	 *				to load then the archive will be set to the error state.
 	 */
 	[[nodiscard]] bool TryLoad(FArchive& Ar);
 
@@ -434,7 +435,8 @@ public:
 	 * and so will attempt to read the footer first and use that to find the start of the trailer in order to read the header.
 	 * 
 	 * @param Ar	The archive to load the trailer from
-	 * @return		True if a valid trailer was found and was able to be loaded, otherwise false.
+	 * @return		True if a valid trailer was found and was able to be loaded, otherwise false. If the trailer was found but failed
+	 *				to load then the archive will be set to the error state.
 	 */
 	[[nodiscard]] bool TryLoadBackwards(FArchive& Ar);
 
