FLinkerLoad will now be marked as failed if the package trailer is actually corrupted and fails to serialize.

#rnx

- We currently don't require the existence of the packagetrailer (might be something we'd want to introduced based on a package version number now that the trailer system is no longer optional) and so were assuming that if the trailer failed to load that it just didn't exist.
- We now the archive to see if an error was raised during the serialization of the package trailer and if so we log an error, notify the user and mark the FLinkerLoad as failed so the loading of the package is halted.
- UE_ASSET_LOG is used to log the error in a format that Horde can understand.
- FMessageLog is used to notify the actual user but is supressed from the log to avoid duplicate entries.

diff --git a/Engine/Source/Runtime/CoreUObject/Private/UObject/LinkerLoad.cpp b/Engine/Source/Runtime/CoreUObject/Private/UObject/LinkerLoad.cpp
--- a/Engine/Source/Runtime/CoreUObject/Private/UObject/LinkerLoad.cpp
+++ b/Engine/Source/Runtime/CoreUObject/Private/UObject/LinkerLoad.cpp
@@ -1609,6 +1609,21 @@ FLinkerLoad::ELinkerStatus FLinkerLoad::SerializePackageTrailer()
 		PackageTrailer = MakeUnique<UE::FPackageTrailer>();
 		if (!PackageTrailer->TryLoad(*this))
 		{
+			// If the archive has an error then we found a package trailer but it failed to serialize
+			// correctly and we most likely have a problem with the file.
+			// If the load failed but the archive is fine then the package is just of an older format
+			// and there never was a package trailer to load.
+			if (IsError())
+			{
+				UE_ASSET_LOG(LogLinker, Error, PackagePath, TEXT("Package has a corrupted package trailer"));
+
+				FMessageLog("LoadErrors").SuppressLoggingToOutputLog(true)
+					.Error(FText::Format(NSLOCTEXT("Core", "LinkerLoad_CorruptTrailer", "Package {0} has a corrupted package trailer"),
+					FText::FromString(GetDebugName())));
+
+				return LINKER_Failed;
+			}
+
 			PackageTrailer.Reset();
 		}
 
