TSubclassOf split into UClass and FFieldClass variations, allowing them to be used with forward-declared types.
Fixed the one use of TSubclassOf<FField> in Sequencer.

diff --git a/Engine/Source/Editor/Sequencer/Public/AnimatedPropertyKey.h b/Engine/Source/Editor/Sequencer/Public/AnimatedPropertyKey.h
--- a/Engine/Source/Editor/Sequencer/Public/AnimatedPropertyKey.h
+++ b/Engine/Source/Editor/Sequencer/Public/AnimatedPropertyKey.h
@@ -4,7 +4,7 @@
 
 #include "CoreMinimal.h"
 #include "UObject/UnrealType.h"
-#include "Templates/SubclassOf.h"
+#include "Templates/SubclassOfField.h"
 
 struct FAnimatedPropertyKey
 {
@@ -84,7 +84,7 @@ struct FAnimatedPropertyKey
 		return Definition;
 	}
 
-	static FAnimatedPropertyKey FromPropertyType(TSubclassOf<FProperty> PropertyType)
+	static FAnimatedPropertyKey FromPropertyType(TSubclassOfField<FProperty> PropertyType)
 	{
 		FAnimatedPropertyKey Definition;
 		Definition.PropertyTypeName = PropertyType->GetFName();
diff --git a/Engine/Source/Runtime/CoreUObject/Public/Templates/SubclassOf.h b/Engine/Source/Runtime/CoreUObject/Public/Templates/SubclassOf.h
--- a/Engine/Source/Runtime/CoreUObject/Public/Templates/SubclassOf.h
+++ b/Engine/Source/Runtime/CoreUObject/Public/Templates/SubclassOf.h
@@ -4,8 +4,6 @@
 
 #include "CoreMinimal.h"
 #include "UObject/Class.h"
-#include "UObject/Field.h"
-#include "Templates/ChooseClass.h"
 
 #include <type_traits>
 
@@ -24,15 +22,11 @@ template <typename T> struct TIsTSubclassOf<      volatile TSubclassOf<T>> { enu
 template <typename T> struct TIsTSubclassOf<const volatile TSubclassOf<T>> { enum { Value = true }; };
 
 /**
- * Template to allow TClassTypes to be passed around with type safety 
+ * Template to allow UClass types to be passed around with type safety
  */
 template <typename T>
 class TSubclassOf
 {
-public:
-	typedef typename TChooseClass<TIsDerivedFrom<T, FField>::IsDerived, FFieldClass, UClass>::Result TClassType;
-	typedef typename TChooseClass<TIsDerivedFrom<T, FField>::IsDerived, FField, UObject>::Result TBaseType;
-
 private:
 	template <typename U>
 	friend class TSubclassOf;
@@ -45,21 +39,18 @@ public:
 	TSubclassOf& operator=(const TSubclassOf&) = default;
 	~TSubclassOf() = default;
 
-	/** Constructor that takes a UClass and does a runtime check to make sure this is a compatible class */
-	FORCEINLINE TSubclassOf(TClassType* From)
+	/** Constructor that takes a UClass*. */
+	FORCEINLINE TSubclassOf(UClass* From)
 		: Class(From)
 	{
 	}
 
-	/**
-	 * Construct from a UClass* or FFieldClass* (or something implicitly convertible to those)
-	 * if T is a UObject or a FField type respectively.
-	 */
+	/** Construct from a UClass* (or something implicitly convertible to it) */
 	template <
 		typename U,
 		std::enable_if_t<
 			!TIsTSubclassOf<std::decay_t<U>>::Value,
-			decltype(ImplicitConv<TClassType*>(std::declval<U>()))
+			decltype(ImplicitConv<UClass*>(std::declval<U>()))
 		>* = nullptr
 	>
 	FORCEINLINE TSubclassOf(U&& From)
@@ -72,8 +63,8 @@ public:
 		typename OtherT,
 		decltype(ImplicitConv<T*>((OtherT*)nullptr))* = nullptr
 	>
-	FORCEINLINE TSubclassOf(const TSubclassOf<OtherT>& From)
-		: Class(*From)
+	FORCEINLINE TSubclassOf(const TSubclassOf<OtherT>& Other)
+		: Class(Other.Class)
 	{
 	}
 
@@ -82,27 +73,25 @@ public:
 		typename OtherT,
 		decltype(ImplicitConv<T*>((OtherT*)nullptr))* = nullptr
 	>
-	FORCEINLINE TSubclassOf& operator=(const TSubclassOf<OtherT>& From)
+	FORCEINLINE TSubclassOf& operator=(const TSubclassOf<OtherT>& Other)
 	{
-		Class = *From;
+		Class = Other.Class;
 		return *this;
 	}
 
-	/** Assign from a UClass* or FFieldClass*. */
-	FORCEINLINE TSubclassOf& operator=(TClassType* From)
+	/** Assign from a UClass*. */
+	FORCEINLINE TSubclassOf& operator=(UClass* From)
 	{
 		Class = From;
 		return *this;
 	}
 
-	/**
-	 * Assign from a UClass* or FFieldClass* (or something implicitly convertible to those).
-	 */
+	/** Assign from a UClass* (or something implicitly convertible to it). */
 	template <
 		typename U,
 		std::enable_if_t<
 			!TIsTSubclassOf<std::decay_t<U>>::Value,
-			decltype(ImplicitConv<TClassType*>(std::declval<U>()))
+			decltype(ImplicitConv<UClass*>(std::declval<U>()))
 		>* = nullptr
 	>
 	FORCEINLINE TSubclassOf& operator=(U&& From)
@@ -111,8 +100,8 @@ public:
 		return *this;
 	}
 
-	/** Dereference back into a UClass* or FFieldClass*, does runtime type checking. */
-	FORCEINLINE TClassType* operator*() const
+	/** Dereference back into a UClass*, does runtime type checking. */
+	FORCEINLINE UClass* operator*() const
 	{
 		if (!Class || !Class->IsChildOf(T::StaticClass()))
 		{
@@ -120,21 +109,21 @@ public:
 		}
 		return Class;
 	}
-	
-	/** Dereference back into a UClass* or FFieldClass*, does runtime type checking. */
-	FORCEINLINE TClassType* Get() const
+
+	/** Dereference back into a UClass*, does runtime type checking. */
+	FORCEINLINE UClass* Get() const
 	{
 		return **this;
 	}
 
-	/** Dereference back into a UClass* or FFieldClass*, does runtime type checking. */
-	FORCEINLINE TClassType* operator->() const
+	/** Dereference back into a UClass*, does runtime type checking. */
+	FORCEINLINE UClass* operator->() const
 	{
 		return **this;
 	}
 
-	/** Implicit conversion to UClass* or FFieldClass*, does runtime type checking. */
-	FORCEINLINE operator TClassType*() const
+	/** Implicit conversion to UClass*, does runtime type checking. */
+	FORCEINLINE operator UClass*() const
 	{
 		return **this;
 	}
@@ -146,7 +135,7 @@ public:
 	 */
 	FORCEINLINE T* GetDefaultObject() const
 	{
-		TBaseType* Result = nullptr;
+		UObject* Result = nullptr;
 		if (Class)
 		{
 			Result = Class->GetDefaultObject();
@@ -155,7 +144,7 @@ public:
 		return (T*)Result;
 	}
 
-	FORCEINLINE void SerializeTSubclassOf(FArchive& Ar)
+	FORCEINLINE void Serialize(FArchive& Ar)
 	{
 		Ar << Class;
 	}
@@ -175,13 +164,13 @@ public:
 #endif
 
 private:
-	TClassType* Class = nullptr;
+	UClass* Class = nullptr;
 };
 
 template <typename T>
 FArchive& operator<<(FArchive& Ar, TSubclassOf<T>& SubclassOf)
 {
-	SubclassOf.SerializeTSubclassOf(Ar);
+	SubclassOf.Serialize(Ar);
 	return Ar;
 }
 
diff --git a/Engine/Source/Runtime/CoreUObject/Public/Templates/SubclassOfField.h b/Engine/Source/Runtime/CoreUObject/Public/Templates/SubclassOfField.h
new file mode 100644
--- /dev/null
+++ b/Engine/Source/Runtime/CoreUObject/Public/Templates/SubclassOfField.h
@@ -0,0 +1,167 @@
+// Copyright Epic Games, Inc. All Rights Reserved.
+
+#pragma once
+
+#include "CoreMinimal.h"
+#include "UObject/Field.h"
+
+#include <type_traits>
+
+template <typename T>
+class TSubclassOfField;
+
+template <typename T>
+struct TIsTSubclassOfField
+{
+	enum { Value = false };
+};
+
+template <typename T> struct TIsTSubclassOfField<               TSubclassOfField<T>> { enum { Value = true }; };
+template <typename T> struct TIsTSubclassOfField<const          TSubclassOfField<T>> { enum { Value = true }; };
+template <typename T> struct TIsTSubclassOfField<      volatile TSubclassOfField<T>> { enum { Value = true }; };
+template <typename T> struct TIsTSubclassOfField<const volatile TSubclassOfField<T>> { enum { Value = true }; };
+
+/**
+ * Template to allow FFieldClass types to be passed around with type safety
+ */
+template <typename T>
+class TSubclassOfField
+{
+private:
+	template <typename U>
+	friend class TSubclassOfField;
+
+public:
+	TSubclassOfField() = default;
+	TSubclassOfField(TSubclassOfField&&) = default;
+	TSubclassOfField(const TSubclassOfField&) = default;
+	TSubclassOfField& operator=(TSubclassOfField&&) = default;
+	TSubclassOfField& operator=(const TSubclassOfField&) = default;
+	~TSubclassOfField() = default;
+
+	/** Constructor that takes a FFieldClass*. */
+	FORCEINLINE TSubclassOfField(FFieldClass* From)
+		: Class(From)
+	{
+	}
+
+	/** Construct from a FFieldClass* (or something implicitly convertible to it) */
+	template <
+		typename U,
+		std::enable_if_t<
+			!TIsTSubclassOf<std::decay_t<U>>::Value,
+			decltype(ImplicitConv<FFieldClass*>(std::declval<U>()))
+		>* = nullptr
+	>
+	FORCEINLINE TSubclassOfField(U&& From)
+		: Class(From)
+	{
+	}
+
+	/** Construct from another TSubclassOfField, only if types are compatible */
+	template <
+		typename OtherT,
+		decltype(ImplicitConv<T*>((OtherT*)nullptr))* = nullptr
+	>
+	FORCEINLINE TSubclassOfField(const TSubclassOfField<OtherT>& Other)
+		: Class(Other.Class)
+	{
+	}
+
+	/** Assign from another TSubclassOfField, only if types are compatible */
+	template <
+		typename OtherT,
+		decltype(ImplicitConv<T*>((OtherT*)nullptr))* = nullptr
+	>
+	FORCEINLINE TSubclassOfField& operator=(const TSubclassOfField<OtherT>& Other)
+	{
+		Class = Other.Class;
+		return *this;
+	}
+
+	/** Assign from a FFieldClass*. */
+	FORCEINLINE TSubclassOfField& operator=(FFieldClass* From)
+	{
+		Class = From;
+		return *this;
+	}
+
+	/** Assign from a FFieldClass* (or something implicitly convertible to it). */
+	template <
+		typename U,
+		std::enable_if_t<
+			!TIsTSubclassOf<std::decay_t<U>>::Value,
+			decltype(ImplicitConv<FFieldClass*>(std::declval<U>()))
+		>* = nullptr
+	>
+	FORCEINLINE TSubclassOfField& operator=(U&& From)
+	{
+		Class = From;
+		return *this;
+	}
+
+	/** Dereference back into a FFieldClass*, does runtime type checking. */
+	FORCEINLINE FFieldClass* operator*() const
+	{
+		if (!Class || !Class->IsChildOf(T::StaticClass()))
+		{
+			return nullptr;
+		}
+		return Class;
+	}
+
+	/** Dereference back into a FFieldClass*, does runtime type checking. */
+	FORCEINLINE FFieldClass* Get() const
+	{
+		return **this;
+	}
+
+	/** Dereference back into a FFieldClass*, does runtime type checking. */
+	FORCEINLINE FFieldClass* operator->() const
+	{
+		return **this;
+	}
+
+	/** Implicit conversion to FFieldClass*, does runtime type checking. */
+	FORCEINLINE operator FFieldClass*() const
+	{
+		return **this;
+	}
+
+	/**
+	 * Get the CDO if we are referencing a valid class
+	 *
+	 * @return the CDO, or null if class is null
+	 */
+	FORCEINLINE T* GetDefaultObject() const
+	{
+		FField* Result = nullptr;
+		if (Class)
+		{
+			Result = Class->GetDefaultObject();
+			check(Result && Result->IsA(T::StaticClass()));
+		}
+		return (T*)Result;
+	}
+
+	FORCEINLINE void Serialize(FArchive& Ar)
+	{
+		Ar << Class;
+	}
+
+	friend uint32 GetTypeHash(const TSubclassOfField& SubclassOf)
+	{
+		return GetTypeHash(SubclassOf.Class);
+	}
+
+private:
+	FFieldClass* Class = nullptr;
+};
+
+template <typename T>
+FArchive& operator<<(FArchive& Ar, TSubclassOfField<T>& SubclassOf)
+{
+	SubclassOf.Serialize(Ar);
+	return Ar;
+}
+
