Non-unity fix.

diff --git a/Engine/Source/Runtime/CoreUObject/Public/Templates/SubclassOfField.h b/Engine/Source/Runtime/CoreUObject/Public/Templates/SubclassOfField.h
--- a/Engine/Source/Runtime/CoreUObject/Public/Templates/SubclassOfField.h
+++ b/Engine/Source/Runtime/CoreUObject/Public/Templates/SubclassOfField.h
@@ -49,7 +49,7 @@ public:
 	template <
 		typename U,
 		std::enable_if_t<
-			!TIsTSubclassOf<std::decay_t<U>>::Value,
+			!TIsTSubclassOfField<std::decay_t<U>>::Value,
 			decltype(ImplicitConv<FFieldClass*>(std::declval<U>()))
 		>* = nullptr
 	>
@@ -90,7 +90,7 @@ public:
 	template <
 		typename U,
 		std::enable_if_t<
-			!TIsTSubclassOf<std::decay_t<U>>::Value,
+			!TIsTSubclassOfField<std::decay_t<U>>::Value,
 			decltype(ImplicitConv<FFieldClass*>(std::declval<U>()))
 		>* = nullptr
 	>
