Removed PPF_ExportCpp, which is a vestige of the blueprint C++ compiler.

diff --git a/Engine/Plugins/Experimental/CommonUI/Source/CommonUI/Public/UITag.h b/Engine/Plugins/Experimental/CommonUI/Source/CommonUI/Public/UITag.h
--- a/Engine/Plugins/Experimental/CommonUI/Source/CommonUI/Public/UITag.h
+++ b/Engine/Plugins/Experimental/CommonUI/Source/CommonUI/Public/UITag.h
@@ -51,12 +51,8 @@ class TTypedTagStaticImpl
 
 	static bool ExportTextItem(const TagT& Tag, FString& ValueStr, int32 PortFlags)
 	{
-		if (0 == (PortFlags & EPropertyPortFlags::PPF_ExportCpp))
-		{
-			ValueStr += Tag.GetTagName().ToString();
-			return true;
-		}
-		return false;
+		ValueStr += Tag.GetTagName().ToString();
+		return true;
 	}
 
 	static TagT TryConvert(FGameplayTag VanillaTag, bool bChecked)
diff --git a/Engine/Plugins/Runtime/DataRegistry/Source/DataRegistry/Private/DataRegistryId.cpp b/Engine/Plugins/Runtime/DataRegistry/Source/DataRegistry/Private/DataRegistryId.cpp
--- a/Engine/Plugins/Runtime/DataRegistry/Source/DataRegistry/Private/DataRegistryId.cpp
+++ b/Engine/Plugins/Runtime/DataRegistry/Source/DataRegistry/Private/DataRegistryId.cpp
@@ -13,11 +13,7 @@ const FDataRegistryType FDataRegistryType::CustomContextType = FDataRegistryType
 
 bool FDataRegistryType::ExportTextItem(FString& ValueStr, FDataRegistryType const& DefaultValue, UObject* Parent, int32 PortFlags, UObject* ExportRootScope) const
 {
-	if (0 != (PortFlags & PPF_ExportCpp))
-	{
-		ValueStr += FString::Printf(TEXT("FDataRegistryType(TEXT(\"%s\"))"), *ToString().ReplaceCharWithEscapedChar());
-	}
-	else if (!(PortFlags & PPF_Delimited))
+	if (!(PortFlags & PPF_Delimited))
 	{
 		ValueStr += ToString();
 	}
@@ -68,11 +64,7 @@ bool FDataRegistryType::SerializeFromMismatchedTag(struct FPropertyTag const& Ta
 
 bool FDataRegistryId::ExportTextItem(FString& ValueStr, FDataRegistryId const& DefaultValue, UObject* Parent, int32 PortFlags, UObject* ExportRootScope) const
 {
-	if (0 != (PortFlags & PPF_ExportCpp))
-	{
-		ValueStr += FString::Printf(TEXT("FDataRegistryId(TEXT(\"%s\"))"), *ToString().ReplaceCharWithEscapedChar());
-	}
-	else if (!(PortFlags & PPF_Delimited))
+	if (!(PortFlags & PPF_Delimited))
 	{
 		ValueStr += ToString();
 	}
diff --git a/Engine/Source/Runtime/Core/Private/Math/UnrealMath.cpp b/Engine/Source/Runtime/Core/Private/Math/UnrealMath.cpp
--- a/Engine/Source/Runtime/Core/Private/Math/UnrealMath.cpp
+++ b/Engine/Source/Runtime/Core/Private/Math/UnrealMath.cpp
@@ -3241,16 +3241,6 @@ void FMath::ApplyScaleToFloat(float& Dst, const FVector& DeltaScale, float Magni
 	Dst = FMath::Max( 0.0f, Dst );
 }
 
-bool FRandomStream::ExportTextItem(FString& ValueStr, FRandomStream const& DefaultValue, class UObject* Parent, int32 PortFlags, class UObject* ExportRootScope) const
-{
-	if (0 != (PortFlags & EPropertyPortFlags::PPF_ExportCpp))
-	{
-		ValueStr += FString::Printf(TEXT("FRandomStream(%i)"), DefaultValue.GetInitialSeed());
-		return true;
-	}
-	return false;
-}
-
 // Implementation of 1D, 2D and 3D Perlin noise based on Ken Perlin's improved version https://mrl.nyu.edu/~perlin/noise/
 // (See Random3.tps for additional third party software info.)
 namespace FMathPerlinHelpers
diff --git a/Engine/Source/Runtime/Core/Private/Misc/DateTime.cpp b/Engine/Source/Runtime/Core/Private/Misc/DateTime.cpp
--- a/Engine/Source/Runtime/Core/Private/Misc/DateTime.cpp
+++ b/Engine/Source/Runtime/Core/Private/Misc/DateTime.cpp
@@ -73,12 +73,6 @@ FDateTime::FDateTime(int32 Year, int32 Month, int32 Day, int32 Hour, int32 Minut
 
 bool FDateTime::ExportTextItem(FString& ValueStr, FDateTime const& DefaultValue, UObject* Parent, int32 PortFlags, UObject* ExportRootScope) const
 {
-	if (0 != (PortFlags & EPropertyPortFlags::PPF_ExportCpp))
-	{
-		ValueStr += FString::Printf(TEXT("FDateTime(0x%016X)"), Ticks);
-		return true;
-	}
-
 	ValueStr += ToString();
 
 	return true;
diff --git a/Engine/Source/Runtime/Core/Private/Misc/Guid.cpp b/Engine/Source/Runtime/Core/Private/Misc/Guid.cpp
--- a/Engine/Source/Runtime/Core/Private/Misc/Guid.cpp
+++ b/Engine/Source/Runtime/Core/Private/Misc/Guid.cpp
@@ -24,11 +24,6 @@
 
 bool FGuid::ExportTextItem(FString& ValueStr, FGuid const& DefaultValue, UObject* Parent, int32 PortFlags, class UObject* ExportRootScope) const
 {
-	if (0 != (PortFlags & EPropertyPortFlags::PPF_ExportCpp))
-	{
-		return false;
-	}
-
 	ValueStr += ToString();
 
 	return true;
diff --git a/Engine/Source/Runtime/Core/Private/Misc/Timespan.cpp b/Engine/Source/Runtime/Core/Private/Misc/Timespan.cpp
--- a/Engine/Source/Runtime/Core/Private/Misc/Timespan.cpp
+++ b/Engine/Source/Runtime/Core/Private/Misc/Timespan.cpp
@@ -16,12 +16,6 @@
 
 bool FTimespan::ExportTextItem(FString& ValueStr, FTimespan const& DefaultValue, UObject* Parent, int32 PortFlags, UObject* ExportRootScope) const
 {
-	if ((PortFlags & EPropertyPortFlags::PPF_ExportCpp) != 0)
-	{
-		ValueStr += FString::Printf(TEXT("FTimespan(0x%016X)"), Ticks);
-		return true;
-	}
-
 	ValueStr += ToString(TEXT("%D.%h:%m:%s.%n"));
 
 	return true;
diff --git a/Engine/Source/Runtime/Core/Public/Math/RandomStream.h b/Engine/Source/Runtime/Core/Public/Math/RandomStream.h
--- a/Engine/Source/Runtime/Core/Public/Math/RandomStream.h
+++ b/Engine/Source/Runtime/Core/Public/Math/RandomStream.h
@@ -322,19 +322,6 @@ public:
 		}
 	}
 
-	/**
-	 * Exports the RandomStreams value to a string.
-	 *
-	 * @param ValueStr Will hold the string value.
-	 * @param DefaultValue The default value.
-	 * @param Parent Not used.
-	 * @param PortFlags Not used.
-	 * @param ExportRootScope Not used.
-	 * @return true on success, false otherwise.
-	 * @see ImportTextItem
-	 */
-	CORE_API bool ExportTextItem(FString& ValueStr, FRandomStream const& DefaultValue, class UObject* Parent, int32 PortFlags, class UObject* ExportRootScope) const;
-
 	/**
 	 * Get a textual representation of the RandomStream.
 	 *
diff --git a/Engine/Source/Runtime/Core/Public/UObject/PropertyPortFlags.h b/Engine/Source/Runtime/Core/Public/UObject/PropertyPortFlags.h
--- a/Engine/Source/Runtime/Core/Public/UObject/PropertyPortFlags.h
+++ b/Engine/Source/Runtime/Core/Public/UObject/PropertyPortFlags.h
@@ -106,8 +106,7 @@ enum EPropertyPortFlags
 	/** Ignores CPF_Deprecated flag */
 	PPF_UseDeprecatedProperties		= 0x08000000,
 
-	/** Export in C++ form */
-	PPF_ExportCpp					= 0x10000000,
+	//								= 0x10000000,
 
 	/** Ignores CPF_SkipSerialization flag when using tagged serialization */
 	PPF_ForceTaggedSerialization	= 0x20000000,
diff --git a/Engine/Source/Runtime/CoreUObject/Private/UObject/Class.cpp b/Engine/Source/Runtime/CoreUObject/Private/UObject/Class.cpp
--- a/Engine/Source/Runtime/CoreUObject/Private/UObject/Class.cpp
+++ b/Engine/Source/Runtime/CoreUObject/Private/UObject/Class.cpp
@@ -2934,11 +2934,6 @@ void UScriptStruct::ExportText(FString& ValueStr, const void* Value, const void*
 		}
 	}
 
-	if (0 != (PortFlags & PPF_ExportCpp))
-	{
-		return;
-	}
-
 	int32 Count = 0;
 
 	// if this struct is configured to be serialized as a unit, it must be exported as a unit as well.
diff --git a/Engine/Source/Runtime/CoreUObject/Private/UObject/EnumProperty.cpp b/Engine/Source/Runtime/CoreUObject/Private/UObject/EnumProperty.cpp
--- a/Engine/Source/Runtime/CoreUObject/Private/UObject/EnumProperty.cpp
+++ b/Engine/Source/Runtime/CoreUObject/Private/UObject/EnumProperty.cpp
@@ -282,28 +282,6 @@ void FEnumProperty::ExportText_Internal(FString& ValueStr, const void* PropertyV
 		PropertyValue = PointerToValuePtr(PropertyValueOrContainer, PropertyPointerType);
 	}
 
-	if (PortFlags & PPF_ExportCpp)
-	{
-		const int64 ActualValue = LocalUnderlyingProp->GetSignedIntPropertyValue(PropertyValue);
-		const int64 MaxValue = Enum->GetMaxEnumValue();
-		const int64 GoodValue = Enum->IsValidEnumValue(ActualValue) ? ActualValue : MaxValue;
-		const bool bNonNativeEnum = Enum->GetClass() != UEnum::StaticClass();
-		ensure(!bNonNativeEnum || Enum->CppType.IsEmpty());
-		const FString FullyQualifiedEnumName = bNonNativeEnum ? ::UnicodeToCPPIdentifier(Enum->GetName(), false, TEXT("E__"))
-			: (Enum->CppType.IsEmpty() ? Enum->GetName() : Enum->CppType);
-		if (GoodValue == MaxValue)
-		{
-			// not all native enums have Max value declared
-			ValueStr += FString::Printf(TEXT("(%s)(%llu)"), *FullyQualifiedEnumName, ActualValue);
-		}
-		else
-		{
-			ValueStr += FString::Printf(TEXT("%s::%s"), *FullyQualifiedEnumName,
-				*Enum->GetNameStringByValue(GoodValue));
-		}
-		return;
-	}
-
 	if (PortFlags & PPF_ConsoleVariable)
 	{
 		UnderlyingProp->ExportText_Internal(ValueStr, PropertyValue, EPropertyPointerType::Direct, DefaultValue, Parent, PortFlags, ExportRootScope);
diff --git a/Engine/Source/Runtime/CoreUObject/Private/UObject/FieldPathProperty.cpp b/Engine/Source/Runtime/CoreUObject/Private/UObject/FieldPathProperty.cpp
--- a/Engine/Source/Runtime/CoreUObject/Private/UObject/FieldPathProperty.cpp
+++ b/Engine/Source/Runtime/CoreUObject/Private/UObject/FieldPathProperty.cpp
@@ -98,13 +98,7 @@ void FFieldPathProperty::ExportText_Internal( FString& ValueStr, const void* Pro
 		Value = GetPropertyValue(PointerToValuePtr(PropertyValueOrContainer, PropertyPointerType));
 	}
 
-	if (PortFlags & PPF_ExportCpp)
-	{
-		ValueStr += TEXT("TEXT(\"");
-		ValueStr += Value.ToString();
-		ValueStr += TEXT("\")");
-	}
-	else if (PortFlags & PPF_PropertyWindow)
+	if (PortFlags & PPF_PropertyWindow)
 	{
 		if (PortFlags & PPF_Delimited)
 		{
diff --git a/Engine/Source/Runtime/CoreUObject/Private/UObject/PrimaryAssetId.cpp b/Engine/Source/Runtime/CoreUObject/Private/UObject/PrimaryAssetId.cpp
--- a/Engine/Source/Runtime/CoreUObject/Private/UObject/PrimaryAssetId.cpp
+++ b/Engine/Source/Runtime/CoreUObject/Private/UObject/PrimaryAssetId.cpp
@@ -8,11 +8,7 @@
 
 bool FPrimaryAssetType::ExportTextItem(FString& ValueStr, FPrimaryAssetType const& DefaultValue, UObject* Parent, int32 PortFlags, UObject* ExportRootScope) const
 {
-	if (0 != (PortFlags & PPF_ExportCpp))
-	{
-		ValueStr += FString::Printf(TEXT("FPrimaryAssetType(TEXT(\"%s\"))"), *ToString().ReplaceCharWithEscapedChar());
-	}
-	else if (!(PortFlags & PPF_Delimited))
+	if (!(PortFlags & PPF_Delimited))
 	{
 		ValueStr += ToString();
 	}
@@ -85,11 +81,7 @@ FPrimaryAssetId FPrimaryAssetId::ParseTypeAndName(FName TypeAndName)
 
 bool FPrimaryAssetId::ExportTextItem(FString& ValueStr, FPrimaryAssetId const& DefaultValue, UObject* Parent, int32 PortFlags, UObject* ExportRootScope) const
 {
-	if (0 != (PortFlags & PPF_ExportCpp))
-	{
-		ValueStr += FString::Printf(TEXT("FPrimaryAssetId(TEXT(\"%s\"))"), *ToString().ReplaceCharWithEscapedChar());
-	}
-	else if (!(PortFlags & PPF_Delimited))
+	if (!(PortFlags & PPF_Delimited))
 	{
 		ValueStr += ToString();
 	}
diff --git a/Engine/Source/Runtime/CoreUObject/Private/UObject/Property.cpp b/Engine/Source/Runtime/CoreUObject/Private/UObject/Property.cpp
--- a/Engine/Source/Runtime/CoreUObject/Private/UObject/Property.cpp
+++ b/Engine/Source/Runtime/CoreUObject/Private/UObject/Property.cpp
@@ -447,7 +447,6 @@ struct TStructOpsTypeTraits<FRandomStream> : public TStructOpsTypeTraitsBase2<FR
 {
 	enum 
 	{
-		WithExportTextItem = true,
 		WithNoInitConstructor = true,
 		WithZeroConstructor = true,
 	};
@@ -2061,16 +2060,6 @@ UPropertyWrapper* FProperty::GetUPropertyWrapper()
 }
 #endif //  WITH_EDITORONLY_DATA
 
-void FFloatProperty::ExportText_Internal(FString& ValueStr, const void* PropertyValueOrContainer, EPropertyPointerType PropertyPointerType, const void* DefaultValue, UObject* Parent, int32 PortFlags, UObject* ExportRootScope) const
-{
-	Super::ExportText_Internal(ValueStr, PropertyValueOrContainer, PropertyPointerType, DefaultValue, Parent, PortFlags, ExportRootScope);
-
-	if (0 != (PortFlags & PPF_ExportCpp))
-	{
-		ValueStr += TEXT("f");
-	}
-}
-
 
 FProperty* UStruct::FindPropertyByName(FName InName) const
 {
diff --git a/Engine/Source/Runtime/CoreUObject/Private/UObject/PropertyArray.cpp b/Engine/Source/Runtime/CoreUObject/Private/UObject/PropertyArray.cpp
--- a/Engine/Source/Runtime/CoreUObject/Private/UObject/PropertyArray.cpp
+++ b/Engine/Source/Runtime/CoreUObject/Private/UObject/PropertyArray.cpp
@@ -463,14 +463,6 @@ void FArrayProperty::ExportText_Internal( FString& ValueStr, const void* Contain
 {
 	checkSlow(Inner);
 
-	if (0 != (PortFlags & PPF_ExportCpp))
-	{
-		FString ExtendedTypeText;
-		FString TypeText = GetCPPType(&ExtendedTypeText, EPropertyExportCPPFlags::CPPF_BlueprintCppBackend);
-		ValueStr += FString::Printf(TEXT("%s%s()"), *TypeText, *ExtendedTypeText);
-		return;
-	}
-
 	uint8* TempArrayStorage = nullptr;
 	void* PropertyValuePtr = nullptr;
 	if (PropertyPointerType == EPropertyPointerType::Container && HasGetter())
diff --git a/Engine/Source/Runtime/CoreUObject/Private/UObject/PropertyBaseObject.cpp b/Engine/Source/Runtime/CoreUObject/Private/UObject/PropertyBaseObject.cpp
--- a/Engine/Source/Runtime/CoreUObject/Private/UObject/PropertyBaseObject.cpp
+++ b/Engine/Source/Runtime/CoreUObject/Private/UObject/PropertyBaseObject.cpp
@@ -243,17 +243,6 @@ void FObjectPropertyBase::ExportText_Internal( FString& ValueStr, const void* Pr
 		Temp = GetObjectPtrPropertyValue(PointerToValuePtr(PropertyValueOrContainer, PropertyPointerType));
 	}
 
-	if (0 != (PortFlags & PPF_ExportCpp))
-	{
-		ValueStr += Temp
-			? FString::Printf(TEXT("LoadObject<%s%s>(nullptr, TEXT(\"%s\"))")
-				, PropertyClass->GetPrefixCPP()
-				, *PropertyClass->GetName()
-				, *(Temp->GetPathName().ReplaceCharWithEscapedChar()))
-			: TEXT("nullptr");
-		return;
-	}
-
 	if (!Temp)
 	{
 		ValueStr += TEXT("None");
diff --git a/Engine/Source/Runtime/CoreUObject/Private/UObject/PropertyBool.cpp b/Engine/Source/Runtime/CoreUObject/Private/UObject/PropertyBool.cpp
--- a/Engine/Source/Runtime/CoreUObject/Private/UObject/PropertyBool.cpp
+++ b/Engine/Source/Runtime/CoreUObject/Private/UObject/PropertyBool.cpp
@@ -355,15 +355,7 @@ void FBoolProperty::ExportText_Internal( FString& ValueStr, const void* Containe
 		LocalByteValue = *((uint8*)PointerToValuePtr(ContainerOrPropertyPtr, PropertyPointerType) + ByteOffset);
 		bValue = 0 != (LocalByteValue & FieldMask);
 	}
-	const TCHAR* Temp = nullptr;
-	if (0 != (PortFlags & PPF_ExportCpp))
-	{
-		Temp = (bValue ? TEXT("true") : TEXT("false"));
-	}
-	else
-	{
-		Temp = (bValue ? TEXT("True") : TEXT("False"));
-	}
+	const TCHAR* Temp = (bValue ? TEXT("True") : TEXT("False"));
 	ValueStr += FString::Printf( TEXT("%s"), Temp );
 }
 const TCHAR* FBoolProperty::ImportText_Internal( const TCHAR* Buffer, void* ContainerOrPropertyPtr, EPropertyPointerType PropertyPointerType, UObject* Parent, int32 PortFlags, FOutputDevice* ErrorText ) const
diff --git a/Engine/Source/Runtime/CoreUObject/Private/UObject/PropertyByte.cpp b/Engine/Source/Runtime/CoreUObject/Private/UObject/PropertyByte.cpp
--- a/Engine/Source/Runtime/CoreUObject/Private/UObject/PropertyByte.cpp
+++ b/Engine/Source/Runtime/CoreUObject/Private/UObject/PropertyByte.cpp
@@ -364,45 +364,6 @@ void FByteProperty::AppendSchemaHash(FBlake3& Builder, bool bSkipEditorOnly) con
 
 void FByteProperty::ExportText_Internal( FString& ValueStr, const void* PropertyValueOrContainer, EPropertyPointerType PropertyPointerType, const void* DefaultValue, UObject* Parent, int32 PortFlags, UObject* ExportRootScope ) const
 {
-	if (0 != (PortFlags & PPF_ExportCpp))
-	{
-		if (Enum)
-		{
-			int64 ActualValue = 0;
-			if (PropertyPointerType == EPropertyPointerType::Container && HasGetter())
-			{
-				uint8 Value = 0;
-				GetValue_InContainer(PropertyValueOrContainer, &Value);
-				ActualValue = Value;
-			}
-			else
-			{
-				ActualValue = *(const uint8*)PointerToValuePtr(PropertyValueOrContainer, PropertyPointerType);
-			}
-			const int64 MaxValue = Enum->GetMaxEnumValue();
-			const int64 GoodValue = Enum->IsValidEnumValue(ActualValue) ? ActualValue : MaxValue;
-			const bool bNonNativeEnum = Enum->GetClass() != UEnum::StaticClass();
-			ensure(!bNonNativeEnum || Enum->CppType.IsEmpty());
-			const FString FullyQualifiedEnumName = bNonNativeEnum ? ::UnicodeToCPPIdentifier(Enum->GetName(), false, TEXT("E__"))
-				: (Enum->CppType.IsEmpty() ? Enum->GetName() : Enum->CppType);
-			if (GoodValue == MaxValue)
-			{
-				// not all native enums have Max value declared
-				ValueStr += FString::Printf(TEXT("(%s)(%d)"), *FullyQualifiedEnumName, ActualValue);
-			}
-			else
-			{
-				ValueStr += FString::Printf(TEXT("%s::%s"), *FullyQualifiedEnumName,
-					*Enum->GetNameStringByValue(GoodValue));
-			}
-		}
-		else
-		{
-			Super::ExportText_Internal(ValueStr, PropertyValueOrContainer, PropertyPointerType, DefaultValue, Parent, PortFlags, ExportRootScope);
-		}
-		return;
-	}
-
 	if( Enum && (PortFlags & PPF_ConsoleVariable) == 0 )
 	{
 		int64 ActualValue = 0;
diff --git a/Engine/Source/Runtime/CoreUObject/Private/UObject/PropertyDelegate.cpp b/Engine/Source/Runtime/CoreUObject/Private/UObject/PropertyDelegate.cpp
--- a/Engine/Source/Runtime/CoreUObject/Private/UObject/PropertyDelegate.cpp
+++ b/Engine/Source/Runtime/CoreUObject/Private/UObject/PropertyDelegate.cpp
@@ -142,12 +142,6 @@ FString FDelegateProperty::GetCPPTypeForwardDeclaration() const
 
 void FDelegateProperty::ExportText_Internal( FString& ValueStr, const void* PropertyValueOrContainer, EPropertyPointerType PropertyPointerType, const void* DefaultValue, UObject* Parent, int32 PortFlags, UObject* ExportRootScope ) const
 {
-	if (0 != (PortFlags & PPF_ExportCpp))
-	{
-		ValueStr += TEXT("{}");
-		return;
-	}
-
 	auto ExportDelegateAsText = [&ValueStr](FScriptDelegate* ScriptDelegate)
 	{
 		check(ScriptDelegate != NULL);
diff --git a/Engine/Source/Runtime/CoreUObject/Private/UObject/PropertyInterface.cpp b/Engine/Source/Runtime/CoreUObject/Private/UObject/PropertyInterface.cpp
--- a/Engine/Source/Runtime/CoreUObject/Private/UObject/PropertyInterface.cpp
+++ b/Engine/Source/Runtime/CoreUObject/Private/UObject/PropertyInterface.cpp
@@ -194,17 +194,6 @@ void FInterfaceProperty::ExportText_Internal( FString& ValueStr, const void* Pro
 		Temp = InterfaceValue->GetObject();
 	}
 
-	if (0 != (PortFlags & PPF_ExportCpp))
-	{
-		const FString GetObjectStr = Temp
-			? FString::Printf(TEXT("LoadObject<UObject>(nullptr, TEXT(\"%s\"))"), *Temp->GetPathName().ReplaceCharWithEscapedChar())
-			: TEXT("");
-		ValueStr += FString::Printf(TEXT("TScriptInterface<I%s>(%s)")
-			, (InterfaceClass ? *InterfaceClass->GetName() : TEXT("Interface"))
-			, *GetObjectStr);
-		return;
-	}
-
 	if( Temp != NULL )
 	{
 		bool bExportFullyQualified = true;
diff --git a/Engine/Source/Runtime/CoreUObject/Private/UObject/PropertyMap.cpp b/Engine/Source/Runtime/CoreUObject/Private/UObject/PropertyMap.cpp
--- a/Engine/Source/Runtime/CoreUObject/Private/UObject/PropertyMap.cpp
+++ b/Engine/Source/Runtime/CoreUObject/Private/UObject/PropertyMap.cpp
@@ -635,12 +635,6 @@ FString FMapProperty::GetCPPMacroType( FString& ExtendedTypeText ) const
 
 void FMapProperty::ExportText_Internal(FString& ValueStr, const void* ContainerOrPropertyPtr, EPropertyPointerType PropertyPointerType, const void* DefaultValue, UObject* Parent, int32 PortFlags, UObject* ExportRootScope) const
 {
-	if (0 != (PortFlags & PPF_ExportCpp))
-	{
-		ValueStr += TEXT("{}");
-		return;
-	}
-
 	checkSlow(KeyProp);
 	checkSlow(ValueProp);
 
diff --git a/Engine/Source/Runtime/CoreUObject/Private/UObject/PropertyMulticastDelegate.cpp b/Engine/Source/Runtime/CoreUObject/Private/UObject/PropertyMulticastDelegate.cpp
--- a/Engine/Source/Runtime/CoreUObject/Private/UObject/PropertyMulticastDelegate.cpp
+++ b/Engine/Source/Runtime/CoreUObject/Private/UObject/PropertyMulticastDelegate.cpp
@@ -181,12 +181,6 @@ FString FMulticastDelegateProperty::GetCPPTypeForwardDeclaration() const
 
 void FMulticastDelegateProperty::ExportText_Internal( FString& ValueStr, const void* PropertyValueOrContainer, EPropertyPointerType PropertyPointerType, const void* DefaultValue, UObject* Parent, int32 PortFlags, UObject* ExportRootScope ) const
 {
-	if (0 != (PortFlags & PPF_ExportCpp))
-	{
-		ValueStr += TEXT("{}");
-		return;
-	}
-
 	const FMulticastScriptDelegate::FInvocationList* InvocationList = nullptr;
 	
 	if (PropertyPointerType == EPropertyPointerType::Container && HasGetter())
diff --git a/Engine/Source/Runtime/CoreUObject/Private/UObject/PropertyName.cpp b/Engine/Source/Runtime/CoreUObject/Private/UObject/PropertyName.cpp
--- a/Engine/Source/Runtime/CoreUObject/Private/UObject/PropertyName.cpp
+++ b/Engine/Source/Runtime/CoreUObject/Private/UObject/PropertyName.cpp
@@ -27,13 +27,7 @@ void FNameProperty::ExportText_Internal( FString& ValueStr, const void* Property
 	{
 		Temp = *(FName*)PointerToValuePtr(PropertyValueOrContainer, PropertyPointerType);
 	}
-	if (0 != (PortFlags & PPF_ExportCpp))
-	{
-		ValueStr += (Temp == NAME_None) 
-			? TEXT("FName()") 
-			: FString::Printf(TEXT("FName(TEXT(\"%s\"))"), *(Temp.ToString().ReplaceCharWithEscapedChar()));
-	}
-	else if( !(PortFlags & PPF_Delimited) )
+	if( !(PortFlags & PPF_Delimited) )
 	{
 		ValueStr += Temp.ToString();
 	}
diff --git a/Engine/Source/Runtime/CoreUObject/Private/UObject/PropertySet.cpp b/Engine/Source/Runtime/CoreUObject/Private/UObject/PropertySet.cpp
--- a/Engine/Source/Runtime/CoreUObject/Private/UObject/PropertySet.cpp
+++ b/Engine/Source/Runtime/CoreUObject/Private/UObject/PropertySet.cpp
@@ -549,12 +549,6 @@ FString FSetProperty::GetCPPTypeForwardDeclaration() const
 
 void FSetProperty::ExportText_Internal(FString& ValueStr, const void* ContainerOrPropertyPtr, EPropertyPointerType PropertyPointerType, const void* DefaultValue, UObject* Parent, int32 PortFlags, UObject* ExportRootScope) const
 {
-	if (0 != (PortFlags & PPF_ExportCpp))
-	{
-		ValueStr += TEXT("{}");
-		return;
-	}
-
 	checkSlow(ElementProp);
 
 	uint8* TempSetStorage = nullptr;
diff --git a/Engine/Source/Runtime/CoreUObject/Private/UObject/PropertySoftObjectPtr.cpp b/Engine/Source/Runtime/CoreUObject/Private/UObject/PropertySoftObjectPtr.cpp
--- a/Engine/Source/Runtime/CoreUObject/Private/UObject/PropertySoftObjectPtr.cpp
+++ b/Engine/Source/Runtime/CoreUObject/Private/UObject/PropertySoftObjectPtr.cpp
@@ -124,12 +124,6 @@ void FSoftObjectProperty::ExportText_Internal( FString& ValueStr, const void* Pr
 		SoftObjectPath = SoftObjectPtr.GetUniqueID();
 	}
 
-	if (0 != (PortFlags & PPF_ExportCpp))
-	{
-		ValueStr += FString::Printf(TEXT("FSoftObjectPath(TEXT(\"%s\"))"), *SoftObjectPath.ToString().ReplaceCharWithEscapedChar());
-		return;
-	}
-
 	SoftObjectPath.ExportTextItem(ValueStr, SoftObjectPath, Parent, PortFlags, ExportRootScope);
 }
 
diff --git a/Engine/Source/Runtime/CoreUObject/Private/UObject/PropertyStr.cpp b/Engine/Source/Runtime/CoreUObject/Private/UObject/PropertyStr.cpp
--- a/Engine/Source/Runtime/CoreUObject/Private/UObject/PropertyStr.cpp
+++ b/Engine/Source/Runtime/CoreUObject/Private/UObject/PropertyStr.cpp
@@ -38,126 +38,6 @@ FString FStrProperty::GetCPPTypeForwardDeclaration() const
 	return FString();
 }
 
-// Necessary to fix Compiler Error C2026 and C1091
-FString FStrProperty::ExportCppHardcodedText(const FString& InSource, const FString& Indent)
-{
-	constexpr FAsciiSet EscapableCharacters("\\\"\n\r\t");
-
-	auto EstimateExportedStringLength = [&EscapableCharacters](const UTF16CHAR* InStr)
-	{
-		int32 EstimatedLen = 0;
-		while (const UTF16CHAR Char = *InStr++)
-		{
-			if (EscapableCharacters.Contains(Char))
-			{
-				// Exported escaped
-				EstimatedLen += 2;
-			}
-			else if (Char > 0x7f)
-			{
-				// Exported as a UTF-16 sequence
-				EstimatedLen += 4;
-			}
-			else
-			{
-				// Exported verbatim
-				EstimatedLen += 1;
-			}
-		}
-		return EstimatedLen;
-	};
-
-	const int32 PreferredLineSize = 256;
-	const int32 LinesPerString = 16;
-
-	// Note: This is a no-op on platforms that are using a 16-bit TCHAR
-	FTCHARToUTF16 UTF16SourceString(*InSource, InSource.Len() + 1); // include the null terminator
-
-	const bool bUseSubStrings = EstimateExportedStringLength(UTF16SourceString.Get()) > (LinesPerString * PreferredLineSize);
-	int32 LineNum = 0;
-
-	TStringBuilder<1024> Result;
-
-	if (bUseSubStrings)
-	{
-		Result << TEXT("*(FString(");
-	}
-
-	const UTF16CHAR* SourceBegin = UTF16SourceString.Get();
-	const UTF16CHAR* SourceIt = SourceBegin;
-	do
-	{
-		if (SourceIt > SourceBegin)
-		{
-			Result << TEXT("\n");
-			Result << Indent;
-		}
-
-		++LineNum;
-		if (bUseSubStrings && (LineNum % LinesPerString) == 0)
-		{
-			Result << TEXT(") + FString(");
-		}
-
-		Result << TEXT("TEXT(\"");
-		{
-			int32 ResultStartLen = Result.Len();
-			while (*SourceIt && (Result.Len() - ResultStartLen) < PreferredLineSize)
-			{
-				if (EscapableCharacters.Contains(*SourceIt))
-				{
-					const TCHAR CharToEscape = (TCHAR)*SourceIt++;
-
-					Result << TEXT('\\');
-					switch (CharToEscape)
-					{
-					case TEXT('\n'):
-						Result << TEXT('n');
-						break;
-					case TEXT('\r'):
-						Result << TEXT('r');
-						break;
-					case TEXT('\t'):
-						Result << TEXT('t');
-						break;
-					default:
-						Result << CharToEscape;
-						break;
-					}
-				}
-				else if (*SourceIt > 0x7f)
-				{
-					// If this character is part of a surrogate pair, then combine them and write them as a single UTF-32 sequence
-					// Otherwise just write out the character as a UTF-16 sequence
-					if (StringConv::IsHighSurrogate(*SourceIt) && StringConv::IsLowSurrogate(*(SourceIt + 1)))
-					{
-						const UTF32CHAR Codepoint = StringConv::EncodeSurrogate(*SourceIt, *(SourceIt + 1));
-						Result.Appendf(TEXT("\\U%08x"), Codepoint);
-						SourceIt += 2;
-					}
-					else
-					{
-						Result.Appendf(TEXT("\\u%04x"), *SourceIt++);
-					}
-				}
-				else
-				{
-					Result << (TCHAR)*SourceIt++;
-				}
-			}
-		}
-		Result << TEXT("\")");
-	}
-	while (*SourceIt);
-
-	if (bUseSubStrings)
-	{
-		Result << TEXT("))");
-	}
-
-	return Result.ToString();
-}
-
 void FStrProperty::ExportText_Internal( FString& ValueStr, const void* PropertyValueOrContainer, EPropertyPointerType PropertyPointerType, const void* DefaultValue, UObject* Parent, int32 PortFlags, UObject* ExportRootScope ) const
 {
 	FString StringValue;
@@ -170,11 +50,7 @@ void FStrProperty::ExportText_Internal( FString& ValueStr, const void* PropertyV
 		StringValue = *(FString*)PointerToValuePtr(PropertyValueOrContainer, PropertyPointerType);
 	}
 
-	if (0 != (PortFlags & PPF_ExportCpp))
-	{
-		ValueStr += FString::Printf(TEXT("FString(%s)"), *ExportCppHardcodedText(StringValue, FString()));
-	}
-	else if (!(PortFlags & PPF_Delimited))
+	if (!(PortFlags & PPF_Delimited))
 	{
 		ValueStr += StringValue;
 	}
diff --git a/Engine/Source/Runtime/CoreUObject/Private/UObject/SoftObjectPath.cpp b/Engine/Source/Runtime/CoreUObject/Private/UObject/SoftObjectPath.cpp
--- a/Engine/Source/Runtime/CoreUObject/Private/UObject/SoftObjectPath.cpp
+++ b/Engine/Source/Runtime/CoreUObject/Private/UObject/SoftObjectPath.cpp
@@ -327,11 +327,6 @@ bool FSoftObjectPath::operator==(FSoftObjectPath const& Other) const
 
 bool FSoftObjectPath::ExportTextItem(FString& ValueStr, FSoftObjectPath const& DefaultValue, UObject* Parent, int32 PortFlags, UObject* ExportRootScope) const
 {
-	if (0 != (PortFlags & EPropertyPortFlags::PPF_ExportCpp))
-	{
-		return false;
-	}
-
 	if (!IsNull())
 	{
 		// Fixup any redirectors
diff --git a/Engine/Source/Runtime/CoreUObject/Private/UObject/TextProperty.cpp b/Engine/Source/Runtime/CoreUObject/Private/UObject/TextProperty.cpp
--- a/Engine/Source/Runtime/CoreUObject/Private/UObject/TextProperty.cpp
+++ b/Engine/Source/Runtime/CoreUObject/Private/UObject/TextProperty.cpp
@@ -126,11 +126,7 @@ void FTextProperty::ExportText_Internal( FString& ValueStr, const void* Property
 		TextValue = GetPropertyValue(PointerToValuePtr(PropertyValueOrContainer, PropertyPointerType));
 	}
 
-	if (PortFlags & PPF_ExportCpp)
-	{
-		ValueStr += GenerateCppCodeForTextValue(TextValue, FString());
-	}
-	else if (PortFlags & PPF_PropertyWindow)
+	if (PortFlags & PPF_PropertyWindow)
 	{
 		if (PortFlags & PPF_Delimited)
 		{
@@ -205,48 +201,6 @@ const TCHAR* FTextProperty::ImportText_Internal( const TCHAR* Buffer, void* Cont
 	return Result;
 }
 
-FString FTextProperty::GenerateCppCodeForTextValue(const FText& InValue, const FString& Indent)
-{
-	FString CppCode;
-
-	if (InValue.IsEmpty())
-	{
-		CppCode += TEXT("FText::GetEmpty()");
-	}
-	else if (InValue.IsCultureInvariant())
-	{
-		const FString& StringValue = FTextInspector::GetDisplayString(InValue);
-
-		// Produces FText::AsCultureInvariant(TEXT("..."))
-		CppCode += TEXT("FText::AsCultureInvariant(\n");
-		CppCode += FStrProperty::ExportCppHardcodedText(StringValue, Indent + TEXT("\t\t"));
-		CppCode += TEXT("\t)");
-	}
-	else
-	{
-		FString ExportedText;
-		FTextStringHelper::WriteToBuffer(ExportedText, InValue);
-
-		if (FTextStringHelper::IsComplexText(*ExportedText))
-		{
-			// Produces FTextStringHelper::CreateFromBuffer(TEXT("..."))
-			CppCode += TEXT("FTextStringHelper::CreateFromBuffer(\n");
-			CppCode += FStrProperty::ExportCppHardcodedText(ExportedText, Indent + TEXT("\t\t"));
-			CppCode += Indent;
-			CppCode += TEXT("\t)");
-		}
-		else
-		{
-			// Produces FText::FromString(TEXT("..."))
-			CppCode += TEXT("FText::FromString(\n");
-			CppCode += FStrProperty::ExportCppHardcodedText(ExportedText, Indent + TEXT("\t\t"));
-			CppCode += TEXT("\t)");
-		}
-	}
-
-	return CppCode;
-}
-
 FString FTextProperty::GetCPPTypeForwardDeclaration() const
 {
 	return FString();
diff --git a/Engine/Source/Runtime/CoreUObject/Private/UObject/TopLevelAssetPath.cpp b/Engine/Source/Runtime/CoreUObject/Private/UObject/TopLevelAssetPath.cpp
--- a/Engine/Source/Runtime/CoreUObject/Private/UObject/TopLevelAssetPath.cpp
+++ b/Engine/Source/Runtime/CoreUObject/Private/UObject/TopLevelAssetPath.cpp
@@ -175,11 +175,6 @@ bool FTopLevelAssetPath::TrySetPath(FAnsiStringView Path)
 
 bool FTopLevelAssetPath::ExportTextItem(FString& ValueStr, FTopLevelAssetPath const& DefaultValue, UObject* Parent, int32 PortFlags, UObject* ExportRootScope) const
 {
-	if (0 != (PortFlags & EPropertyPortFlags::PPF_ExportCpp))
-	{
-		return false;
-	}
-
 	if (!IsNull())
 	{
 		if (PortFlags & PPF_Delimited)
diff --git a/Engine/Source/Runtime/CoreUObject/Public/UObject/TextProperty.h b/Engine/Source/Runtime/CoreUObject/Public/UObject/TextProperty.h
--- a/Engine/Source/Runtime/CoreUObject/Public/UObject/TextProperty.h
+++ b/Engine/Source/Runtime/CoreUObject/Public/UObject/TextProperty.h
@@ -68,8 +68,5 @@ public:
 	virtual FString GetCPPTypeForwardDeclaration() const override;
 	// End of FProperty interface
 
-	/** Generate the correct C++ code for the given text value */
-	static FString GenerateCppCodeForTextValue(const FText& InValue, const FString& Indent);
-
 	static bool Identical_Implementation(const FText& A, const FText& B, uint32 PortFlags);
 };
diff --git a/Engine/Source/Runtime/CoreUObject/Public/UObject/UnrealType.h b/Engine/Source/Runtime/CoreUObject/Public/UObject/UnrealType.h
--- a/Engine/Source/Runtime/CoreUObject/Public/UObject/UnrealType.h
+++ b/Engine/Source/Runtime/CoreUObject/Public/UObject/UnrealType.h
@@ -2369,11 +2369,6 @@ PRAGMA_ENABLE_DEPRECATION_WARNINGS
 	{
 	}
 #endif // WITH_EDITORONLY_DATA
-
-	// FProperty interface
-protected:
-	virtual void ExportText_Internal(FString& ValueStr, const void* PropertyValueOrContainer, EPropertyPointerType PropertyPointerType, const void* DefaultValue, UObject* Parent, int32 PortFlags, UObject* ExportRootScope) const override;
-	// End of FProperty interface
 };
 
 /*-----------------------------------------------------------------------------
@@ -3636,9 +3631,6 @@ public:
 	virtual FString GetCPPTypeForwardDeclaration() const override;
 	uint32 GetValueTypeHashInternal(const void* Src) const override;
 	// End of FProperty interface
-
-	// Necessary to fix Compiler Error C2026
-	static FString ExportCppHardcodedText(const FString& InSource, const FString& Indent);
 };
 
 /*-----------------------------------------------------------------------------
diff --git a/Engine/Source/Runtime/Engine/Private/OnlineReplStructs.cpp b/Engine/Source/Runtime/Engine/Private/OnlineReplStructs.cpp
--- a/Engine/Source/Runtime/Engine/Private/OnlineReplStructs.cpp
+++ b/Engine/Source/Runtime/Engine/Private/OnlineReplStructs.cpp
@@ -463,11 +463,6 @@ bool FUniqueNetIdRepl::Serialize(FArchive& Ar)
 
 bool FUniqueNetIdRepl::ExportTextItem(FString& ValueStr, FUniqueNetIdRepl const& DefaultValue, UObject* Parent, int32 PortFlags, UObject* ExportRootScope) const
 {
-	if (0 != (PortFlags & EPropertyPortFlags::PPF_ExportCpp))
-	{
-		return false;
-	}
-
 	if (IsValid())
 	{
 		if (IsV1())
diff --git a/Engine/Source/Runtime/IOS/IOSRuntimeSettings/Classes/IOSRuntimeSettings.h b/Engine/Source/Runtime/IOS/IOSRuntimeSettings/Classes/IOSRuntimeSettings.h
--- a/Engine/Source/Runtime/IOS/IOSRuntimeSettings/Classes/IOSRuntimeSettings.h
+++ b/Engine/Source/Runtime/IOS/IOSRuntimeSettings/Classes/IOSRuntimeSettings.h
@@ -85,11 +85,6 @@ struct FIOSBuildResourceFilePath
 	 */
 	bool ExportTextItem(FString& ValueStr, FIOSBuildResourceFilePath const& DefaultValue, UObject* Parent, int32 PortFlags, UObject* ExportRootScope) const
 	{
-		if (0 != (PortFlags & EPropertyPortFlags::PPF_ExportCpp))
-		{
-			return false;
-		}
-
 		ValueStr += FilePath;
 		return true;
 	}
@@ -138,11 +133,6 @@ struct FIOSBuildResourceDirectory
 	 */
 	bool ExportTextItem(FString& ValueStr, FIOSBuildResourceDirectory const& DefaultValue, UObject* Parent, int32 PortFlags, UObject* ExportRootScope) const
 	{
-		if (0 != (PortFlags & EPropertyPortFlags::PPF_ExportCpp))
-		{
-			return false;
-		}
-
 		ValueStr += Path;
 		return true;
 	}
diff --git a/Engine/Source/Runtime/InputCore/Private/InputCoreTypes.cpp b/Engine/Source/Runtime/InputCore/Private/InputCoreTypes.cpp
--- a/Engine/Source/Runtime/InputCore/Private/InputCoreTypes.cpp
+++ b/Engine/Source/Runtime/InputCore/Private/InputCoreTypes.cpp
@@ -1458,12 +1458,6 @@ bool FKey::SerializeFromMismatchedTag(struct FPropertyTag const& Tag, FStructure
 
 bool FKey::ExportTextItem(FString& ValueStr, FKey const& DefaultValue, UObject* Parent, int32 PortFlags, UObject* ExportRootScope) const
 {
-	if (0 != (PortFlags & EPropertyPortFlags::PPF_ExportCpp))
-	{
-		ValueStr += FString::Printf(TEXT("FKey(TEXT(\"%s\"))"), *KeyName.ToString());
-		return true;
-	}
-
 	ValueStr += KeyName.ToString();
 	return true;
 }
diff --git a/Engine/Source/Runtime/MovieScene/Private/MovieSceneFrameMigration.cpp b/Engine/Source/Runtime/MovieScene/Private/MovieSceneFrameMigration.cpp
--- a/Engine/Source/Runtime/MovieScene/Private/MovieSceneFrameMigration.cpp
+++ b/Engine/Source/Runtime/MovieScene/Private/MovieSceneFrameMigration.cpp
@@ -57,11 +57,6 @@ bool FMovieSceneFrameRange::SerializeFromMismatchedTag(FPropertyTag const& Tag,
 
 bool FMovieSceneFrameRange::ExportTextItem(FString& ValueStr, FMovieSceneFrameRange const& DefaultValue, UObject* Parent, int32 PortFlags, UObject* ExportRootScope) const
 {
-	if (0 != (PortFlags & EPropertyPortFlags::PPF_ExportCpp))
-	{
-		return false;
-	}
-
 	FString String;
 	if (Value.GetLowerBound().IsOpen())
 	{
