AsyncLoading2: Improve handling of renamed loaded objects and packages.
Remove the special case callback OnLeakedPackageRename that was executed on the GT but was mostly modifying ALT state.
Package and object renames can theoretically happen at any time on the GT, and will modify the UPackage/UObject FName and the UObject global hash tables directly.
Check for package renames on the ALT every time we start loading or importing a package.
When a GT package rename is detected on the ALT, then update the ALT FGlobalImportStore and the UPackage PackageId accordingly.
Object renames are only detected on the ALT when a package is reloaded and we overwrite the old public export entry with the new export (similar to how to they are kept around in linkers).

#rnx

diff --git a/Engine/Source/Runtime/CoreUObject/Private/Serialization/AsyncLoading2.cpp b/Engine/Source/Runtime/CoreUObject/Private/Serialization/AsyncLoading2.cpp
--- a/Engine/Source/Runtime/CoreUObject/Private/Serialization/AsyncLoading2.cpp
+++ b/Engine/Source/Runtime/CoreUObject/Private/Serialization/AsyncLoading2.cpp
@@ -544,9 +544,7 @@ private:
 	class FPublicExportMap
 	{
 	public:
-		FPublicExportMap()
-		{
-		}
+		FPublicExportMap() = default;
 
 		FPublicExportMap(const FPublicExportMap&) = delete;
 
@@ -732,8 +730,9 @@ private:
 		int32 SingleItemValue = -1;
 	};
 
-	UPackage* Package = nullptr;
 	FPublicExportMap PublicExportMap;
+	FName OriginalPackageName;
+	int32 PackageObjectIndex = -1;
 	int32 RefCount = 0;
 	bool bAreAllPublicExportsLoaded = false;
 	bool bIsMissing = false;
@@ -741,9 +740,7 @@ private:
 	bool bHasBeenLoadedDebug = false;
 
 public:
-	FLoadedPackageRef()
-	{
-	}
+	FLoadedPackageRef() = default;
 
 	FLoadedPackageRef(const FLoadedPackageRef& Other) = delete;
 
@@ -758,20 +755,23 @@ public:
 		return RefCount;
 	}
 
+	inline FName GetOriginalPackageName() const
+	{
+		return OriginalPackageName;
+	}
+
+	inline bool HasPackage() const
+	{
+		return PackageObjectIndex >= 0;
+	}
+
 	inline UPackage* GetPackage() const
 	{
-#if DO_CHECK
-		if (Package)
-		{
-			check(!bIsMissing);
-			check(!Package->IsUnreachable());
-		}
-		else
+		if (HasPackage())
 		{
-			check(!bAreAllPublicExportsLoaded);
+			return static_cast<UPackage*>(GUObjectArray.IndexToObject(PackageObjectIndex)->Object);
 		}
-#endif
-		return Package;
+		return nullptr;
 	}
 
 	inline void SetPackage(UPackage* InPackage)
@@ -779,20 +779,39 @@ public:
 		check(!bAreAllPublicExportsLoaded);
 		check(!bIsMissing);
 		check(!bHasFailed);
-		check(!Package);
-		Package = InPackage;
+		check(!HasPackage());
+		PackageObjectIndex = GUObjectArray.ObjectToIndex(InPackage);
+		OriginalPackageName = InPackage->GetFName();
+	}
+
+	void RemoveCompletedRenamedPackage()
+	{
+		check(RefCount == 0);
+		*this = FLoadedPackageRef();
+	}
+
+	void ReplaceReferencedRenamedPackage(UPackage* NewPackage)
+	{
+		// keep RefCount and PublicExportMap while resetting all other state,
+		// the public exports will be replaced one by one from StoreGlobalObject
+		bAreAllPublicExportsLoaded = false;
+		bIsMissing = false;
+		bHasFailed = false;
+		bHasBeenLoadedDebug = false;
+		PackageObjectIndex = GUObjectArray.ObjectToIndex(NewPackage);
+		OriginalPackageName = NewPackage->GetFName();
 	}
 
 	inline bool AreAllPublicExportsLoaded() const
 	{
-		return bAreAllPublicExportsLoaded;
+		return bAreAllPublicExportsLoaded && OriginalPackageName == GetPackage()->GetFName();
 	}
 
 	inline void SetAllPublicExportsLoaded()
 	{
 		check(!bIsMissing);
 		check(!bHasFailed);
-		check(Package);
+		check(HasPackage());
 		bIsMissing = false;
 		bAreAllPublicExportsLoaded = true;
 		bHasBeenLoadedDebug = true;
@@ -801,7 +820,7 @@ public:
 	inline void SetIsMissingPackage()
 	{
 		check(!bAreAllPublicExportsLoaded);
-		check(!Package);
+		check(!HasPackage());
 		bIsMissing = true;
 		bAreAllPublicExportsLoaded = false;
 	}
@@ -835,7 +854,7 @@ public:
 	void RemovePublicExport(uint64 ExportHash)
 	{
 		check(!bIsMissing);
-		check(Package);
+		check(HasPackage());
 		bAreAllPublicExportsLoaded = false;
 		PublicExportMap.Remove(ExportHash);
 	}
@@ -847,6 +866,7 @@ public:
 
 	void PinPublicExportsForGC()
 	{
+		UPackage* Package = GetPackage();
 		UE_ASYNC_UPACKAGE_DEBUG(Package);
 
 		if (GUObjectArray.IsDisregardForGC(Package))
@@ -860,6 +880,7 @@ public:
 
 	void UnpinPublicExportsForGC()
 	{
+		UPackage* Package = GetPackage();
 		UE_ASYNC_UPACKAGE_DEBUG(Package);
 
 		if (GUObjectArray.IsDisregardForGC(Package))
@@ -916,10 +937,11 @@ public:
 		return Packages.Find(PackageId);
 	}
 
-	inline FLoadedPackageRef& FindPackageRefChecked(FPackageId PackageId)
+	inline FLoadedPackageRef& FindPackageRefChecked(FPackageId PackageId, FName Name = FName())
 	{
 		FLoadedPackageRef* PackageRef = FindPackageRef(PackageId);
-		UE_CLOG(!PackageRef, LogStreaming, Fatal, TEXT("FindPackageRefChecked: Package with id 0x%llX has been deleted"), PackageId.ValueForDebugging());
+		UE_CLOG(!PackageRef, LogStreaming, Fatal, TEXT("FindPackageRefChecked: Package %s (0x%llX) has been deleted"),
+			*Name.ToString(), PackageId.ValueForDebugging());
 		return *PackageRef;
 	}
 
@@ -932,7 +954,7 @@ public:
 		if (PackageRef.RefCount == 0)
 		{
 #if WITH_EDITOR
-			if (!PackageRef.Package && !PackageNameIfKnown.IsNone())
+			if (!PackageRef.HasPackage() && !PackageNameIfKnown.IsNone())
 			{
 				UPackage* FoundPackage = FindObjectFast<UPackage>(nullptr, PackageNameIfKnown);
 				if (FoundPackage)
@@ -947,9 +969,22 @@ public:
 				}
 			}
 #endif
-			if (PackageRef.Package)
+			if (UPackage* Package = PackageRef.GetPackage())
 			{
-				PackageRef.PinPublicExportsForGC();
+				if (PackageRef.GetOriginalPackageName() == Package->GetFName())
+				{
+					PackageRef.PinPublicExportsForGC();
+				}
+				else
+				{
+					UE_LOG(LogStreaming, Display,
+						TEXT("FGlobalImportStore:AddPackageRef: Dropping renamed package %s before reloading %s (0x%llX)"),
+						*Package->GetName(),
+						*PackageRef.GetOriginalPackageName().ToString(),
+						Package->GetPackageId().ValueForDebugging());
+
+					RemoveCompletedRenamedPackage(PackageRef);
+				}
 			}
 		}
 		++PackageRef.RefCount;
@@ -968,7 +1003,7 @@ public:
 			TEXT("LoadedPackageRef from None (0x%llX) to %s (0x%llX) should not have been released when the package is not complete.")
 			TEXT("RefCount=%d, AreAllExportsLoaded=%d, IsMissing=%d, HasFailed=%d, HasBeenLoaded=%d"),
 			FromPackageId.ValueForDebugging(),
-			PackageRef.Package ? *PackageRef.Package->GetName() : TEXT("None"),
+			*PackageRef.GetOriginalPackageName().ToString(),
 			PackageId.Value(),
 			PackageRef.RefCount,
 			PackageRef.bAreAllPublicExportsLoaded,
@@ -986,7 +1021,7 @@ public:
 		}
 #endif
 		// is this the last reference to a loaded package?
-		if (PackageRef.RefCount == 0 && PackageRef.Package)
+		if (PackageRef.RefCount == 0 && PackageRef.HasPackage())
 		{
 			PackageRef.UnpinPublicExportsForGC();
 		}
@@ -1007,12 +1042,37 @@ public:
 		}
 	}
 
-	void RemovePackage(UPackage* Package)
+	void RemoveCompletedRenamedPackage(FLoadedPackageRef& PackageRef)
 	{
-		TRACE_CPUPROFILER_EVENT_SCOPE(RemovePackage);
-		VerifyPackageForRemoval(Package);
-		FPackageId PackageId = Package->GetPackageId();
-		RemovePackage(PackageId, Package->GetFName());
+		UPackage* OldPackage = PackageRef.GetPackage();
+		UE_ASYNC_UPACKAGE_DEBUG(OldPackage);
+
+		if (GVerifyUnreachableObjects)
+		{
+			VerifyPackageForRemoval(PackageRef);
+		}
+		for (int32 ObjectIndex : PackageRef.GetPublicExportObjectIndices())
+		{
+			if (ObjectIndex >= 0)
+			{
+				ObjectIndexToPublicExport.Remove(ObjectIndex);
+			}
+		}
+		PackageRef.RemoveCompletedRenamedPackage();
+		// Update PackageId to prevent us from removing our updated package ref from GC
+		OldPackage->SetPackageId(FPackageId::FromName(OldPackage->GetFName()));
+	}
+
+	void ReplaceReferencedRenamedPackage(FLoadedPackageRef& PackageRef, UPackage* NewPackage)
+	{
+		UPackage* OldPackage = PackageRef.GetPackage();
+		UE_ASYNC_UPACKAGE_DEBUG(OldPackage);
+
+		PackageRef.ReplaceReferencedRenamedPackage(NewPackage);
+		// Clear LoaderImport so GC may destroy this package
+		OldPackage->AtomicallyClearInternalFlags(EInternalObjectFlags::LoaderImport);
+		// Update PackageId to prevent us from removing our updated package ref from GC
+		OldPackage->SetPackageId(FPackageId::FromName(OldPackage->GetFName()));
 	}
 
 	void RemovePackages(const FUnreachableObjects& ObjectsToRemove)
@@ -1036,15 +1096,6 @@ public:
 		bool bRemoved = Packages.RemoveAndCopyValue(PackageId, PackageRef);
 		if (bRemoved)
 		{
-			if (PackageRef.GetRefCount() > 0)
-			{
-				UE_LOG(LogStreaming, Fatal,
-					TEXT("FGlobalImportStore::RemovePackage: %s (0x%llX) - ")
-					TEXT("Package removed while still being referenced, RefCount %d > 0."),
-					*PackageName.ToString(),
-					PackageId.ValueForDebugging(),
-					PackageRef.GetRefCount());
-			}
 			for (int32 ObjectIndex : PackageRef.GetPublicExportObjectIndices())
 			{
 				if (ObjectIndex >= 0)
@@ -1128,32 +1179,24 @@ public:
 		}
 	}
 
-	void VerifyPackageForRemoval(UPackage* Package)
+	void VerifyPackageForRemoval(FLoadedPackageRef& PackageRef)
 	{
-		check(Package && Package->CanBeImported());
-
+		UPackage* Package = PackageRef.GetPackage();
 		FPackageId PackageId = Package->GetPackageId();
-		FLoadedPackageRef* PackageRef = FindPackageRef(PackageId);
-
-		UE_CLOG(!PackageRef, LogStreaming, Fatal,
-			TEXT("FGlobalImportStore::VerifyPackageForRemoval: %s (0x%llX) - ")
-			TEXT("A loaded package that can be imported by other packages is missing in GlobalImportStore."),
-			*Package->GetName(),
-			PackageId.ValueForDebugging());
 
-		UE_CLOG(PackageRef->GetRefCount() > 0, LogStreaming, Fatal,
+		UE_CLOG(PackageRef.GetRefCount() > 0, LogStreaming, Fatal,
 			TEXT("FGlobalImportStore::VerifyPackageForRemoval: %s (0x%llX) - ")
 			TEXT("Package removed while still being referenced, RefCount %d > 0."),
 			*Package->GetName(),
 			PackageId.ValueForDebugging(),
-			PackageRef->GetRefCount());
+			PackageRef.GetRefCount());
 
-		for (int32 ObjectIndex : PackageRef->GetPublicExportObjectIndices())
+		for (int32 ObjectIndex : PackageRef.GetPublicExportObjectIndices())
 		{
 			if (ObjectIndex >= 0)
 			{
 				UObject* Object = static_cast<UObject*>(GUObjectArray.IndexToObject(ObjectIndex)->Object);
-				ensureMsgf(!Object->HasAnyInternalFlags(EInternalObjectFlags::LoaderImport),
+				ensureMsgf(!Object->HasAnyInternalFlags(EInternalObjectFlags::LoaderImport) || GUObjectArray.IsDisregardForGC(Object),
 						TEXT("FGlobalImportStore::VerifyPackageForRemoval: The loaded public export object '%s' with flags (ObjectFlags=%x, InternalObjectFlags=%x) and id 0x%llX:0x%llX is probably still referenced by the loader."),
 						*Object->GetFullName(),
 						Object->GetFlags(),
@@ -1220,7 +1263,7 @@ public:
 
 			UE_LOG(LogStreaming, Warning,
 				TEXT("FGlobalImportStore::StoreGlobalObject: The constructed public export object '%s' with index %d and id 0x%llX:0x%llX collides with object '%s' (ObjectFlags=%X, InternalObjectFlags=%x) with index %d in GlobalImportStore. ")
-				TEXT("The existing object will be replaced since it was most likely renamed within the package after it was loaded."),
+				TEXT("The existing object will be replaced since it or its package was most likely renamed after it was loaded the first time."),
 				Object ? *Object->GetFullName() : TEXT("null"),
 				ObjectIndex,
 				Key.GetPackageId().Value(), Key.GetExportHash(),
@@ -1411,9 +1454,9 @@ struct FPackageImportStore
 	}
 
 public:
-	void AddImportedPackageReference(FPackageId ImportedPackageId, FName PackageNameIfKnown)
+	FLoadedPackageRef& AddImportedPackageReference(FPackageId ImportedPackageId, FName PackageNameIfKnown)
 	{
-		GlobalImportStore.AddPackageRef(ImportedPackageId, PackageNameIfKnown);
+		return GlobalImportStore.AddPackageRef(ImportedPackageId, PackageNameIfKnown);
 	}
 
 	void AddPackageReference(const FAsyncPackageDesc2& Desc)
@@ -2982,8 +3025,6 @@ private:
 		return bDidSomething;
 	}
 
-	void OnLeakedPackageRename(UPackage* Package);
-
 	void OnPreGarbageCollect();
 
 	void CollectUnreachableObjects(TArrayView<FUObjectItem*> UnreachableObjectItems, FUnreachableObjects& OutUnreachableObjects);
@@ -4110,8 +4151,6 @@ void FAsyncPackage2::ImportPackagesRecursiveInner(FAsyncLoadingThreadState2& Thr
 		}
 #endif
 
-		ImportStore.AddImportedPackageReference(ImportedPackageId, ImportedPackageUPackageName);
-
 		FPackagePath ImportedPackagePath;
 #if ALT2_ENABLE_LINKERLOAD_SUPPORT
 		bool bIsZenPackage = !LinkerLoadState.IsSet();
@@ -4131,7 +4170,7 @@ void FAsyncPackage2::ImportPackagesRecursiveInner(FAsyncLoadingThreadState2& Thr
 		constexpr bool bIsZenPackageImport = true;
 #endif
 
-		FLoadedPackageRef& ImportedPackageRef = AsyncLoadingThread.GlobalImportStore.FindPackageRefChecked(ImportedPackageId);
+		FLoadedPackageRef& ImportedPackageRef = ImportStore.AddImportedPackageReference(ImportedPackageId, ImportedPackageUPackageName);
 #if WITH_EDITOR
 		if (AsyncLoadingThread.UncookedPackageLoader && ImportedPackageStatus == EPackageStoreEntryStatus::Ok && !ImportedPackageEntry.UncookedPackageName.IsNone())
 		{
@@ -4699,7 +4738,7 @@ EEventLoadNodeExecutionResult FAsyncPackage2::ProcessLinkerLoadPackageSummary(FA
 	HeaderData.ImportedPublicExportHashes = LinkerLoadState->LinkerLoadHeaderData.ImportedPublicExportHashes;
 	HeaderData.ImportMap = LinkerLoadState->LinkerLoadHeaderData.ImportMap;
 
-	FLoadedPackageRef& PackageRef = AsyncLoadingThread.GlobalImportStore.FindPackageRefChecked(Desc.UPackageId);
+	FLoadedPackageRef& PackageRef = AsyncLoadingThread.GlobalImportStore.FindPackageRefChecked(Desc.UPackageId, Desc.UPackageName);
 	if (!bLoadHasFailed)
 	{
 		PackageRef.ReserveSpaceForPublicExports(LinkerLoadState->LinkerLoadHeaderData.PublicExportHashes.Num());
@@ -4959,7 +4998,7 @@ EEventLoadNodeExecutionResult FAsyncPackage2::Event_ProcessPackageSummary(FAsync
 	{
 		if (Package->Desc.bCanBeImported)
 		{
-			FLoadedPackageRef& PackageRef = Package->AsyncLoadingThread.GlobalImportStore.FindPackageRefChecked(Package->Desc.UPackageId);
+			FLoadedPackageRef& PackageRef = Package->AsyncLoadingThread.GlobalImportStore.FindPackageRefChecked(Package->Desc.UPackageId, Package->Desc.UPackageName);
 			PackageRef.SetHasFailed();
 		}
 	}
@@ -4998,7 +5037,7 @@ EEventLoadNodeExecutionResult FAsyncPackage2::Event_ProcessPackageSummary(FAsync
 				}
 			}
 #endif
-			FLoadedPackageRef& PackageRef = Package->AsyncLoadingThread.GlobalImportStore.FindPackageRefChecked(Package->Desc.UPackageId);
+			FLoadedPackageRef& PackageRef = Package->AsyncLoadingThread.GlobalImportStore.FindPackageRefChecked(Package->Desc.UPackageId, Package->Desc.UPackageName);
 			if (PublicExportsCount)
 			{
 				PackageRef.ReserveSpaceForPublicExports(PublicExportsCount);
@@ -5655,7 +5694,7 @@ EEventLoadNodeExecutionResult FAsyncPackage2::Event_ExportsDone(FAsyncLoadingThr
 
 	if (!Package->bLoadHasFailed && Package->Desc.bCanBeImported)
 	{
-		FLoadedPackageRef& PackageRef = Package->AsyncLoadingThread.GlobalImportStore.FindPackageRefChecked(Package->Desc.UPackageId);
+		FLoadedPackageRef& PackageRef = Package->AsyncLoadingThread.GlobalImportStore.FindPackageRefChecked(Package->Desc.UPackageId, Package->Desc.UPackageName);
 		PackageRef.SetAllPublicExportsLoaded();
 	}
 
@@ -6676,7 +6715,6 @@ FAsyncLoadingThread2::FAsyncLoadingThread2(FIoDispatcher& InIoDispatcher, IAsync
 	ThreadResumedEvent = FPlatformProcess::GetSynchEventFromPool();
 	AsyncLoadingTickCounter = 0;
 
-	FCoreUObjectInternalDelegates::GetOnLeakedPackageRenameDelegate().AddRaw(this, &FAsyncLoadingThread2::OnLeakedPackageRename);
 	FCoreUObjectDelegates::GetPreGarbageCollectDelegate().AddRaw(this, &FAsyncLoadingThread2::OnPreGarbageCollect);
 
 	FAsyncLoadingThreadState2::TlsSlot = FPlatformTLS::AllocTlsSlot();
@@ -6700,7 +6738,6 @@ FAsyncLoadingThread2::~FAsyncLoadingThread2()
 void FAsyncLoadingThread2::ShutdownLoading()
 {
 	FCoreUObjectDelegates::GetPreGarbageCollectDelegate().RemoveAll(this);
-	FCoreUObjectInternalDelegates::GetOnLeakedPackageRenameDelegate().RemoveAll(this);
 
 	delete Thread;
 	Thread = nullptr;
@@ -7141,60 +7178,22 @@ void FAsyncLoadingThread2::CollectUnreachableObjects(
 				UPackage* Package = static_cast<UPackage*>(Object);
 				if (Package->bCanBeImported)
 				{
-					GlobalImportStore.VerifyPackageForRemoval(Package);
+					FPackageId PackageId = Package->GetPackageId();
+					FLoadedPackageRef* PackageRef = GlobalImportStore.FindPackageRef(PackageId);
+					if (PackageRef)
+					{
+						GlobalImportStore.VerifyPackageForRemoval(*PackageRef);
+					}
 				}
-				GlobalImportStore.VerifyObjectForRemoval(Object);
 			}
+			GlobalImportStore.VerifyObjectForRemoval(Object);
 		});
 	}
 }
 
-void FAsyncLoadingThread2::OnLeakedPackageRename(UPackage* Package)
-{
-	TRACE_CPUPROFILER_EVENT_SCOPE(OnLeakedPackageRename);
-	check(IsInGameThread());
-
-	// We don't care about levelstreaming /Temp/ packages that are never imported by other packages
-	if (!Package->CanBeImported())
-	{
-		return;
-	}
-
-	// Code such as LoadMap or LevelStreaming is about to rename a loaded package which was detected as leaking so that we can load another copy of it.
-	UE_LOG(LogStreaming, Verbose, TEXT("Renaming leaked package %s (0x%llX)"), *Package->GetName(), Package->GetPackageId().ValueForDebugging());
-
-	if (!FGCCSyncObject::Get().IsGCLocked())
-	{
-		// If this function is called outside of a GC (CheckAndHandleStaleWorldObjectReferences), then flush async loading so that we can be sure nothing is modifying the stores,
-		// and that nothing is depending on this package. This will also empty the deferred delete queue.
-		FlushAsyncLoading(); 
-	}
-	else
-	{
-		// If this function is called from within a GC the ALT will be inactive and not touching the deferred delete queue.
-		// Flush the delete queue to clear references to the renamed package, if there are still references after this we're not done with loading the package and it shouldn't have been renamed.
-		ProcessDeferredDeletePackagesQueue();
-	}
-
-	// We need to protect UnreachableObjects since the GC will wake up the ALT from NotifyUnreachableObjects to run RemoveUnreachableObjects on the ALT.
-	// FlushAsyncLoading and ProcessDeferredDeletePackagesQueue will not empty the UnreachableObjects array that was just filled by GC.
-	FScopeLock UnreachableObjectsLock(&UnreachableObjectsCritical);
-
-	// unreachable objects from last GC should typically have been processed already,
-	// if not handle them here before processing new ones
-	RemoveUnreachableObjects(UnreachableObjects);
-
-	// If a package that can be imported was leaked and renamed, then go ahead and remove it.
-	// Note that it must exist in the global import store at this point since it is normally only trimmed during GC.
-	GlobalImportStore.RemovePackage(Package);
-
-	// Clear the CanBeImportedFlag so that this package is only removed once,
-	// else we would try and remove it again during GC, which would instead remove the reloaded package if it exists
-	Package->SetCanBeImportedFlag(false);
-}
-
 void FAsyncLoadingThread2::OnPreGarbageCollect()
 {
+	TRACE_CPUPROFILER_EVENT_SCOPE(FAsyncLoadingThread2::OnPreGarbageCollect);
 	// Flush the delete queue so that we don't prevent packages from being garbage collected if we're done with them
 	ProcessDeferredDeletePackagesQueue();
 }
@@ -7529,50 +7528,15 @@ void FAsyncPackage2::CreateUPackage()
 {
 	check(!LinkerRoot);
 
-	// temp packages are never stored or found in global import store
-	FLoadedPackageRef* PackageRef = nullptr;
-
-	// Try to find existing package or create it if not already present.
-	UPackage* ExistingPackage = nullptr;
 	{
 		TRACE_CPUPROFILER_EVENT_SCOPE(UPackageFind);
-		if (Desc.bCanBeImported)
-		{
-			PackageRef = AsyncLoadingThread.GlobalImportStore.FindPackageRef(Desc.UPackageId);
-			UE_ASYNC_PACKAGE_CLOG(!PackageRef, Fatal, Desc, TEXT("CreateUPackage"), TEXT("Package has been destroyed by GC."));
-			LinkerRoot = PackageRef->GetPackage();
-#if DO_CHECK
-			if (LinkerRoot)
-			{
-				UPackage* FoundPackage = FindObjectFast<UPackage>(nullptr, Desc.UPackageName);
-				checkf(LinkerRoot == FoundPackage,
-					TEXT("LinkerRoot '%s' (%p) is different from FoundPackage '%s' (%p)"),
-					*LinkerRoot->GetName(), LinkerRoot, FoundPackage ? *FoundPackage->GetName() : TEXT("null"), FoundPackage);
-			}
-#endif
-		}
-		if (!LinkerRoot)
-		{
-			// Packages can be created outside the loader, i.e from ResolveName via StaticLoadObject
-			ExistingPackage = FindObjectFast<UPackage>(nullptr, Desc.UPackageName);
-		}
+		LinkerRoot = FindObjectFast<UPackage>(/*Outer*/nullptr, Desc.UPackageName);
 	}
 	if (!LinkerRoot)
 	{
 		TRACE_CPUPROFILER_EVENT_SCOPE(UPackageCreate);
-		if (ExistingPackage)
-		{
-			LinkerRoot = ExistingPackage;
-		}
-		else 
-		{
-			LinkerRoot = NewObject<UPackage>(/*Outer*/nullptr, Desc.UPackageName);
-			bCreatedLinkerRoot = true;
-		}
-		if (PackageRef)
-		{
-			PackageRef->SetPackage(LinkerRoot);
-		}
+		LinkerRoot = NewObject<UPackage>(/*Outer*/nullptr, Desc.UPackageName);
+		bCreatedLinkerRoot = true;
 	}
 
 	LinkerRoot->SetFlags(RF_Public | RF_WasLoaded);
@@ -7591,6 +7555,27 @@ void FAsyncPackage2::CreateUPackage()
 	}
 	LinkerRoot->SetInternalFlags(FlagsToSet);
 
+	// update global import store
+	// temp packages can't be imported and are never stored or found in global import store
+	if (Desc.bCanBeImported)
+	{
+		FLoadedPackageRef& PackageRef = AsyncLoadingThread.GlobalImportStore.FindPackageRefChecked(Desc.UPackageId, Desc.UPackageName);
+		UPackage* ExistingPackage = PackageRef.GetPackage();
+		if (!ExistingPackage)
+		{
+			PackageRef.SetPackage(LinkerRoot);
+		}
+		else if (ExistingPackage != LinkerRoot)
+		{
+			UE_ASYNC_PACKAGE_LOG(Warning, Desc, TEXT("CreateUPackage: ReplacePackage"),
+				TEXT("Replacing renamed package %s (0x%llX) while being referenced by the loader, RefCount=%d"),
+				*ExistingPackage->GetName(), ExistingPackage->GetPackageId().ValueForDebugging(),
+				PackageRef.GetRefCount());
+
+			AsyncLoadingThread.GlobalImportStore.ReplaceReferencedRenamedPackage(PackageRef, LinkerRoot);
+		}
+	}
+
 	if (bCreatedLinkerRoot)
 	{
 		UE_ASYNC_PACKAGE_LOG_VERBOSE(VeryVerbose, Desc, TEXT("CreateUPackage: AddPackage"),
