ZenLoader: Implemented priority support for the loader event queue instead of just relying on io request priorities

diff --git a/Engine/Source/Runtime/Core/Internal/IO/IoPriorityQueue.h b/Engine/Source/Runtime/Core/Internal/IO/IoPriorityQueue.h
new file mode 100644
--- /dev/null
+++ b/Engine/Source/Runtime/Core/Internal/IO/IoPriorityQueue.h
@@ -0,0 +1,328 @@
+// Copyright Epic Games, Inc. All Rights Reserved.
+
+#pragma once
+
+#include "Algo/IsSorted.h"
+
+template<typename T>
+class TIoPriorityQueue
+{
+private:
+	struct TInternalQueue
+	{
+		TInternalQueue(int32 InPriority)
+			: Priority(InPriority)
+		{
+
+		}
+
+		// Returns true if outer queue should be removed
+		bool Remove(T* Item)
+		{
+			T* Prev = Item->Prev;
+			T* Next = Item->Next;
+			Item->Prev = nullptr;
+			Item->Next = nullptr;
+			if (Prev && Next)
+			{
+				Prev->Next = Next;
+				Next->Prev = Prev;
+				return false;
+			}
+			if (Prev)
+			{
+				Prev->Next = Next;
+			}
+			else
+			{
+				Head = Next;
+			}
+			if (Next)
+			{
+				Next->Prev = Prev;
+			}
+			else
+			{
+				Tail = Prev;
+			}
+			if (!Head)
+			{
+				check(!Tail);
+				return true;
+			}
+			return false;
+		}
+		TInternalQueue* NextFree = nullptr;
+		T* Head = nullptr;
+		T* Tail = nullptr;
+		int32 Priority;
+	};
+
+public:
+	class TIterator
+	{
+	public:
+		friend TIoPriorityQueue;
+
+		TIterator& operator++()
+		{
+			Current = Next;
+			Next = Current ? Current->Next : nullptr;
+			return *this;
+		}
+
+		TIterator operator++(int)
+		{
+			TIterator Tmp(*this);
+			Current = Next;
+			Next = Current ? Current->Next : nullptr;
+			return Tmp;
+		}
+
+		T& operator*() const
+		{
+			check(Current);
+			check(InternalQueue);
+			check(Current->Priority == InternalQueue->Priority);
+			return *Current;
+		}
+
+		T* operator->() const
+		{
+			check(Current);
+			check(InternalQueue);
+			check(Current->Priority == InternalQueue->Priority);
+			return Current;
+		}
+
+		explicit operator bool() const
+		{
+			return !!Current;
+		}
+
+		void RemoveCurrent()
+		{
+			check(Current);
+			check(InternalQueue);
+			check(Current->Priority == InternalQueue->Priority);
+			if (InternalQueue->Remove(Current))
+			{
+				Outer.RemoveQueueByPriority(InternalQueue->Priority);
+			}
+			Current = reinterpret_cast<T*>(-1);
+		}
+
+	private:
+		TIterator(TIoPriorityQueue<T>& InOuter, int32 InPriority)
+			: Outer(InOuter)
+		{
+			InternalQueue = Outer.FindQueueByPriority(InPriority);
+			if (InternalQueue)
+			{
+				Current = InternalQueue->Head;
+				Next = Current->Next;
+			}
+			else
+			{
+				Current = Next = nullptr;
+			}
+		}
+
+		TIoPriorityQueue<T>& Outer;
+		TInternalQueue* InternalQueue;
+		T* Current;
+		T* Next;
+	};
+
+	bool IsEmpty()
+	{
+		return InternalQueues.IsEmpty();
+	}
+
+	int32 GetMaxPriority() const
+	{
+		if (InternalQueues.IsEmpty())
+		{
+			return MIN_int32;
+		}
+		return InternalQueues.Last()->Priority;
+	}
+
+	void Push(T* Item, int32 Priority)
+	{
+		check(!Item->Prev);
+		check(!Item->Next);
+		Item->Priority = Priority;
+		TInternalQueue& Queue = FindOrAddQueueByPriority(Priority);
+		if (!Queue.Head)
+		{
+			check(!Queue.Tail);
+			Queue.Head = Queue.Tail = Item;
+		}
+		else
+		{
+			Item->Prev = Queue.Tail;
+			Queue.Tail->Next = Item;
+			Queue.Tail = Item;
+		}
+	}
+
+	T* Pop()
+	{
+		if (InternalQueues.IsEmpty())
+		{
+			return nullptr;
+		}
+		int QueueIndex = InternalQueues.Num() - 1;
+		TInternalQueue* QueueWithHighestPriority = InternalQueues[QueueIndex];
+		check(QueueWithHighestPriority);
+		check(QueueWithHighestPriority->Head);
+		check(QueueWithHighestPriority->Tail);
+
+		T* Item = QueueWithHighestPriority->Head;
+		check(Item);
+		check(Item->Priority == QueueWithHighestPriority->Priority);
+		T* Next = Item->Next;
+		if (Next)
+		{
+			Next->Prev = nullptr;
+			QueueWithHighestPriority->Head = Next;
+		}
+		else
+		{
+			check(Item == QueueWithHighestPriority->Tail);
+			QueueWithHighestPriority->Head = QueueWithHighestPriority->Tail = nullptr;
+			RemoveQueueAtIndex(QueueIndex);
+		}
+		Item->Prev = Item->Next = nullptr;
+
+		return Item;
+	}
+
+	TIterator CreateIterator(int32 Priority)
+	{
+		return TIterator(*this, Priority);
+	}
+
+	void Remove(T* Item)
+	{
+		T* Prev = Item->Prev;
+		T* Next = Item->Next;
+		if (Prev && Next)
+		{
+			Item->Prev = nullptr;
+			Item->Next = nullptr;
+			Prev->Next = Next;
+			Next->Prev = Prev;
+			return;
+		}
+		int32 QueueIndex;
+		TInternalQueue* InternalQueue = FindQueueByPriority(Item->Priority, QueueIndex);
+		check(InternalQueue);
+		if (InternalQueue->Remove(Item))
+		{
+			RemoveQueueAtIndex(QueueIndex);
+		}
+	}
+
+	void Reprioritize(T* Item, int32 NewPriority)
+	{
+		Remove(Item);
+		Push(Item, NewPriority);
+	}
+
+	void MergeInto(TIoPriorityQueue<T>& Other, int32 Priority)
+	{
+		int32 QueueIndex;
+		TInternalQueue* Queue = FindQueueByPriority(Priority, QueueIndex);
+		if (!Queue)
+		{
+			return;
+		}
+		TInternalQueue& OtherQueue = Other.FindOrAddQueueByPriority(Priority);
+		if (OtherQueue.Head)
+		{
+			check(OtherQueue.Tail);
+			Queue->Head->Prev = OtherQueue.Tail;
+			OtherQueue.Tail->Next = Queue->Head;
+			OtherQueue.Tail = Queue->Tail;
+		}
+		else
+		{
+			OtherQueue.Head = Queue->Head;
+			OtherQueue.Tail = Queue->Tail;
+		}
+		Queue->Head = Queue->Tail = nullptr;
+		RemoveQueueAtIndex(QueueIndex);
+	}
+
+private:
+	TInternalQueue* FindQueueByPriority(int32 Priority) const
+	{
+		int32 Index = Algo::LowerBoundBy(InternalQueues, Priority, &TInternalQueue::Priority);
+		if (InternalQueues.IsValidIndex(Index) && InternalQueues[Index]->Priority == Priority)
+		{
+			return InternalQueues[Index];
+		}
+		return nullptr;
+	}
+
+	TInternalQueue* FindQueueByPriority(int32 Priority, int32& OutIndex) const
+	{
+		int32 Index = Algo::LowerBoundBy(InternalQueues, Priority, &TInternalQueue::Priority);
+		if (InternalQueues.IsValidIndex(Index) && InternalQueues[Index]->Priority == Priority)
+		{
+			OutIndex = Index;
+			return InternalQueues[Index];
+		}
+		return nullptr;
+	}
+
+	TInternalQueue& FindOrAddQueueByPriority(int32 Priority)
+	{
+		int32 Index = Algo::LowerBoundBy(InternalQueues, Priority, &TInternalQueue::Priority);
+		if (!InternalQueues.IsValidIndex(Index) || InternalQueues[Index]->Priority != Priority)
+		{
+			InternalQueues.Insert(AllocInternalQueue(Priority), Index);
+		}
+		return *InternalQueues[Index];
+	}
+
+	void RemoveQueueByPriority(int32 Priority)
+	{
+		int32 QueueIndex;
+		if (TInternalQueue* InternalQueue = FindQueueByPriority(Priority, QueueIndex))
+		{
+			RemoveQueueAtIndex(QueueIndex);
+		}
+	}
+
+	void RemoveQueueAtIndex(int32 Index)
+	{
+		FreeInternalQueue(InternalQueues[Index]);
+		InternalQueues.RemoveAt(Index, 1, false);
+	}
+
+	TInternalQueue* AllocInternalQueue(int32 Priority)
+	{
+		if (TInternalQueue* FromPool = FirstFreeQueue)
+		{
+			FirstFreeQueue = FromPool->NextFree;
+			FromPool->NextFree = nullptr;
+			FromPool->Priority = Priority;
+			return FromPool;
+		}
+		return new TInternalQueue(Priority);
+	}
+
+	void FreeInternalQueue(TInternalQueue* InternalQueue)
+	{
+		check(!InternalQueue->Head);
+		check(!InternalQueue->Tail);
+		InternalQueue->NextFree = FirstFreeQueue;
+		FirstFreeQueue = InternalQueue;
+	}
+
+	TArray<TInternalQueue*> InternalQueues; // Sorted by priority
+	TInternalQueue* FirstFreeQueue = nullptr;
+};
\ No newline at end of file
diff --git a/Engine/Source/Runtime/CoreUObject/Private/Serialization/AsyncLoading2.cpp b/Engine/Source/Runtime/CoreUObject/Private/Serialization/AsyncLoading2.cpp
--- a/Engine/Source/Runtime/CoreUObject/Private/Serialization/AsyncLoading2.cpp
+++ b/Engine/Source/Runtime/CoreUObject/Private/Serialization/AsyncLoading2.cpp
@@ -69,6 +69,7 @@
 #include "Misc/PathViews.h"
 #include "UObject/LinkerLoad.h"
 #include "Containers/SpscQueue.h"
+#include "IO/IoPriorityQueue.h"
 
 #include <atomic>
 
@@ -1805,6 +1806,7 @@ public:
 
 private:
 	friend class FAsyncLoadEventQueue2;
+	friend TIoPriorityQueue<FEventLoadNode2>;
 
 	FEventLoadNode2()
 	{
@@ -1814,21 +1816,29 @@ private:
 	void ProcessDependencies(FAsyncLoadingThreadState2& ThreadState);
 	void Fire(FAsyncLoadingThreadState2* ThreadState);
 
+	const FAsyncLoadEventSpec* Spec = nullptr;
+	FAsyncPackage2* Package = nullptr;
 	union
 	{
 		FEventLoadNode2* SingleDependent;
 		FEventLoadNode2** MultipleDependents;
 	};
-	std::atomic<FEventLoadNode2*> Next { nullptr };
+	FEventLoadNode2* Prev = nullptr;
+	FEventLoadNode2* Next = nullptr;
 	uint32 DependenciesCount = 0;
 	uint32 DependenciesCapacity = 0;
+	int32 Priority = 0;
+	int32 ImportOrExportIndex = -1;
 	std::atomic<int32> BarrierCount { 0 };
+	enum EQueueStatus
+	{
+		QueueStatus_None = 0,
+		QueueStatus_Local = 1,
+		QueueStatus_External = 2
+	};
+	uint8 QueueStatus = QueueStatus_None;
 	std::atomic<bool> bIsUpdatingDependencies { false };
 	std::atomic<bool> bIsDone { false };
-
-	const FAsyncLoadEventSpec* Spec = nullptr;
-	FAsyncPackage2* Package = nullptr;
-	int32 ImportOrExportIndex = -1;
 };
 
 class FAsyncLoadEventGraphAllocator
@@ -1871,22 +1881,52 @@ public:
 		Zenaphore = InZenaphore;
 	}
 
+	void SetOwnerThread(const FAsyncLoadingThreadState2* InOwnerThread)
+	{
+		OwnerThread = InOwnerThread;
+	}
+
 	bool PopAndExecute(FAsyncLoadingThreadState2& ThreadState);
-	void Push(FEventLoadNode2* Node);
-	bool ExecuteSyncLoadEvents(FAsyncLoadingThreadState2& ThreadState, FAsyncLoadingSyncLoadContext& SyncLoadContext);
+	void Push(FAsyncLoadingThreadState2* ThreadState, FEventLoadNode2* Node);
+	bool ExecuteSyncLoadEvents(FAsyncLoadingThreadState2& ThreadState);
+	void UpdatePackagePriority(FAsyncPackage2* Package);
 
 private:
-	void Deplete();
+	void PushLocal(FEventLoadNode2* Node);
+	void PushExternal(FEventLoadNode2* Node);
+	bool GetMaxPriorityInExternalQueue(int32& OutMaxPriority)
+	{
+		int64 StateValue = ExternalQueueState.load();
+		if (StateValue == MIN_int64)
+		{
+			return false;
+		}
+		else
+		{
+			OutMaxPriority = static_cast<int32>(StateValue);
+			return true;
+		}
+	}
+	void UpdateExternalQueueState()
+	{
+		if (ExternalQueue.IsEmpty())
+		{
+			ExternalQueueState.store(MIN_int64);
+		}
+		else
+		{
+			ExternalQueueState.store(ExternalQueue.GetMaxPriority());
+		}
+	}
 
+	const FAsyncLoadingThreadState2* OwnerThread = nullptr;
 	FZenaphore* Zenaphore = nullptr;
-	FEventLoadNode2 Sentinel; // `Sentinel.Next` is the head of the queue
-	std::atomic<FEventLoadNode2*> Tail{ &Sentinel };
-	FEventLoadNode2* LocalHead = nullptr;
-	FEventLoadNode2* LocalTail = nullptr;
+	TIoPriorityQueue<FEventLoadNode2> LocalQueue;
+	FCriticalSection ExternalCritical;
+	TIoPriorityQueue<FEventLoadNode2> ExternalQueue;
+	std::atomic<int64> ExternalQueueState = MIN_int64;
 	FEventLoadNode2* TimedOutEventNode = nullptr;
 	int32 ExecuteSyncLoadEventsCallCounter = 0;
-	bool bIsInPopAndExecute = false;
-
 };
 
 struct FAsyncLoadEventSpec
@@ -2008,6 +2048,7 @@ struct FAsyncLoadingThreadState2
 	TArray<FAsyncLoadingSyncLoadContext*> SyncLoadContextStack;
 	TArray<FAsyncPackage2*> PackagesOnStack;
 	TSpscQueue<FAsyncLoadingSyncLoadContext*> SyncLoadContextsCreatedOnGameThread;
+	TSpscQueue<FAsyncPackage2*> PackagesToReprioritize;
 	bool bIsAsyncLoadingThread = false;
 	bool bCanAccessAsyncLoadingThreadData = true;
 	bool bShouldFireNodes = true;
@@ -2308,7 +2349,23 @@ struct FAsyncPackage2
 
 	void AddRef()
 	{
-		++RefCount;
+		RefCount.fetch_add(1);
+	}
+
+	bool TryAddRef()
+	{
+		for (;;)
+		{
+			int32 CurrentRefCount = RefCount.load();
+			if (CurrentRefCount == 0)
+			{
+				return false;
+			}
+			if (RefCount.compare_exchange_strong(CurrentRefCount, CurrentRefCount + 1))
+			{
+				return true;
+			}
+		}
 	}
 
 	void ReleaseRef();
@@ -2424,7 +2481,7 @@ private:
 	std::atomic<uint64>			SyncLoadContextId = 0;
 	/** Time load begun. This is NOT the time the load was requested in the case of pending requests.	*/
 	double						LoadStartTime = 0.0;
-	TAtomic<int32> RefCount{ 0 };
+	std::atomic<int32>			RefCount = 0;
 	bool						bHasStartedImportingPackages = false;
 	int32						ProcessedExportBundlesCount = 0;
 	/** Current bundle entry index in the current export bundle */
@@ -2856,7 +2913,8 @@ public:
 		return AsyncPackageLookup.FindRef(PackageId);
 	}
 
-	void UpdatePackagePriority(FAsyncLoadingThreadState2& ThreadState, FAsyncPackage2* Package, int32 NewPriority);
+	void UpdatePackagePriority(FAsyncLoadingThreadState2& ThreadState, FAsyncPackage2* Package);
+	void UpdatePackagePriorityRecursive(FAsyncLoadingThreadState2& ThreadState, FAsyncPackage2* Package, int32 NewPriority);
 
 	FAsyncPackage2* FindOrInsertPackage(FAsyncLoadingThreadState2& ThreadState, FAsyncPackageDesc2& InDesc, bool& bInserted, TUniquePtr<FLoadPackageAsyncDelegate>&& PackageLoadedDelegate = TUniquePtr<FLoadPackageAsyncDelegate>());
 	void QueueMissingPackage(FAsyncPackageDesc2& PackageDesc, TUniquePtr<FLoadPackageAsyncDelegate>&& LoadPackageAsyncDelegate);
@@ -3114,7 +3172,7 @@ private:
 	EAsyncPackageState::Type ProcessLoadedPackagesFromGameThread(FAsyncLoadingThreadState2& ThreadState, bool& bDidSomething, int32 FlushRequestID = INDEX_NONE);
 
 	void IncludePackageInSyncLoadContextRecursive(FAsyncLoadingThreadState2& ThreadState, uint64 ContextId, FAsyncPackage2* Package);
-	FAsyncLoadingSyncLoadContext* UpdateSyncLoadContext(FAsyncLoadingThreadState2& ThreadState);
+	void UpdateSyncLoadContext(FAsyncLoadingThreadState2& ThreadState);
 
 	bool CreateAsyncPackagesFromQueue(FAsyncLoadingThreadState2& ThreadState);
 
@@ -3340,17 +3398,58 @@ void FAsyncLoadingThread2::InitializeLoading()
 	UE_LOG(LogStreaming, Display, TEXT("AsyncLoading2 - Initialized"));
 }
 
-void FAsyncLoadingThread2::UpdatePackagePriority(FAsyncLoadingThreadState2& ThreadState, FAsyncPackage2* Package, int32 NewPriority)
+void FAsyncLoadingThread2::UpdatePackagePriority(FAsyncLoadingThreadState2& ThreadState, FAsyncPackage2* Package)
 {
 	TRACE_CPUPROFILER_EVENT_SCOPE(UpdatePackagePriority);
-	Package->Desc.Priority = NewPriority;
-	Package->SerializationState.IoRequest.UpdatePriority(NewPriority);
+	EAsyncPackageLoadingState2 LoadingState = Package->AsyncPackageLoadingState;
+	check(ThreadState.bCanAccessAsyncLoadingThreadData);
+
+	// TODO: Shader requests are not tracked so they won't be reprioritized correctly here
+	if (LoadingState <= EAsyncPackageLoadingState2::WaitingForIo)
+	{
+		Package->SerializationState.IoRequest.UpdatePriority(Package->Desc.Priority);
 #if WITH_EDITOR
-	if (Package->OptionalSegmentSerializationState.IsSet())
+		if (Package->OptionalSegmentSerializationState.IsSet())
+		{
+			Package->OptionalSegmentSerializationState->IoRequest.UpdatePriority(Package->Desc.Priority);
+		}
+#endif
+	}
+	if (LoadingState <= EAsyncPackageLoadingState2::PostLoad)
 	{
-		Package->OptionalSegmentSerializationState->IoRequest.UpdatePriority(NewPriority);
+		EventQueue.UpdatePackagePriority(Package);
+	}
+	if (LoadingState <= EAsyncPackageLoadingState2::DeferredPostLoad)
+	{
+		if (ThreadState.bIsAsyncLoadingThread)
+		{
+			if (Package->TryAddRef())
+			{
+				GameThreadState->PackagesToReprioritize.Enqueue(Package);
+			}
+		}
+		else
+		{
+			MainThreadEventQueue.UpdatePackagePriority(Package);
+		}
 	}
-#endif
+}
+
+void FAsyncLoadingThread2::UpdatePackagePriorityRecursive(FAsyncLoadingThreadState2& ThreadState, FAsyncPackage2* Package, int32 NewPriority)
+{
+	if (Package->Desc.Priority >= NewPriority)
+	{
+		return;
+	}
+	Package->Desc.Priority = NewPriority;
+	for (FAsyncPackage2* ImportedPackage : Package->Data.ImportedAsyncPackages)
+	{
+		if (ImportedPackage)
+		{
+			UpdatePackagePriorityRecursive(ThreadState, ImportedPackage, NewPriority);
+		}
+	}
+	UpdatePackagePriority(ThreadState, Package);
 }
 
 FAsyncPackage2* FAsyncLoadingThread2::FindOrInsertPackage(FAsyncLoadingThreadState2& ThreadState, FAsyncPackageDesc2& Desc, bool& bInserted, TUniquePtr<FLoadPackageAsyncDelegate>&& PackageLoadedDelegate)
@@ -3379,7 +3478,8 @@ FAsyncPackage2* FAsyncLoadingThread2::FindOrInsertPackage(FAsyncLoadingThreadSta
 			}
 			if (Desc.Priority > Package->Desc.Priority)
 			{
-				UpdatePackagePriority(ThreadState, Package, Desc.Priority);
+				TRACE_CPUPROFILER_EVENT_SCOPE(UpdatePackagePriority);
+				UpdatePackagePriorityRecursive(ThreadState, Package, Desc.Priority);
 			}
 		}
 		if (PackageLoadedDelegate.IsValid())
@@ -3410,6 +3510,11 @@ void FAsyncLoadingThread2::IncludePackageInSyncLoadContextRecursive(FAsyncLoadin
 			IncludePackageInSyncLoadContextRecursive(ThreadState, ContextId, ImportedPackage);
 		}
 	}
+	if (Package->Desc.Priority < MAX_int32)
+	{
+		Package->Desc.Priority = MAX_int32;
+		UpdatePackagePriority(ThreadState, Package);
+	}
 }
 
 bool FAsyncLoadingThread2::CreateAsyncPackagesFromQueue(FAsyncLoadingThreadState2& ThreadState)
@@ -3606,10 +3711,10 @@ bool FAsyncLoadingThread2::CreateAsyncPackagesFromQueue(FAsyncLoadingThreadState
 }
 
 FEventLoadNode2::FEventLoadNode2(const FAsyncLoadEventSpec* InSpec, FAsyncPackage2* InPackage, int32 InImportOrExportIndex, int32 InBarrierCount)
-	: BarrierCount(InBarrierCount)
-	, Spec(InSpec)
+	: Spec(InSpec)
 	, Package(InPackage)
 	, ImportOrExportIndex(InImportOrExportIndex)
+	, BarrierCount(InBarrierCount)
 {
 	check(Spec);
 	check(Package);
@@ -3696,7 +3801,7 @@ void FEventLoadNode2::Fire(FAsyncLoadingThreadState2* ThreadState)
 	}
 	else
 	{
-		Spec->EventQueue->Push(this);
+		Spec->EventQueue->Push(ThreadState, this);
 	}
 }
 
@@ -3773,61 +3878,44 @@ FAsyncLoadEventQueue2::~FAsyncLoadEventQueue2()
 {
 }
 
-void FAsyncLoadEventQueue2::Push(FEventLoadNode2* Node)
+void FAsyncLoadEventQueue2::Push(FAsyncLoadingThreadState2* ThreadState, FEventLoadNode2* Node)
 {
-	// switch `Tail` to the new node and only then link the old tail to the new one. The list is not fully linked between these ops,
-	// this is explicitly handled by the consumer by waiting for the link
-	FEventLoadNode2* Prev = Tail.exchange(Node, std::memory_order_acq_rel); // sync with consumer in Deplete()
-	check(Prev->Next.load(std::memory_order_relaxed) == nullptr); // `Tail` is assigned before its Next
-	Prev->Next.store(Node, std::memory_order_relaxed);
-	if (Zenaphore)
+	if (OwnerThread == ThreadState)
 	{
-		bool bWasEmpty = Prev == &Sentinel;
-		if (bWasEmpty)
-		{
-			Zenaphore->NotifyOne();
-		}
+		PushLocal(Node);
 	}
-}
-
-void FAsyncLoadEventQueue2::Deplete()
-{
-	FEventLoadNode2* First = Sentinel.Next.load(std::memory_order_relaxed);
-	if (First == nullptr)
+	else
 	{
-		return; // empty
+		PushExternal(Node);
 	}
+}
 
-	// reset the head so the next consumption can detect that the queue is empty
-	// `Sentinel.Next` is not touched by producers right now because it's already not null
-	Sentinel.Next.store(nullptr, std::memory_order_relaxed);
-
-	// reset the queue to the empty state. this redirects producers to start from `Sentinel` again.
-	// take note of the tail on resetting it because the list can be still not fully linked and so `Node.Next == nullptr` can't be 
-	// used to detect the end of the list
-	FEventLoadNode2* Last = Tail.exchange(&Sentinel, std::memory_order_acq_rel); // after setting Sentinel.Next to nullptr and before producers' tail modifications
-	check(Last->Next.load(std::memory_order_relaxed) == nullptr); // `Tail` is assigned before its Next
-	// the previously queued items are detached from the instance (as a linked list, though potentially not fully linked yet)
-
-	check(Last != &Sentinel); // can't be empty because of `First != nullptr` above
+void FAsyncLoadEventQueue2::PushLocal(FEventLoadNode2* Node)
+{
+	check(!Node->QueueStatus);
+	int32 Priority = Node->Package->Desc.Priority;
+	Node->QueueStatus = FEventLoadNode2::QueueStatus_Local;
+	LocalQueue.Push(Node, Priority);
+}
 
-	if (LocalHead)
+void FAsyncLoadEventQueue2::PushExternal(FEventLoadNode2* Node)
+{
 	{
-		check(LocalTail);
-		check(LocalTail->Next.load(std::memory_order_relaxed) == nullptr);
-		LocalTail->Next.store(First, std::memory_order_relaxed);
+		int32 Priority = Node->Package->Desc.Priority;
+		FScopeLock Lock(&ExternalCritical);
+		check(!Node->QueueStatus);
+		Node->QueueStatus = FEventLoadNode2::QueueStatus_External;
+		ExternalQueue.Push(Node, Priority);
+		UpdateExternalQueueState();
 	}
-	else
+	if (Zenaphore)
 	{
-		LocalHead = First;
+		Zenaphore->NotifyOne();
 	}
-	LocalTail = Last;
 }
 
 bool FAsyncLoadEventQueue2::PopAndExecute(FAsyncLoadingThreadState2& ThreadState)
 {
-	check(!bIsInPopAndExecute);
-	TGuardValue<bool> GuardIsInPopAndExecute(bIsInPopAndExecute, true);
 	if (TimedOutEventNode)
 	{
 		EEventLoadNodeExecutionResult Result = TimedOutEventNode->Execute(ThreadState);
@@ -3838,49 +3926,55 @@ bool FAsyncLoadEventQueue2::PopAndExecute(FAsyncLoadingThreadState2& ThreadState
 		return true;
 	}
 
-	if (!LocalHead)
-	{
-		Deplete();
-	}
-	if (LocalHead)
+	bool bPopFromExternalQueue = false;
+	int32 MaxPriorityInExternalQueue;
+	if (GetMaxPriorityInExternalQueue(MaxPriorityInExternalQueue))
 	{
-		FEventLoadNode2* Next = nullptr;
-		if (LocalHead != LocalTail)
-		{
-			// producers can still be updating `Next`, wait until the link to the next element is established
-			do
-			{
-				Next = LocalHead->Next.load(std::memory_order_relaxed);
-			} while (Next == nullptr);
-		}
-		FEventLoadNode2* Node = LocalHead;
-		LocalHead = Next;
-		if (!LocalHead)
+		if (LocalQueue.IsEmpty() || MaxPriorityInExternalQueue > LocalQueue.GetMaxPriority())
 		{
-			LocalTail = nullptr;
+			bPopFromExternalQueue = true;
 		}
-
-		EEventLoadNodeExecutionResult Result = Node->Execute(ThreadState);
-		if (Result == EEventLoadNodeExecutionResult::Timeout)
-		{
-			TimedOutEventNode = Node;
-		}
-		return true;
 	}
-	return false;
+	FEventLoadNode2* Node;
+	if (bPopFromExternalQueue)
+	{
+		FScopeLock Lock(&ExternalCritical);
+		Node = ExternalQueue.Pop();
+		check(Node);
+		UpdateExternalQueueState();
+	}
+	else
+	{
+		Node = LocalQueue.Pop();
+	}
+	if (!Node)
+	{
+		return false;
+	}
+	Node->QueueStatus = FEventLoadNode2::QueueStatus_None;
+	
+	EEventLoadNodeExecutionResult Result = Node->Execute(ThreadState);
+	if (Result == EEventLoadNodeExecutionResult::Timeout)
+	{
+		TimedOutEventNode = Node;
+	}
+	return true;
 }
 
-bool FAsyncLoadEventQueue2::ExecuteSyncLoadEvents(FAsyncLoadingThreadState2& ThreadState, FAsyncLoadingSyncLoadContext& SyncLoadContext)
+bool FAsyncLoadEventQueue2::ExecuteSyncLoadEvents(FAsyncLoadingThreadState2& ThreadState)
 {
+	check(!ThreadState.SyncLoadContextStack.IsEmpty());
+	FAsyncLoadingSyncLoadContext& SyncLoadContext = *ThreadState.SyncLoadContextStack.Top();
+
 	int32 ThisCallCounter = ++ExecuteSyncLoadEventsCallCounter;
 
-	auto ShouldExecuteNode = [&SyncLoadContext](FEventLoadNode2* Node) -> bool
+	auto ShouldExecuteNode = [&SyncLoadContext](FEventLoadNode2& Node) -> bool
 	{
-		return Node->Package->SyncLoadContextId >= SyncLoadContext.ContextId;
+		return Node.Package->SyncLoadContextId >= SyncLoadContext.ContextId;
 	};
 
 	bool bDidSomething = false;
-	if (TimedOutEventNode && ShouldExecuteNode(TimedOutEventNode))
+	if (TimedOutEventNode && ShouldExecuteNode(*TimedOutEventNode))
 	{
 		EEventLoadNodeExecutionResult Result = TimedOutEventNode->Execute(ThreadState);
 		check(Result == EEventLoadNodeExecutionResult::Complete); // we can't timeout during a sync load operation
@@ -3888,42 +3982,25 @@ bool FAsyncLoadEventQueue2::ExecuteSyncLoadEvents(FAsyncLoadingThreadState2& Thr
 		bDidSomething = true;
 	}
 
-	Deplete();
-
-	FEventLoadNode2* Node = LocalHead;
-	FEventLoadNode2* Prev = nullptr;
-	while (Node)
+	int32 MaxPriorityInExternalQueue;
+	bool bTakeFromExternalQueue = GetMaxPriorityInExternalQueue(MaxPriorityInExternalQueue) && MaxPriorityInExternalQueue == MAX_int32;
+	if (bTakeFromExternalQueue)
 	{
-		FEventLoadNode2* Next = nullptr;
-		if (Node != LocalTail)
-		{
-			// producers can still be updating `Next`, wait until the link to the next element is established
-			do
-			{
-				Next = Node->Next.load(std::memory_order_relaxed);
-			} while (Next == nullptr);
-		}
-
+		TRACE_CPUPROFILER_EVENT_SCOPE(MergeIntoLocalQueue);
+		// Take all the max prio items from the external queue and put in the local queue. This breaks the queue status value
+		// of the items in the external queue but we know that we'll never reprioritize them again so it doesn't matter
+		FScopeLock Lock(&ExternalCritical);
+		ExternalQueue.MergeInto(LocalQueue, MAX_int32);
+		UpdateExternalQueueState();
+	}
+	for (auto It = LocalQueue.CreateIterator(MAX_int32); It; ++It)
+	{
+		FEventLoadNode2& Node = *It;
 		if (ShouldExecuteNode(Node))
 		{
-			if (Prev)
-			{
-				Prev->Next.store(Next, std::memory_order_relaxed);
-			}
-			if (Node == LocalHead)
-			{
-				LocalHead = Next;
-				if (!LocalHead)
-				{
-					LocalTail = nullptr;
-				}
-			}
-			if (Node == LocalTail)
-			{
-				LocalTail = Prev;
-			}
-
-			EEventLoadNodeExecutionResult Result = Node->Execute(ThreadState);
+			It.RemoveCurrent();
+			Node.QueueStatus = FEventLoadNode2::QueueStatus_None;
+			EEventLoadNodeExecutionResult Result = Node.Execute(ThreadState);
 			check(Result == EEventLoadNodeExecutionResult::Complete); // we can't timeout during a sync load operation
 			if (ExecuteSyncLoadEventsCallCounter != ThisCallCounter)
 			{
@@ -3932,11 +4009,6 @@ bool FAsyncLoadEventQueue2::ExecuteSyncLoadEvents(FAsyncLoadingThreadState2& Thr
 			}
 			bDidSomething = true;
 		}
-		else
-		{
-			Prev = Node;
-		}
-		Node = Next;
 	}
 	if (!bDidSomething && ThreadState.bIsAsyncLoadingThread)
 	{
@@ -3945,6 +4017,35 @@ bool FAsyncLoadEventQueue2::ExecuteSyncLoadEvents(FAsyncLoadingThreadState2& Thr
 	return bDidSomething;
 }
 
+void FAsyncLoadEventQueue2::UpdatePackagePriority(FAsyncPackage2* Package)
+{
+	FScopeLock Lock(&ExternalCritical);
+	auto ReprioritizeNode = [this](FEventLoadNode2& Node)
+	{
+		if (Node.Spec->EventQueue == this && Node.Priority < Node.Package->Desc.Priority)
+		{
+			if (Node.QueueStatus == FEventLoadNode2::QueueStatus_Local)
+			{
+				LocalQueue.Reprioritize(&Node, Node.Package->Desc.Priority);
+			}
+			else if (Node.QueueStatus == FEventLoadNode2::QueueStatus_External)
+			{
+				ExternalQueue.Reprioritize(&Node, Node.Package->Desc.Priority);
+			}
+		}
+	};
+
+	for (FEventLoadNode2& Node : Package->PackageNodes)
+	{
+		ReprioritizeNode(Node);
+	}
+	for (FEventLoadNode2& Node : Package->Data.ExportBundleNodes)
+	{
+		ReprioritizeNode(Node);
+	}
+	UpdateExternalQueueState();
+}
+
 FUObjectSerializeContext* FAsyncPackage2::GetSerializeContext()
 {
 	return FUObjectThreadContext::Get().GetSerializeContext();
@@ -6212,7 +6313,7 @@ FEventLoadNode2& FAsyncPackage2::GetExportBundleNode(EEventLoadNode2 Phase, uint
 	return Data.ExportBundleNodes[ExportBundleNodeIndex];
 }
 
-FAsyncLoadingSyncLoadContext* FAsyncLoadingThread2::UpdateSyncLoadContext(FAsyncLoadingThreadState2& ThreadState)
+void FAsyncLoadingThread2::UpdateSyncLoadContext(FAsyncLoadingThreadState2& ThreadState)
 {
 	if (ThreadState.bIsAsyncLoadingThread)
 	{
@@ -6224,7 +6325,7 @@ FAsyncLoadingSyncLoadContext* FAsyncLoadingThread2::UpdateSyncLoadContext(FAsync
 	}
 	if (ThreadState.SyncLoadContextStack.IsEmpty())
 	{
-		return nullptr;
+		return;
 	}
 	FAsyncLoadingSyncLoadContext* SyncLoadContext = ThreadState.SyncLoadContextStack.Top();
 	if (ThreadState.bIsAsyncLoadingThread)
@@ -6235,14 +6336,14 @@ FAsyncLoadingSyncLoadContext* FAsyncLoadingThread2::UpdateSyncLoadContext(FAsync
 			ThreadState.SyncLoadContextStack.Pop();
 			if (ThreadState.SyncLoadContextStack.IsEmpty())
 			{
-				return nullptr;
+				return;
 			}
 			SyncLoadContext = ThreadState.SyncLoadContextStack.Top();
 		}
 	}
 	else if (!ContainsRequestID(SyncLoadContext->RequestId))
 	{
-		return nullptr;
+		return;
 	}
 	if (ThreadState.bCanAccessAsyncLoadingThreadData && !SyncLoadContext->bHasFoundRequestedPackage)
 	{
@@ -6262,8 +6363,6 @@ FAsyncLoadingSyncLoadContext* FAsyncLoadingThread2::UpdateSyncLoadContext(FAsync
 		UE_CLOG(!bPreloadIsDone, LogStreaming, Fatal, TEXT("Flushing package %s while it's being preloaded in the same callstack is not permitted"), *SyncLoadContext->RequestedPackage->Desc.UPackageName.ToString());
 		RemovePendingRequests(TArrayView<int32>(&SyncLoadContext->RequestId, 1));
 	}
-
-	return SyncLoadContext;
 }
 
 EAsyncPackageState::Type FAsyncLoadingThread2::ProcessAsyncLoadingFromGameThread(FAsyncLoadingThreadState2& ThreadState, bool& bDidSomething)
@@ -6310,9 +6409,9 @@ EAsyncPackageState::Type FAsyncLoadingThread2::ProcessAsyncLoadingFromGameThread
 				}
 			}
 
-			if (!ThreadState.SyncLoadContextStack.IsEmpty())
+			if (!ThreadState.SyncLoadContextStack.IsEmpty() && ThreadState.SyncLoadContextStack.Top()->ContextId)
 			{
-				if (EventQueue.ExecuteSyncLoadEvents(ThreadState, *ThreadState.SyncLoadContextStack.Top()))
+				if (EventQueue.ExecuteSyncLoadEvents(ThreadState))
 				{
 					bDidSomething = true;
 					break;
@@ -6381,9 +6480,15 @@ EAsyncPackageState::Type FAsyncLoadingThread2::ProcessLoadedPackagesFromGameThre
 		}
 
 		bool bLocalDidSomething = false;
-		if (!ThreadState.SyncLoadContextStack.IsEmpty())
+		FAsyncPackage2* PackageToRepriortize;
+		while (ThreadState.PackagesToReprioritize.Dequeue(PackageToRepriortize))
+		{
+			MainThreadEventQueue.UpdatePackagePriority(PackageToRepriortize);
+			PackageToRepriortize->ReleaseRef();
+		}
+		if (!ThreadState.SyncLoadContextStack.IsEmpty() && ThreadState.SyncLoadContextStack.Top()->ContextId)
 		{
-			bLocalDidSomething |= MainThreadEventQueue.ExecuteSyncLoadEvents(ThreadState, *ThreadState.SyncLoadContextStack.Top());
+			bLocalDidSomething |= MainThreadEventQueue.ExecuteSyncLoadEvents(ThreadState);
 		}
 		else
 		{
@@ -6719,6 +6824,8 @@ FAsyncLoadingThread2::FAsyncLoadingThread2(FIoDispatcher& InIoDispatcher, IAsync
 
 	FAsyncLoadingThreadState2::TlsSlot = FPlatformTLS::AllocTlsSlot();
 	GameThreadState = MakeUnique<FAsyncLoadingThreadState2>(GraphAllocator, IoDispatcher);
+	EventQueue.SetOwnerThread(GameThreadState.Get());
+	MainThreadEventQueue.SetOwnerThread(GameThreadState.Get());
 	FAsyncLoadingThreadState2::Set(GameThreadState.Get());
 	
 	UE_LOG(LogStreaming, Display, TEXT("AsyncLoading2 - Created: Event Driven Loader: %s, Async Loading Thread: %s, Async Post Load: %s"),
@@ -6763,6 +6870,7 @@ void FAsyncLoadingThread2::StartThread()
 	if (FAsyncLoadingThreadSettings::Get().bAsyncLoadingThreadEnabled && !Thread)
 	{
 		AsyncLoadingThreadState = MakeUnique<FAsyncLoadingThreadState2>(GraphAllocator, IoDispatcher);
+		EventQueue.SetOwnerThread(AsyncLoadingThreadState.Get());
 		AsyncLoadingThreadState->bIsAsyncLoadingThread = true;
 		AsyncLoadingThreadState->bCanAccessAsyncLoadingThreadData = true;
 		GameThreadState->bCanAccessAsyncLoadingThreadData = false;
@@ -6888,9 +6996,10 @@ uint32 FAsyncLoadingThread2::Run()
 						do 
 						{
 							bPopped = false;
-							if (FAsyncLoadingSyncLoadContext* SyncLoadContext = UpdateSyncLoadContext(ThreadState))
+							UpdateSyncLoadContext(ThreadState);
+							if (!ThreadState.SyncLoadContextStack.IsEmpty() && ThreadState.SyncLoadContextStack.Top()->ContextId)
 							{
-								if (EventQueue.ExecuteSyncLoadEvents(ThreadState, *SyncLoadContext))
+								if (EventQueue.ExecuteSyncLoadEvents(ThreadState))
 								{
 									bPopped = true;
 									bDidSomething = true;
@@ -7382,7 +7491,7 @@ FAsyncPackage2::~FAsyncPackage2()
 #endif
 	ImportStore.ReleasePackageReference(Desc);
 
-	checkf(RefCount == 0, TEXT("RefCount is not 0 when deleting package %s"),
+	checkf(RefCount.load(std::memory_order_relaxed) == 0, TEXT("RefCount is not 0 when deleting package %s"),
 		*Desc.PackagePathToLoad.GetPackageFName().ToString());
 
 	checkf(ConstructedObjects.Num() == 0, TEXT("ClearConstructedObjects() has not been called for package %s"),
@@ -7393,8 +7502,9 @@ FAsyncPackage2::~FAsyncPackage2()
 
 void FAsyncPackage2::ReleaseRef()
 {
-	check(RefCount > 0);
-	if (--RefCount == 0)
+	int32 OldRefCount = RefCount.fetch_sub(1);
+	check(OldRefCount > 0);
+	if (OldRefCount == 1)
 	{
 		FAsyncLoadingThread2& AsyncLoadingThreadLocal = AsyncLoadingThread;
 		AsyncLoadingThreadLocal.DeferredDeletePackages.Enqueue(this);
