#rnx
PR #9856: int64 instead of int32 type when copy-assigning in FBulkDataBuffer (Contributed by [REDACTED_PERSON])

- The problem was introduced in CL 16399515 which converted the ViewType from TArrayView to TArrayView64

diff --git a/Engine/Source/Runtime/CoreUObject/Public/Serialization/BulkDataBuffer.h b/Engine/Source/Runtime/CoreUObject/Public/Serialization/BulkDataBuffer.h
--- a/Engine/Source/Runtime/CoreUObject/Public/Serialization/BulkDataBuffer.h
+++ b/Engine/Source/Runtime/CoreUObject/Public/Serialization/BulkDataBuffer.h
@@ -72,7 +72,7 @@ public:
 
 		if (this != &Other)
 		{
-			const int32 BufferSize = Other.View.Num();
+			const int64 BufferSize = Other.View.Num();
 
 			DataType* BufferCopy = (DataType*)FMemory::Malloc(BufferSize);
 			FMemory::Memcpy(BufferCopy, Other.View.GetData(), BufferSize);
