EditorDomain and Localization: make the creation deterministic for PackageLocalizationNamespaceValue. This removes one source of indeterminism in EditorDomain saves.
#rnx

diff --git a/Engine/Source/Runtime/CoreUObject/Private/Internationalization/TextPackageNamespaceUtil.cpp b/Engine/Source/Runtime/CoreUObject/Private/Internationalization/TextPackageNamespaceUtil.cpp
--- a/Engine/Source/Runtime/CoreUObject/Private/Internationalization/TextPackageNamespaceUtil.cpp
+++ b/Engine/Source/Runtime/CoreUObject/Private/Internationalization/TextPackageNamespaceUtil.cpp
@@ -1,10 +1,12 @@
 // Copyright Epic Games, Inc. All Rights Reserved.
 
 #include "Internationalization/TextPackageNamespaceUtil.h"
+
+#include "Hash/Blake3.h"
 #include "Misc/Guid.h"
+#include "Misc/PackageName.h"
 #include "UObject/Package.h"
 #include "UObject/MetaData.h"
-#include "Misc/PackageName.h"
 
 #if USE_STABLE_LOCALIZATION_KEYS
 
@@ -40,7 +42,20 @@ FString FindOrAddPackageNamespace(UPackage* InPackage, const bool bCanAdd)
 				FString& PackageLocalizationNamespaceValue = PackageMetaData->RootMetaDataMap.FindOrAdd(PackageLocalizationNamespaceKey);
 				if (PackageLocalizationNamespaceValue.IsEmpty())
 				{
-					PackageLocalizationNamespaceValue = FGuid::NewGuid().ToString();
+					// Make a determinstic new guid that will vary based on the package
+					FBlake3 Builder;
+					FString PackagePath = InPackage->GetName();
+					FGuid NonUniqueGuid = InPackage->GetPersistentGuid(); // Can be the same for duplicated packages
+					Builder.Update(&NonUniqueGuid, sizeof(FGuid));
+					Builder.Update(*PackagePath, PackagePath.Len() * sizeof(**PackagePath));
+					FBlake3Hash Hash = Builder.Finalize();
+					// We use the first 16 bytes of the FIoHash to create the guid, there is
+					// no specific reason why these were chosen, we could take any pattern or combination
+					// of bytes.
+					uint32* HashBytes = (uint32*)Hash.GetBytes();
+					FGuid NewGuid(HashBytes[0], HashBytes[1], HashBytes[2], HashBytes[3]);
+
+					PackageLocalizationNamespaceValue = NewGuid.ToString();
 				}
 				return PackageLocalizationNamespaceValue;
 			}
