AsyncLoading2: Fix static analysis error "pointer was utilized before it was verified against nullptr" (introduced in cl 23285991)

diff --git a/Engine/Source/Runtime/CoreUObject/Private/Serialization/AsyncLoading2.cpp b/Engine/Source/Runtime/CoreUObject/Private/Serialization/AsyncLoading2.cpp
--- a/Engine/Source/Runtime/CoreUObject/Private/Serialization/AsyncLoading2.cpp
+++ b/Engine/Source/Runtime/CoreUObject/Private/Serialization/AsyncLoading2.cpp
@@ -781,8 +781,16 @@ public:
 		check(!bIsMissing);
 		check(!bHasFailed);
 		check(!HasPackage());
-		PackageObjectIndex = GUObjectArray.ObjectToIndex(InPackage);
-		OriginalPackageName = InPackage->GetFName();
+		if (InPackage)
+		{
+			PackageObjectIndex = GUObjectArray.ObjectToIndex(InPackage);
+			OriginalPackageName = InPackage->GetFName();
+		}
+		else
+		{
+			PackageObjectIndex = -1;
+			OriginalPackageName = FName();
+		}
 	}
 
 	void RemoveCompletedRenamedPackage()
