Fix non-unity issues.

diff --git a/Engine/Source/Runtime/CoreUObject/Private/UObject/ObjectRename.cpp b/Engine/Source/Runtime/CoreUObject/Private/UObject/ObjectRename.cpp
--- a/Engine/Source/Runtime/CoreUObject/Private/UObject/ObjectRename.cpp
+++ b/Engine/Source/Runtime/CoreUObject/Private/UObject/ObjectRename.cpp
@@ -4,6 +4,7 @@
 
 #include "Misc/StringBuilder.h"
 #include "UObject/Object.h"
+#include "UObject/Package.h"
 #include "UObject/UObjectGlobalsInternal.h"
 
 namespace UE::Object
diff --git a/Engine/Source/Runtime/CoreUObject/Public/UObject/ObjectRename.h b/Engine/Source/Runtime/CoreUObject/Public/UObject/ObjectRename.h
--- a/Engine/Source/Runtime/CoreUObject/Public/UObject/ObjectRename.h
+++ b/Engine/Source/Runtime/CoreUObject/Public/UObject/ObjectRename.h
@@ -1,5 +1,7 @@
 // Copyright Epic Games, Inc. All Rights Reserved.
 
+class UPackage;
+
 namespace UE::Object
 {
 	// Helper function to rename a package which ought to have been freed but which can't be GC'd
