Fixed bug with script Let*  commands when a custom setter is used on a property.

When SetValue_InContainer is called, we need to ensure that we use the property container associated with the l-value in the assignment. MostRecentPropertyContainer can change when evaluating the r-value, which leads to an assert or runtime error. execLet already has this fix, but the other variants don't.

diff --git a/Engine/Source/Runtime/CoreUObject/Private/UObject/ScriptCore.cpp b/Engine/Source/Runtime/CoreUObject/Private/UObject/ScriptCore.cpp
--- a/Engine/Source/Runtime/CoreUObject/Private/UObject/ScriptCore.cpp
+++ b/Engine/Source/Runtime/CoreUObject/Private/UObject/ScriptCore.cpp
@@ -2870,6 +2870,7 @@ DEFINE_FUNCTION(UObject::execLetObj)
 	}
 
 	void* ObjAddr = Stack.MostRecentPropertyAddress;
+	void* PropertyContainer = Stack.MostRecentPropertyContainer;
 	FObjectPropertyBase* ObjectProperty = CastField<FObjectPropertyBase>(Stack.MostRecentProperty);
 	if (ObjectProperty == NULL)
 	{
@@ -2889,8 +2890,8 @@ DEFINE_FUNCTION(UObject::execLetObj)
 		checkSlow(ObjectProperty);
 		if (ObjectProperty->HasSetter())
 		{
-			check(Stack.MostRecentPropertyContainer != nullptr);
-			ObjectProperty->SetValue_InContainer(Stack.MostRecentPropertyContainer, &NewValue);
+			check(PropertyContainer != nullptr);
+			ObjectProperty->SetValue_InContainer(PropertyContainer, &NewValue);
 		}
 		else
 		{
@@ -2916,6 +2917,7 @@ DEFINE_FUNCTION(UObject::execLetWeakObjPtr)
 	}
 
 	void* ObjAddr = Stack.MostRecentPropertyAddress;
+	void* PropertyContainer = Stack.MostRecentPropertyContainer;
 	FObjectPropertyBase* ObjectProperty = CastField<FObjectPropertyBase>(Stack.MostRecentProperty);
 	if (ObjectProperty == NULL)
 	{
@@ -2935,9 +2937,9 @@ DEFINE_FUNCTION(UObject::execLetWeakObjPtr)
 		checkSlow(ObjectProperty);
 		if (ObjectProperty->HasSetter())
 		{
-			check(Stack.MostRecentPropertyContainer != nullptr);
+			check(PropertyContainer != nullptr);
 			FWeakObjectPtr NewWeakPtrValue(NewValue);
-			ObjectProperty->SetValue_InContainer(Stack.MostRecentPropertyContainer, &NewWeakPtrValue);
+			ObjectProperty->SetValue_InContainer(PropertyContainer, &NewWeakPtrValue);
 		}
 		else
 		{
@@ -2970,6 +2972,7 @@ DEFINE_FUNCTION(UObject::execLetBool)
 		we'll need to check for that.
 	*/
 	uint8* BoolAddr = (uint8*)Stack.MostRecentPropertyAddress;
+	void* PropertyContainer = Stack.MostRecentPropertyContainer;
 	FBoolProperty* BoolProperty = ExactCastField<FBoolProperty>(Stack.MostRecentProperty);
 	if (BoolProperty == NULL)
 	{
@@ -2989,8 +2992,8 @@ DEFINE_FUNCTION(UObject::execLetBool)
 		checkSlow(CastField<FBoolProperty>(BoolProperty));
 		if (BoolProperty->HasSetter())
 		{
-			check(Stack.MostRecentPropertyContainer != nullptr);
-			BoolProperty->SetValue_InContainer(Stack.MostRecentPropertyContainer, &NewValue);
+			check(PropertyContainer != nullptr);
+			BoolProperty->SetValue_InContainer(PropertyContainer, &NewValue);
 		}
 		else
 		{
