Removed multiple temporary FString allocations.

Before: 372ms
After: 149ms

diff --git a/Engine/Source/Runtime/CoreUObject/Private/UObject/PropertyBaseObject.cpp b/Engine/Source/Runtime/CoreUObject/Private/UObject/PropertyBaseObject.cpp
--- a/Engine/Source/Runtime/CoreUObject/Private/UObject/PropertyBaseObject.cpp
+++ b/Engine/Source/Runtime/CoreUObject/Private/UObject/PropertyBaseObject.cpp
@@ -181,7 +181,9 @@ void FObjectPropertyBase::AddReferencedObjects(FReferenceCollector& Collector)
 
 FString FObjectPropertyBase::GetExportPath(FTopLevelAssetPath ClassPathName, const FString& ObjectPathName)
 {
-	return FString::Printf(TEXT("%s.%s'%s'"), *ClassPathName.GetPackageName().ToString(), *ClassPathName.GetAssetName().ToString(), *ObjectPathName);
+	TStringBuilder<256> StringBuilder;
+	StringBuilder << ClassPathName.GetPackageName() << "." << ClassPathName.GetAssetName() << "'" << ObjectPathName << "'";
+	return FString(StringBuilder);
 }
 
 FString FObjectPropertyBase::GetExportPath(const TObjectPtr<const UObject>& Object, const UObject* Parent /*= nullptr*/, const UObject* ExportRootScope /*= nullptr*/, const uint32 PortFlags /*= PPF_None*/)
