Fix missing header guard.

diff --git a/Engine/Source/Runtime/CoreUObject/Public/UObject/ObjectRename.h b/Engine/Source/Runtime/CoreUObject/Public/UObject/ObjectRename.h
--- a/Engine/Source/Runtime/CoreUObject/Public/UObject/ObjectRename.h
+++ b/Engine/Source/Runtime/CoreUObject/Public/UObject/ObjectRename.h
@@ -1,5 +1,7 @@
 // Copyright Epic Games, Inc. All Rights Reserved.
 
+#pragma once
+
 #include "CoreTypes.h"
 #include "HAL/Platform.h"
 
