Fixed regression of CPF_HasGetValueTypeHash not being set on FBoolProperty after property constructor change in CL# 19033587.

diff --git a/Engine/Source/Runtime/CoreUObject/Private/UObject/PropertyBool.cpp b/Engine/Source/Runtime/CoreUObject/Private/UObject/PropertyBool.cpp
--- a/Engine/Source/Runtime/CoreUObject/Private/UObject/PropertyBool.cpp
+++ b/Engine/Source/Runtime/CoreUObject/Private/UObject/PropertyBool.cpp
@@ -38,7 +38,7 @@ PRAGMA_ENABLE_DEPRECATION_WARNINGS
 }
 
 FBoolProperty::FBoolProperty(FFieldVariant InOwner, const UECodeGen_Private::FBoolPropertyParams& Prop)
-	: FProperty(InOwner, (const UECodeGen_Private::FPropertyParamsBaseWithoutOffset&)Prop)
+	: FProperty(InOwner, (const UECodeGen_Private::FPropertyParamsBaseWithoutOffset&)Prop, CPF_HasGetValueTypeHash)
 	, FieldSize(0)
 	, ByteOffset(0)
 	, ByteMask(1)
