EditorDomain: Implement support for PackageTrailers in the same package as BulkData in a regular BulkData section for PackageWriters, so that EditorDomain can have both references to BulkData in a WorkspaceDomain PackageTrailer and embedded BulkDatas that needed to be resaved.
Because the PackageFooter comes before the PackageTrailer, we need to change accounting for the PackageFooter in PackageWriters to include adjusting the VirtualExportsFileOffset that is passed into the PackageTrailer.
#rnx

diff --git a/Engine/Source/Developer/IoStoreUtilities/Private/ZenStoreWriter.cpp b/Engine/Source/Developer/IoStoreUtilities/Private/ZenStoreWriter.cpp
--- a/Engine/Source/Developer/IoStoreUtilities/Private/ZenStoreWriter.cpp
+++ b/Engine/Source/Developer/IoStoreUtilities/Private/ZenStoreWriter.cpp
@@ -352,6 +352,12 @@ void FZenStoreWriter::WriteLinkerAdditionalData(const FLinkerAdditionalDataInfo&
 	checkNoEntry();
 }
 
+void FZenStoreWriter::WritePackageTrailer(const FPackageTrailerInfo& Info, const FIoBuffer& Data)
+{
+	// PackageTrailers are not yet implemented in this writer; it is only used for EditorBulkData which is not used in cooked content
+	checkNoEntry();
+}
+
 void FZenStoreWriter::Initialize(const FCookInfo& Info)
 {
 	CookMode = Info.CookMode;
diff --git a/Engine/Source/Developer/IoStoreUtilities/Public/ZenStoreWriter.h b/Engine/Source/Developer/IoStoreUtilities/Public/ZenStoreWriter.h
--- a/Engine/Source/Developer/IoStoreUtilities/Public/ZenStoreWriter.h
+++ b/Engine/Source/Developer/IoStoreUtilities/Public/ZenStoreWriter.h
@@ -67,6 +67,8 @@ public:
 	IOSTOREUTILITIES_API virtual void WritePackageData(const FPackageInfo& Info, FLargeMemoryWriter& ExportsArchive, const TArray<FFileRegion>& FileRegions) override;
 	IOSTOREUTILITIES_API virtual void WriteAdditionalFile(const FAdditionalFileInfo& Info, const FIoBuffer& FileData) override;
 	IOSTOREUTILITIES_API virtual void WriteLinkerAdditionalData(const FLinkerAdditionalDataInfo& Info, const FIoBuffer& Data, const TArray<FFileRegion>& FileRegions) override;
+	IOSTOREUTILITIES_API virtual void WritePackageTrailer(const FPackageTrailerInfo& Info, const FIoBuffer& Data) override;
+
 
 	IOSTOREUTILITIES_API virtual void WriteBulkData(const FBulkDataInfo& Info, const FIoBuffer& BulkData, const TArray<FFileRegion>& FileRegions) override;
 	IOSTOREUTILITIES_API virtual void Initialize(const FCookInfo& Info) override;
diff --git a/Engine/Source/Editor/UnrealEd/Private/Cooker/DiffPackageWriter.h b/Engine/Source/Editor/UnrealEd/Private/Cooker/DiffPackageWriter.h
--- a/Engine/Source/Editor/UnrealEd/Private/Cooker/DiffPackageWriter.h
+++ b/Engine/Source/Editor/UnrealEd/Private/Cooker/DiffPackageWriter.h
@@ -37,9 +37,13 @@ public:
 	{
 		Inner->WriteLinkerAdditionalData(Info, Data, FileRegions);
 	}
-	virtual void AddToExportsSize(int64& ExportsSize) override
+	virtual void WritePackageTrailer(const FPackageTrailerInfo& Info, const FIoBuffer& Data) override
 	{
-		Inner->AddToExportsSize(ExportsSize);
+		Inner->WritePackageTrailer(Info, Data);
+	}
+	virtual int64 GetExportsFooterSize() override
+	{
+		return Inner->GetExportsFooterSize();
 	}
 	virtual TUniquePtr<FLargeMemoryWriter> CreateLinkerArchive(FName PackageName, UObject* Asset) override;
 	virtual bool IsPreSaveCompleted() const override
@@ -159,9 +163,13 @@ public:
 	{
 		Inner->WriteLinkerAdditionalData(Info, Data, FileRegions);
 	}
-	virtual void AddToExportsSize(int64& ExportsSize) override
+	virtual void WritePackageTrailer(const FPackageTrailerInfo& Info, const FIoBuffer& Data) override
+	{
+		Inner->WritePackageTrailer(Info, Data);
+	}
+	virtual int64 GetExportsFooterSize() override
 	{
-		Inner->AddToExportsSize(ExportsSize);
+		return Inner->GetExportsFooterSize();
 	}
 	virtual TUniquePtr<FLargeMemoryWriter> CreateLinkerArchive(FName PackageName, UObject* Asset) override
 	{
diff --git a/Engine/Source/Editor/UnrealEd/Private/Cooker/LooseCookedPackageWriter.cpp b/Engine/Source/Editor/UnrealEd/Private/Cooker/LooseCookedPackageWriter.cpp
--- a/Engine/Source/Editor/UnrealEd/Private/Cooker/LooseCookedPackageWriter.cpp
+++ b/Engine/Source/Editor/UnrealEd/Private/Cooker/LooseCookedPackageWriter.cpp
@@ -101,6 +101,7 @@ void FLooseCookedPackageWriter::AsyncSave(FRecord& Record, const FCommitPackageI
 	CollectForSaveLinkerAdditionalDataRecords(Record, Context);
 	CollectForSaveAdditionalFileRecords(Record, Context);
 	CollectForSaveExportsFooter(Record, Context);
+	CollectForSaveExportsPackageTrailer(Record, Context);
 	CollectForSaveExportsBuffers(Record, Context);
 
 	AsyncSaveOutputFiles(Record, Context);
@@ -132,6 +133,15 @@ void FLooseCookedPackageWriter::CompleteExportsArchiveForDiff(const FPackageInfo
 
 	uint32 FooterData = PACKAGE_FILE_TAG;
 	ExportsArchive << FooterData;
+
+	for (FPackageTrailerRecord& PackageTrailer : Record.PackageTrailers)
+	{
+		if (PackageTrailer.Info.MultiOutputIndex == Info.MultiOutputIndex)
+		{
+			ExportsArchive.Serialize(const_cast<void*>(PackageTrailer.Buffer.GetData()),
+				PackageTrailer.Buffer.GetSize());
+		}
+	}
 }
 
 
@@ -212,9 +222,24 @@ void FLooseCookedPackageWriter::CollectForSaveExportsFooter(FRecord& Record, FCo
 	}
 }
 
-void FLooseCookedPackageWriter::AddToExportsSize(int64& ExportsSize)
+int64 FLooseCookedPackageWriter::GetExportsFooterSize()
+{
+	return sizeof(uint32);
+}
+
+void FLooseCookedPackageWriter::CollectForSaveExportsPackageTrailer(FRecord& Record, FCommitContext& Context)
 {
-	ExportsSize += sizeof(uint32); // Footer size
+	if (Record.bCompletedExportsArchiveForDiff)
+	{
+		// Already Added in CompleteExportsArchiveForDiff
+		return;
+	}
+
+	for (FPackageTrailerRecord& PackageTrailer : Record.PackageTrailers)
+	{
+		Context.ExportsBuffers[PackageTrailer.Info.MultiOutputIndex].Add(
+			FExportBuffer{ PackageTrailer.Buffer, TArray<FFileRegion>() });
+	}
 }
 
 void FLooseCookedPackageWriter::CollectForSaveExportsBuffers(FRecord& Record, FCommitContext& Context)
diff --git a/Engine/Source/Editor/UnrealEd/Private/Cooker/LooseCookedPackageWriter.h b/Engine/Source/Editor/UnrealEd/Private/Cooker/LooseCookedPackageWriter.h
--- a/Engine/Source/Editor/UnrealEd/Private/Cooker/LooseCookedPackageWriter.h
+++ b/Engine/Source/Editor/UnrealEd/Private/Cooker/LooseCookedPackageWriter.h
@@ -51,7 +51,7 @@ public:
 	}
 
 	virtual void BeginPackage(const FBeginPackageInfo& Info) override;
-	virtual void AddToExportsSize(int64& ExportsSize) override;
+	virtual int64 GetExportsFooterSize() override;
 
 	virtual FDateTime GetPreviousCookTime() const override;
 	virtual void Initialize(const FCookInfo& Info) override;
@@ -131,6 +131,7 @@ private:
 	void CollectForSaveLinkerAdditionalDataRecords(FRecord& Record, FCommitContext& Context);
 	void CollectForSaveAdditionalFileRecords(FRecord& Record, FCommitContext& Context);
 	void CollectForSaveExportsFooter(FRecord& Record, FCommitContext& Context);
+	void CollectForSaveExportsPackageTrailer(FRecord& Record, FCommitContext& Context);
 	void CollectForSaveExportsBuffers(FRecord& Record, FCommitContext& Context);
 	void AsyncSaveOutputFiles(FRecord& Record, FCommitContext& Context);
 	void UpdateManifest(FRecord& Record);
diff --git a/Engine/Source/Editor/UnrealEd/Private/EditorDomain/EditorDomainUtils.cpp b/Engine/Source/Editor/UnrealEd/Private/EditorDomain/EditorDomainUtils.cpp
--- a/Engine/Source/Editor/UnrealEd/Private/EditorDomain/EditorDomainUtils.cpp
+++ b/Engine/Source/Editor/UnrealEd/Private/EditorDomain/EditorDomainUtils.cpp
@@ -203,7 +203,7 @@ FBlake3Hash GGlobalConstructClassesHash;
 int64 GMaxBulkDataSize = -1;
 
 // Change to a new guid when EditorDomain needs to be invalidated
-const TCHAR* EditorDomainVersion = TEXT("6D2972C7108F4366A13B7886DEDE68D5");
+const TCHAR* EditorDomainVersion = TEXT("74B1CEDEF6774CC89F3A45BEA18B9B8B");
 
 // Identifier of the CacheBuckets for EditorDomain tables
 const TCHAR* EditorDomainPackageBucketName = TEXT("EditorDomainPackage");
@@ -1569,7 +1569,7 @@ public:
 	/** Deserialize the CustomVersions out of the PackageFileSummary that was serialized into the header */
 	bool TryGetClassesAndVersions(TArray<FName>& OutImportedClasses, FCustomVersionContainer& OutVersions)
 	{
-		FMemoryReaderView HeaderArchive(WritePackageRecord.Buffer.GetView());
+		FMemoryReaderView HeaderArchive(SavedRecord.Packages[0].Buffer.GetView());
 		FPackageReader PackageReader;
 		if (!PackageReader.OpenPackageFile(&HeaderArchive))
 		{
@@ -1628,9 +1628,8 @@ protected:
 		}
 		
 		// EditorDomain save does not support multi package outputs
-		check(Record.Packages.Num() == 1);
+		checkf(Record.Packages.Num() == 1, TEXT("Multioutput not supported"));
 		FPackageWriterRecords::FWritePackage& Package = Record.Packages[0];
-		WritePackageRecord = Package;
 
 		TArray<FSharedBuffer> AttachmentBuffers;
 		for (const FFileRegion& FileRegion : Package.Regions)
@@ -1683,6 +1682,15 @@ protected:
 				TEXT("Expects all LinkerAdditionalData to be in a region."));
 			BulkDataSize += AdditionalRecord.Buffer.GetSize();
 		}
+		checkf(Record.PackageTrailers.Num() <= 1, TEXT("MultiOutput not supported"));
+		if (Record.PackageTrailers.Num() == 1)
+		{
+			FPackageTrailerRecord& PackageTrailer = Record.PackageTrailers[0];
+			const uint8* BufferStart = reinterpret_cast<const uint8*>(PackageTrailer.Buffer.GetData());
+			int64 BufferSize = PackageTrailer.Buffer.GetSize();
+			AttachmentBuffers.Add(FSharedBuffer::MakeView(BufferStart, BufferSize));
+			BulkDataSize += BufferSize;
+		}
 
 		// We use a counter for ValueIds rather than hashes of the Attachments. We do this because
 		// some attachments may be identical, and Attachments are not allowed to have identical ValueIds.
@@ -1709,12 +1717,15 @@ protected:
 			Attachments.Add(FAttachment{ Buffer, IntToValueId(AttachmentIndex++) });
 			FileSize += Buffer.GetSize();
 		}
-		WritePackageRecord = MoveTemp(Package);
+
+		// Our Attachments are views into SharedBuffers that are stored on the Record. Keep those SharedBuffers
+		// alive until we are done with them. We also need access to Record.Packages[0] later.
+		SavedRecord = MoveTemp(Record);
 	}
 
 private:
 	TArray<FAttachment> Attachments;
-	FPackageWriterRecords::FWritePackage WritePackageRecord;
+	FPackageWriterRecords::FPackage SavedRecord;
 	uint64 FileSize = 0;
 	uint64 BulkDataSize = 0;
 };
diff --git a/Engine/Source/Runtime/Core/Public/Serialization/PackageWriter.h b/Engine/Source/Runtime/Core/Public/Serialization/PackageWriter.h
--- a/Engine/Source/Runtime/Core/Public/Serialization/PackageWriter.h
+++ b/Engine/Source/Runtime/Core/Public/Serialization/PackageWriter.h
@@ -156,13 +156,20 @@ public:
 	/** Write separate data written by UObjects via FLinkerSave::AdditionalDataToAppend. */
 	virtual void WriteLinkerAdditionalData(const FLinkerAdditionalDataInfo& Info, const FIoBuffer& Data, const TArray<FFileRegion>& FileRegions) = 0;
 
-	/** Increase the referenced ExportsSize by the size in bytes of the data that will be added on to it during
-	 * commit before writing to disk. Used for accurate disk size reporting on the UPackage and AssetRegistry.
-	 */
-	virtual void AddToExportsSize(int64& InOutExportsSize)
+	/** Report the size of the Footer that is added after Exports and BulkData but before the PackageTrailer */ 
+	virtual int64 GetExportsFooterSize()
 	{
+		return 0;
 	}
 
+	struct FPackageTrailerInfo
+	{
+		FName PackageName;
+		uint16 MultiOutputIndex = 0;
+	};
+	/** Write the PackageTrailer, a separate segment for some bulkdata that is written the end of the file. */
+	virtual void WritePackageTrailer(const FPackageTrailerInfo& Info, const FIoBuffer& Data) = 0;
+
 	/** Create the FLargeMemoryWriter to which the Header and Exports are written during the save. */
 	COREUOBJECT_API virtual TUniquePtr<FLargeMemoryWriter> CreateLinkerArchive(FName PackageName, UObject* Asset);
 
diff --git a/Engine/Source/Runtime/CoreUObject/Internal/Serialization/PackageWriterToSharedBuffer.h b/Engine/Source/Runtime/CoreUObject/Internal/Serialization/PackageWriterToSharedBuffer.h
--- a/Engine/Source/Runtime/CoreUObject/Internal/Serialization/PackageWriterToSharedBuffer.h
+++ b/Engine/Source/Runtime/CoreUObject/Internal/Serialization/PackageWriterToSharedBuffer.h
@@ -6,7 +6,6 @@
 #include "HAL/CriticalSection.h"
 #include "Memory/SharedBuffer.h"
 #include "Misc/AssertionMacros.h"
-#include "Misc/Optional.h"
 #include "Serialization/PackageWriter.h"
 #include "Templates/UniquePtr.h"
 
@@ -24,6 +23,7 @@ public:
 		const FIoBuffer& FileData);
 	COREUOBJECT_API void WriteLinkerAdditionalData(const IPackageWriter::FLinkerAdditionalDataInfo& Info,
 		const FIoBuffer& Data, const TArray<FFileRegion>& FileRegions);
+	COREUOBJECT_API void WritePackageTrailer(const IPackageWriter::FPackageTrailerInfo& Info, const FIoBuffer& Data);
 
 	struct FWritePackage
 	{
@@ -48,6 +48,11 @@ public:
 		FSharedBuffer Buffer;
 		TArray<FFileRegion> Regions;
 	};
+	struct FPackageTrailer
+	{
+		IPackageWriter::FPackageTrailerInfo Info;
+		FSharedBuffer Buffer;
+	};
 	struct FPackage
 	{
 		/** TPackageWriterToSharedBuffer can allocate FPackage subclasses, so the base needs a virtual destructor. */
@@ -61,7 +66,8 @@ public:
 		TArray<FWritePackage> Packages;
 		TArray<FBulkData> BulkDatas;
 		TArray<FAdditionalFile> AdditionalFiles;
-		TArray<FLinkerAdditionalData>  LinkerAdditionalDatas;
+		TArray<FLinkerAdditionalData> LinkerAdditionalDatas;
+		TArray<FPackageTrailer> PackageTrailers;
 	};
 
 	/** Get the Record created by BeginPackage for the given PackageName; assert that it is valid */
@@ -91,6 +97,7 @@ public:
 	using FBulkDataRecord = FPackageWriterRecords::FBulkData;
 	using FAdditionalFileRecord = FPackageWriterRecords::FAdditionalFile;
 	using FLinkerAdditionalDataRecord = FPackageWriterRecords::FLinkerAdditionalData;
+	using FPackageTrailerRecord = FPackageWriterRecords::FPackageTrailer;
 	using FPackageRecord = FPackageWriterRecords::FPackage;
 
 	virtual void BeginPackage(const IPackageWriter::FBeginPackageInfo& Info) override
@@ -117,6 +124,10 @@ public:
 	{
 		Records.WriteLinkerAdditionalData(Info, Data, FileRegions);
 	}
+	virtual void WritePackageTrailer(const IPackageWriter::FPackageTrailerInfo& Info, const FIoBuffer& Data) override
+	{
+		Records.WritePackageTrailer(Info, Data);
+	}
 	virtual void CommitPackage(IPackageWriter::FCommitPackageInfo&& Info) override
 	{
 		TUniquePtr<FPackageRecord> Record = Records.FindAndRemoveRecordChecked(Info.PackageName);
diff --git a/Engine/Source/Runtime/CoreUObject/Private/Serialization/PackageWriterToSharedBuffer.cpp b/Engine/Source/Runtime/CoreUObject/Private/Serialization/PackageWriterToSharedBuffer.cpp
--- a/Engine/Source/Runtime/CoreUObject/Private/Serialization/PackageWriterToSharedBuffer.cpp
+++ b/Engine/Source/Runtime/CoreUObject/Private/Serialization/PackageWriterToSharedBuffer.cpp
@@ -43,6 +43,15 @@ void FPackageWriterRecords::WritePackageData(const IPackageWriter::FPackageInfo&
 	checkf(DataSize > 0, TEXT("IPackageWriter->WritePackageData must not be called with an empty ExportsArchive"));
 	checkf(static_cast<uint64>(DataSize) >= Info.HeaderSize,
 		TEXT("IPackageWriter->WritePackageData must not be called with HeaderSize > ExportsArchive.TotalSize"));
+	if (Record.Packages.FindByPredicate(
+		[&Info](const FWritePackage& Existing) { return Existing.Info.MultiOutputIndex == Info.MultiOutputIndex; }))
+	{
+		UE_LOG(LogPackageWriter, Error,
+			TEXT("PackageWriter->WritePackageData called more than once for package (Package: %s). Ignoring second value."),
+			*Info.PackageName.ToString());
+		return;
+	}
+
 	FSharedBuffer Buffer = FSharedBuffer::TakeOwnership(ExportsArchive.ReleaseOwnership(), DataSize,
 		FMemory::Free);
 	Record.Packages.Insert(FWritePackage{ Info, MoveTemp(Buffer), FileRegions }, Info.MultiOutputIndex);
@@ -79,6 +88,20 @@ void FPackageWriterRecords::WriteLinkerAdditionalData(const IPackageWriter::FLin
 		FLinkerAdditionalData{ Info, IoBufferToSharedBuffer(Data), FileRegions });
 }
 
+void FPackageWriterRecords::WritePackageTrailer(const IPackageWriter::FPackageTrailerInfo& Info, const FIoBuffer& Data)
+{
+	FPackage& Record = FindRecordChecked(Info.PackageName);
+	if (Record.PackageTrailers.FindByPredicate([&Info](const FPackageTrailer& Existing) { return Existing.Info.MultiOutputIndex == Info.MultiOutputIndex; }))
+	{
+		UE_LOG(LogPackageWriter, Error,
+			TEXT("PackageWriter->WritePackageTrailer called more than once for package (Package: %s). Ignoring second value."),
+			*Info.PackageName.ToString());
+		return;
+	}
+
+	Record.PackageTrailers.Emplace(FPackageTrailer { Info, IoBufferToSharedBuffer(Data) });
+}
+
 FPackageWriterRecords::FPackage& FPackageWriterRecords::FindRecordChecked(FName InPackageName) const
 {
 	const TUniquePtr<FPackage>* Record;
diff --git a/Engine/Source/Runtime/CoreUObject/Private/UObject/PackageTrailer.cpp b/Engine/Source/Runtime/CoreUObject/Private/UObject/PackageTrailer.cpp
--- a/Engine/Source/Runtime/CoreUObject/Private/UObject/PackageTrailer.cpp
+++ b/Engine/Source/Runtime/CoreUObject/Private/UObject/PackageTrailer.cpp
@@ -288,6 +288,12 @@ bool FPackageTrailerBuilder::UpdatePayloadAsLocal(const FIoHash& Identifier, FCo
 }
 
 bool FPackageTrailerBuilder::BuildAndAppendTrailer(FLinkerSave* Linker, FArchive& DataArchive)
+{
+	int64 CurrentOffset = DataArchive.Tell();
+	return BuildAndAppendTrailer(Linker, DataArchive, CurrentOffset);
+}
+
+bool FPackageTrailerBuilder::BuildAndAppendTrailer(FLinkerSave* Linker, FArchive& DataArchive, int64& InOutPackageFileOffset)
 {
 	// Note that we do not serialize containers directly as we want a file format that is 
 	// 100% under our control. This will allow people to create external scripts that can
@@ -348,34 +354,35 @@ bool FPackageTrailerBuilder::BuildAndAppendTrailer(FLinkerSave* Linker, FArchive
 
 	// Now that we have the complete trailer we can serialize it to the archive
 
-	Trailer.TrailerPositionInFile = DataArchive.Tell();
+	Trailer.TrailerPositionInFile = InOutPackageFileOffset;
 
+	const int64 TrailerPositionInDataArchive = DataArchive.Tell();
 	DataArchive << Trailer.Header;
 
-	checkf((Trailer.TrailerPositionInFile + Trailer.Header.HeaderLength) == DataArchive.Tell(), 
+	checkf((TrailerPositionInDataArchive + Trailer.Header.HeaderLength) == DataArchive.Tell(),
 		TEXT("Header length was calculated as %d bytes but we wrote %" INT64_FMT " bytes!"), 
 		Trailer.Header.HeaderLength, 
-		DataArchive.Tell() - Trailer.TrailerPositionInFile);
+		DataArchive.Tell() - TrailerPositionInDataArchive);
 
-	const int64 PayloadPosInFile = DataArchive.Tell();
+	const int64 PayloadPosInDataArchive = DataArchive.Tell();
 
 	for (TPair<FIoHash, LocalEntry>& It : LocalEntries)
 	{
 		DataArchive << It.Value.Payload;
 	}
 
-	checkf((PayloadPosInFile + Trailer.Header.PayloadsDataLength) == DataArchive.Tell(), 
+	checkf((PayloadPosInDataArchive + Trailer.Header.PayloadsDataLength) == DataArchive.Tell(),
 		TEXT("Total payload length was calculated as %" INT64_FMT " bytes but we wrote %" INT64_FMT " bytes!"), 
 		Trailer.Header.PayloadsDataLength, 
-		DataArchive.Tell() - PayloadPosInFile);
+		DataArchive.Tell() - PayloadPosInDataArchive);
 
 	FPackageTrailer::FFooter Footer = Trailer.CreateFooter();
 	DataArchive << Footer;
 
-	checkf((Trailer.TrailerPositionInFile + Footer.TrailerLength) == DataArchive.Tell(), 
+	checkf((TrailerPositionInDataArchive + Footer.TrailerLength) == DataArchive.Tell(),
 		TEXT("Trailer length was calculated as %" INT64_FMT " bytes but we wrote %" INT64_FMT " bytes!"), 
 		Footer.TrailerLength, 
-		DataArchive.Tell() - Trailer.TrailerPositionInFile);
+		DataArchive.Tell() - TrailerPositionInDataArchive);
 
 	// Invoke any registered callbacks and pass in the trailer, this allows the callbacks to poll where 
 	// in the output archive the payload has been stored.
@@ -388,7 +395,8 @@ bool FPackageTrailerBuilder::BuildAndAppendTrailer(FLinkerSave* Linker, FArchive
 	}
 
 	// Minor sanity check that ::GetTrailerLength works
-	check(CalculateTrailerLength() == (DataArchive.Tell() - Trailer.TrailerPositionInFile) || DataArchive.IsError());
+	check(CalculateTrailerLength() == (DataArchive.Tell() - TrailerPositionInDataArchive) || DataArchive.IsError());
+	InOutPackageFileOffset += DataArchive.Tell() - TrailerPositionInDataArchive;
 
 	return !DataArchive.IsError();
 }
diff --git a/Engine/Source/Runtime/CoreUObject/Private/UObject/SavePackage2.cpp b/Engine/Source/Runtime/CoreUObject/Private/UObject/SavePackage2.cpp
--- a/Engine/Source/Runtime/CoreUObject/Private/UObject/SavePackage2.cpp
+++ b/Engine/Source/Runtime/CoreUObject/Private/UObject/SavePackage2.cpp
@@ -61,7 +61,7 @@ static FAutoConsoleVariableRef CVar_FixupStandaloneFlags(
 	TEXT("If non-zero, when the UAsset of a package is missing RF_Standalone, the flag is added. If zero, the flags are not changed and the save fails.")
 );
 
-ESavePackageResult WriteAdditionalFiles(FSaveContext& SaveContext, FScopedSlowTask& SlowTask, int64 LinkerSize);
+ESavePackageResult WriteAdditionalFiles(FSaveContext& SaveContext, FScopedSlowTask& SlowTask, int64& VirtualExportsFileOffset);
 
 ESavePackageResult ReturnSuccessOrCancel()
 {
@@ -905,7 +905,8 @@ ESavePackageResult CreateLinker(FSaveContext& SaveContext)
 			{
 				if (const FLinkerLoad* LinkerLoad = FLinkerLoad::FindExistingLinkerForPackage(SaveContext.GetPackage()))
 				{
-					if (const UE::FPackageTrailer* Trailer = LinkerLoad->GetPackageTrailer())
+					const UE::FPackageTrailer* Trailer = LinkerLoad->GetPackageTrailer();
+					if (Trailer && Trailer->GetNumPayloads(UE::EPayloadStorageType::Any) > 0)
 					{
 						SaveContext.GetLinker()->PackageTrailerBuilder = UE::FPackageTrailerBuilder::CreateReferenceToTrailer(*Trailer, SaveContext.GetPackage()->GetName());
 					}
@@ -1901,7 +1902,7 @@ ESavePackageResult WriteExports(FStructuredArchive::FRecord& StructuredArchiveRo
 	return Linker->IsError() ? ESavePackageResult::Error : ReturnSuccessOrCancel();
 }
 
-[[nodiscard]] ESavePackageResult BuildAndWriteTrailer(FStructuredArchive::FRecord& StructuredArchiveRoot, FSaveContext& SaveContext)
+[[nodiscard]] ESavePackageResult BuildAndWriteTrailer(IPackageWriter* PackageWriter, FStructuredArchive::FRecord& StructuredArchiveRoot, FSaveContext& SaveContext, int64& InOutCurrentOffset)
 {
 	SaveContext.GetLinker()->Summary.PayloadTocOffset = INDEX_NONE;
 
@@ -1917,10 +1918,25 @@ ESavePackageResult WriteExports(FStructuredArchive::FRecord& StructuredArchiveRo
 
 		checkf(SaveContext.IsTextFormat() == false, TEXT("Attempting to build a package trailer for text based asset '%s', this is not supported!"), *SaveContext.GetPackage()->GetName());
 
-		SaveContext.GetLinker()->Summary.PayloadTocOffset = SaveContext.GetLinker()->Tell();
-		if (!SaveContext.GetLinker()->PackageTrailerBuilder->BuildAndAppendTrailer(SaveContext.GetLinker(), *SaveContext.GetLinker()))
+		SaveContext.GetLinker()->Summary.PayloadTocOffset = InOutCurrentOffset;
+		if (!PackageWriter)
 		{
-			return ESavePackageResult::Error;
+			if (!SaveContext.GetLinker()->PackageTrailerBuilder->BuildAndAppendTrailer(SaveContext.GetLinker(), *SaveContext.GetLinker(), InOutCurrentOffset))
+			{
+				return ESavePackageResult::Error;
+			}
+		}
+		else
+		{
+			FLargeMemoryWriter TrailerData(0, true /* IsPersistent */);
+			if (!SaveContext.GetLinker()->PackageTrailerBuilder->BuildAndAppendTrailer(SaveContext.GetLinker(), TrailerData, InOutCurrentOffset))
+			{
+				return ESavePackageResult::Error;
+			}
+
+			IPackageWriter::FPackageTrailerInfo TrailerInfo;
+			TrailerInfo.PackageName = SaveContext.GetPackage()->GetFName();
+			PackageWriter->WritePackageTrailer(TrailerInfo, FIoBuffer(FIoBuffer::AssumeOwnership, TrailerData.ReleaseOwnership(), TrailerData.TotalSize()));
 		}
 
 		SaveContext.GetLinker()->PackageTrailerBuilder.Reset();
@@ -2411,53 +2427,55 @@ ESavePackageResult SaveHarvestedRealms(FSaveContext& SaveContext, ESaveRealm Har
 	}
 
 	IPackageWriter* PackageWriter = SaveContext.GetPackageWriter();
+	const int64 EndOfExportsOffset = SaveContext.GetLinker()->Tell();
+	// When not using a PackageWriter, VirtualExportsFileOffset is identical to the offset in the
+	// Exports archive: SaveContext.GetLinker()->Tell. When using a PackageWriter however, additional
+	// blobs such as bulkdata are not written into the exports archive, they are stored as separate archives
+	// in the PackageWriter. But various structs need to know the "offset" in the combined file that would
+	// be created by appending all of these blobs after the exports. VirtualExportsFileOffset holds that value.
+	int64 VirtualExportsFileOffset = EndOfExportsOffset;
+	SaveContext.Result = WriteAdditionalFiles(SaveContext, SlowTask, VirtualExportsFileOffset /* In/Out */);
+	if (SaveContext.Result != ESavePackageResult::Success)
+	{
+		return SaveContext.Result;
+	}
+
+	// Write out a tag at the end of the exports and AdditionalFiles
 	if (PackageWriter)
 	{
-		const int64 ExportsSize = SaveContext.GetLinker()->Tell();
-		SaveContext.Result = WriteAdditionalFiles(SaveContext, SlowTask, ExportsSize);
-		checkf(SaveContext.GetLinker()->Tell() == ExportsSize, TEXT("The writing of additional files is not allowed to append to the LinkerSave when using a PackageWriter."));
-		if (SaveContext.Result != ESavePackageResult::Success)
-		{
-			return SaveContext.Result;
-		}
+		VirtualExportsFileOffset += PackageWriter->GetExportsFooterSize();
 	}
 	else
 	{
-		// AdditionalFiles are appended to the Linker's archive, and so must be appended before we can calculate the full size of the Package
-		SaveContext.Result = WriteAdditionalFiles(SaveContext, SlowTask, -1);
-		if (SaveContext.Result != ESavePackageResult::Success)
+		if (!SaveContext.IsTextFormat())
 		{
-			return SaveContext.Result;
+			uint32 Tag = PACKAGE_FILE_TAG;
+			StructuredArchiveRoot.GetUnderlyingArchive() << Tag;
+			VirtualExportsFileOffset += sizeof(Tag);
 		}
 	}
 
-	// Write out a tag to the end of the package
-	if (PackageWriter == nullptr && !SaveContext.IsTextFormat())
-	{
-		uint32 Tag = PACKAGE_FILE_TAG;
-		StructuredArchiveRoot.GetUnderlyingArchive() << Tag;
-	}
 
 	// Now that the package is written out we can write the package trailer that is appended
 	// to the file. This should be the last thing written to the file!
 	SlowTask.EnterProgressFrame();
-	SaveContext.Result = BuildAndWriteTrailer(StructuredArchiveRoot, SaveContext);
+	SaveContext.Result = BuildAndWriteTrailer(PackageWriter, StructuredArchiveRoot, SaveContext, VirtualExportsFileOffset);
 	if (SaveContext.Result != ESavePackageResult::Success)
 	{
 		return SaveContext.Result;
 	}	
-
-	int64 ExportsSize = SaveContext.GetLinker()->Tell();
 	if (PackageWriter)
 	{
-		PackageWriter->AddToExportsSize(ExportsSize);
+		checkf(SaveContext.GetLinker()->Tell() == EndOfExportsOffset,
+			TEXT("The writing of additional files is not allowed to append to the LinkerSave when using a PackageWriter."));
 	}
+
 	// Store the package header and export size of the non optional realm
 	if (SaveContext.GetCurrentHarvestingRealm() != ESaveRealm::Optional)
 	{
-		SaveContext.PackageHeaderAndExportSize = ExportsSize;
+		SaveContext.PackageHeaderAndExportSize = VirtualExportsFileOffset;
 	}
-	SaveContext.TotalPackageSizeUncompressed += ExportsSize;
+	SaveContext.TotalPackageSizeUncompressed += VirtualExportsFileOffset;
 	for (const FSavePackageOutputFile& File : SaveContext.AdditionalPackageFiles)
 	{
 		SaveContext.TotalPackageSizeUncompressed += File.DataSize;
@@ -2582,19 +2600,18 @@ ESavePackageResult InnerSave(FSaveContext& SaveContext)
  * @param SaveContext The context for the overall save, including data about the additional payloads
  * @param LinkerSize If the Linker has finished writing, this is the size of the Linker's archive: Linker->Tell(). Otherwise it is -1.
  */
-ESavePackageResult WriteAdditionalFiles(FSaveContext& SaveContext, FScopedSlowTask& SlowTask, int64 LinkerSize)
+ESavePackageResult WriteAdditionalFiles(FSaveContext& SaveContext, FScopedSlowTask& SlowTask, int64& VirtualExportsFileOffset)
 {
 	// Save Bulk Data
 	SlowTask.EnterProgressFrame();
-	int64 DataStartOffset = LinkerSize >= 0 ? LinkerSize : SaveContext.GetLinker()->Tell();
-	SaveContext.Result = SavePackageUtilities::SaveBulkData(SaveContext.GetLinker(), DataStartOffset,
+	SaveContext.Result = SavePackageUtilities::SaveBulkData(SaveContext.GetLinker(), VirtualExportsFileOffset,
 		SaveContext.GetPackage(), SaveContext.GetFilename(), SaveContext.GetTargetPlatform(),
 		SaveContext.GetSavePackageContext(), SaveContext.GetSaveArgs().SaveFlags, SaveContext.IsTextFormat(),
 		SaveContext.TotalPackageSizeUncompressed,
 		SaveContext.GetCurrentHarvestingRealm() == ESaveRealm::Optional);
 
 	// Add any pending data blobs to the end of the file by invoking the callbacks
-	ESavePackageResult Result = SavePackageUtilities::AppendAdditionalData(*SaveContext.GetLinker(), DataStartOffset, SaveContext.GetSavePackageContext());
+	ESavePackageResult Result = SavePackageUtilities::AppendAdditionalData(*SaveContext.GetLinker(), VirtualExportsFileOffset, SaveContext.GetSavePackageContext());
 	if (Result != ESavePackageResult::Success)
 	{
 		return Result;
diff --git a/Engine/Source/Runtime/CoreUObject/Public/UObject/PackageTrailer.h b/Engine/Source/Runtime/CoreUObject/Public/UObject/PackageTrailer.h
--- a/Engine/Source/Runtime/CoreUObject/Public/UObject/PackageTrailer.h
+++ b/Engine/Source/Runtime/CoreUObject/Public/UObject/PackageTrailer.h
@@ -281,7 +281,18 @@ public:
 	 * @return True if the builder was created and appended successfully and false if any error was encountered
 	 */
 	[[nodiscard]] bool BuildAndAppendTrailer(FLinkerSave* Linker, FArchive& DataArchive);
-	
+
+	/**
+	 * @param ExportsArchive	The linker associated with the package being written to disk.
+	 * @param DataArchive		The archive where the package data has been written to. This is where the FPackageTrailer will be written to
+	 * @param InOutPackageFileOffset The offset at which the trailer is written in the package file. In this function version,
+	 *        DataArchive might be an archive for the entire package file, or it might be a separate archive solely for the package trailer.
+	 *        The caller must provide the initial value with the accumulated offset, and the function will increment this value by how many
+	 *        bytes are written to DataArchive.
+	 * @return True if the builder was created and appended successfully and false if any error was encountered
+	 */
+	[[nodiscard]] bool BuildAndAppendTrailer(FLinkerSave* Linker, FArchive& DataArchive, int64& InOutPackageFileOffset);
+
 	/** Returns if the builder has any payload entries or not */
 	[[nodiscard]] bool IsEmpty() const;
 
