AsyncLoading2: Fix infinite wait for io issue caused by decreasing PendingIoRequestsCounter too late in cl 22264831.

diff --git a/Engine/Source/Runtime/CoreUObject/Private/Serialization/AsyncLoading2.cpp b/Engine/Source/Runtime/CoreUObject/Private/Serialization/AsyncLoading2.cpp
--- a/Engine/Source/Runtime/CoreUObject/Private/Serialization/AsyncLoading2.cpp
+++ b/Engine/Source/Runtime/CoreUObject/Private/Serialization/AsyncLoading2.cpp
@@ -4583,6 +4583,7 @@ void FAsyncPackage2::StartLoading(FAsyncLoadingThreadState2& ThreadState, FIoBat
 				int32 LocalPendingIoRequestsCounter = AsyncLoadingThread.PendingIoRequestsCounter.DecrementExchange() - 1;
 				TRACE_COUNTER_SET(AsyncLoadingPendingIoRequests, LocalPendingIoRequestsCounter);
 				FAsyncLoadingThread2& LocalAsyncLoadingThread = AsyncLoadingThread;
+				FAsyncLoadingThread2& LocalAsyncLoadingThread = AsyncLoadingThread;
 				GetPackageNode(EEventLoadNode2::Package_ProcessSummary).ReleaseBarrier(nullptr);
 				if (LocalPendingIoRequestsCounter == 0)
 				{
@@ -4614,6 +4615,7 @@ void FAsyncPackage2::StartLoading(FAsyncLoadingThreadState2& ThreadState, FIoBat
 			int32 LocalPendingIoRequestsCounter = AsyncLoadingThread.PendingIoRequestsCounter.DecrementExchange() - 1;
 			TRACE_COUNTER_SET(AsyncLoadingPendingIoRequests, LocalPendingIoRequestsCounter);
 			FAsyncLoadingThread2& LocalAsyncLoadingThread = AsyncLoadingThread;
+			FAsyncLoadingThread2& LocalAsyncLoadingThread = AsyncLoadingThread;
 			GetPackageNode(EEventLoadNode2::Package_ProcessSummary).ReleaseBarrier(nullptr);
 			if (LocalPendingIoRequestsCounter == 0)
 			{
@@ -4636,6 +4638,7 @@ void FAsyncPackage2::StartLoading(FAsyncLoadingThreadState2& ThreadState, FIoBat
 					int32 LocalPendingIoRequestsCounter = AsyncLoadingThread.PendingIoRequestsCounter.DecrementExchange() - 1;
 					TRACE_COUNTER_SET(AsyncLoadingPendingIoRequests, LocalPendingIoRequestsCounter);
 					FAsyncLoadingThread2& LocalAsyncLoadingThread = AsyncLoadingThread;
+					FAsyncLoadingThread2& LocalAsyncLoadingThread = AsyncLoadingThread;
 					GetPackageNode(Package_ExportsSerialized).ReleaseBarrier(nullptr);
 					if (LocalPendingIoRequestsCounter == 0)
 					{
