UClass::TryFindTypeSlow will now properly handle NAME_None converted to string and passed as InPathNameOrShortName

diff --git a/Engine/Source/Runtime/CoreUObject/Private/UObject/Class.cpp b/Engine/Source/Runtime/CoreUObject/Private/UObject/Class.cpp
--- a/Engine/Source/Runtime/CoreUObject/Private/UObject/Class.cpp
+++ b/Engine/Source/Runtime/CoreUObject/Private/UObject/Class.cpp
@@ -6043,7 +6043,8 @@ UField* UClass::TryFindTypeSlow(UClass* TypeClass, const FString& InPathNameOrSh
 	checkf(TypeClass && TypeClass->IsChildOf(UField::StaticClass()), TEXT("TryFindType requires a valid TypeClass parameter which is a subclass of UField. \"%s\" provided."), *GetFullNameSafe(TypeClass));
 
 	UField* FoundType = nullptr;
-	if (!InPathNameOrShortName.IsEmpty())
+	// We check against "None" is it's not uncommon to pass an FName converted to string as InPathNameOrShortName
+	if (!InPathNameOrShortName.IsEmpty() && InPathNameOrShortName != TEXT("None"))
 	{
 		if (!FPackageName::IsShortPackageName(InPathNameOrShortName))
 		{
