Updating comment to make it match what I was thinking when I wrote it

diff --git a/Engine/Source/Runtime/CoreUObject/Private/UObject/Class.cpp b/Engine/Source/Runtime/CoreUObject/Private/UObject/Class.cpp
--- a/Engine/Source/Runtime/CoreUObject/Private/UObject/Class.cpp
+++ b/Engine/Source/Runtime/CoreUObject/Private/UObject/Class.cpp
@@ -6043,7 +6043,7 @@ UField* UClass::TryFindTypeSlow(UClass* TypeClass, const FString& InPathNameOrSh
 	checkf(TypeClass && TypeClass->IsChildOf(UField::StaticClass()), TEXT("TryFindType requires a valid TypeClass parameter which is a subclass of UField. \"%s\" provided."), *GetFullNameSafe(TypeClass));
 
 	UField* FoundType = nullptr;
-	// We check against "None" is it's not uncommon to pass an FName converted to string as InPathNameOrShortName
+	// We check against "None" as it's not uncommon to pass an FName converted to string as InPathNameOrShortName
 	if (!InPathNameOrShortName.IsEmpty() && InPathNameOrShortName != TEXT("None"))
 	{
 		if (!FPackageName::IsShortPackageName(InPathNameOrShortName))
