compile fix for missing include

diff --git a/Engine/Source/Runtime/CoreUObject/Tests/Serialization/BulkDataTests.cpp b/Engine/Source/Runtime/CoreUObject/Tests/Serialization/BulkDataTests.cpp
--- a/Engine/Source/Runtime/CoreUObject/Tests/Serialization/BulkDataTests.cpp
+++ b/Engine/Source/Runtime/CoreUObject/Tests/Serialization/BulkDataTests.cpp
@@ -2,6 +2,7 @@
 
 #if WITH_LOW_LEVEL_TESTS
 
+#include "IO/IoDispatcher.h"
 #include "TestHarness.h"
 #include "TestMacros/Assertions.h"
 #include "Memory/SharedBuffer.h"
