ARM64 MTE/HWASAN fix - use bit 55 to distinguish fixed from loosely allocated asset tag values

diff --git a/Engine/Source/Runtime/CoreUObject/Public/AssetRegistry/AssetDataTagMap.h b/Engine/Source/Runtime/CoreUObject/Public/AssetRegistry/AssetDataTagMap.h
--- a/Engine/Source/Runtime/CoreUObject/Public/AssetRegistry/AssetDataTagMap.h
+++ b/Engine/Source/Runtime/CoreUObject/Public/AssetRegistry/AssetDataTagMap.h
@@ -171,6 +171,16 @@ namespace FixedTagPrivate
 		}
 	};
 
+	// This bit is always zero in user mode addresses and most likely won't be used by current or future
+	// CPU features like ARM's PAC / Top-Byte Ignore or Intel's Linear Address Masking / 5-Level Paging
+#if defined(__x86_64__) || defined(_M_X64)
+	static constexpr uint32 KernelAddressBit = 63;
+#elif defined(__aarch64__) || defined(_M_ARM64)
+	static constexpr uint32 KernelAddressBit = 55;
+#else
+	#error Unsupported architecture, please declare which address bit distinguish user space from kernel space
+#endif
+
 } // end namespace FixedTagPrivate
 
 /**
@@ -185,7 +195,9 @@ class COREUOBJECT_API FAssetTagValueRef
 
 	class FFixedTagValue
 	{
-		static constexpr uint64 FixedMask = uint64(1) << 63;
+		static constexpr uint64 FixedMask = uint64(1) << FixedTagPrivate::KernelAddressBit;
+		static_assert(FixedTagPrivate::FMapHandle::StoreIndexBits <= (FixedTagPrivate::KernelAddressBit - 32), 
+			"Too few bits remain for the StoreIndex. Consider using other high bits but note that ARM64 use the top byte for HWASAN & MTE.)");
 
 		uint64 Bits;
 
