ZenLoader in editor: Fix ensure about missing NotifyRegistrationEvent for the script subobject /Script/ControlRig.Default__FKControlRig:VM that is created after the CDO itself.

#rnx

diff --git a/Engine/Source/Runtime/CoreUObject/Private/Serialization/AsyncLoading2.cpp b/Engine/Source/Runtime/CoreUObject/Private/Serialization/AsyncLoading2.cpp
--- a/Engine/Source/Runtime/CoreUObject/Private/Serialization/AsyncLoading2.cpp
+++ b/Engine/Source/Runtime/CoreUObject/Private/Serialization/AsyncLoading2.cpp
@@ -1301,6 +1301,7 @@ public:
 		ObjectIndexToPublicExport.Add(ObjectIndex, Key);
 	}
 
+	void FindAllScriptObjects(bool bVerifyOnly);
 	void RegistrationComplete();
 
 	void AddScriptObject(FStringView PackageName, FStringView Name, UObject* Object)
@@ -4154,50 +4155,72 @@ void FAsyncPackage2::SetupScriptDependencies()
 }
 
 
-void FGlobalImportStore::RegistrationComplete()
+void FGlobalImportStore::FindAllScriptObjects(bool bVerifyOnly)
 {
-#if DO_CHECK
-	{
-		TRACE_CPUPROFILER_EVENT_SCOPE(FindAllScriptObjectsDebug);
-		TStringBuilder<FName::StringBufferSize> Name;
-		TArray<UPackage*> ScriptPackages;
-		TArray<UObject*> Objects;
-		FindAllRuntimeScriptPackages(ScriptPackages);
+	TStringBuilder<FName::StringBufferSize> Name;
+	TArray<UPackage*> ScriptPackages;
+	TArray<UObject*> Objects;
+	FindAllRuntimeScriptPackages(ScriptPackages);
 
-		for (UPackage* Package : ScriptPackages)
-		{
+	for (UPackage* Package : ScriptPackages)
+	{
 #if WITH_EDITOR
-			Name.Reset();
-			Package->GetPathName(nullptr, Name);
-			FPackageObjectIndex PackageGlobalImportIndex = FPackageObjectIndex::FromScriptPath(Name);
-			if (!ScriptObjects.Contains(PackageGlobalImportIndex))
+		Name.Reset();
+		Package->GetPathName(nullptr, Name);
+		FPackageObjectIndex PackageGlobalImportIndex = FPackageObjectIndex::FromScriptPath(Name);
+		if (!ScriptObjects.Contains(PackageGlobalImportIndex))
+		{
+			if (bVerifyOnly)
 			{
-				ScriptObjects.Add(PackageGlobalImportIndex, Package);
 				UE_LOG(LogStreaming, Display, TEXT("Script package %s (0x%016llX) is missing a NotifyRegistrationEvent from the initial load phase."),
 					*Package->GetFullName(),
 					PackageGlobalImportIndex.Value());
 			}
+			else
+			{
+				ScriptObjects.Add(PackageGlobalImportIndex, Package);
+			}
+		}
 #endif
-			Objects.Reset();
-			GetObjectsWithOuter(Package, Objects, /*bIncludeNestedObjects*/true);
-			for (UObject* Object : Objects)
+		Objects.Reset();
+		GetObjectsWithOuter(Package, Objects, /*bIncludeNestedObjects*/true);
+		for (UObject* Object : Objects)
+		{
+			if (Object->HasAnyFlags(RF_Public))
 			{
-				if (Object->HasAnyFlags(RF_Public))
+				Name.Reset();
+				Object->GetPathName(nullptr, Name);
+				FPackageObjectIndex GlobalImportIndex = FPackageObjectIndex::FromScriptPath(Name);
+				if (!ScriptObjects.Contains(GlobalImportIndex))
 				{
-					Name.Reset();
-					Object->GetPathName(nullptr, Name);
-					FPackageObjectIndex GlobalImportIndex = FPackageObjectIndex::FromScriptPath(Name);
-					if (!ScriptObjects.Contains(GlobalImportIndex))
+					if (bVerifyOnly)
 					{
-						ScriptObjects.Add(GlobalImportIndex, Object);
-						ensureMsgf(false, TEXT("Script object %s (0x%016llX) is missing a NotifyRegistrationEvent from the initial load phase."),
+						UE_LOG(LogStreaming, Warning, TEXT("Script object %s (0x%016llX) is missing a NotifyRegistrationEvent from the initial load phase."),
 							*Object->GetFullName(),
 							GlobalImportIndex.Value());
 					}
+					else
+					{
+						ScriptObjects.Add(GlobalImportIndex, Object);
+					}
 				}
 			}
 		}
 	}
+}
+
+void FGlobalImportStore::RegistrationComplete()
+{
+#if WITH_EDITOR
+	{
+		TRACE_CPUPROFILER_EVENT_SCOPE(FindAllScriptObjects);
+		FindAllScriptObjects(/*bVerifyOnly*/false);
+	}
+#elif DO_CHECK
+	{
+		TRACE_CPUPROFILER_EVENT_SCOPE(FindAllScriptObjectsVerify);
+		FindAllScriptObjects(/*bVerifyOnly*/true);
+	}
 #endif
 	ScriptObjects.Shrink();
 }
