Skip private asset referencing check in FLinkerLoad if the package is loaded for diff because it's not always in its original location (usually in /Temp/)  so we can't reliably compare mount points

diff --git a/Engine/Source/Runtime/CoreUObject/Private/UObject/LinkerLoad.cpp b/Engine/Source/Runtime/CoreUObject/Private/UObject/LinkerLoad.cpp
--- a/Engine/Source/Runtime/CoreUObject/Private/UObject/LinkerLoad.cpp
+++ b/Engine/Source/Runtime/CoreUObject/Private/UObject/LinkerLoad.cpp
@@ -4683,11 +4683,18 @@ bool FLinkerLoad::IsPackageReferenceAllowed(UPackage* InPackage)
 {
 	if (InPackage && !InPackage->IsExternallyReferenceable())
 	{
-		FName MountPointName = FPackageName::GetPackageMountPoint(LinkerRoot->GetName());
-		FName ImportMountPointName = FPackageName::GetPackageMountPoint(InPackage->GetName());
-		if (MountPointName != ImportMountPointName)
-		{
-			return false;
+#if WITH_EDITOR
+		// Package loaded for diff is not always in its original location (usually in /Temp/) 
+		// so we can't reliably compare mount points here
+		if ((LoadFlags & LOAD_ForDiff) == 0)
+#endif //if WITH_EDITOR
+		{
+			FName MountPointName = FPackageName::GetPackageMountPoint(LinkerRoot->GetName());
+			FName ImportMountPointName = FPackageName::GetPackageMountPoint(InPackage->GetName());
+			if (MountPointName != ImportMountPointName)
+			{
+				return false;
+			}
 		}
 	}
 	return true;
