Fixed AsyncLoading2 handling of instanced partitioned world packages

AsyncLoading2 will now correctly fix-up any soft-object references to/from the generated cell packages for a partitioned world and/or the persistent world package. Previously only soft-object references within the same package would be handled, which broke cross-cell or persistent<->cell references.

Note: As part of this the naming convention of instanced cell packages has been changed to use the same suffix as the persistent world package.

diff --git a/Engine/Source/Runtime/CoreUObject/Private/Serialization/AsyncLoading2.cpp b/Engine/Source/Runtime/CoreUObject/Private/Serialization/AsyncLoading2.cpp
--- a/Engine/Source/Runtime/CoreUObject/Private/Serialization/AsyncLoading2.cpp
+++ b/Engine/Source/Runtime/CoreUObject/Private/Serialization/AsyncLoading2.cpp
@@ -1494,7 +1494,14 @@ public:
 		}
 	}
 };
-	
+
+struct FAsyncPackageInstanceContext
+{
+	TArray<TTuple<FString, FString>, TInlineAllocator<2>> InstancedPackageMappings;
+	FString GeneratedPackagesFolder;
+	FString InstancedPackageSuffix;
+};
+
 class FExportArchive final : public FArchive
 {
 public:
@@ -1686,7 +1693,6 @@ public:
 		FArchive& Ar = *this;
 		FSoftObjectPath ID;
 		ID.Serialize(Ar);
-		FixupSoftObjectPathForInstancedPackage(ID);
 		Value = ID;
 		return Ar;
 	}
@@ -1737,29 +1743,52 @@ private:
 	FPackageImportStore* ImportStore = nullptr;
 	TArray<FExternalReadCallback>* ExternalReadDependencies;
 	const FAsyncPackageHeaderData* HeaderData = nullptr;
+	const FAsyncPackageInstanceContext* InstanceContext = nullptr;
 	TArrayView<const FExportObject> Exports;
 	UObject* CurrentExport = nullptr;
 	uint64 CookedSerialOffset = 0;
 	uint64 CookedSerialSize = 0;
 	uint64 BufferSerialOffset = 0;
 
-	/** Set when the package is being loaded as an instance; empty otherwise. */
-	FNameBuilder InstancedPackageSourceName;
-	FNameBuilder InstancedPackageInstanceName;
 	void FixupSoftObjectPathForInstancedPackage(FSoftObjectPath& InOutSoftObjectPath)
 	{
-		if (InstancedPackageSourceName.Len() > 0 && InstancedPackageInstanceName.Len() > 0)
+		if (InstanceContext && InstanceContext->InstancedPackageMappings.Num() > 0)
 		{
 			FNameBuilder TmpSoftObjectPathBuilder;
 			InOutSoftObjectPath.ToString(TmpSoftObjectPathBuilder);
 
-			FStringView InstancedPackageSourceNameView = InstancedPackageSourceName.ToView();
-			FStringView TmpSoftObjectPathView = TmpSoftObjectPathBuilder.ToView();
+			// Is this a basic package remapping?
+			for (const TTuple<FString, FString>& InstancedPackageMapping : InstanceContext->InstancedPackageMappings)
+			{
+				FStringView TmpSoftObjectPathView = TmpSoftObjectPathBuilder.ToView();
+				if (TmpSoftObjectPathView.StartsWith(InstancedPackageMapping.Key) && (TmpSoftObjectPathView.Len() == InstancedPackageMapping.Key.Len() || TmpSoftObjectPathView[InstancedPackageMapping.Key.Len()] == TEXT('.')))
+				{
+					TmpSoftObjectPathBuilder.ReplaceAt(0, InstancedPackageMapping.Key.Len(), InstancedPackageMapping.Value);
+					InOutSoftObjectPath.SetPath(TmpSoftObjectPathBuilder.ToView());
+					return;
+				}
+			}
 
-			if (TmpSoftObjectPathView.StartsWith(InstancedPackageSourceNameView) && (TmpSoftObjectPathView.Len() == InstancedPackageSourceNameView.Len() || TmpSoftObjectPathView[InstancedPackageSourceNameView.Len()] == TEXT('.')))
+			// Is this a generated package reference that we care about?
+			if (!InstanceContext->GeneratedPackagesFolder.IsEmpty())
 			{
-				TmpSoftObjectPathBuilder.ReplaceAt(0, InstancedPackageSourceNameView.Len(), InstancedPackageInstanceName.ToView());
-				InOutSoftObjectPath.SetPath(TmpSoftObjectPathBuilder.ToView());
+				check(!InstanceContext->InstancedPackageSuffix.IsEmpty());
+
+				// Does this package path start with the generated folder path?
+				FStringView TmpSoftObjectPathView = TmpSoftObjectPathBuilder.ToView();
+				if (TmpSoftObjectPathView.StartsWith(InstanceContext->GeneratedPackagesFolder))
+				{
+					// ... and is that generated folder path immediately preceding the package name?
+					if (const int32 ExtraSlashIndex = TmpSoftObjectPathView.Find(TEXTVIEW("/"), InstanceContext->GeneratedPackagesFolder.Len()); ExtraSlashIndex == INDEX_NONE)
+					{
+						FNameBuilder PackageNameBuilder;
+						PackageNameBuilder.Append(InOutSoftObjectPath.GetLongPackageName());
+						PackageNameBuilder.Append(InstanceContext->InstancedPackageSuffix);
+						FTopLevelAssetPath SuffixTopLevelAsset(FName(PackageNameBuilder.ToString()), InOutSoftObjectPath.GetAssetFName());
+						InOutSoftObjectPath = FSoftObjectPath(SuffixTopLevelAsset, InOutSoftObjectPath.GetSubPathString());
+						return;
+					}
+				}
 			}
 		}
 	}
@@ -2612,6 +2641,9 @@ private:
 	using FCompletionCallback = TUniquePtr<FLoadPackageAsyncDelegate>;
 	TArray<FCompletionCallback, TInlineAllocator<2>> CompletionCallbacks;
 
+	/** Set when the package is being loaded as an instance; null otherwise. */
+	TUniquePtr<FAsyncPackageInstanceContext> InstanceContext;
+
 public:
 
 	FAsyncLoadingThread2& GetAsyncLoadingThread()
@@ -5229,6 +5261,49 @@ EEventLoadNodeExecutionResult FAsyncPackage2::Event_ProcessPackageSummary(FAsync
 			Package->LinkerRoot->SetLinkerLicenseeVersion(GPackageFileLicenseeUEVersion);
 		}
 
+		// Check if the package is instanced
+		FName PackageNameToLoad = Package->Desc.PackagePathToLoad.GetPackageFName();
+		if (Package->Desc.UPackageName != PackageNameToLoad)
+		{
+			Package->InstanceContext = MakeUnique<FAsyncPackageInstanceContext>();
+			Package->InstanceContext->InstancedPackageMappings.Add(MakeTuple(PackageNameToLoad.ToString(), Package->Desc.UPackageName.ToString()));
+
+			// Stash the suffix used for this instance so we can also apply it to generated packages
+			if (Package->InstanceContext->InstancedPackageMappings[0].Value.StartsWith(Package->InstanceContext->InstancedPackageMappings[0].Key))
+			{
+				Package->InstanceContext->InstancedPackageSuffix = Package->InstanceContext->InstancedPackageMappings[0].Value.Mid(Package->InstanceContext->InstancedPackageMappings[0].Key.Len());
+			}
+
+			// Is this a generated partitioned map package? If so, we'll also need to handle re-mapping paths to our persistent map package
+			if (!Package->InstanceContext->InstancedPackageSuffix.IsEmpty())
+			{
+				check(Package->InstanceContext->InstancedPackageMappings.Num() == 1);
+
+				const FStringView GeneratedFolderName = TEXTVIEW("/_Generated_/");
+				const FStringView PackageNameToLoadView = Package->InstanceContext->InstancedPackageMappings[0].Key;
+
+				// Does this package path include the generated folder?
+				if (const int32 GeneratedFolderStartIndex = PackageNameToLoadView.Find(GeneratedFolderName); GeneratedFolderStartIndex != INDEX_NONE)
+				{
+					// ... and is that generated folder immediately preceding the package name?
+					const int32 GeneratedFolderEndIndex = GeneratedFolderStartIndex + GeneratedFolderName.Len();
+					if (const int32 ExtraSlashIndex = PackageNameToLoadView.Find(TEXTVIEW("/"), GeneratedFolderEndIndex); ExtraSlashIndex == INDEX_NONE)
+					{
+						Package->InstanceContext->GeneratedPackagesFolder = PackageNameToLoadView.Left(GeneratedFolderEndIndex);
+
+						FString PersistentPackageName = FString(PackageNameToLoadView.Left(GeneratedFolderStartIndex));
+						FString PersistentPackageInstanceName = PersistentPackageName + Package->InstanceContext->InstancedPackageSuffix;
+						Package->InstanceContext->InstancedPackageMappings.Add(MakeTuple(MoveTemp(PersistentPackageName), MoveTemp(PersistentPackageInstanceName)));
+					}
+				}
+				if (Package->InstanceContext->GeneratedPackagesFolder.IsEmpty())
+				{
+					Package->InstanceContext->GeneratedPackagesFolder =  PackageNameToLoadView;
+					Package->InstanceContext->GeneratedPackagesFolder += GeneratedFolderName;
+				}
+			}
+		}
+
 		TRACE_LOADTIME_PACKAGE_SUMMARY(Package, Package->HeaderData.PackageName, Package->HeaderData.PackageSummary->HeaderSize, Package->HeaderData.ImportMap.Num(), Package->HeaderData.ExportMap.Num());
 	}
 
@@ -5383,14 +5458,7 @@ EEventLoadNodeExecutionResult FAsyncPackage2::Event_ProcessExportBundle(FAsyncLo
 			Ar.ImportStore = &Package->ImportStore;
 			Ar.Exports = Exports;
 			Ar.ExternalReadDependencies = &Package->ExternalReadDependencies;
-
-			// Check if the package is instanced
-			FName PackageNameToLoad = Package->Desc.PackagePathToLoad.GetPackageFName();
-			if (Package->Desc.UPackageName != PackageNameToLoad)
-			{
-				PackageNameToLoad.ToString(Ar.InstancedPackageSourceName);
-				Package->Desc.UPackageName.ToString(Ar.InstancedPackageInstanceName);
-			}
+			Ar.InstanceContext = Package->InstanceContext.Get();
 		}
 
 		while (Package->ExportBundleEntryIndex < int32(ExportBundle->EntryCount))
diff --git a/Engine/Source/Runtime/Engine/Private/WorldPartition/WorldPartitionRuntimeLevelStreamingCell.cpp b/Engine/Source/Runtime/Engine/Private/WorldPartition/WorldPartitionRuntimeLevelStreamingCell.cpp
--- a/Engine/Source/Runtime/Engine/Private/WorldPartition/WorldPartitionRuntimeLevelStreamingCell.cpp
+++ b/Engine/Source/Runtime/Engine/Private/WorldPartition/WorldPartitionRuntimeLevelStreamingCell.cpp
@@ -358,8 +358,33 @@ UWorldPartitionLevelStreamingDynamic* UWorldPartitionRuntimeLevelStreamingCell::
 		ULevel* PartitionLevel = WorldPartition->GetTypedOuter<ULevel>();
 		if (PartitionLevel->IsInstancedLevel())
 		{
-			FString PackageShortName = FPackageName::GetShortName(PartitionLevel->GetPackage());
-			FString InstancedLevelPackageName = FString::Printf(TEXT("%s_InstanceOf_%s"), *LevelStreaming->PackageNameToLoad.ToString(), *PackageShortName);
+			// Try and extract the instance suffix that was applied to the instanced level itself
+			FString InstancedLevelSuffix;
+			{
+				UPackage* PartitionLevelPackage = PartitionLevel->GetPackage();
+
+				FNameBuilder SourcePackageName(PartitionLevelPackage->GetLoadedPath().GetPackageFName());
+				FNameBuilder InstancedPackageName(PartitionLevelPackage->GetFName());
+				
+				FStringView SourcePackageNameView = SourcePackageName.ToView();
+				FStringView InstancedPackageNameView = InstancedPackageName.ToView();
+
+				if (InstancedPackageNameView.StartsWith(SourcePackageNameView))
+				{
+					InstancedLevelSuffix = InstancedPackageNameView.Mid(SourcePackageNameView.Len());
+				}
+				else
+				{
+					InstancedLevelSuffix =  TEXT("_");
+					InstancedLevelSuffix += FPackageName::GetShortName(PartitionLevelPackage);
+				}
+			}
+			check(!InstancedLevelSuffix.IsEmpty());
+
+			FNameBuilder InstancedLevelPackageName;
+			LevelStreaming->PackageNameToLoad.AppendString(InstancedLevelPackageName);
+			InstancedLevelPackageName += InstancedLevelSuffix;
+
 			LevelStreaming->SetWorldAssetByPackageName(FName(InstancedLevelPackageName));
 		}
 	}
