[CoreUObject]
* ObjectPtr - Added a little IWYU snippet in one of the assignment operator implementations to make IWYU aware that this call requires knowledge about the from type in order to be able to do the implicit cast.

Unfortunately there is no easy way to add support in IWYU for this scenario since it requires full ADL resolving and there are no public APIs in clang to do this atm

diff --git a/Engine/Source/Runtime/CoreUObject/Public/UObject/ObjectPtr.h b/Engine/Source/Runtime/CoreUObject/Public/UObject/ObjectPtr.h
--- a/Engine/Source/Runtime/CoreUObject/Public/UObject/ObjectPtr.h
+++ b/Engine/Source/Runtime/CoreUObject/Public/UObject/ObjectPtr.h
@@ -414,6 +414,9 @@ public:
 	>
 	FORCEINLINE TObjectPtr<T>& operator=(const TObjectPtr<U>& Other)
 	{
+		#ifdef PLATFORM_COMPILER_IWYU
+		enum { IWYU = sizeof(U) }; // Unfortunately the decltype expression above is not visible to IWYU so it can't tell that the code needs to know about type U
+		#endif
 		ObjectPtr = Other.ObjectPtr;
 		return *this;
 	}
