Removing the UObject scope in FAsyncPackage::PostLoadObjects to reduce memory overhead while UObject scope cycle counters are enabled.

diff --git a/Engine/Source/Runtime/CoreUObject/Private/Serialization/AsyncLoading.cpp b/Engine/Source/Runtime/CoreUObject/Private/Serialization/AsyncLoading.cpp
--- a/Engine/Source/Runtime/CoreUObject/Private/Serialization/AsyncLoading.cpp
+++ b/Engine/Source/Runtime/CoreUObject/Private/Serialization/AsyncLoading.cpp
@@ -6906,8 +6906,6 @@ EAsyncPackageState::Type FAsyncPackage::PostLoadObjects()
 			{
 				SCOPED_ACCUM_LOADTIME(PostLoad, StaticGetNativeClassName(Object->GetClass()));
 
-				FScopeCycleCounterUObject ConstructorScope(Object, GET_STATID(STAT_FAsyncPackage_PostLoadObjects));
-
 				// We want this check only with EDL enabled
 				check(!GEventDrivenLoaderEnabled || !Object->HasAnyFlags(RF_NeedLoad));
 
