GC reachability optimizations. ~1.5x speedup and ~2.5MB memory saved for an internal title.

List of optimizations and changes:
* Token stream structure
    * Split token stream into strong-only and a mixed (weak+strong) stream
    * Split token stream into a builder and a tighter view class which reduces sizeof(UClass)
    * Implemented ref-counted token stream view sharing
    * Removed Class and Outer from token stream
    * Allow empty token streams (enabled by removing Class/Outer) to avoid touching token stream data
    * Placed ARO (AddReferencedObjects) last to reduce per object cache thrashing, improve control flow predicability and avoid reading the last EndOfPointer and EndOfStream tokens
    * FPrefetchingObjectIterator that bring in Class/Outer, class' tokenstream view and the first token data ahead of time
    * Decode token bitfield once and ahead of time
* Reference queues and batch processing
    * Introduced bounded queues: ref arrays -> unvalidated refs -> validated (non-null / non-permanent) refs
    * Split all these queues for killable vs immutable references
    * Stack-living references still handled synchronously. With removal of Class/Outer (prefetched ahead of time) few instances remain outside of ARO calls.
    * Outer queues hold 32 items and get flushed when full.
* AddReferencedObjects (ARO) optimizations
    * Misc optimizations in many ARO implementations
    * New FReferenceCollector API to queue up ARO references (AddStableReference), old sync API (HandleObjectReference) still available
    * New AddPropertyReferences traversal that replaces SerializeBin and PropertyValueIterator
        * 4.5x faster than PropertyValueIterator
        * Uses CLASSCAST dispatch instead of virtual SerializeItem dispatch.
        * Step towards new unified token stream replacement shared by class token processing, structs and ARO
    * Replaced StructUtil::AddReferencedObjects with AddPropertyReferences traversal, ~8x speedup and collects more references for CitySample
* Parallelism
    * Single long-running task per worker
    * Improved work-stealing / load-balancing, workers can steal full blocks, ARO calls and initial references
    * Queue up slow ARO calls to improve load balancing and avoid late stragglers. Motivated by certain ARO calls taking over 2ms for a few specific objects.
    * Kick tasks manually to avoid ParallelFor end synchronization
* FGCObject
    * Initial reference collector runs in parallel with mark phase
    * New FGCObject constructor API (AddStableNativeReferencesOnly) to opt-in to initial reference collection, used by StreamableManager
    * Same constructor API allows FGCObjects to defer registration until they become active (RegisterLater), reduces number of active GCObjects
* Reduced memory usage
    * Allocate reached objects in scratch pages (FWorkBlock) and reuse processed blocks, instead of swapping two big TArray<UObject*> per worker
    * Reduced sizeof(UClass)
    * Shareable token streams
* Misc optimizations
    * New API to test if an object is in the permanent object pool. Old API read two global pointers for every visited reference.
    * Fixed signed integer usage in GUObjectArray lookup that led to bad codegen
    * FPropertyIterator optimizations
    * SerializeBin optimizations
* Other changes
    * Moved many helpers into UE::GC namespace
    * Replaced TFastReferenceCollector API with simplified CollectReferences call. Needed to break this API  any way.
    * Introduced FGCInternals to avoid forward-declaring TFastReferenceCollector and depend on the options enum in common headers
    * Moved and outlined code from GarbageCollection.h / FastReferenceCollector.h to GarbageCollection.cpp
    * Moved GC History and Garbage Reference Tracking into a synchronous TDebugReachabilityProcessor
    * Removed PersistentGarbage flag since it wasn't used in practice
    * Improved const correctness

diff --git a/Engine/Plugins/Developer/TraceSourceFiltering/Source/SourceFilteringTrace/Private/ActorFiltering.h b/Engine/Plugins/Developer/TraceSourceFiltering/Source/SourceFilteringTrace/Private/ActorFiltering.h
--- a/Engine/Plugins/Developer/TraceSourceFiltering/Source/SourceFilteringTrace/Private/ActorFiltering.h
+++ b/Engine/Plugins/Developer/TraceSourceFiltering/Source/SourceFilteringTrace/Private/ActorFiltering.h
@@ -204,10 +204,9 @@ public:
 		return ActorArray.Num();
 	}
 
-
 	void AddReferencedObjects(FReferenceCollector& Collector) override
 	{
-		Collector.AddReferencedObjects(AllFilteredClasses);
+		Collector.AddStableReferenceArray(&AllFilteredClasses);
 	}
 	virtual FString GetReferencerName() const override
 	{
diff --git a/Engine/Plugins/Experimental/StructUtils/Source/StructUtils/Private/InstancedStruct.cpp b/Engine/Plugins/Experimental/StructUtils/Source/StructUtils/Private/InstancedStruct.cpp
--- a/Engine/Plugins/Experimental/StructUtils/Source/StructUtils/Private/InstancedStruct.cpp
+++ b/Engine/Plugins/Experimental/StructUtils/Source/StructUtils/Private/InstancedStruct.cpp
@@ -382,7 +382,8 @@ void FInstancedStruct::AddStructReferencedObjects(class FReferenceCollector& Col
 {
 	if (ScriptStruct != nullptr)
 	{
-		UE::StructUtils::AddReferencedObjects(Collector, ScriptStruct, GetMutableMemory());
+		Collector.AddReferencedObject(ScriptStruct);
+		Collector.AddPropertyReferencesWithStructARO(ScriptStruct, GetMutableMemory());
 	}
 }
 
diff --git a/Engine/Plugins/Experimental/StructUtils/Source/StructUtils/Private/InstancedStructContainer.cpp b/Engine/Plugins/Experimental/StructUtils/Source/StructUtils/Private/InstancedStructContainer.cpp
--- a/Engine/Plugins/Experimental/StructUtils/Source/StructUtils/Private/InstancedStructContainer.cpp
+++ b/Engine/Plugins/Experimental/StructUtils/Source/StructUtils/Private/InstancedStructContainer.cpp
@@ -363,14 +363,15 @@ void FInstancedStructContainer::Empty()
 	AllocatedSize = 0;
 }
 
-void FInstancedStructContainer::AddStructReferencedObjects(class FReferenceCollector& Collector) const
+void FInstancedStructContainer::AddStructReferencedObjects(FReferenceCollector& Collector) const
 {
-	for (int32 Index = 0; Index < NumItems; Index++)
+	for (int32 Index = 0, Num = NumItems; Index < Num; Index++)
 	{
 		FItem& Item = GetItem(Index);
 		if (Item.ScriptStruct != nullptr)
 		{
-			UE::StructUtils::AddReferencedObjects(Collector, Item.ScriptStruct, Memory + Item.Offset);
+			Collector.AddReferencedObject(Item.ScriptStruct);
+			Collector.AddPropertyReferencesWithStructARO(Item.ScriptStruct, Memory + Item.Offset);
 		}
 	}
 }
diff --git a/Engine/Plugins/Experimental/StructUtils/Source/StructUtils/Private/StructUtilsTypes.cpp b/Engine/Plugins/Experimental/StructUtils/Source/StructUtils/Private/StructUtilsTypes.cpp
--- a/Engine/Plugins/Experimental/StructUtils/Source/StructUtils/Private/StructUtilsTypes.cpp
+++ b/Engine/Plugins/Experimental/StructUtils/Source/StructUtils/Private/StructUtilsTypes.cpp
@@ -26,38 +26,4 @@ namespace UE::StructUtils
 		}
 		return 0;
 	}
-
-	STRUCTUTILS_API void AddReferencedObjects(FReferenceCollector& Collector, const UScriptStruct*& ScriptStruct, void* StructMemory,
-										 const UObject* ReferencingObject, const FProperty* ReferencingProperty)
-	{
-		check(ScriptStruct != nullptr);
-		check(StructMemory != nullptr);
-
-		Collector.AddReferencedObject(ScriptStruct, ReferencingObject, ReferencingProperty);
-
-		// If the script struct explicitly provided an implementation of AddReferencedObjects, make sure to capture its referenced objects
-		if (ScriptStruct->StructFlags & STRUCT_AddStructReferencedObjects)
-		{
-			ScriptStruct->GetCppStructOps()->AddStructReferencedObjects()(StructMemory, Collector);
-		}
-
-		// Iterate through all object and struct properties within the struct (will also search through structs within the struct)
-		for (TPropertyValueIterator<const FProperty> PropertyIter(ScriptStruct, StructMemory); PropertyIter; ++PropertyIter)
-		{
-			if (const FObjectProperty* ObjectProperty = CastField<FObjectProperty>(PropertyIter.Key()))
-			{
-				TObjectPtr<UObject>& ObjectPtrRef = ObjectProperty->GetObjectPtrPropertyValueRef(PropertyIter.Value());
-				Collector.AddReferencedObject(ObjectPtrRef, ReferencingObject, ReferencingProperty);
-			}
-			else if (const FStructProperty* StructProperty = CastField<FStructProperty>(PropertyIter.Key()))
-			{
-				const UScriptStruct* ChildScriptStruct = StructProperty->Struct; 
-				if (ChildScriptStruct && ChildScriptStruct->StructFlags & STRUCT_AddStructReferencedObjects)
-				{
-					ChildScriptStruct->GetCppStructOps()->AddStructReferencedObjects()(const_cast<void*>(PropertyIter.Value()), Collector);
-				}
-			}
-		}
-	}
-
 }
diff --git a/Engine/Plugins/Experimental/StructUtils/Source/StructUtils/Public/StructUtilsTypes.h b/Engine/Plugins/Experimental/StructUtils/Source/StructUtils/Public/StructUtilsTypes.h
--- a/Engine/Plugins/Experimental/StructUtils/Source/StructUtils/Public/StructUtilsTypes.h
+++ b/Engine/Plugins/Experimental/StructUtils/Source/StructUtils/Public/StructUtilsTypes.h
@@ -30,13 +30,6 @@ namespace UE::StructUtils
 	{
 		return T::StaticClass();
 	}
-
-	/**
-	 * Adds references objects from specific struct, similar to FReferenceCollector::AddReferencedObjects(),
-	 * but calls AddStructReferencedObjects recursively on child structs.
-	 */
-	extern STRUCTUTILS_API void AddReferencedObjects(FReferenceCollector& Collector, const UScriptStruct*& ScriptStruct, void* StructMemory,
-													 const UObject* ReferencingObject = nullptr, const FProperty* ReferencingProperty = nullptr);
 }
 
 /* Predicate useful to find a struct of a specific type in an container */
diff --git a/Engine/Plugins/Runtime/DataRegistry/Source/DataRegistry/Private/DataRegistry.cpp b/Engine/Plugins/Runtime/DataRegistry/Source/DataRegistry/Private/DataRegistry.cpp
--- a/Engine/Plugins/Runtime/DataRegistry/Source/DataRegistry/Private/DataRegistry.cpp
+++ b/Engine/Plugins/Runtime/DataRegistry/Source/DataRegistry/Private/DataRegistry.cpp
@@ -1123,14 +1123,13 @@ void UDataRegistry::AddReferencedObjects(UObject* InThis, FReferenceCollector& C
 	// Need to emit references for referenced items (unless there's no properties that reference UObjects)
 	if (ItemStruct != nullptr && ItemStruct->RefLink != nullptr && This->Cache->LookupCache.Num() > 0)
 	{
-		FVerySlowReferenceCollectorArchiveScope CollectorScope(Collector.GetVerySlowReferenceCollectorArchive(), This);
 		for (TPair<FDataRegistryLookup, TUniquePtr<FCachedDataRegistryItem>>& Pair : This->Cache->LookupCache)
 		{
 			FCachedDataRegistryItem* Item = Pair.Value.Get();
 
 			if (Item && Item->ItemMemory)
 			{
-				ItemStruct->SerializeBin(CollectorScope.GetArchive(), Item->ItemMemory);
+				Collector.AddPropertyReferences(ItemStruct, Item->ItemMemory, This);
 			}
 		}
 	}
diff --git a/Engine/Plugins/Runtime/StateTree/Source/StateTreeModule/Private/StateTree.cpp b/Engine/Plugins/Runtime/StateTree/Source/StateTreeModule/Private/StateTree.cpp
--- a/Engine/Plugins/Runtime/StateTree/Source/StateTreeModule/Private/StateTree.cpp
+++ b/Engine/Plugins/Runtime/StateTree/Source/StateTreeModule/Private/StateTree.cpp
@@ -123,9 +123,7 @@ void UStateTree::AddReferencedObjects(UObject* InThis, FReferenceCollector& Coll
 	{
 		if (InstanceData.IsValid())
 		{
-			uint8* StructMemory = (uint8*)InstanceData.Get();
-			const UScriptStruct* ScriptStruct = FStateTreeInstanceData::StaticStruct();
-			UE::StructUtils::AddReferencedObjects(Collector, ScriptStruct, StructMemory);
+			Collector.AddPropertyReferencesWithStructARO(FStateTreeInstanceData::StaticStruct(), InstanceData.Get(), StateTree);
 		}
 	}
 }
diff --git a/Engine/Plugins/Runtime/WorldConditions/Source/WorldConditions/Private/WorldConditionQuery.cpp b/Engine/Plugins/Runtime/WorldConditions/Source/WorldConditions/Private/WorldConditionQuery.cpp
--- a/Engine/Plugins/Runtime/WorldConditions/Source/WorldConditions/Private/WorldConditionQuery.cpp
+++ b/Engine/Plugins/Runtime/WorldConditions/Source/WorldConditions/Private/WorldConditionQuery.cpp
@@ -215,7 +215,7 @@ void FWorldConditionQueryState::AddStructReferencedObjects(class FReferenceColle
 	check(SharedDefinition);
 	check(NumConditions == SharedDefinition->Conditions.Num());
 
-	for (int32 Index = 0; Index < SharedDefinition->Conditions.Num(); Index++)
+	for (int32 Index = 0, Num = SharedDefinition->Conditions.Num(); Index < Num; Index++)
 	{
 		const FWorldConditionBase& Condition = SharedDefinition->Conditions[Index].GetMutable<FWorldConditionBase>();
 		if (Condition.StateDataOffset == 0)
@@ -233,7 +233,8 @@ void FWorldConditionQueryState::AddStructReferencedObjects(class FReferenceColle
 		{
 			if (const UScriptStruct* StateScriptStruct = Cast<UScriptStruct>(Condition.GetRuntimeStateType()))
 			{
-				UE::StructUtils::AddReferencedObjects(Collector, StateScriptStruct, StateMemory);
+				Collector.AddReferencedObject(StateScriptStruct, Owner);
+				Collector.AddPropertyReferencesWithStructARO(StateScriptStruct, StateMemory, Owner);
 			}
 		}
 	}
diff --git a/Engine/Source/Editor/KismetCompiler/Private/KismetCompiler.cpp b/Engine/Source/Editor/KismetCompiler/Private/KismetCompiler.cpp
--- a/Engine/Source/Editor/KismetCompiler/Private/KismetCompiler.cpp
+++ b/Engine/Source/Editor/KismetCompiler/Private/KismetCompiler.cpp
@@ -2996,7 +2996,7 @@ void FKismetCompilerContext::FinishCompilingClass(UClass* Class)
 	if (ParentClass != NULL)
 	{
 		// Propagate the new parent's inheritable class flags
-		Class->ReferenceTokenStream.Empty();
+		Class->ReferenceTokens.Reset();
 		Class->ClassFlags &= ~CLASS_RecompilerClear;
 		Class->ClassFlags |= (ParentClass->ClassFlags & CLASS_ScriptInherit);//@TODO: ChangeParentClass had this, but I don't think I want it: | UClass::StaticClassFlags;  // will end up with CLASS_Intrinsic
 		Class->ClassCastFlags |= ParentClass->ClassCastFlags;
diff --git a/Engine/Source/Editor/Sequencer/Private/Sequencer.cpp b/Engine/Source/Editor/Sequencer/Private/Sequencer.cpp
--- a/Engine/Source/Editor/Sequencer/Private/Sequencer.cpp
+++ b/Engine/Source/Editor/Sequencer/Private/Sequencer.cpp
@@ -4241,7 +4241,7 @@ void FSequencer::AddReferencedObjects( FReferenceCollector& Collector )
 		Collector.AddReferencedObject( RootSequencePtr );
 	}
 
-	FMovieSceneRootEvaluationTemplateInstance::StaticStruct()->SerializeBin(Collector.GetVerySlowReferenceCollectorArchive(), &RootTemplateInstance);
+	Collector.AddPropertyReferences(FMovieSceneRootEvaluationTemplateInstance::StaticStruct(), &RootTemplateInstance, nullptr);
 }
 
 FString FSequencer::GetReferencerName() const
diff --git a/Engine/Source/Editor/UnrealEd/Private/Kismet2/BlueprintEditorUtils.cpp b/Engine/Source/Editor/UnrealEd/Private/Kismet2/BlueprintEditorUtils.cpp
--- a/Engine/Source/Editor/UnrealEd/Private/Kismet2/BlueprintEditorUtils.cpp
+++ b/Engine/Source/Editor/UnrealEd/Private/Kismet2/BlueprintEditorUtils.cpp
@@ -126,6 +126,7 @@
 #include "AssetRegistry/AssetRegistryModule.h"
 #include "Misc/FeedbackContext.h"
 
+#include "Containers/ArrayView.h"
 #include "UObject/FastReferenceCollector.h"
 #include "Elements/Framework/TypedElementRegistry.h"
 #include "Elements/Interfaces/TypedElementHierarchyInterface.h"
@@ -8329,7 +8330,7 @@ public:
 			}
 		}
 	}
-	FORCEINLINE_DEBUGGABLE void HandleTokenStreamObjectReference(FGCArrayStruct& ObjectsToSerializeStruct, UObject* ReferencingObject, UObject*& Object, const int32 TokenIndex, const EGCTokenType TokenType, bool bAllowReferenceElimination)
+	FORCEINLINE_DEBUGGABLE void HandleTokenStreamObjectReference(FGCArrayStruct& ObjectsToSerializeStruct, UObject* ReferencingObject, UObject*& Object, UE::GC::FTokenId, EGCTokenType, bool)
 	{
 		if (!ReferencingObject)
 		{
@@ -8363,17 +8364,12 @@ void FBlueprintEditorUtils::GetActorReferenceMap(UWorld* InWorld, TArray<UClass*
 {
 	TRACE_CPUPROFILER_EVENT_SCOPE(FBlueprintEditorUtils::GetActorReferenceMap);
 	
-	FGCArrayStruct ArrayStruct;
-	
-	FActorMapReferenceProcessor Processor(InWorld, ArrayStruct.ObjectsToSerialize, InClassesToIgnore, OutReferencingActors);
-	TFastReferenceCollector<
-		FActorMapReferenceProcessor, 
-		TDefaultReferenceCollector<FActorMapReferenceProcessor>, 
-		FGCArrayPool, 
-		EFastReferenceCollectorOptions::AutogenerateTokenStream | EFastReferenceCollectorOptions::ProcessNoOpTokens
-	> ReferenceCollector(Processor, FGCArrayPool::Get());
-	
-	ReferenceCollector.CollectReferences(ArrayStruct);
+	TArray<UObject*> InitialObjects;
+	FActorMapReferenceProcessor Processor(InWorld, /* out */ InitialObjects, InClassesToIgnore, OutReferencingActors);
+
+	UE::GC::FWorkerContext Context;
+	Context.SetInitialObjectsUnpadded(InitialObjects);
+	CollectReferences(Processor, Context);
 }
 
 void FBlueprintEditorUtils::FixLevelScriptActorBindings(ALevelScriptActor* LevelScriptActor, const ULevelScriptBlueprint* ScriptBlueprint)
@@ -10558,7 +10554,7 @@ static FAutoConsoleCommand AuditFunctionCallsForBlueprint(
 
 			struct FFunctionReferenceProcessor : public FSimpleReferenceProcessorBase
 			{
-				FORCEINLINE void HandleTokenStreamObjectReference(FGCArrayStruct& ObjectsToSerializeStruct, UObject* ReferencingObject, UObject*& Object, const int32 TokenIndex, const EGCTokenType TokenType, bool bAllowReferenceElimination)
+				FORCEINLINE void HandleTokenStreamObjectReference(FGCArrayStruct& ObjectsToSerializeStruct, UObject* ReferencingObject, UObject*& Object, UE::GC::FTokenId, EGCTokenType TokenType, bool)
 				{
 					if (UFunction* Function = Cast<UFunction>(Object))
 					{
@@ -10577,15 +10573,11 @@ static FAutoConsoleCommand AuditFunctionCallsForBlueprint(
 				}
 			} Processor;
 
-			TFastReferenceCollector<
-				FFunctionReferenceProcessor,
-				TDefaultReferenceCollector<FFunctionReferenceProcessor>,
-				FGCArrayPool,
-				EFastReferenceCollectorOptions::AutogenerateTokenStream | EFastReferenceCollectorOptions::ProcessNoOpTokens
-			> ReferenceCollector(Processor, FGCArrayPool::Get());
-			FGCArrayStruct ArrayStruct;
-			ArrayStruct.ObjectsToSerialize.Add(Blueprint->GeneratedClass);
-			ReferenceCollector.CollectReferences(ArrayStruct);
+			
+			TArray<UObject*> InitialObjects = {Blueprint->GeneratedClass};
+			UE::GC::FWorkerContext Context;
+			Context.SetInitialObjectsUnpadded(InitialObjects);
+			CollectReferences(Processor, Context);
 
 			UE_LOG(LogBlueprint, Display, TEXT("--- END audit all BlueprintThreadSafe functions ---"));
 		}));
diff --git a/Engine/Source/Programs/Shared/EpicGames.Core/UnrealEngineTypes.cs b/Engine/Source/Programs/Shared/EpicGames.Core/UnrealEngineTypes.cs
--- a/Engine/Source/Programs/Shared/EpicGames.Core/UnrealEngineTypes.cs
+++ b/Engine/Source/Programs/Shared/EpicGames.Core/UnrealEngineTypes.cs
@@ -657,11 +657,6 @@ namespace EpicGames.Core
 		/// </summary>
 		Garbage = 1 << 21,
 
-		/// <summary>
-		/// Same as above but referenced through a persistent reference so it can't be GC'd
-		/// </summary>
-		PersistentGarbage = 1 << 22,
-
 		/// <summary>
 		/// External reference to object in cluster exists
 		/// </summary>
@@ -710,7 +705,7 @@ namespace EpicGames.Core
 		GarbageCollectionKeepFlags = Native | Async | AsyncLoading | LoaderImport,
 
 		//~ Make sure this is up to date!
-		AllFlags = LoaderImport | Garbage | PersistentGarbage | ReachableInCluster | ClusterRoot | Native | Async | AsyncLoading | Unreachable | PendingKill | RootSet | PendingConstruction
+		AllFlags = LoaderImport | Garbage | ReachableInCluster | ClusterRoot | Native | Async | AsyncLoading | Unreachable | PendingKill | RootSet | PendingConstruction
 	};
 
 	/// <summary>
diff --git a/Engine/Source/Runtime/CoreUObject/Private/Misc/GCObjectReferencer.cpp b/Engine/Source/Runtime/CoreUObject/Private/Misc/GCObjectReferencer.cpp
--- a/Engine/Source/Runtime/CoreUObject/Private/Misc/GCObjectReferencer.cpp
+++ b/Engine/Source/Runtime/CoreUObject/Private/Misc/GCObjectReferencer.cpp
@@ -25,8 +25,25 @@ static FAutoConsoleVariableRef CVerifyGCObjectNames(
 extern bool GObjIncrementalPurgeIsInProgress;
 extern bool GObjUnhashUnreachableIsInProgress;
 
+namespace UE::GC
+{
+	inline constexpr FStringView UnknownGCObjectName = TEXTVIEW("Unknown FGCObject");
+}
+
+struct UGCObjectReferencer::FImpl
+{
+	/** Critical section used when adding and removing objects */
+	FCriticalSection ReferencedObjectsCritical;
+	/** Objects without EFlags::AddStableNativeReferencesOnly */
+	TArray<FGCObject*> RemainingReferencedObjects;
+	/** Objects with EFlags::AddStableNativeReferencesOnly */
+	TArray<FGCObject*> InitialReferencedObjects;
+};
+
+
 UGCObjectReferencer::UGCObjectReferencer(const FObjectInitializer& ObjectInitializer)
 	: Super(ObjectInitializer)
+	, Impl(new FImpl)
 {
 #if VERIFY_DISREGARD_GC_ASSUMPTIONS
 	if (!HasAnyFlags(RF_ClassDefaultObject))
@@ -36,31 +53,47 @@ UGCObjectReferencer::UGCObjectReferencer(const FObjectInitializer& ObjectInitial
 #endif // VERIFY_DISREGARD_GC_ASSUMPTIONS
 }
 
+UGCObjectReferencer::UGCObjectReferencer(FVTableHelper& Helper)
+: Super(Helper)
+, Impl(new FImpl)
+{}
+
 void UGCObjectReferencer::AddReferencedObjects(UObject* InThis, FReferenceCollector& Collector)
 {
 	UGCObjectReferencer* This = CastChecked<UGCObjectReferencer>(InThis);
 
-	checkSlow(!This->bIsAddingReferencedObjects);
-	This->bIsAddingReferencedObjects = true;
+	check(!This->CurrentlySerializingObject);
+
 	// Note we're not locking ReferencedObjectsCritical here because we guard
 	// against adding new references during GC in AddObject and RemoveObject.
 	// Let each registered object handle its AddReferencedObjects call
-	for (FGCObject* Object : This->ReferencedObjects)
+	
+	if (Collector.NeedsInitialReferences())
+	{
+		for (FGCObject* Object : This->Impl->InitialReferencedObjects)
+		{
+			This->CurrentlySerializingObject = Object;
+			Object->AddReferencedObjects(Collector);
+		}
+	}
+
+	for (FGCObject* Object : This->Impl->RemainingReferencedObjects)
 	{
-		check(Object);
 		This->CurrentlySerializingObject = Object;
 		Object->AddReferencedObjects(Collector);
 	}
+
 	This->CurrentlySerializingObject = nullptr;
 	Super::AddReferencedObjects(This, Collector);
-	This->bIsAddingReferencedObjects = false;
 }
 
 void UGCObjectReferencer::AddObject(FGCObject* Object)
 {
 	check(Object);
 	check(GObjUnhashUnreachableIsInProgress || GObjIncrementalPurgeIsInProgress || !IsGarbageCollectingAndLockingUObjectHashTables());
-	FScopeLock ReferencedObjectsLock(&ReferencedObjectsCritical);
+	
+	FScopeLock ReferencedObjectsLock(&Impl->ReferencedObjectsCritical);
+	TArray<FGCObject*>& ReferencedObjects = Object->bCanMakeInitialReferences ? Impl->InitialReferencedObjects : Impl->RemainingReferencedObjects; 
 	// Make sure there are no duplicates. Should be impossible...
 	checkSlow(!ReferencedObjects.Contains(Object));
 	ReferencedObjects.Add(Object);
@@ -77,7 +110,9 @@ void UGCObjectReferencer::RemoveObject(FGCObject* Object)
 {
 	check(Object);
 	check(GObjUnhashUnreachableIsInProgress || GObjIncrementalPurgeIsInProgress || !IsGarbageCollectingAndLockingUObjectHashTables());
-	FScopeLock ReferencedObjectsLock(&ReferencedObjectsCritical);
+	
+	FScopeLock ReferencedObjectsLock(&Impl->ReferencedObjectsCritical);
+	TArray<FGCObject*>& ReferencedObjects = Object->bCanMakeInitialReferences ? Impl->InitialReferencedObjects : Impl->RemainingReferencedObjects; 
 	int32 NumRemoved = ReferencedObjects.RemoveSingleSwap(Object);
 	check(NumRemoved == 1);
 #if VERIFY_DISREGARD_GC_ASSUMPTIONS
@@ -89,7 +124,7 @@ bool UGCObjectReferencer::GetReferencerName(UObject* Object, FString& OutName, b
 {
 	if (bOnlyIfAddingReferenced)
 	{
-		if (!bIsAddingReferencedObjects || !CurrentlySerializingObject)
+		if (!CurrentlySerializingObject)
 		{
 			return false;
 		}
@@ -103,24 +138,27 @@ bool UGCObjectReferencer::GetReferencerName(UObject* Object, FString& OutName, b
 	}
 
 	// Let each registered object handle its AddReferencedObjects call
-	for (int32 i = 0; i < ReferencedObjects.Num(); i++)
+	for (TArrayView<FGCObject*> ReferencedObjects : { MakeArrayView(Impl->RemainingReferencedObjects),
+													  MakeArrayView(Impl->InitialReferencedObjects) })
 	{
-		TArray<UObject*> ObjectArray;
-		FReferenceFinder Collector(ObjectArray);
+		for (FGCObject* GCReporter : ReferencedObjects)
+		{
+			TArray<UObject*> ObjectArray;
+			FReferenceFinder Collector(ObjectArray);
 
-		FGCObject* GCReporter = ReferencedObjects[i];
-		check(GCReporter);
-		GCReporter->AddReferencedObjects(Collector);
+			check(GCReporter);
+			GCReporter->AddReferencedObjects(Collector);
 
-		if (ObjectArray.Contains(Object))
-		{
-			OutName = GCReporter->GetReferencerName();
-			FString ReferencerProperty;
-			if (GCReporter->GetReferencerPropertyName(Object, ReferencerProperty))
+			if (ObjectArray.Contains(Object))
 			{
-				OutName += TEXT(":") + ReferencerProperty;
+				OutName = GCReporter->GetReferencerName();
+				FString ReferencerProperty;
+				if (GCReporter->GetReferencerPropertyName(Object, ReferencerProperty))
+				{
+					OutName += TEXT(":") + ReferencerProperty;
+				}
+				return true;
 			}
-			return true;
 		}
 	}
 
@@ -134,13 +172,66 @@ void UGCObjectReferencer::FinishDestroy()
 		// Make sure FGCObjects that are around after exit purge don't
 		// reference this object.
 		check( FGCObject::GGCObjectReferencer == this );
-		FGCObject::GGCObjectReferencer = NULL;
-		ReferencedObjects.Empty();
+		FGCObject::GGCObjectReferencer = nullptr;
+		Impl.Reset();
 	}
 
 	Super::FinishDestroy();
 }
 
+
+class FInitialReferenceCollector final : public FReferenceCollector
+{
+	TArray<UObject**>& Result;
+
+	virtual void AddStableReference(UObject** Object) override
+	{
+		Result.Add(Object);
+	}
+
+	virtual void AddStableReferenceArray(TArray<UObject*>* Objects) override
+	{
+		for (UObject*& Object : *Objects)
+		{
+			Result.Add(&Object);
+		}
+	}
+
+	virtual void AddStableReferenceSet(TSet<UObject*>* Objects) override
+	{
+		for (UObject*& Object : *Objects)
+		{
+			Result.Add(&Object);
+		}
+	}
+
+	virtual void HandleObjectReference(UObject*& InObject, const UObject* InReferencingObject, const FProperty* InReferencingProperty) override
+	{
+		checkf(false, TEXT("FGCObject constructed with AddStableNativeReferencesOnly should only call AddStableReference, not HandleObjectReference"));
+	}
+
+	virtual void SetIsProcessingNativeReferences(bool) override
+	{
+		checkf(false, TEXT("FGCObject constructed with AddStableNativeReferencesOnly should never flip to non-native references"));
+	}
+
+	virtual bool IsIgnoringArchetypeRef() const override { return false;}
+	virtual bool IsIgnoringTransient() const override {	return false; }
+
+public:
+	FInitialReferenceCollector(TArray<UObject**>& Out) : Result(Out) {}
+};
+
+
+void UGCObjectReferencer::AddInitialReferences(TArray<UObject**>& Out)
+{
+	FInitialReferenceCollector Collector(Out);
+	for (FGCObject* Object : Impl->InitialReferencedObjects)
+	{
+		Object->AddReferencedObjects(Collector);
+	}
+}
+
 /** 
   * Helper class to make sure the ensure that verifies GetReferencerName is fired from within a callstack that can be
   * associated with the offending FGCObject-derived class.
@@ -164,7 +255,7 @@ public:
 	{
 		if (!bCurrentObjectVerified && CurrentGCObject)
 		{
-			if (!ensureAlwaysMsgf(CurrentGCObject->GetReferencerName() != FGCObject::UnknownGCObjectName,
+			if (!ensureAlwaysMsgf(CurrentGCObject->GetReferencerName() !=  UE::GC::UnknownGCObjectName,
 				TEXT("Please make sure all FGCObject derived classes have a unique name by overriding FGCObject::GetReferencerName() function. FGCObject::GetReferencerName() will become pure virtual in the next engine release. See callstack for details.")))
 			{
 				bCurrentObjectVerified = true;
@@ -198,11 +289,15 @@ void UGCObjectReferencer::VerifyGCObjectNames()
 	{
 		FVerifyReferencerNameCollector VerifyReferencerNameCollector;
 
-		for (FGCObject* GCObject : ReferencedObjects)
+		for (TArrayView<FGCObject*> ReferencedObjects : { MakeArrayView(Impl->RemainingReferencedObjects),
+														  MakeArrayView(Impl->InitialReferencedObjects) })
 		{
-			check(GCObject);
-			VerifyReferencerNameCollector.SetCurrentGCObject(GCObject);
-			GCObject->AddReferencedObjects(VerifyReferencerNameCollector);
+			for (FGCObject* GCObject : ReferencedObjects)
+			{
+				check(GCObject);
+				VerifyReferencerNameCollector.SetCurrentGCObject(GCObject);
+				GCObject->AddReferencedObjects(VerifyReferencerNameCollector);
+			}
 		}
 	}
 }
@@ -213,9 +308,41 @@ IMPLEMENT_CORE_INTRINSIC_CLASS(UGCObjectReferencer, UObject,
 	}
 );
 
-/** Static used for calling AddReferencedObjects on non-UObject objects */
 UGCObjectReferencer* FGCObject::GGCObjectReferencer = nullptr;
-/** Default name for unnamed FGCObjects */
-const TCHAR* FGCObject::UnknownGCObjectName = TEXT("Unknown FGCObject");
 
+// This variable allows inline declaration of GetReferencerName(), which
+// works around a Linux/.so undefined typeinfo link error for plugins using -frtti
+const TCHAR* FGCObject::UnknownGCObjectName = UE::GC::UnknownGCObjectName.GetData();
+
+void FGCObject::StaticInit()
+{
+	if (GGCObjectReferencer == nullptr)
+	{
+		GGCObjectReferencer = NewObject<UGCObjectReferencer>();
+		GGCObjectReferencer->AddToRoot();
+	}
+}
+
+void FGCObject::RegisterGCObject()
+{
+	// Some objects can get created after the engine started shutting down (lazy init of singletons etc).
+	if (!IsEngineExitRequested() && !bReferenceAdded)
+	{
+		StaticInit();
+
+		// Add this instance to the referencer's list
+		GGCObjectReferencer->AddObject(this);
+		bReferenceAdded = true;
+	}
+}
 
+void FGCObject::UnregisterGCObject()
+{
+	// GObjectSerializer will be NULL if this object gets destroyed after the exit purge.
+	// We want to make sure we remove any objects that were added to the GGCObjectReferencer during Init when exiting
+	if (GGCObjectReferencer && bReferenceAdded)
+	{
+		GGCObjectReferencer->RemoveObject(this);
+		bReferenceAdded = false;
+	}
+}
diff --git a/Engine/Source/Runtime/CoreUObject/Private/UObject/Class.cpp b/Engine/Source/Runtime/CoreUObject/Private/UObject/Class.cpp
--- a/Engine/Source/Runtime/CoreUObject/Private/UObject/Class.cpp
+++ b/Engine/Source/Runtime/CoreUObject/Private/UObject/Class.cpp
@@ -80,6 +80,8 @@ static_assert(sizeof(__underlying_type(EClassFlags)) == sizeof(uint32), "expecti
 
 //////////////////////////////////////////////////////////////////////////
 
+namespace { UE_CALL_ONCE(UE::GC::RegisterSlowImplementation, &UClass::AddReferencedObjects, UE::GC::EAROFlags::Unbalanced); }
+
 FThreadSafeBool& InternalSafeGetTokenStreamDirtyFlag()
 {
 	static FThreadSafeBool TokenStreamDirty(true);
@@ -574,7 +576,8 @@ void UField::SetAssociatedFField(FField* InField)
 
 IMPLEMENT_CORE_INTRINSIC_CLASS(UField, UObject,
 	{
-		Class->EmitObjectReference(STRUCT_OFFSET(UField, Next), TEXT("Next"));
+		UE::GC::FTokenStreamBuilder& Builder = UE::GC::FIntrinsicClassTokens::AllocateBuilder(Class);
+		Builder.EmitObjectReference(STRUCT_OFFSET(UField, Next), TEXT("Next"));
 	}
 );
 
@@ -1081,8 +1084,14 @@ void UStruct::SerializeBin( FStructuredArchive::FSlot Slot, void* Data ) const
 	
 	if (UnderlyingArchive.IsObjectReferenceCollector() && !UnderlyingArchive.IsModifyingWeakAndStrongReferences())
 	{
+		// The FProperty instance might start in the middle of a cache line	
+		static constexpr uint32 ExtraPrefetchBytes = PLATFORM_CACHE_LINE_SIZE - /* min alignment */ 16;
+		// Prefetch vtable, PropertyFlags and NextRef. NextRef comes last.
+		static constexpr uint32 PropertyPrefetchBytes = offsetof(FProperty, NextRef) + ExtraPrefetchBytes;
+		FPlatformMisc::PrefetchBlock(RefLink, PropertyPrefetchBytes);
 		for( FProperty* RefLinkProperty=RefLink; RefLinkProperty!=NULL; RefLinkProperty=RefLinkProperty->NextRef )
 		{
+			FPlatformMisc::PrefetchBlock(RefLinkProperty->NextRef, PropertyPrefetchBytes);
 			RefLinkProperty->SerializeBinProperty(PropertyStream.EnterElement(), Data );
 		}
 	}
@@ -2179,13 +2188,14 @@ void UStruct::InstanceSubobjectTemplates( void* Data, void const* DefaultData, U
 IMPLEMENT_CORE_INTRINSIC_CLASS(UStruct, UField,
 	{
 		Class->CppClassStaticFunctions = UOBJECT_CPPCLASS_STATICFUNCTIONS_FORCLASS(UStruct);
-		Class->EmitObjectReference(STRUCT_OFFSET(UStruct, SuperStruct), TEXT("SuperStruct"));
-		Class->EmitObjectReference(STRUCT_OFFSET(UStruct, Children), TEXT("Children"));
+
+		UE::GC::FTokenStreamBuilder& Builder = UE::GC::FIntrinsicClassTokens::AllocateBuilder(Class);
+		Builder.EmitObjectReference(STRUCT_OFFSET(UStruct, SuperStruct), TEXT("SuperStruct"));
+		Builder.EmitObjectReference(STRUCT_OFFSET(UStruct, Children), TEXT("Children"));
 
 		// Note: None of the *Link members need to be emitted, as they only contain properties
 		// that are in the Children chain or SuperStruct->Children chains.
-
-		Class->EmitObjectArrayReference(STRUCT_OFFSET(UStruct, ScriptAndPropertyObjectReferences), TEXT("ScriptAndPropertyObjectReferences"));
+		Builder.EmitObjectArrayReference(STRUCT_OFFSET(UStruct, ScriptAndPropertyObjectReferences), TEXT("ScriptAndPropertyObjectReferences"));
 	}
 );
 
@@ -3948,19 +3958,23 @@ void UClass::AddReferencedObjects(UObject* InThis, FReferenceCollector& Collecto
 	UClass* This = CastChecked<UClass>(InThis);
 	for (FImplementedInterface& Inter : This->Interfaces)
 	{
-		Collector.AddReferencedObject( Inter.Class, This );
+		Collector.AddStableReference( &Inter.Class );
 	}
 
-	Collector.AddReferencedObjects( This->FuncMap, This );
-	Collector.AddReferencedObject( This->ClassWithin, This );
+	for (TPair<FName, UFunction*>& Pair : This->FuncMap)
+	{
+		Collector.AddStableReference( &Pair.Value );
+	}
+	
+	Collector.AddStableReference( &This->ClassWithin );
 
 #if WITH_EDITORONLY_DATA
-	Collector.AddReferencedObject( This->ClassGeneratedBy, This );
+	Collector.AddStableReference( &This->ClassGeneratedBy );
 #endif
 
 	if ( !Collector.IsIgnoringArchetypeRef() )
 	{
-		Collector.AddReferencedObject( This->ClassDefaultObject, This );
+		Collector.AddStableReference( &This->ClassDefaultObject );
 	}
 	else if( This->ClassDefaultObject != NULL)
 	{
@@ -3968,7 +3982,7 @@ void UClass::AddReferencedObjects(UObject* InThis, FReferenceCollector& Collecto
 		This->CallAddReferencedObjects(This->ClassDefaultObject, Collector);
 	}
 
-	Collector.AddReferencedObject( This->SparseClassDataStruct, This );
+	Collector.AddStableReference( &This->SparseClassDataStruct );
 
 	// Add sparse class data
 	if (This->SparseClassDataStruct && This->SparseClassData)
@@ -3979,12 +3993,9 @@ void UClass::AddReferencedObjects(UObject* InThis, FReferenceCollector& Collecto
 		}
 		else
 		{
-			// The iterator will recursively loop through all structs in structs too.
-			for (TPropertyValueIterator<FObjectProperty> It(This->SparseClassDataStruct, This->SparseClassData); It; ++It)
-			{
-				// @TODO: OBJPTR: 
-				Collector.AddReferencedObject(It.Key()->GetObjectPtrPropertyValueRef(It.Value()));
-			}
+			// Skip FFieldPathProperty and FInterfaceProperty references since original code iterated over
+			// FObjectProperty only. Unsure if sparse class author intended this.
+			Collector.AddPropertyReferencesLimitedToObjectProperties(This->SparseClassDataStruct, This->SparseClassData, This);
 		}
 	}
 	
@@ -6301,10 +6312,11 @@ IMPLEMENT_CORE_INTRINSIC_CLASS(UClass, UStruct,
 	{
 		Class->CppClassStaticFunctions = UOBJECT_CPPCLASS_STATICFUNCTIONS_FORCLASS(UClass);
 
-		Class->EmitObjectReference(STRUCT_OFFSET(UClass, ClassDefaultObject), TEXT("ClassDefaultObject"));
-		Class->EmitObjectReference(STRUCT_OFFSET(UClass, ClassWithin), TEXT("ClassWithin"));
-		Class->EmitObjectReference(STRUCT_OFFSET(UClass, ClassGeneratedBy), TEXT("ClassGeneratedBy"));
-		Class->EmitObjectArrayReference(STRUCT_OFFSET(UClass, NetFields), TEXT("NetFields"));
+		UE::GC::FTokenStreamBuilder& Builder = UE::GC::FIntrinsicClassTokens::AllocateBuilder(Class);
+		Builder.EmitObjectReference(STRUCT_OFFSET(UClass, ClassDefaultObject), TEXT("ClassDefaultObject"));
+		Builder.EmitObjectReference(STRUCT_OFFSET(UClass, ClassWithin), TEXT("ClassWithin"));
+		Builder.EmitObjectReference(STRUCT_OFFSET(UClass, ClassGeneratedBy), TEXT("ClassGeneratedBy"));
+		Builder.EmitObjectArrayReference(STRUCT_OFFSET(UClass, NetFields), TEXT("NetFields"));
 	}
 );
 #else
@@ -6312,9 +6324,10 @@ IMPLEMENT_CORE_INTRINSIC_CLASS(UClass, UStruct,
 	{
 		Class->CppClassStaticFunctions = UOBJECT_CPPCLASS_STATICFUNCTIONS_FORCLASS(UClass);
 
-		Class->EmitObjectReference(STRUCT_OFFSET(UClass, ClassDefaultObject), TEXT("ClassDefaultObject"));
-		Class->EmitObjectReference(STRUCT_OFFSET(UClass, ClassWithin), TEXT("ClassWithin"));
-		Class->EmitObjectArrayReference(STRUCT_OFFSET(UClass, NetFields), TEXT("NetFields"));
+		UE::GC::FTokenStreamBuilder& Builder = UE::GC::FIntrinsicClassTokens::AllocateBuilder(Class);
+		Builder.EmitObjectReference(STRUCT_OFFSET(UClass, ClassDefaultObject), TEXT("ClassDefaultObject"));
+		Builder.EmitObjectReference(STRUCT_OFFSET(UClass, ClassWithin), TEXT("ClassWithin"));
+		Builder.EmitObjectArrayReference(STRUCT_OFFSET(UClass, NetFields), TEXT("NetFields"));
 	}
 );
 #endif
diff --git a/Engine/Source/Runtime/CoreUObject/Private/UObject/CoreNative.cpp b/Engine/Source/Runtime/CoreUObject/Private/UObject/CoreNative.cpp
--- a/Engine/Source/Runtime/CoreUObject/Private/UObject/CoreNative.cpp
+++ b/Engine/Source/Runtime/CoreUObject/Private/UObject/CoreNative.cpp
@@ -72,7 +72,9 @@ UClass* Z_Construct_UClass_UObject()
 	{
 		Z_Registration_Info_UClass_UObject.OuterSingleton = UObject::StaticClass();
 		UObjectForceRegistration(Z_Registration_Info_UClass_UObject.OuterSingleton);
-		UObjectBase::EmitBaseReferences(Z_Registration_Info_UClass_UObject.OuterSingleton);
+#if WITH_EDITOR
+		static_assert(false, "Need to EmitBaseReferences here");
+#endif
 		Z_Registration_Info_UClass_UObject.OuterSingleton->StaticLink();
 	}
 	check(Z_Registration_Info_UClass_UObject.OuterSingleton->GetClass());
diff --git a/Engine/Source/Runtime/CoreUObject/Private/UObject/GCObjectInfo.cpp b/Engine/Source/Runtime/CoreUObject/Private/UObject/GCObjectInfo.cpp
--- a/Engine/Source/Runtime/CoreUObject/Private/UObject/GCObjectInfo.cpp
+++ b/Engine/Source/Runtime/CoreUObject/Private/UObject/GCObjectInfo.cpp
@@ -73,10 +73,9 @@ FString FGCObjectInfo::GetFullName() const
 	return FString::Printf(TEXT("%s %s"), *GetClassName(), *GetPathName());
 }
 
-FGCObjectInfo* FGCObjectInfo::FindOrAddInfoHelper(UObject* InObject, TMap<UObject*, FGCObjectInfo*>& InOutObjectToInfoMap)
+FGCObjectInfo* FGCObjectInfo::FindOrAddInfoHelper(const UObject* InObject, TMap<const UObject*, FGCObjectInfo*>& InOutObjectToInfoMap)
 {
-	FGCObjectInfo** ExistingObjInfo = InOutObjectToInfoMap.Find(InObject);
-	if (ExistingObjInfo)
+	if (FGCObjectInfo** ExistingObjInfo = InOutObjectToInfoMap.Find(InObject))
 	{
 		return *ExistingObjInfo;
 	}
