GC reachability optimizations. ~1.5x speedup and ~2.5MB memory saved for an internal title.

diff --git a/Engine/Source/Runtime/CoreUObject/Private/UObject/GarbageCollection.cpp b/Engine/Source/Runtime/CoreUObject/Private/UObject/GarbageCollection.cpp
--- a/Engine/Source/Runtime/CoreUObject/Private/UObject/GarbageCollection.cpp
+++ b/Engine/Source/Runtime/CoreUObject/Private/UObject/GarbageCollection.cpp
@@ -5,11 +5,14 @@
 =============================================================================*/
 
 #include "UObject/GarbageCollection.h"
+#include "Containers/Deque.h"
+#include "Containers/StaticBitArray.h"
 #include "HAL/ThreadSafeBool.h"
 #include "Misc/TimeGuard.h"
 #include "HAL/IConsoleManager.h"
 #include "Misc/App.h"
 #include "Misc/CoreDelegates.h"
+#include "Tasks/Task.h"
 #include "UObject/ScriptInterface.h"
 #include "UObject/UObjectAllocator.h"
 #include "UObject/UObjectBase.h"
@@ -45,6 +48,9 @@
 
 DEFINE_LOG_CATEGORY(LogGarbage);
 
+static_assert(sizeof(FGCSkipInfo) == sizeof(uint32),		"Must have size of a GC token");
+static_assert(sizeof(FGCReferenceInfo) == sizeof(uint32),	"Must have size of a GC token");
+
 /** Object count during last mark phase																				*/
 FThreadSafeCounter		GObjectCountDuringLastMarkPhase;
 /** Whether UObject hash tables are locked by GC */
@@ -77,15 +83,8 @@ static FCriticalSection GUnreachableObjectsCritical;
 static int32 GUnrechableObjectIndex = 0;
 static double GLastGCTime = 0;
 
-/** Helpful constant for determining how many token slots we need to store a pointer **/
-static const uint32 GNumTokensPerPointer = sizeof(void*) / sizeof(uint32); //-V514
-
 bool GIsGarbageCollecting = false;
 
-int32 FGCReferenceTokenStream::MaxStackSize = 1;
-
-const FGCReferenceInfo FGCReferenceInfo::EndOfStreamToken(GCRT_EndOfStream, 0);
-
 /**
 * Call back into the async loading code to inform of the destruction of serialized objects
 */
@@ -222,15 +221,6 @@ bool IsGarbageCollectionLocked()
 	return FGCCSyncObject::Get().IsAsyncLocked();
 }
 
-// Minimum number of objects to spawn a GC sub-task for
-static int32 GMinDesiredObjectsPerSubTask = 128;
-static FAutoConsoleVariableRef CVarMinDesiredObjectsPerSubTask(
-	TEXT("gc.MinDesiredObjectsPerSubTask"),
-	GMinDesiredObjectsPerSubTask,
-	TEXT("Minimum number of objects to spawn a GC sub-task for."),
-	ECVF_Default
-	);
-
 static int32 GIncrementalBeginDestroyEnabled = 1;
 static FAutoConsoleVariableRef CIncrementalBeginDestroyEnabled(
 	TEXT("gc.IncrementalBeginDestroyEnabled"),
@@ -248,8 +238,10 @@ static FAutoConsoleVariableRef CMultithreadedDestructionEnabled(
 );
 
 #if UE_BUILD_SHIPPING
+static constexpr bool GRedoReachabilityToTrackGarbage = false;
 static constexpr int32 GGarbageReferenceTrackingEnabled = 0;
 #else
+static bool GRedoReachabilityToTrackGarbage;
 int32 GGarbageReferenceTrackingEnabled = 0;
 static FAutoConsoleVariableRef CGarbageReferenceTrackingEnabled(
 	TEXT("gc.GarbageReferenceTrackingEnabled"),
@@ -259,26 +251,6 @@ static FAutoConsoleVariableRef CGarbageReferenceTrackingEnabled(
 );
 #endif
 
-#if !UE_BUILD_SHIPPING
-static FAutoConsoleCommand CmdCalculateTokenStreamSize(
-	TEXT("gc.CalculateTokenStreamSize"),
-	TEXT(""),
-	FConsoleCommandWithArgsDelegate::CreateLambda([](const TArray<FString>& Args)
-		{
-			int64 TokenSize = 0;
-			int64 DebugSize = 0;
-			int32 NumClasses = 0;
-			for (TObjectIterator<UClass> It; It; ++It)
-			{
-				TokenSize += It->ReferenceTokenStream.GetTokenAllocatedSize();
-				DebugSize += It->ReferenceTokenStream.GetDebugInfoAllocatedSize();
-				NumClasses++;
-			}
-			UE_LOG(LogGarbage, Display, TEXT("Memory allocated for GC Tokens: %lldb, Debug Info: %lldb, Total: %lldb (%d classes)"), TokenSize, DebugSize, (TokenSize + DebugSize), NumClasses);
-		})
-);
-#endif // !UE_BUILD_SHIPPING
-
 #if PERF_DETAILED_PER_CLASS_GC_STATS
 /** Map from a UClass' FName to the number of objects that were purged during the last purge phase of this class.	*/
 static TMap<const FName,uint32> GClassToPurgeCountMap;
@@ -668,528 +640,2089 @@ bool IsInGarbageCollectorThread()
 	return GAsyncPurge ? GAsyncPurge->IsInAsyncPurgeThread() : IsInGameThread();
 }
 
-/** Called on shutdown to free GC memory */
-void ShutdownGarbageCollection()
-{
-	FGCArrayPool::Get().Cleanup();
-	delete GAsyncPurge;
-	GAsyncPurge = nullptr;
-#if ENABLE_GC_HISTORY
-	FGCHistory::Get().Cleanup();
-#endif
-}
+namespace UE::GC {
 
-/**
-* Handles UObject references found by TFastReferenceCollector
-*/
+using AROFunc = void (*)(UObject*, FReferenceCollector&);
 
 #if UE_WITH_GC
-template <EFastReferenceCollectorOptions Options>
-class FGCReferenceProcessor
+
+//////////////////////////////////////////////////////////////////////////
+
+/** Pool for reusing contexts between CollectReferences calls */
+class FContextPool
 {
-	// Local copies of globals to make sure they're cached
-	const int32 MinDesiredObjectsPerSubTask;
-	const bool bGarbageReferenceTrackingEnabled;
-#if ENABLE_GC_HISTORY
-	const bool bWithHistory;
-#endif // ENABLE_GC_HISTORY
+	friend class FContextPoolScope;
+	FContextPool();
+	
+	TArray<TUniquePtr<FWorkerContext>> Reusable;
+	int32 NumAllocated = 0;
+};
 
-public:
+// Helps validate single-threaded access and check for leaks
+//
+// Bit of defensive coding during a large refactor, can be removed in the long run
+class FContextPoolScope
+{
+	FContextPool& Pool;
+	bool bNested;
+	static FContextPool& Get();
 
-	constexpr static FORCEINLINE bool IsParallel()
-	{
-		return !!(Options & EFastReferenceCollectorOptions::Parallel);
-	}
-	constexpr static FORCEINLINE bool IsWithClusters()
-	{
-		return !!(Options & EFastReferenceCollectorOptions::WithClusters);
-	}
-	constexpr static FORCEINLINE bool IsWithPendingKill()
+	void CheckGameThread()
 	{
-		return !!(Options & EFastReferenceCollectorOptions::WithPendingKill);
+		checkf(IsInGameThread(), TEXT("Context pool use restricted to game thread"));
 	}
+public:
+	FContextPoolScope() : Pool(Get()) { CheckGameThread(); }
+	~FContextPoolScope() { CheckGameThread(); }
 
-	FGCReferenceProcessor()
-		: MinDesiredObjectsPerSubTask(GMinDesiredObjectsPerSubTask)
-		, bGarbageReferenceTrackingEnabled(!!GGarbageReferenceTrackingEnabled)
-#if ENABLE_GC_HISTORY
-		, bWithHistory(FGCHistory::Get().IsActive())
-#endif // ENABLE_GC_HISTORY
-	{
-	}
+	FWorkerContext* AllocateFromPool();
+	void ReturnToPool(FWorkerContext* Context);
+	TConstArrayView<TUniquePtr<FWorkerContext>> PeekFree() { return Pool.Reusable; }
+	
+	void Cleanup();
+	int32 NumAllocated() const { return Pool.NumAllocated; }
+};
 
-	FORCEINLINE int32 GetMinDesiredObjectsPerSubTask() const
+FContextPool& FContextPoolScope::Get()
+{
+	static FContextPool Singleton;
+	return Singleton;
+}
+
+FWorkerContext* FContextPoolScope::AllocateFromPool()
+{
+	CheckGameThread();
+
+	++Pool.NumAllocated;
+	if (Pool.Reusable.IsEmpty())
 	{
-		return MinDesiredObjectsPerSubTask;
+		return new FWorkerContext();
 	}
+	
+	FWorkerContext* Out = Pool.Reusable.Pop().Release();
+	Out->AllocateWorkerIndex();
+	return Out;
+}
 
-	void UpdateDetailedStats(UObject* CurrentObject, uint32 DeltaCycles)
-	{
-#if PERF_DETAILED_PER_CLASS_GC_STATS
-		// Keep track of how many refs we encountered for the object's class.
-		const FName& ClassName = CurrentObject->GetClass()->GetFName();
-		// Refs to objects that reside in permanent object pool.
-		uint32 ClassDisregardedObjRefs = GClassToDisregardedObjectRefsMap.FindRef(ClassName);
-		GClassToDisregardedObjectRefsMap.Add(ClassName, ClassDisregardedObjRefs + GCurrentObjectDisregardedObjectRefs);
-		// Refs to regular objects.
-		uint32 ClassRegularObjRefs = GClassToRegularObjectRefsMap.FindRef(ClassName);
-		GClassToRegularObjectRefsMap.Add(ClassName, ClassRegularObjRefs + GCurrentObjectRegularObjectRefs);
-		// Track per class cycle count spent in GC.
-		uint32 ClassCycles = GClassToCyclesMap.FindRef(ClassName);
-		GClassToCyclesMap.Add(ClassName, ClassCycles + DeltaCycles);
-		// Reset current counts.
-		GCurrentObjectDisregardedObjectRefs = 0;
-		GCurrentObjectRegularObjectRefs = 0;
+void FContextPoolScope::ReturnToPool(FWorkerContext* Context)
+{
+	CheckGameThread();
+	check(Pool.NumAllocated >= 1);
+	--Pool.NumAllocated;
+	Context->FreeWorkerIndex();
+#if !UE_BUILD_SHIPPING
+	GRedoReachabilityToTrackGarbage |= Context->bDetectedGarbageReference;
+	Context->bDetectedGarbageReference = false;
 #endif
-	}
+	Pool.Reusable.Emplace(Context);
+}
 
-	void LogDetailedStatsSummary()
+void FContextPoolScope::Cleanup()
+{
+	SIZE_T FreedMemory = 0;
+	for (const TUniquePtr<FWorkerContext>& Context : Pool.Reusable)
 	{
-#if PERF_DETAILED_PER_CLASS_GC_STATS
-		LogClassCountInfo(TEXT("references to regular objects from"), GClassToRegularObjectRefsMap, 20, 0);
-		LogClassCountInfo(TEXT("references to permanent objects from"), GClassToDisregardedObjectRefsMap, 20, 0);
-		LogClassCountInfo(TEXT("cycles for GC"), GClassToCyclesMap, 20, 0);
-#endif
+		checkf(Context->WeakReferences.Num() == 0, TEXT("Cleaning up with active weak references"));
+		FreedMemory += Context->GetAllocatedSize();
 	}
+	UE_LOG(LogGarbage, Log, TEXT("Freed %" SIZE_T_FMT "b from %d GC contexts"), FreedMemory, Pool.Reusable.Num());
 
-	/** Marks all objects that can't be directly in a cluster but are referenced by it as reachable */
-	static FORCENOINLINE bool MarkClusterMutableObjectsAsReachable(FUObjectCluster& Cluster, TArray<UObject*>& ObjectsToSerialize)
-	{
-		check(IsWithClusters());
+	Pool.Reusable.Empty();
+}
+
+static void DumpMemoryStats(FOutputDevice& OutputDevice);
+
+FContextPool::FContextPool()
+{
+	static FAutoConsoleCommandWithOutputDevice Cmd(TEXT("gc.DumpMemoryStats"), TEXT("Dump GC memory stats"),
+		FConsoleCommandWithOutputDeviceDelegate::CreateStatic(&DumpMemoryStats));
+}
+
+//////////////////////////////////////////////////////////////////////////
+	
+FORCEINLINE constexpr bool IsWithClusters(EGCOptions Options) {	return !!(Options & EGCOptions::WithClusters); }
+
+template<EGCOptions Options, class ContainerType>
+static void MarkReferencedClustersAsReachable(int32 ClusterIndex, ContainerType& ObjectsToSerialize);
+
+// Let MarkReferencedClustersAsReachable generically Add() to both TArray<UObject*> and FWorkBlockifier
+template<EGCOptions Options>
+struct TAddAdapter
+{
+	FWorkBlockifier& Queue;
+	FORCEINLINE void Add(UObject* Object) { Queue.Add<Options>(Object); }
+};
+
+template<EGCOptions Options>
+FORCEINLINE static void MarkReferencedClustersAsReachableThunk(int32 ClusterIndex, FWorkBlockifier& ObjectsToSerialize)
+{
+	TAddAdapter<Options> Adapter{ObjectsToSerialize};
+	MarkReferencedClustersAsReachable<Options>(ClusterIndex, Adapter);
+}
 
-		// This is going to be the return value and basically means that we ran across some pending kill objects
-		bool bAddClusterObjectsToSerialize = false;
-		for (int32& ReferencedMutableObjectIndex : Cluster.MutableObjects)
+/** Marks all objects that can't be directly in a cluster but are referenced by it as reachable */
+template<EGCOptions Options, class ContainerType>
+static bool MarkClusterMutableObjectsAsReachable(FUObjectCluster& Cluster, ContainerType& ObjectsToSerialize)
+{
+	// This is going to be the return value and basically means that we ran across some pending kill objects
+	bool bAddClusterObjectsToSerialize = false;
+	for (int32& ReferencedMutableObjectIndex : Cluster.MutableObjects)
+	{
+		if (ReferencedMutableObjectIndex >= 0) // Pending kill support
 		{
-			if (ReferencedMutableObjectIndex >= 0) // Pending kill support
+			FUObjectItem* ReferencedMutableObjectItem = GUObjectArray.IndexToObjectUnsafeForGC(ReferencedMutableObjectIndex);
+			if constexpr (IsParallel(Options))
 			{
-				FUObjectItem* ReferencedMutableObjectItem = GUObjectArray.IndexToObjectUnsafeForGC(ReferencedMutableObjectIndex);
-				if (IsParallel())
+				PRAGMA_DISABLE_DEPRECATION_WARNINGS
+				if (!ReferencedMutableObjectItem->HasAnyFlags(EInternalObjectFlags::PendingKill | EInternalObjectFlags::Garbage))
+				PRAGMA_ENABLE_DEPRECATION_WARNINGS
 				{
-					PRAGMA_DISABLE_DEPRECATION_WARNINGS
-					if (!ReferencedMutableObjectItem->HasAnyFlags(EInternalObjectFlags::PendingKill | EInternalObjectFlags::Garbage))
-					PRAGMA_ENABLE_DEPRECATION_WARNINGS
+					if (ReferencedMutableObjectItem->IsUnreachable())
 					{
-						if (ReferencedMutableObjectItem->IsUnreachable())
+						if (ReferencedMutableObjectItem->ThisThreadAtomicallyClearedRFUnreachable())
 						{
-							if (ReferencedMutableObjectItem->ThisThreadAtomicallyClearedRFUnreachable())
-							{
-								// Needs doing because this is either a normal unclustered object (clustered objects are never unreachable) or a cluster root
-								ObjectsToSerialize.Add(static_cast<UObject*>(ReferencedMutableObjectItem->Object));
+							// Needs doing because this is either a normal unclustered object (clustered objects are never unreachable) or a cluster root
+							ObjectsToSerialize.Add(static_cast<UObject*>(ReferencedMutableObjectItem->Object));
 
-								// So is this a cluster root maybe?
-								if (ReferencedMutableObjectItem->GetOwnerIndex() < 0)
-								{
-									MarkReferencedClustersAsReachable(ReferencedMutableObjectItem->GetClusterIndex(), ObjectsToSerialize);
-								}
+							// So is this a cluster root maybe?
+							if (ReferencedMutableObjectItem->GetOwnerIndex() < 0)
+							{
+								MarkReferencedClustersAsReachable<Options>(ReferencedMutableObjectItem->GetClusterIndex(), ObjectsToSerialize);
 							}
 						}
-						else if (ReferencedMutableObjectItem->GetOwnerIndex() > 0 && !ReferencedMutableObjectItem->HasAnyFlags(EInternalObjectFlags::ReachableInCluster))
+					}
+					else if (ReferencedMutableObjectItem->GetOwnerIndex() > 0 && !ReferencedMutableObjectItem->HasAnyFlags(EInternalObjectFlags::ReachableInCluster))
+					{
+						// This is a clustered object that maybe hasn't been processed yet
+						if (ReferencedMutableObjectItem->ThisThreadAtomicallySetFlag(EInternalObjectFlags::ReachableInCluster))
 						{
-							// This is a clustered object that maybe hasn't been processed yet
-							if (ReferencedMutableObjectItem->ThisThreadAtomicallySetFlag(EInternalObjectFlags::ReachableInCluster))
+							// Needs doing, we need to get its cluster root and process it too
+							FUObjectItem* ReferencedMutableObjectsClusterRootItem = GUObjectArray.IndexToObjectUnsafeForGC(ReferencedMutableObjectItem->GetOwnerIndex());
+							if (ReferencedMutableObjectsClusterRootItem->IsUnreachable())
 							{
-								// Needs doing, we need to get its cluster root and process it too
-								FUObjectItem* ReferencedMutableObjectsClusterRootItem = GUObjectArray.IndexToObjectUnsafeForGC(ReferencedMutableObjectItem->GetOwnerIndex());
-								if (ReferencedMutableObjectsClusterRootItem->IsUnreachable())
+								// The root is also maybe unreachable so process it and all the referenced clusters
+								if (ReferencedMutableObjectsClusterRootItem->ThisThreadAtomicallyClearedRFUnreachable())
 								{
-									// The root is also maybe unreachable so process it and all the referenced clusters
-									if (ReferencedMutableObjectsClusterRootItem->ThisThreadAtomicallyClearedRFUnreachable())
-									{
-										MarkReferencedClustersAsReachable(ReferencedMutableObjectsClusterRootItem->GetClusterIndex(), ObjectsToSerialize);
-									}
+									MarkReferencedClustersAsReachable<Options>(ReferencedMutableObjectsClusterRootItem->GetClusterIndex(), ObjectsToSerialize);
 								}
 							}
 						}
 					}
-					else
-					{
-						// Pending kill support for clusters (multi-threaded case)
-						ReferencedMutableObjectIndex = -1;
-						bAddClusterObjectsToSerialize = true;
-					}
 				}
-				PRAGMA_DISABLE_DEPRECATION_WARNINGS
-				else if (!ReferencedMutableObjectItem->HasAnyFlags(EInternalObjectFlags::PendingKill | EInternalObjectFlags::Garbage))
-				PRAGMA_ENABLE_DEPRECATION_WARNINGS
+				else
 				{
-					if (ReferencedMutableObjectItem->IsUnreachable())
-					{
-						// Needs doing because this is either a normal unclustered object (clustered objects are never unreachable) or a cluster root
-						ReferencedMutableObjectItem->ClearFlags(EInternalObjectFlags::Unreachable);
-						ObjectsToSerialize.Add(static_cast<UObject*>(ReferencedMutableObjectItem->Object));
-
-						// So is this a cluster root?
-						if (ReferencedMutableObjectItem->GetOwnerIndex() < 0)
-						{
-							MarkReferencedClustersAsReachable(ReferencedMutableObjectItem->GetClusterIndex(), ObjectsToSerialize);
-						}
-					}
-					else if (ReferencedMutableObjectItem->GetOwnerIndex() > 0 && !ReferencedMutableObjectItem->HasAnyFlags(EInternalObjectFlags::ReachableInCluster))
-					{
-						// This is a clustered object that hasn't been processed yet
-						ReferencedMutableObjectItem->SetFlags(EInternalObjectFlags::ReachableInCluster);
+					// Pending kill support for clusters (multi-threaded case)
+					ReferencedMutableObjectIndex = -1;
+					bAddClusterObjectsToSerialize = true;
+				}
+			}
+			PRAGMA_DISABLE_DEPRECATION_WARNINGS
+			else if (!ReferencedMutableObjectItem->HasAnyFlags(EInternalObjectFlags::PendingKill | EInternalObjectFlags::Garbage))
+			PRAGMA_ENABLE_DEPRECATION_WARNINGS
+			{
+				if (ReferencedMutableObjectItem->IsUnreachable())
+				{
+					// Needs doing because this is either a normal unclustered object (clustered objects are never unreachable) or a cluster root
+					ReferencedMutableObjectItem->ClearFlags(EInternalObjectFlags::Unreachable);
+					ObjectsToSerialize.Add(static_cast<UObject*>(ReferencedMutableObjectItem->Object));
 						
-						// If the root is also unreachable, process it and all its referenced clusters
-						FUObjectItem* ReferencedMutableObjectsClusterRootItem = GUObjectArray.IndexToObjectUnsafeForGC(ReferencedMutableObjectItem->GetOwnerIndex());
-						if (ReferencedMutableObjectsClusterRootItem->IsUnreachable())
-						{
-							ReferencedMutableObjectsClusterRootItem->ClearFlags(EInternalObjectFlags::Unreachable);
-							MarkReferencedClustersAsReachable(ReferencedMutableObjectsClusterRootItem->GetClusterIndex(), ObjectsToSerialize);
-						}
+					// So is this a cluster root?
+					if (ReferencedMutableObjectItem->GetOwnerIndex() < 0)
+					{
+						MarkReferencedClustersAsReachable<Options>(ReferencedMutableObjectItem->GetClusterIndex(), ObjectsToSerialize);
 					}
 				}
-				else
+				else if (ReferencedMutableObjectItem->GetOwnerIndex() > 0 && !ReferencedMutableObjectItem->HasAnyFlags(EInternalObjectFlags::ReachableInCluster))
 				{
-					// Pending kill support for clusters (single-threaded case)
-					ReferencedMutableObjectIndex = -1;
-					bAddClusterObjectsToSerialize = true;
+					// This is a clustered object that hasn't been processed yet
+					ReferencedMutableObjectItem->SetFlags(EInternalObjectFlags::ReachableInCluster);
+						
+					// If the root is also unreachable, process it and all its referenced clusters
+					FUObjectItem* ReferencedMutableObjectsClusterRootItem = GUObjectArray.IndexToObjectUnsafeForGC(ReferencedMutableObjectItem->GetOwnerIndex());
+					if (ReferencedMutableObjectsClusterRootItem->IsUnreachable())
+					{
+						ReferencedMutableObjectsClusterRootItem->ClearFlags(EInternalObjectFlags::Unreachable);
+						MarkReferencedClustersAsReachable<Options>(ReferencedMutableObjectsClusterRootItem->GetClusterIndex(), ObjectsToSerialize);
+					}
 				}
 			}
+			else
+			{
+				// Pending kill support for clusters (single-threaded case)
+				ReferencedMutableObjectIndex = -1;
+				bAddClusterObjectsToSerialize = true;
+			}
 		}
-		return bAddClusterObjectsToSerialize;
 	}
+	return bAddClusterObjectsToSerialize;
+}
 
-	/** Marks all clusters referenced by another cluster as reachable */
-	static FORCENOINLINE void MarkReferencedClustersAsReachable(int32 ClusterIndex, TArray<UObject*>& ObjectsToSerialize)
-	{
-		check(IsWithClusters());
+/** Marks all clusters referenced by another cluster as reachable */
+template<EGCOptions Options, class ContainerType>
+static FORCENOINLINE void MarkReferencedClustersAsReachable(int32 ClusterIndex, ContainerType& ObjectsToSerialize)
+{
+	static_assert(IsWithClusters(Options));
 
-		// If we run across some PendingKill objects we need to add all objects from this cluster
-		// to ObjectsToSerialize so that we can properly null out all the references.
-		// It also means this cluster will have to be dissolved because we may no longer guarantee all cross-cluster references are correct.
+	// If we run across some PendingKill objects we need to add all objects from this cluster
+	// to ObjectsToSerialize so that we can properly null out all the references.
+	// It also means this cluster will have to be dissolved because we may no longer guarantee all cross-cluster references are correct.
 
-		bool bAddClusterObjectsToSerialize = false;
-		FUObjectCluster& Cluster = GUObjectClusters[ClusterIndex];
-		// Also mark all referenced objects from outside of the cluster as reachable
-		for (int32& ReferncedClusterIndex : Cluster.ReferencedClusters)
+	bool bAddClusterObjectsToSerialize = false;
+	FUObjectCluster& Cluster = GUObjectClusters[ClusterIndex];
+	// Also mark all referenced objects from outside of the cluster as reachable
+	for (int32& ReferncedClusterIndex : Cluster.ReferencedClusters)
+	{
+		if (ReferncedClusterIndex >= 0) // Pending Kill support
 		{
-			if (ReferncedClusterIndex >= 0) // Pending Kill support
+			FUObjectItem* ReferencedClusterRootObjectItem = GUObjectArray.IndexToObjectUnsafeForGC(ReferncedClusterIndex);
+			PRAGMA_DISABLE_DEPRECATION_WARNINGS
+			if (!ReferencedClusterRootObjectItem->HasAnyFlags(EInternalObjectFlags::PendingKill | EInternalObjectFlags::Garbage))
+			PRAGMA_ENABLE_DEPRECATION_WARNINGS
 			{
-				FUObjectItem* ReferencedClusterRootObjectItem = GUObjectArray.IndexToObjectUnsafeForGC(ReferncedClusterIndex);
-				PRAGMA_DISABLE_DEPRECATION_WARNINGS
-				if (!ReferencedClusterRootObjectItem->HasAnyFlags(EInternalObjectFlags::PendingKill | EInternalObjectFlags::Garbage))
-				PRAGMA_ENABLE_DEPRECATION_WARNINGS
+				if (ReferencedClusterRootObjectItem->IsUnreachable())
 				{
-					// This condition should get collapsed by the compiler based on the template argument
-					if (IsParallel())
+					if constexpr (IsParallel(Options))
 					{
-						if (ReferencedClusterRootObjectItem->IsUnreachable())
-						{
-							ReferencedClusterRootObjectItem->ThisThreadAtomicallyClearedFlag(EInternalObjectFlags::Unreachable);
-						}
+						ReferencedClusterRootObjectItem->ThisThreadAtomicallyClearedFlag(EInternalObjectFlags::Unreachable);
 					}
 					else
 					{
 						ReferencedClusterRootObjectItem->ClearFlags(EInternalObjectFlags::Unreachable);
 					}
 				}
-				else
-				{
-					// Pending kill support for clusters
-					ReferncedClusterIndex = -1;
-					bAddClusterObjectsToSerialize = true;
-				}
 			}
-		}
-		if (MarkClusterMutableObjectsAsReachable(Cluster, ObjectsToSerialize))
-		{
-			bAddClusterObjectsToSerialize = true;
-		}
-		if (bAddClusterObjectsToSerialize)
-		{
-			// We need to process all cluster objects to handle PendingKill objects we nulled out (-1) from the cluster.
-			for (int32 ClusterObjectIndex : Cluster.Objects)
+			else
 			{
-				FUObjectItem* ClusterObjectItem = GUObjectArray.IndexToObjectUnsafeForGC(ClusterObjectIndex);
-				UObject* ClusterObject = static_cast<UObject*>(ClusterObjectItem->Object);
-				ObjectsToSerialize.Add(ClusterObject);
+				// Pending kill support for clusters
+				ReferncedClusterIndex = -1;
+				bAddClusterObjectsToSerialize = true;
 			}
-			Cluster.bNeedsDissolving = true;
-			GUObjectClusters.SetClustersNeedDissolving();
 		}
 	}
-
-	FORCENOINLINE void HandleGarbageReference(FGCArrayStruct& ObjectsToSerializeStruct, const UObject* const ReferencingObject, UObject*& Object, const int32 TokenIndex)
+	if (MarkClusterMutableObjectsAsReachable<Options>(Cluster, ObjectsToSerialize))
 	{
-#if !UE_BUILD_SHIPPING
-		FName PropertyName;
-		UObject* GarbageReferencer = ReferencingObject ? const_cast<UObject*>(ReferencingObject) : ObjectsToSerializeStruct.GetReferencingObject();
-		if (IsValid(GarbageReferencer))
+		bAddClusterObjectsToSerialize = true;
+	}
+	if (bAddClusterObjectsToSerialize)
+	{
+		// We need to process all cluster objects to handle PendingKill objects we nulled out (-1) from the cluster.
+		for (int32 ClusterObjectIndex : Cluster.Objects)
 		{
-			if (GarbageReferencer != FGCObject::GGCObjectReferencer)
-			{
-				if (TokenIndex >= 0)
-				{
-					UClass* Class = GarbageReferencer->GetClass();
-					FTokenInfo TokenInfo = Class->ReferenceTokenStream.GetTokenInfo(TokenIndex);
-					PropertyName = TokenInfo.Name;
-				}
-				ObjectsToSerializeStruct.GarbageReferences.Add(FGarbageReferenceInfo(GarbageReferencer, Object, Object, PropertyName));
-			}
-			else
-			{
-				FGCObject* GCObjectReferencer = FGCObject::GGCObjectReferencer->GetCurrentlySerializingObject();
-				ObjectsToSerializeStruct.GarbageReferences.Add(FGarbageReferenceInfo(GCObjectReferencer, Object, Object));
-			}
+			FUObjectItem* ClusterObjectItem = GUObjectArray.IndexToObjectUnsafeForGC(ClusterObjectIndex);
+			UObject* ClusterObject = static_cast<UObject*>(ClusterObjectItem->Object);
+			ObjectsToSerialize.Add(ClusterObject);
 		}
+		Cluster.bNeedsDissolving = true;
+		GUObjectClusters.SetClustersNeedDissolving();
+	}
+}
+
+//////////////////////////////////////////////////////////////////////////
+
+enum class EKillable { No, Yes };
+
+// Retains address of reference to allow killing it
+struct FMutableReference { UObject** Object; };
+struct FImmutableReference { UObject* Object; };
+struct FReferenceArray { TArray<UObject*>* Array; };
+
+// Cached immutable reference to avoid rereading it
+struct FResolvedMutableReference : FImmutableReference
+{
+	UObject** Mutable;
+};
+
+FORCEINLINE_DEBUGGABLE static bool ValidateReference(UObject* Object, FPermanentObjectPoolExtents PermanentPool, const UObject* ReferencingObject, FTokenId TokenId)
+{
+	bool bOk = (!PermanentPool.Contains(Object)) & (!!Object) & IsObjectHandleResolved(reinterpret_cast<FObjectHandle&>(Object));
+
+#if ENABLE_GC_OBJECT_CHECKS
+	if (bOk)
+	{
+		if (
+#if DO_POINTER_CHECKS_ON_GC
+			!IsPossiblyAllocatedUObjectPointer(Object) ||
 #endif
+			!Object->IsValidLowLevelFast())
+		{
+			UClass* ReferencingClass = ReferencingObject ? ReferencingObject->GetClass() : nullptr;
+			FString TokenName = ReferencingClass ? ReferencingClass->ReferenceTokens.GetTokenInfo(TokenId).Name.ToString() : TEXT("N/A");
+
+			UE_LOG(LogGarbage, Fatal, TEXT("Invalid object in GC: 0x%016llx, ReferencingObject: %s, TokenId %s (%d)"),
+				(int64)(PTRINT)Object,
+				ReferencingObject ? *ReferencingObject->GetFullName() : TEXT("N/A"),
+				*TokenName, TokenId.AsPrintableIndex());
+		}
 	}
+#endif // ENABLE_GC_OBJECT_CHECKS
 
-#if ENABLE_GC_HISTORY
-	FORCEINLINE void HandleHistoryReference(FGCArrayStruct& ObjectsToSerializeStruct, const UObject* const ReferencingObject, UObject*& Object, const int32 TokenIndex)
+	return bOk;
+}
+
+struct FReferenceMetadata 
+{
+	FReferenceMetadata() {} // Uninitialized
+
+	explicit FReferenceMetadata(int32 ObjectIndex)
 	{
-		UObject* Referencer = ReferencingObject ? const_cast<UObject*>(ReferencingObject) : ObjectsToSerializeStruct.GetReferencingObject();		
-		FGCDirectReference Ref(Object);
+		ObjectItem = GUObjectArray.IndexToObjectUnsafeForGC(ObjectIndex);
+		Flags = ObjectItem->GetFlags();
+	}
+
+	EInternalObjectFlags Flags;
+	FUObjectItem* ObjectItem;
+
+	FORCEINLINE bool Has(EInternalObjectFlags InFlags) const { return !!(int32(Flags) & int32(InFlags)); }
+};
+
+//////////////////////////////////////////////////////////////////////////
+
+class FReferenceBatcherBase
+{
+protected:
+	template<class ReferenceType>
+	struct TEntry
+	{
+		ReferenceType Reference;
+	};
+
+	template<class ReferenceType, uint32 Capacity, uint32 PrefetchCapacity = 0>
+	struct TBatchQueue
+	{
+		uint32 Num = 0;
+		TEntry<ReferenceType> Entries[Capacity + PrefetchCapacity];
 
-		if (TokenIndex >= 0)
+		FORCEINLINE static constexpr uint32 Max() { return Capacity; }
+		FORCEINLINE bool IsFull() const { return Num == Capacity; }
+		FORCEINLINE void Push(TEntry<ReferenceType> Entry) { Entries[Num++] = Entry; }
+		FORCEINLINE uint32 Slack() const { return Capacity - Num; }
+		FORCEINLINE TEntry<ReferenceType>& operator[](uint32 Idx) { return Entries[Idx]; }
+
+		TBatchQueue() { FMemory::Memzero(Entries); } // Zero out so prefetching ahead fetches same (near) null addresses
+		~TBatchQueue() { checkf(Num == 0, TEXT("Failed to flush")); }
+	};
+
+	static constexpr uint32 ArrayBatchSize = 32;
+	static constexpr uint32 UnvalidatedBatchSize = 32;
+	static constexpr uint32 ValidatedBatchSize = 1024;
+	static constexpr uint32 ValidatedPrefetchAhead = 64;
+
+	struct FValidatedBitmask
+	{
+		FORCEINLINE_DEBUGGABLE void Set(uint32 Idx, bool Value) { Words[Idx/64] |= (uint64(Value) << (Idx%64)); }
+		FORCEINLINE_DEBUGGABLE uint64 Get(uint32 Idx) { return (uint64(1) << (Idx%64)) & Words[Idx/64]; }
+		FORCEINLINE_DEBUGGABLE uint32 CountBits()
 		{
-			UClass* Class = Referencer->GetClass();	
-			Ref.ReferencerName = Class->ReferenceTokenStream.GetTokenInfo(TokenIndex).Name;
+			uint32 Out = 0;
+			for (uint64 Word : Words)
+				Out += FPlatformMath::CountBits(Word);
+			return Out;
 		}
-		else if (FGCObject::GGCObjectReferencer == Referencer)
+
+		FORCEINLINE_DEBUGGABLE static FValidatedBitmask Or(const FValidatedBitmask& A, const FValidatedBitmask& B)
 		{
-			Ref.ReferencerName = *FGCObject::GGCObjectReferencer->GetCurrentlySerializingObject()->GetReferencerName();
+			FValidatedBitmask Out;
+			for (uint32 WordIdx = 0; WordIdx < NumWords; ++WordIdx)
+				Out.Words[WordIdx] = A.Words[WordIdx] | B.Words[WordIdx];
+			return Out;
 		}
 
-		TArray<FGCDirectReference>*& DirectReferences = ObjectsToSerializeStruct.History.FindOrAdd(Referencer);
-		if (!DirectReferences)
+		FORCEINLINE_DEBUGGABLE static FValidatedBitmask And(const FValidatedBitmask& A, const FValidatedBitmask& B)
 		{
-			DirectReferences = new TArray<FGCDirectReference>();
+			FValidatedBitmask Out;
+			for (uint32 WordIdx = 0; WordIdx < NumWords; ++WordIdx)
+				Out.Words[WordIdx] = A.Words[WordIdx] & B.Words[WordIdx];
+			return Out;
 		}
-		DirectReferences->Add(Ref);
-	}
-#endif // ENABLE_GC_HISTORY
 
-	/**
-	 * Handles object reference, potentially NULL'ing
-	 *
-	 * @param Object						Object pointer passed by reference
-	 * @param ReferencingObject UObject which owns the reference (can be NULL)
-	 * @param bAllowReferenceElimination	Whether to allow NULL'ing the reference if RF_PendingKill is set
-	*/
-	FORCEINLINE void HandleObjectReference(FGCArrayStruct& ObjectsToSerializeStruct, const UObject * const ReferencingObject, UObject*& Object, const int32 TokenIndex, const EGCTokenType TokenType, const bool bAllowReferenceElimination)
+
+		static const uint32 NumWords = 1 + (UnvalidatedBatchSize - 1) / 64;
+		uint64 Words[NumWords] = {};
+	};
+};
+
+// TReferenceBatcher helpers
+FORCEINLINE UObject*						GetObject(FImmutableReference In) { return In.Object; }
+FORCEINLINE UObject*						GetObject(FMutableReference In) { return *In.Object; }
+FORCEINLINE void							PrefetchObjectPointer(FImmutableReference) {}
+FORCEINLINE void							PrefetchObjectPointer(FMutableReference In) { FPlatformMisc::Prefetch(In.Object); }
+FORCEINLINE FImmutableReference				ToImmutableReference(FMutableReference In) { return {*In.Object}; }
+FORCEINLINE FImmutableReference				ToResolvedReference(FImmutableReference In) { return In; }
+FORCEINLINE FResolvedMutableReference		ToResolvedReference(FMutableReference In) { return { ToImmutableReference(In), In.Object }; }
+template<class RefType> RefType				ToReference(FMutableReference In);
+template<> FORCEINLINE FImmutableReference	ToReference(FMutableReference In) { return ToImmutableReference(In); }
+template<> FORCEINLINE FMutableReference	ToReference(FMutableReference In) { return In; }
+template<class RefType> RefType				MakeReference(UObject*& Object);
+template<> FORCEINLINE FImmutableReference	MakeReference(UObject*& Object) { return {Object}; }
+template<> FORCEINLINE FMutableReference	MakeReference(UObject*& Object) { return {&Object}; }
+
+/**
+ * Queues up memory indirection to avoid synchronous cache misses by prefetching data just in time.
+ *
+ * Uses bounded queues to avoid branching, control temp memory consumption, allow compiler to generate vectorized code
+ * and make life easier for the branch predictor.
+ * 
+ * Array queue feed into "unvalidated" references queue, which feed into validated references queue, which gets processed when full.
+ * 
+ * Unvalidated means references that might be nullptr or part of the permanent object pool, neither of which require processing.
+ */
+template <class UnvalidatedReferenceType, class ValidatedReferenceType, class ProcessorType>
+class TReferenceBatcher : public FReferenceBatcherBase
+{
+	FWorkerContext& Context;
+	const FPermanentObjectPoolExtents PermanentPool;
+
+	alignas (PLATFORM_CACHE_LINE_SIZE)	TBatchQueue<FReferenceArray, ArrayBatchSize>									UnvalidatedArrays;
+	alignas (PLATFORM_CACHE_LINE_SIZE)	TBatchQueue<UnvalidatedReferenceType, UnvalidatedBatchSize>						UnvalidatedReferences;
+	alignas (PLATFORM_CACHE_LINE_SIZE)	TBatchQueue<ValidatedReferenceType, ValidatedBatchSize, ValidatedPrefetchAhead>	ValidatedReferences;
+
+public:
+	TReferenceBatcher(FWorkerContext& InContext) : Context(InContext) {}
+	
+	FORCEINLINE_DEBUGGABLE void PushArray(FReferenceArray UnvalidatedArray)
 	{
-		if (Object == nullptr)
+		FPlatformMisc::Prefetch(UnvalidatedArray.Array);
+		UnvalidatedArrays.Push({ UnvalidatedArray });
+		if (UnvalidatedArrays.IsFull())
 		{
-			return;
+			DrainArraysFull();
 		}
+	}
 
-		TArray<UObject*>& ObjectsToSerialize = ObjectsToSerializeStruct.ObjectsToSerialize;
-
-		// Disregard NULL objects and perform very fast check to see whether object is part of permanent
-		// object pool and should therefore be disregarded. The check doesn't touch the object and is
-		// cache friendly as it's just a pointer compare against to globals.
-		if (GUObjectAllocator.ResidesInPermanentPool(Object))
+	FORCEINLINE_DEBUGGABLE void PushReference(UnvalidatedReferenceType UnvalidatedReference)
+	{
+		UnvalidatedReferences.Push({ UnvalidatedReference });
+		if (UnvalidatedReferences.IsFull())
 		{
-#if PERF_DETAILED_PER_CLASS_GC_STATS
-			GCurrentObjectDisregardedObjectRefs++;
-#endif
-			return;
+			DrainUnvalidatedFull();
 		}
-
-		const int32 ObjectIndex = GUObjectArray.ObjectToIndex(Object);
-		FUObjectItem* ObjectItem = GUObjectArray.IndexToObjectUnsafeForGC(ObjectIndex);
-
-		// Check for garbage objects but skip objects that were already marked as persistent garbage
-		if (!IsWithPendingKill() && ObjectItem->HasAnyFlags(EInternalObjectFlags::Garbage))
+	}
+	
+	FORCEINLINE_DEBUGGABLE void PushReferences(TArrayView<UObject*> Range)
+	{
+		if (Range.Num())
 		{
-			if (bAllowReferenceElimination)
-			{
-				// To avoid content changes, allow reference elimination inside of Blueprints
-				if (ReferencingObject && TokenType == EGCTokenType::NonNative)
-				{
-					Object = nullptr;
-					return;
-				}
-
-				// Disable checking this flag because it's too aggressive at removing references from reporting.
-				if (bGarbageReferenceTrackingEnabled /*&& !ObjectItem->HasAnyFlags(EInternalObjectFlags::PersistentGarbage)*/)
-				{
-					HandleGarbageReference(ObjectsToSerializeStruct, ReferencingObject, Object, TokenIndex);
-				}
-			}
-			else if (bGarbageReferenceTrackingEnabled)
+			// Fill up UnvalidatedReferences slack and drain it to avoid inserting an element at a time and testing if full
+			for (uint32 Slack = UnvalidatedReferences.Slack(); (uint32)Range.Num() >= Slack; Slack = UnvalidatedBatchSize)
 			{
-				// This object is being referenced by a persistent reference which means it wouldn't have been GC'd anyway
-				// so no need to track this object
-				ObjectItem->ThisThreadAtomicallySetFlag(EInternalObjectFlags::PersistentGarbage);
+				QueueUnvalidated(Range.Left(Slack));
+				Range.RightChopInline(Slack);
+				check(UnvalidatedReferences.IsFull());
+				DrainUnvalidatedFull();
 			}
-		}
 
-		// Remove references to pending kill objects if we're allowed to do so.
-		if (IsWithPendingKill() && ObjectItem->IsPendingKill() && bAllowReferenceElimination)
+			QueueUnvalidated(Range);
+			check(!UnvalidatedReferences.IsFull());
+		}
+	}
+
+	FORCEINLINE_DEBUGGABLE void FlushQueues()
+	{
+		// Slower path with dynamic Num that the compiler won't unroll or vectorize, unlike DrainXyzFull methods
+		DrainArrays(UnvalidatedArrays.Num);
+		DrainUnvalidated(UnvalidatedReferences.Num);
+		DrainValidated(ValidatedReferences.Num);
+	}
+
+private:
+	// Process queued arrays of references and feed contents to unvalidated queue
+	FORCEINLINE_DEBUGGABLE void DrainArrays(uint32 Num)
+	{
+		check(Num <= ArrayBatchSize);
+
+		for (TEntry<FReferenceArray> Entry : MakeArrayView(UnvalidatedArrays.Entries, Num))
+		{
+			FPlatformMisc::Prefetch(Entry.Reference.Array->GetData());
+		}
+
+		for (TEntry<FReferenceArray> Entry : MakeArrayView(UnvalidatedArrays.Entries, Num))
+		{
+			PushReferences(*Entry.Reference.Array);
+		}
+
+		UnvalidatedArrays.Num = 0;
+	}
+
+	// Helps generate vectorized code with static Num
+	FORCENOINLINE void DrainArraysFull()
+	{
+		DrainArrays(UnvalidatedArrays.Max());
+	}
+
+	FORCEINLINE_DEBUGGABLE void QueueUnvalidated(TArrayView<UObject*> Range)
+	{
+		check((uint32)Range.Num() <= UnvalidatedReferences.Slack());
+
+		TEntry<UnvalidatedReferenceType>* Unvalidated = UnvalidatedReferences.Entries + UnvalidatedReferences.Num;
+		for (int32 Idx = 0; Idx < Range.Num(); ++Idx)
+		{
+			Unvalidated[Idx] = { MakeReference<UnvalidatedReferenceType>(Range[Idx]) };
+		}
+
+		UnvalidatedReferences.Num += Range.Num();
+	}
+
+	// Validate all unvalidated references and feed referencing in need of processing to validated queue
+	FORCEINLINE_DEBUGGABLE void DrainUnvalidated(const uint32 Num)
+	{
+		check(Num <= UnvalidatedBatchSize);
+
+		FPermanentObjectPoolExtents Permanent(PermanentPool);
+		FValidatedBitmask ValidsA, ValidsB;
+		for (uint32 Idx = 0; Idx < Num; ++Idx)
+		{
+			UObject* Object = GetObject(UnvalidatedReferences[Idx].Reference);
+			uint64 ObjectAddress = reinterpret_cast<uint64>(Object);
+			ValidsA.Set(Idx, !Permanent.Contains(Object));
+		}
+
+		for (uint32 Idx = 0; Idx < Num; ++Idx)
+		{
+			UObject* Object = GetObject(UnvalidatedReferences[Idx].Reference);
+			ValidsB.Set(Idx, !!Object & IsObjectHandleResolved(reinterpret_cast<FObjectHandle&>(Object)));
+		}
+
+		FValidatedBitmask Validations = FValidatedBitmask::And(ValidsA, ValidsB);
+		uint32 NumValid = Validations.CountBits();
+		uint32 UnvalidatedIdx = 0;
+		for (uint32 Slack = ValidatedReferences.Slack(); NumValid >= Slack; Slack = ValidatedBatchSize)
+		{
+			QueueValidReferences(Slack, Validations, /* in-out */  UnvalidatedIdx);
+			check(ValidatedReferences.IsFull());
+			DrainValidatedFull();
+			NumValid -= Slack;
+		}
+
+		QueueValidReferences(NumValid, Validations, UnvalidatedIdx);
+		check(!ValidatedReferences.IsFull());
+
+		UnvalidatedReferences.Num = 0;
+	}
+
+	// Helps generate vectorized code with static Num
+	FORCENOINLINE void DrainUnvalidatedFull()
+	{
+		DrainUnvalidated(UnvalidatedReferences.Max());
+	}
+
+	FORCEINLINE_DEBUGGABLE void QueueValidReferences(uint32 NumToAppend, FValidatedBitmask Validations, uint32& InOutIdx)
+	{
+		checkSlow(NumToAppend == Validations.CountBits());
+		checkSlow(ValidatedReferences.Num + NumToAppend <= ValidatedReferences.Max());
+		uint32 NewQueueNum = ValidatedReferences.Num + NumToAppend;
+		for (uint32 QueueIdx = ValidatedReferences.Num; QueueIdx < NewQueueNum; ++InOutIdx)
+		{
+			bool bIsValid = !!Validations.Get(InOutIdx);
+			ValidatedReferences[QueueIdx].Reference = ToResolvedReference(UnvalidatedReferences[InOutIdx].Reference);
+			QueueIdx += bIsValid;
+		}
+
+		ValidatedReferences.Num += NumToAppend;
+	}
+
+	// Process all queued validated references
+	FORCEINLINE_DEBUGGABLE void DrainValidated(const uint32 Num)
+	{
+		check(Num <= ValidatedBatchSize);
+
+		// Prefetch metadata
+		// 
+		// UObjectBase::InternalIndex is 4B at offset 12. We allocate UObjects with 16B alignment so InternalIndex always
+		// shares cacheline with the raw UObject*. Not adding an offset leads to a bit faster and more compact code.
+		static constexpr uint32 InternalIndexPrefetchOffset = 0;
+		static constexpr uint32 PrefetchAhead = ValidatedPrefetchAhead;
+
+		uint32 ObjectIndices[ValidatedBatchSize];
+		if (Num > PrefetchAhead)
+		{
+			for (uint32 Idx = 0; Idx < PrefetchAhead; ++Idx)
+			{
+				FPlatformMisc::Prefetch(ValidatedReferences[Idx].Reference.Object, InternalIndexPrefetchOffset);
+			}
+			for (uint32 Idx = 0; Idx < Num; ++Idx)
+			{
+				ObjectIndices[Idx] = GUObjectArray.ObjectToIndex(ValidatedReferences[Idx].Reference.Object);
+				FPlatformMisc::Prefetch(ValidatedReferences[Idx + PrefetchAhead].Reference.Object, InternalIndexPrefetchOffset);
+			}
+		}
+		else
+		{
+			for (uint32 Idx = 0; Idx < Num; ++Idx)
+			{
+				FPlatformMisc::Prefetch(ValidatedReferences[Idx].Reference.Object, InternalIndexPrefetchOffset);
+			}
+			for (uint32 Idx = 0; Idx < Num; ++Idx)
+			{
+				ObjectIndices[Idx] = GUObjectArray.ObjectToIndex(ValidatedReferences[Idx].Reference.Object);
+			}
+		}
+
+		FChunkedFixedUObjectArray& ObjectArray = GUObjectArray.GetObjectItemArrayUnsafe();
+		FReferenceMetadata Metadatas[ValidatedBatchSize + PrefetchAhead];
+		for (uint32 Idx = 0; Idx < Num; ++Idx)
+		{
+			Metadatas[Idx].ObjectItem = ObjectArray.GetObjectPtr(ObjectIndices[Idx]);
+		}
+
+		if (Num > PrefetchAhead)
+		{
+			FMemory::Memzero(Metadatas + ValidatedBatchSize, PrefetchAhead * sizeof(FReferenceMetadata));
+
+			for (uint32 Idx = 0; Idx < PrefetchAhead; ++Idx)
+			{
+				FPlatformMisc::Prefetch(Metadatas[Idx].ObjectItem, offsetof(FUObjectItem, Flags));
+			}
+			for (uint32 Idx = 0; Idx < Num; ++Idx)
+			{
+				Metadatas[Idx].Flags = Metadatas[Idx].ObjectItem->GetFlags();
+				FPlatformMisc::Prefetch(Metadatas[Idx + PrefetchAhead].ObjectItem, offsetof(FUObjectItem, Flags));
+			}
+		}
+		else
+		{
+			for (uint32 Idx = 0; Idx < Num; ++Idx)
+			{
+				FPlatformMisc::Prefetch(Metadatas[Idx].ObjectItem, offsetof(FUObjectItem, Flags));
+			}
+			for (uint32 Idx = 0; Idx < Num; ++Idx)
+			{
+				Metadatas[Idx].Flags = Metadatas[Idx].ObjectItem->GetFlags();
+			}
+		}
+
+		for (uint32 Idx = 0; Idx < Num; ++Idx)
+		{
+			ProcessorType::HandleBatchedReference(Context, ValidatedReferences[Idx].Reference, Metadatas[Idx]);
+		}
+
+		ValidatedReferences.Num = 0;
+	}
+
+	// Helps generate code with static Num
+	FORCENOINLINE void DrainValidatedFull()
+	{
+		DrainValidated(ValidatedReferences.Max());
+	}
+};
+
+//////////////////////////////////////////////////////////////////////////
+
+FORCEINLINE_DEBUGGABLE static void PadObjects(UObject* LastObject, TArrayView<UObject*> Padding)
+{
+	check(LastObject->IsValidLowLevel());
+
+	for (UObject*& ObjectPastEnd : Padding)
+	{
+		ObjectPastEnd = LastObject;
+	}
+}
+
+void PadObjectArray(TArray<UObject*>& Objects)
+{
+	if (int32 Num = Objects.Num())
+	{
+		Objects.Reserve(Num + ObjectLookahead);
+		PadObjects(Objects.Last(), MakeArrayView(Objects.GetData() + Num, ObjectLookahead));
+	}
+}
+
+FORCEINLINE_DEBUGGABLE void PadBlock(FWorkBlock& Block)
+{
+	PadObjects(Block.GetObjects().Last(), Block.GetPadding());
+}
+
+//////////////////////////////////////////////////////////////////////////
+
+// Work-stealing algorithms are O(N^2), everyone steals from everyone.
+// Might want to improve that before going much wider.
+static constexpr int32 MaxWorkers = 16;
+
+class FWorkerIndexAllocator
+{
+	std::atomic<uint64> Used{0};
+	static_assert(MaxWorkers <= 64, "Currently supports single uint64 word");
+public:
+	int32 AllocateChecked()
+	{
+		while (true)
+		{
+			uint64 UsedNow = Used.load(std::memory_order_relaxed);
+			uint64 FreeIndex = FPlatformMath::CountTrailingZeros64(~UsedNow);
+			checkf(FreeIndex < MaxWorkers, TEXT("Exceeded max active GC worker contexts"));
+
+			uint64 Mask = uint64(1) << FreeIndex;
+			if ((Used.fetch_or(Mask) & Mask) == 0)
+			{
+				return static_cast<int32>(FreeIndex);
+			}
+		}
+	}
+
+	void FreeChecked(int32 Index)
+	{
+		check(Index >= 0 && Index < MaxWorkers);
+		uint64 Mask = uint64(1) << Index;
+		uint64 Old = Used.fetch_and(~Mask);
+		checkf(Old & Mask, TEXT("Index already freed"));
+	}
+};
+static FWorkerIndexAllocator GWorkerIndices;
+
+FWorkerContext::FWorkerContext()
+{
+	AllocateWorkerIndex();
+	ObjectsToSerialize.Init();
+}
+	
+FWorkerContext::~FWorkerContext()
+{
+	FreeWorkerIndex();
+}
+
+void FWorkerContext::AllocateWorkerIndex()
+{
+	check(GetWorkerIndex() == INDEX_NONE);
+	ObjectsToSerialize.SetWorkerIndex(GWorkerIndices.AllocateChecked());
+}
+
+void FWorkerContext::FreeWorkerIndex()
+{
+	if (GetWorkerIndex() != INDEX_NONE)
+	{
+		GWorkerIndices.FreeChecked(GetWorkerIndex());
+		ObjectsToSerialize.SetWorkerIndex(INDEX_NONE);
+	}
+}
+
+//////////////////////////////////////////////////////////////////////////
+
+int32 GetNumSlowAROs();
+
+// Allocates and caches 4K pages used for temporary allocations during GC
+//
+// Currently doesn't free pages until shutdown
+class FPageAllocator
+{
+	struct alignas(PLATFORM_CACHE_LINE_SIZE) FWorkerCache
+	{
+		UE_NONCOPYABLE(FWorkerCache);
+		FWorkerCache() = default;
+		~FWorkerCache()
+		{
+			for (void* Page : MakeArrayView(Pages, Num))
+			{
+				FreePage(Page);
+			}
+		}
+
+		bool Push(void* Page)
+		{
+			check(Page);
+			if (Num < Capacity)
+			{
+				Pages[Num++] = Page;
+				return true;	
+			}
+			return false;
+		}
+
+		void* Pop()
+		{
+			return Num > 0 ? Pages[--Num] : nullptr;
+		}
+
+		static constexpr uint32 Capacity = 512 / MaxWorkers;
+		uint32 Num = 0;
+		void* Pages[Capacity] = {};
+	};	
+
+	struct FSharedCache
+	{
+		UE_NONCOPYABLE(FSharedCache);
+		FSharedCache() = default;
+		~FSharedCache()
+		{
+			for (void* Page : Pages)
+			{
+				FreePage(Page);
+			}
+		}
+
+		void Push(void* Page)
+		{
+			FScopeLock Scope(&Lock);
+			Pages.Add(Page);
+		}
+
+		void PushSurplus(FWorkerCache& Worker)
+		{
+			// Each worker context also owns a page via FWorkBlockifier::Wip
+			const uint32 NumWorkerInitialPages = 1 + GetNumSlowAROs();
+
+			if (Worker.Num > NumWorkerInitialPages)
+			{
+				FScopeLock Scope(&Lock);
+				Pages.Append(Worker.Pages + NumWorkerInitialPages, Worker.Num - NumWorkerInitialPages);
+				Worker.Num = NumWorkerInitialPages;
+			}
+		}
+
+		void* Pop()
+		{
+			FScopeLock Scope(&Lock);
+			return Pages.IsEmpty() ? nullptr : Pages.Pop(/* shrink */ false);
+		}
+
+		FCriticalSection Lock;
+		TArray<void*> Pages;
+	};	
+
+	FSharedCache SharedCache;
+	FWorkerCache WorkerCaches[MaxWorkers];
+
+public:
+	static constexpr uint64 PageSize = 4096;
+
+	void* AllocatePage(int32 WorkerIdx)
+	{
+		check(WorkerIdx >=0 && WorkerIdx < MaxWorkers);
+		if (void* WorkerPage = WorkerCaches[WorkerIdx].Pop())
+		{
+			check(IsValidPage(WorkerPage));
+			return WorkerPage;
+		}
+		else if (void* SharedPage = SharedCache.Pop())
+		{
+			check(IsValidPage(SharedPage));
+			return SharedPage;
+		}
+
+		void* NewPage = FMemory::Malloc(PageSize, PageSize);
+		return NewPage;
+	}
+	
+	void ReturnWorkerPage(int32 WorkerIdx, void* Page)
+	{
+		check(IsValidPage(Page));
+		if (!WorkerCaches[WorkerIdx].Push(Page))
+		{
+			FreePage(Page);
+		}
+	}
+
+	void ReturnSharedPage(void* Page)
+	{
+		check(IsValidPage(Page));
+		SharedCache.Push(Page);
+	}
+
+	static void FreePage(void* Page)
+	{
+		check(IsValidPage(Page));
+		FMemory::Free(Page);
+	}
+
+	uint64 CountBytes() const
+	{
+		uint64 NumPages = SharedCache.Pages.Num();
+		for (const FWorkerCache& WorkerCache : WorkerCaches)
+		{
+			NumPages += WorkerCache.Num;
+		}
+		return NumPages * PageSize;
+	}
+
+	static bool IsValidPage(void* Page)
+	{
+		return Page && UPTRINT(Page) % PageSize == 0;
+	}
+};
+
+static FPageAllocator GScratchPages;
+
+// Bounded SPMC work-stealing queue
+//
+// Derived from task graph TWorkStealingQueueBase2 and added CheckEmpty(),
+// reduced alignment and simplified some code
+class FBoundedWorkstealingQueue
+{
+	static constexpr uint32 Capacity = 16;
+	static constexpr uint64 FreeSlot = 0;
+	static constexpr uint64 TakenSlot = 1;
+
+public:
+	// Called by the single producer
+	bool Push(FWorkBlock* Block)
+	{
+		uint64 Item = reinterpret_cast<uint64>(Block);
+		checkSlow(Item != FreeSlot);
+		checkSlow(Item != TakenSlot);
+
+		uint32 Idx = (Head + 1) % Capacity;
+		uint64 Slot = Slots[Idx].load(std::memory_order_acquire);
+
+		if (Slot == FreeSlot)
+		{
+			Slots[Idx].store(Item, std::memory_order_release);
+			Head++;
+			checkSlow(Head % Capacity == Idx);
+			return true;		
+		}
+
+		return false;
+	}
+
+	// Called by single producer, pops in LIFO order
+	FWorkBlock* Pop()
+	{
+		uint32 Idx = Head % Capacity;
+		uint64 Slot = Slots[Idx].load(std::memory_order_acquire);
+
+		if (Slot > TakenSlot && Slots[Idx].compare_exchange_strong(Slot, FreeSlot, std::memory_order_acq_rel))
+		{
+			Head--;
+			checkSlow((Head + 1) % Capacity == Idx);
+			return reinterpret_cast<FWorkBlock*>(Slot);
+		}
+
+		return nullptr;
+	}
+
+	// Called by other consumers, pops in FIFO order
+	FWorkBlock* Steal()
+	{
+		while (true)
+		{
+			uint32 IdxVer = Tail.load(std::memory_order_acquire);
+			uint32 Idx = IdxVer % Capacity;
+			uint64 Slot = Slots[Idx].load(std::memory_order_acquire);
+
+			if (Slot == FreeSlot)
+			{
+				return nullptr;
+			}
+			else if (Slot != TakenSlot && Slots[Idx].compare_exchange_weak(Slot, TakenSlot, std::memory_order_acq_rel))
+			{
+				if (IdxVer == Tail.load(std::memory_order_acquire))
+				{
+					uint32 Prev = Tail.fetch_add(1, std::memory_order_release);
+					checkSlow(Prev % Capacity == Idx);
+					Slots[Idx].store(FreeSlot, std::memory_order_release);
+					return reinterpret_cast<FWorkBlock*>(Slot);
+				}
+				Slots[Idx].store(Slot, std::memory_order_release);
+			}
+		}
+	}
+
+	// Called after all producers and consumers have stopped working
+	void CheckEmpty()
+	{
+		for (const AlignedAtomicUint64& Slot : Slots)
+		{
+			check(Slot.load() == FreeSlot);
+		}
+	}
+
+private:
+	struct alignas(PLATFORM_CACHE_LINE_SIZE) AlignedAtomicUint64 : public std::atomic<uint64> {};
+
+	uint32 Head { ~0u };
+	alignas(PLATFORM_CACHE_LINE_SIZE) std::atomic_uint32_t Tail { 0 };
+	AlignedAtomicUint64 Slots[Capacity] = {};
+};
+
+// Unbounded SPMC queue where only a bounded number of items can be stolen
+class FWorkstealingQueue
+{
+public:
+	void Push(FWorkBlock* Block)
+	{
+		if (!SharedBlocks.Push(Block))
+		{
+			LocalBlocks.Push(Block);
+		}
+	}
+
+	FWorkBlock* Pop()
+	{
+		int32 LocalNum = LocalBlocks.Num();
+		if (LocalNum == 0)
+		{
+			return SharedBlocks.Pop();	
+		}
+
+		// Pop last block and try moving remaining blocks to shared queue
+		FWorkBlock* Out = LocalBlocks[--LocalNum];
+		while (LocalNum > 0 && SharedBlocks.Push(LocalBlocks[LocalNum - 1]))
+		{
+			--LocalNum;
+		}
+		LocalBlocks.SetNum(LocalNum, /* shrink */ false);
+
+		return Out;
+	}
+
+	FWorkBlock* Steal()
+	{
+		return SharedBlocks.Steal();
+	}
+
+	void CheckEmpty()
+	{
+		SharedBlocks.CheckEmpty();
+		check(LocalBlocks.IsEmpty());
+	}
+
+private:
+	FBoundedWorkstealingQueue SharedBlocks;
+	TArray<FWorkBlock*> LocalBlocks;
+};
+
+struct FWorkstealingManager
+{
+	FWorkstealingQueue Queues[MaxWorkers];
+
+	FWorkBlock* Steal(FWorkstealingQueue* Workless)
+	{
+		PTRINT WorkerPtrDiff = Workless - Queues;
+		checkf(WorkerPtrDiff >= 0 && WorkerPtrDiff < MaxWorkers, TEXT("FWorkstealingQueue not owned by manager"));
+		int32 WorklessIdx = static_cast<int32>(WorkerPtrDiff);
+
+		return TrySteal(WorklessIdx);
+	}
+
+	FWorkBlock* TrySteal(int32 WorklessIdx)
+	{
+		for (int32 Idx = WorklessIdx + 1; Idx < MaxWorkers; ++Idx)
+		{
+			if (FWorkBlock* Stolen = Queues[Idx].Steal())
+			{
+				return Stolen;
+			}
+		}
+
+		for (int32 Idx = 0; Idx < WorklessIdx; ++Idx)
+		{
+			if (FWorkBlock* Stolen = Queues[Idx].Steal())
+			{
+				return Stolen;
+			}
+		}
+
+		return nullptr;
+	}
+};
+
+static FWorkstealingManager GWorkstealingManager;
+
+//////////////////////////////////////////////////////////////////////////
+
+FORCEINLINE_DEBUGGABLE void FWorkBlockifier::AllocateWipBlock()
+{
+	static_assert(sizeof(FWorkBlock) <= FPageAllocator::PageSize, "Block must fit in page");
+	static_assert(sizeof(FWorkBlock) >= FPageAllocator::PageSize, "Implement thread local blocks-in-page allocator");
+	
+	Wip = new (GScratchPages.AllocatePage(WorkerIndex)) FWorkBlock;
+	Wip->Previous = nullptr;
+	WipIt = Wip->Objects;
+}
+
+void FWorkBlockifier::FreeOwningBlock(UObject*const* BlockObjects)
+{
+	if (BlockObjects)
+	{
+		FWorkBlock* Block = reinterpret_cast<FWorkBlock*>(UPTRINT(BlockObjects) - offsetof(FWorkBlock, Objects));
+		check(UPTRINT(Block) % FPageAllocator::PageSize == 0);
+
+		Block->~FWorkBlock();
+		GScratchPages.ReturnWorkerPage(WorkerIndex, Block);
+	}
+}
+
+FWorkBlockifier::~FWorkBlockifier()
+{
+	check(IsUnused());
+	Wip->~FWorkBlock();
+	GScratchPages.FreePage(Wip);
+}
+
+void FWorkBlockifier::ResetAsyncQueue()
+{
+	checkf(PartialNum() == 0, TEXT("Queue not empty"));
+	SyncQueue = nullptr;
+}
+
+void FWorkBlockifier::PushFullBlockSync()
+{
+	PadBlock(*Wip);
+	Wip->Previous = SyncQueue;
+	SyncQueue = Wip;
+	AllocateWipBlock();
+}
+void FWorkBlockifier::PushFullBlockAsync()
+{
+	PadBlock(*Wip);
+	AsyncQueue->Push(Wip);
+	AllocateWipBlock();
+}
+
+FWorkBlock* FWorkBlockifier::PopFullBlockSync()
+{
+	FWorkBlock* Out = SyncQueue;
+	SyncQueue = Out ? Out->Previous : nullptr;
+	return Out;
+}
+
+FWorkBlock* FWorkBlockifier::PopFullBlockAsync()
+{
+	return AsyncQueue->Pop();
+}
+	
+FWorkBlock* FWorkBlockifier::PopWipBlock()
+{
+	PadObjects(WipIt[-1], MakeArrayView(WipIt, ObjectLookahead));
+	FWorkBlock* Out = Wip;
+	AllocateWipBlock();
+	return Out;
+}
+
+FWorkBlock* FWorkBlockifier::StealAsyncBlock() const
+{
+	return GWorkstealingManager.Steal(AsyncQueue);
+}
+
+//////////////////////////////////////////////////////////////////////////
+
+struct FAROBlock
+{
+	static constexpr uint32 NumWords = FPageAllocator::PageSize / sizeof(UObject*);
+	static constexpr uint32 Capacity = NumWords - 1;
+
+	uint32 FirstIndexInNextBlock;
+	UObject* Objects[Capacity];
+};
+
+// Simple bounded lock-free SPMC LIFO queue used to schedule AddReferencedObjects calls
+//
+// * Producer steals in same order as consumer, simpler but less efficient
+// * Popping is batched, in part to reduce producer-consumer tail popping contention
+// * Producer consumes all entries before shutting down, so stealing can fail
+// * ARO queues and blocks outlive producer and consumers to avoid lifetime problems
+// * Bounded by fix number of pages shared between all queues
+// * Last two points allow using ever increasing indices to side-step the ABA problem
+class FAROQueue
+{
+public:
+	struct FFew{};
+	struct FMany{};
+	template<class Quantity> static constexpr int NumPop = std::is_same_v<Quantity, FMany> ? 64 : 4;
+
+	UE_NONCOPYABLE(FAROQueue);
+	explicit FAROQueue(int32 WorkerIdx);
+	~FAROQueue();
+
+	[[nodiscard]] bool TryPush(UObject* Object);								// Called by the single producer thread
+	template<class Quantity> [[nodiscard]] TArrayView<UObject*> Pop();		// Called by the single producer thread
+	template<class Quantity> [[nodiscard]] TArrayView<UObject*> Steal();	// Called by other work-stealing threads 
+	
+private:
+	alignas(PLATFORM_CACHE_LINE_SIZE) std::atomic<uint32> Head;
+	FAROBlock* HeadBlock;
+	const int32 WorkerIndex;
+	alignas(PLATFORM_CACHE_LINE_SIZE) std::atomic<uint32> Tail;
+
+	template<void(*Fence)(), uint32 NumWanted> TArrayView<UObject*> PopImpl();
+};
+
+class FAROQueueStore
+{
+public:
+	[[nodiscard]] FAROBlock* AllocateBlock(int32 WorkerIdx, uint32& OutIdx)
+	{
+		static_assert(sizeof(FAROBlock) == FPageAllocator::PageSize);
+
+		if (bFull.load(std::memory_order_relaxed))
+		{
+			return nullptr;
+		}
+
+		uint32 BlockIdx = NumBlocks.fetch_add(1, std::memory_order_acq_rel);
+		if (BlockIdx >= MaxBlocks)
+		{
+			bFull.store(true, std::memory_order_relaxed);
+			NumBlocks.fetch_sub(1, std::memory_order_acq_rel);
+			return nullptr;
+		}
+
+		check(Blocks[BlockIdx] == nullptr);
+
+		OutIdx = BlockIdx * FAROBlock::NumWords;
+		FAROBlock* Block = new (GScratchPages.AllocatePage(WorkerIdx)) FAROBlock;
+		Blocks[BlockIdx] = Block;
+		return Block;
+	}
+
+	UObject*& operator[](uint32 Idx)
+	{
+		check(Idx / FAROBlock::NumWords < MaxBlocks);
+		check(Idx % FAROBlock::NumWords < FAROBlock::Capacity);
+		check(Blocks[Idx / FAROBlock::NumWords]);
+
+		return Blocks[Idx / FAROBlock::NumWords]->Objects[Idx % FAROBlock::NumWords];
+	}
+
+	uint32 GetFirstIndexInNextBlock(uint32 Idx)
+	{
+		check(Idx / FAROBlock::NumWords < MaxBlocks);
+		check(Idx % FAROBlock::NumWords < FAROBlock::Capacity);
+		check(Blocks[Idx / FAROBlock::NumWords]);
+		check(Blocks[Idx / FAROBlock::NumWords]->FirstIndexInNextBlock % FAROBlock::NumWords == 0);
+		check(Blocks[Idx / FAROBlock::NumWords]->FirstIndexInNextBlock / FAROBlock::NumWords < MaxBlocks);
+
+		return Blocks[Idx / FAROBlock::NumWords]->FirstIndexInNextBlock;
+	}
+
+	void ReturnAllBlocks()
+	{
+		check(NumBlocks.load() <= MaxBlocks);
+
+		for (FAROBlock*& Block : MakeArrayView(Blocks, NumBlocks.load()))
+		{
+			Block->~FAROBlock();
+			GScratchPages.ReturnSharedPage(Block);
+			Block = nullptr;
+		}
+
+		NumBlocks.store(0);
+		bFull.store(false);
+	}
+
+	static constexpr uint32 MaxBlocks = 256;
+private:
+	std::atomic<uint32> NumBlocks{0};
+	std::atomic_bool bFull{false};
+	alignas(PLATFORM_CACHE_LINE_SIZE) FAROBlock* Blocks[MaxBlocks] = {};
+};
+
+static FAROQueueStore GAROBlocks;
+
+FAROQueue::FAROQueue(int32 WorkerIdx)
+: WorkerIndex(WorkerIdx)
+{
+	uint32 HeadIdx;
+	HeadBlock = GAROBlocks.AllocateBlock(WorkerIdx, /* out */ HeadIdx);
+	checkf(HeadBlock, TEXT("One head block per worker must exist during worker setup, "
+			"the assumption MaxWorkers * slow AROs <= MaxBlocks %s with %d slow AROs registered."),
+			(MaxWorkers * GetNumSlowAROs() <= FAROQueueStore::MaxBlocks) ? TEXT("held") : TEXT("failed"), GetNumSlowAROs());
+	HeadBlock->FirstIndexInNextBlock = ~0u;
+	Head.store(HeadIdx, std::memory_order_relaxed);
+	Tail.store(HeadIdx, std::memory_order_relaxed);
+}
+
+FAROQueue::~FAROQueue()
+{
+	checkf(Head.load() == Tail.load(), TEXT("Failed to flush ARO calls"));
+	// HeadBlock is returned by GAROBlocks.FreeAllBlocks()
+}
+
+FORCEINLINE_DEBUGGABLE bool FAROQueue::TryPush(UObject* Object)
+{
+	uint32 HeadIdx = Head.load(std::memory_order_relaxed);
+	HeadBlock->Objects[HeadIdx % FAROBlock::NumWords] = Object;
+	++HeadIdx;
+
+	// Must store OldBlock->FirstIndexInNextBlock before Head.store
+	if (HeadIdx % FAROBlock::NumWords == FAROBlock::Capacity)
+	{
+		FAROBlock* OldBlock = HeadBlock;
+		FAROBlock* NewBlock = GAROBlocks.AllocateBlock(WorkerIndex, /* out */ HeadIdx);
+		if (!NewBlock)
+		{
+			// Head won't move so HeadBlock->Objects[HeadIdx % FAROBlock::NumWords] won't be read
+			return false;
+		}
+		
+		HeadBlock = NewBlock;
+		OldBlock->FirstIndexInNextBlock = HeadIdx;
+	}
+
+	// Synchronizes non-atomic Blocks, Object and FirstIndexInNextBlock stores with work-stealing AcquireFence
+	Head.store(HeadIdx, std::memory_order_release);
+
+	return true;
+}
+
+template<void(*Fence)(), uint32 NumWanted>
+FORCEINLINE_DEBUGGABLE
+TArrayView<UObject*> FAROQueue::PopImpl()
+{
+	check(NumWanted < FAROBlock::Capacity);
+
+	while (true)
+	{
+		const uint32 HeadNow = Head.load(std::memory_order_relaxed);
+		const uint32 TailNow = Tail.load(std::memory_order_relaxed);
+
+		if (TailNow >= HeadNow)
+		{
+			return {};
+		}
+
+		// Can only pop up to end of tail block
+		static constexpr uint32 PageIndexMask = ~(FAROBlock::NumWords - 1);
+		const uint32 LastInTailBlock = (TailNow & PageIndexMask) + FAROBlock::Capacity - 1;
+		const uint32 WantedTail = FMath::Min(HeadNow, TailNow + NumWanted);
+
+		// Atomic-fence synchronization of non-atomic stores via relaxed Head load
+		Fence();
+
+		UObject** TailData = &GAROBlocks[TailNow];
+		uint32 ExpectedTail = TailNow;
+		if (WantedTail <= LastInTailBlock)
+		{
+			if (Tail.compare_exchange_weak(ExpectedTail, WantedTail))
+			{
+				return MakeArrayView(TailData, WantedTail - TailNow);
+			}
+		}
+		else
+		{
+			if (Tail.compare_exchange_weak(ExpectedTail, GAROBlocks.GetFirstIndexInNextBlock(TailNow)))
+			{
+				return MakeArrayView(TailData, FAROBlock::Capacity - (TailNow % FAROBlock::NumWords));
+			}			
+		}
+	}
+}
+
+static FORCEINLINE void AcquireFence() { std::atomic_thread_fence(std::memory_order_acquire); }
+static FORCEINLINE void NoFence() {}
+
+template<class Quantity> FORCEINLINE TArrayView<UObject*> FAROQueue::Pop()		{ return PopImpl<NoFence,		NumPop<Quantity>>(); }
+template<class Quantity> FORCEINLINE TArrayView<UObject*> FAROQueue::Steal()	{ return PopImpl<AcquireFence,	NumPop<Quantity>>(); }
+
+//////////////////////////////////////////////////////////////////////////
+
+class FSlowAROManager
+{
+public:
+	void RegisterImplementation(AROFunc ARO, EAROFlags Flags)
+	{
+		check(IsInGameThread());
+		check(AROs.Find(ARO) == INDEX_NONE);
+		checkf(AROs.Num() < Capacity, TEXT("Don't register this many slow AROs, memory consumption and work-stealing time increases linearly"));
+
+		int32 Idx = AROs.Num();
+		AROs.Add(ARO);
+		if (EnumHasAllFlags(Flags, EAROFlags::Unbalanced))
+		{
+			UnbalancedAROIndices.Add(Idx);
+		}
+
+		static_assert(Capacity <= sizeof(ExtraSlowAROs) * 8);
+		ExtraSlowAROs |= decltype(ExtraSlowAROs)(EnumHasAllFlags(Flags, EAROFlags::ExtraSlow)) << Idx;
+	}
+
+	int32 FindImplementation(AROFunc ARO) const
+	{
+		return AROs.Find(ARO);
+	}
+	
+	void SetupWorkerQueues(int32 InNumWorkers)
+	{
+		check(AllQueues.IsEmpty());
+		NumWorkers = InNumWorkers;
+		AllQueues.SetNumUninitialized(NumWorkers * AROs.Num());
+		for (int32 AROIdx = 0, NumAROs = AROs.Num(); AROIdx < NumAROs; ++AROIdx)
+		{
+			for (int32 WorkerIdx = 0; WorkerIdx < NumWorkers; ++WorkerIdx)
+			{
+				new (&GetQueue(AROIdx, WorkerIdx)) FAROQueue(WorkerIdx);
+			}
+		}
+	}
+
+	FORCEINLINE_DEBUGGABLE void CallSync(int32 AROIdx, UObject* Object, FReferenceCollector& Collector)
+	{
+		AROs[AROIdx](Object, Collector);
+	}
+
+	[[nodiscard]] FORCEINLINE_DEBUGGABLE bool QueueCall(int32 AROIdx, int32 WorkerIdx, UObject* Object)
+	{
+		return GetQueue(AROIdx, WorkerIdx).TryPush(Object);
+	}
+
+	FORCEINLINE_DEBUGGABLE void DrainLocalUnbalancedQueues(int32 WorkerIdx, FReferenceCollector& Collector)
+	{
+		for (uint32 AROIdx : UnbalancedAROIndices)
+		{
+			AROFunc ARO = AROs[AROIdx];
+			FAROQueue& Queue = GetQueue(AROIdx, WorkerIdx);
+			for (TArrayView<UObject*> Objects = Queue.Pop<FAROQueue::FMany>(); Objects.Num(); Objects = Queue.Pop<FAROQueue::FMany>())
+			{
+				for (UObject* Object : Objects)
+				{
+					ARO(Object, Collector);
+				}
+			}
+		}
+	}
+
+	bool ProcessAllQueues(int32 WorkerIdx, FReferenceCollector& Collector)
+	{
+		// Drain all queues for one specific slow ARO type at a time and stop if any ARO calls were made
+		// Workers start with different ARO types to reach the slowest ARO type sooner and reduce stealing contention
+		
+		bool bStop = false;
+		uint32 NumCalls = 0;
+		const uint32 NumAROs = (uint32)AROs.Num();
+		
+		const double StopTime = FPlatformTime::Seconds() + 0.1/1000;
+		for (uint32 OffsetIdx = (uint32)WorkerIdx % NumAROs, EndOffsetIdx = OffsetIdx + NumAROs; OffsetIdx < EndOffsetIdx && !bStop; ++OffsetIdx)
+		{
+			int32 AROIdx = OffsetIdx - (OffsetIdx >= NumAROs ? NumAROs : 0);
+			TArrayView<FAROQueue> WorkerQueues = MakeArrayView(&GetQueue(AROIdx, 0), NumWorkers);
+			bStop = IsExtraSlow(AROIdx)	? ProcessQueues<FAROQueue::FFew >(WorkerIdx, WorkerQueues, AROs[AROIdx], Collector, NumCalls)
+										: ProcessQueues<FAROQueue::FMany>(WorkerIdx, WorkerQueues, AROs[AROIdx], Collector, NumCalls);
+
+			// Don't proceed to next ARO type if we did some work that took a considerable amount of time
+			bStop = bStop || (NumCalls > 0 && FPlatformTime::Seconds() > StopTime);
+		}
+
+		return NumCalls > 0;
+	}
+
+	void ResetWorkerQueues()
+	{
+		AllQueues.Reset();
+		GAROBlocks.ReturnAllBlocks();
+		NumWorkers = 0;
+	}
+
+	int32 NumAROs() const { return AROs.Num(); }
+
+private:
+	static constexpr int32 Capacity = 8;
+	uint32 ExtraSlowAROs = 0;
+	int32 NumWorkers = 0;
+	TArray<FAROQueue> AllQueues;
+	TArray<int32, TFixedAllocator<2>> UnbalancedAROIndices;
+	TArray<AROFunc, TFixedAllocator<Capacity>> AROs;
+
+	FAROQueue& GetQueue(int32 AROIdx, int32 WorkerIdx)
+	{
+		check(WorkerIdx < NumWorkers);	
+		check(AROIdx < AROs.Num());
+		return AllQueues[AROIdx * NumWorkers + WorkerIdx];
+	}
+
+	// @return if processing should stop since NumCalls reached CallLimit
+	template<class Quantity>
+	static bool ProcessQueues(int32 WorkerIdx, TArrayView<FAROQueue> WorkerQueues, AROFunc ARO, FReferenceCollector& Collector, uint32& NumCalls)
+	{
+		static constexpr uint32 CallLimit = FAROQueue::NumPop<Quantity>;
+
+		return	ProcessQueuesUsing<&FAROQueue::Pop<Quantity>,	CallLimit>(WorkerQueues.Slice(WorkerIdx, 1),		ARO, Collector, NumCalls) ||
+				ProcessQueuesUsing<&FAROQueue::Steal<Quantity>,	CallLimit>(WorkerQueues.RightChop(WorkerIdx + 1),	ARO, Collector, NumCalls) ||
+				ProcessQueuesUsing<&FAROQueue::Steal<Quantity>,	CallLimit>(WorkerQueues.Slice(0, WorkerIdx),		ARO, Collector, NumCalls);
+	}
+
+	// @return if processing should stop since NumCalls reached or exceeded CallLimit
+	template<TArrayView<UObject*>(FAROQueue::*PopFunc)(), uint32 CallLimit>
+	FORCEINLINE_DEBUGGABLE static bool ProcessQueuesUsing(TArrayView<FAROQueue> Queues, AROFunc ARO, FReferenceCollector& Collector, uint32& NumCalls)
+	{
+		for (FAROQueue& Queue : Queues)
+		{
+			for (TArrayView<UObject*> Objects = (Queue.*PopFunc)(); Objects.Num(); Objects = (Queue.*PopFunc)())
+			{
+				for (UObject* Object : Objects)
+				{
+					ARO(Object, Collector);
+				}
+
+				NumCalls += Objects.Num();
+				if (NumCalls >= CallLimit)
+				{
+					return true;
+				}
+			}
+		}
+
+		return false;
+	}
+
+	FORCEINLINE bool IsExtraSlow(int32 AROIdx) const { return !!((ExtraSlowAROs >> AROIdx) & 1u); }
+};
+
+
+template<class T>
+class TSingleton
+{
+	// Leave all data uninitialized
+	bool bValid;
+	alignas(T) uint8 Data[sizeof(T)];
+
+	T& Cast() { return *reinterpret_cast<T*>(Data); }
+	
+public:
+	T& Get()
+	{
+		struct FLifetime
+		{
+			TSingleton& Singleton;
+
+			FLifetime(TSingleton& S)
+			: Singleton(S)
+			{
+				check(!Singleton.bValid);
+				new (Singleton.Data) T;
+				Singleton.bValid = true;
+			}
+
+			~FLifetime()
+			{
+				Singleton.bValid = false;
+				Singleton.Cast().~T();
+			}
+		};
+		static FLifetime FunctionStaticLifetime(*this);
+
+		return Cast();
+	}
+
+	T& GetPostInit()
+	{
+		checkf(bValid, TEXT("Must call Get() at least once first"));
+		return Cast();
+	}
+};
+
+static TSingleton<FSlowAROManager> GSlowARO;
+
+
+void FSlowARO::CallSync(uint32 SlowAROIndex, UObject* Object, FReferenceCollector& Collector)
+{
+	GSlowARO.GetPostInit().CallSync(SlowAROIndex, Object, Collector);
+}
+
+bool FSlowARO::TryQueueCall(uint32 SlowAROIndex, UObject* Object, FWorkerContext& Context)
+{
+	return GSlowARO.GetPostInit().QueueCall(SlowAROIndex, Context.GetWorkerIndex(), Object);
+}
+
+void FSlowARO::ProcessUnbalancedCalls(FWorkerContext& Context, FReferenceCollector& Collector)
+{
+	GSlowARO.GetPostInit().DrainLocalUnbalancedQueues(Context.GetWorkerIndex(), Collector);
+}
+bool FSlowARO::ProcessAllCalls(FWorkerContext& Context, FReferenceCollector& Collector)
+{
+	return GSlowARO.GetPostInit().ProcessAllQueues(Context.GetWorkerIndex(), Collector);
+}
+
+void RegisterSlowImplementation(AROFunc ARO, EAROFlags Flags)
+{
+	GSlowARO.Get().RegisterImplementation(ARO, Flags);
+}
+
+int32 FindSlowImplementation(AROFunc ARO)
+{
+	return GSlowARO.GetPostInit().FindImplementation(ARO);
+}
+
+int32 GetNumSlowAROs()
+{
+	return GSlowARO.GetPostInit().NumAROs();
+}
+
+//////////////////////////////////////////////////////////////////////////
+
+// Helps identify code locations that kill references
+FORCEINLINE_DEBUGGABLE void KillReference(UObject*& Object) { Object = nullptr; }
+
+template <EGCOptions Options>
+constexpr FORCEINLINE EKillable MayKill(EGCTokenType TokenType, bool bAllowKill)
+{
+	// To avoid content changes, allow reference elimination inside of Blueprints
+	return (bAllowKill & (IsPendingKill(Options) || TokenType == EGCTokenType::NonNative)) ? EKillable::Yes : EKillable::No;
+}
+
+// Return whether flag was cleared. Only thread-safe for concurrent clear, not concurrent set+clear. Don't use during mark phase.
+FORCEINLINE static bool ClearUnreachableInterlocked(int32& Flags)
+{
+	static constexpr int32 FlagToClear = int32(EInternalObjectFlags::Unreachable);
+	if (Flags & FlagToClear)
+	{
+		int Old = FPlatformAtomics::InterlockedAnd(&Flags, ~FlagToClear);
+		return Old & FlagToClear;
+	}
+
+	return false;
+}
+
+template <EGCOptions InOptions>
+class TReachabilityProcessor
+{
+public:
+	static constexpr EGCOptions Options = InOptions;
+
+	static constexpr FORCEINLINE bool IsWithPendingKill() {	return !!(Options & EGCOptions::WithPendingKill);}
+
+	static void UpdateDetailedStats(UObject* CurrentObject, uint32 DeltaCycles)
+	{
+#if PERF_DETAILED_PER_CLASS_GC_STATS
+		// Keep track of how many refs we encountered for the object's class.
+		const FName& ClassName = CurrentObject->GetClass()->GetFName();
+		// Refs to objects that reside in permanent object pool.
+		uint32 ClassDisregardedObjRefs = GClassToDisregardedObjectRefsMap.FindRef(ClassName);
+		GClassToDisregardedObjectRefsMap.Add(ClassName, ClassDisregardedObjRefs + GCurrentObjectDisregardedObjectRefs);
+		// Refs to regular objects.
+		uint32 ClassRegularObjRefs = GClassToRegularObjectRefsMap.FindRef(ClassName);
+		GClassToRegularObjectRefsMap.Add(ClassName, ClassRegularObjRefs + GCurrentObjectRegularObjectRefs);
+		// Track per class cycle count spent in GC.
+		uint32 ClassCycles = GClassToCyclesMap.FindRef(ClassName);
+		GClassToCyclesMap.Add(ClassName, ClassCycles + DeltaCycles);
+		// Reset current counts.
+		GCurrentObjectDisregardedObjectRefs = 0;
+		GCurrentObjectRegularObjectRefs = 0;
+#endif
+	}
+
+	static void LogDetailedStatsSummary()
+	{
+#if PERF_DETAILED_PER_CLASS_GC_STATS
+		LogClassCountInfo(TEXT("references to regular objects from"), GClassToRegularObjectRefsMap, 20, 0);
+		LogClassCountInfo(TEXT("references to permanent objects from"), GClassToDisregardedObjectRefsMap, 20, 0);
+		LogClassCountInfo(TEXT("cycles for GC"), GClassToCyclesMap, 20, 0);
+#endif
+	}
+
+	PRAGMA_DISABLE_DEPRECATION_WARNINGS
+	static constexpr EInternalObjectFlags KillFlag = IsWithPendingKill() ? EInternalObjectFlags::PendingKill : EInternalObjectFlags::Garbage;
+	PRAGMA_ENABLE_DEPRECATION_WARNINGS
+
+	constexpr static FORCEINLINE EKillable MayKill(EGCTokenType TokenType, bool bAllowKill) { return UE::GC::MayKill<Options>(TokenType, bAllowKill); }	
+
+	static FORCEINLINE void ProcessReferenceDirectly(FWorkerContext& Context, FPermanentObjectPoolExtents PermanentPool, const UObject* ReferencingObject, UObject*& Object, FTokenId TokenId, EKillable Killable)
+	{
+		return Killable == EKillable::No	? ProcessReferenceDirectly<EKillable::No >(Context, PermanentPool, ReferencingObject, Object, TokenId)
+											: ProcessReferenceDirectly<EKillable::Yes>(Context, PermanentPool, ReferencingObject, Object, TokenId);
+	}
+
+	static FORCEINLINE void DetectGarbageReference(FWorkerContext& Context, FReferenceMetadata Metadata)
+	{
+#if !UE_BUILD_SHIPPING
+		Context.bDetectedGarbageReference |= !IsWithPendingKill() && Metadata.Has(KillFlag);
+#endif
+	}
+
+	/**
+	 * Handles object reference, potentially NULL'ing
+	 *
+	 * @param Object						Object pointer passed by reference
+	 * @param ReferencingObject UObject which owns the reference (can be NULL)
+	 * @param bAllowReferenceElimination	Whether to allow NULL'ing the reference if RF_PendingKill is set
+	*/
+	template<EKillable Killable>
+	static FORCEINLINE_DEBUGGABLE void ProcessReferenceDirectly(FWorkerContext& Context, FPermanentObjectPoolExtents PermanentPool, const UObject* ReferencingObject, UObject*& Object, FTokenId TokenId)
+	{
+		if (ValidateReference(Object, PermanentPool, ReferencingObject, TokenId))
+		{
+			const int32 ObjectIndex = GUObjectArray.ObjectToIndex(Object);
+			FImmutableReference Reference = {Object};	
+			FReferenceMetadata Metadata(ObjectIndex);
+			if (Metadata.Has(KillFlag) & (Killable == EKillable::Yes))
+			{
+				check(ReferencingObject || IsWithPendingKill());
+				checkSlow(Metadata.ObjectItem->GetOwnerIndex() <= 0);
+				KillReference(Object);
+				return;
+			}
+			
+			DetectGarbageReference(Context, Metadata);
+			HandleValidReference(Context, /*ReferencingObject, */Reference, Metadata);
+		}
+	}
+
+	FORCEINLINE static void HandleBatchedReference(FWorkerContext& Context, FResolvedMutableReference Reference, FReferenceMetadata Metadata)
+	{
+		if (Metadata.Has(KillFlag))
+		{
+			checkSlow(Metadata.ObjectItem->GetOwnerIndex() <= 0);
+			KillReference(*Reference.Mutable);
+		}
+		else
+		{
+			HandleValidReference(Context, Reference, Metadata);
+		}
+	}
+
+	FORCEINLINE static void HandleBatchedReference(FWorkerContext& Context, FImmutableReference Reference, FReferenceMetadata Metadata)
+	{
+		DetectGarbageReference(Context, Metadata);
+		HandleValidReference(Context, Reference, Metadata);
+	}
+
+	FORCEINLINE static bool HandleValidReference(FWorkerContext& Context, FImmutableReference Reference, FReferenceMetadata Metadata)
+	{
+		if (ClearUnreachableInterlocked(Metadata.ObjectItem->Flags))
 		{
-			checkSlow(ObjectItem->GetOwnerIndex() <= 0);
+			// Objects that are part of a GC cluster should never have the unreachable flag set!
+			checkSlow(Metadata.ObjectItem->GetOwnerIndex() <= 0);
+
+			if (!IsWithClusters(Options) || !Metadata.Has(EInternalObjectFlags::ClusterRoot))
+			{
+				// Add it to the list of objects to serialize.
+				Context.ObjectsToSerialize.Add<Options>(Reference.Object);
+			}
+			else if constexpr (IsWithClusters(Options))
+			{
+				// This is a cluster root reference so mark all referenced clusters as reachable
+				MarkReferencedClustersAsReachableThunk<Options>(Metadata.ObjectItem->GetClusterIndex(), Context.ObjectsToSerialize);
+			}
 
-			// Null out reference.
-			Object = nullptr;
+			return true;
 		}
-		// Add encountered object reference to list of to be serialized objects if it hasn't already been added.
-		else if (ObjectItem->IsUnreachable())
+		else if constexpr (IsWithClusters(Options))
 		{
-			if (IsParallel())
+			if ((Metadata.ObjectItem->GetOwnerIndex() > 0) & !Metadata.Has(EInternalObjectFlags::ReachableInCluster))
 			{
-				// Mark it as reachable.
-				if (ObjectItem->ThisThreadAtomicallyClearedRFUnreachable())
+				// Make sure cluster root object is reachable too
+				FUObjectItem* RootObjectItem = GUObjectArray.IndexToObjectUnsafeForGC(Metadata.ObjectItem->GetOwnerIndex());
+				checkSlow(RootObjectItem->HasAnyFlags(EInternalObjectFlags::ClusterRoot));
+
+				bool bNeedsDoing = true;
+				if constexpr (IsParallel(Options))
 				{
-#if ENABLE_GC_HISTORY
-					if (bWithHistory)
+					bNeedsDoing = Metadata.ObjectItem->ThisThreadAtomicallySetFlag(EInternalObjectFlags::ReachableInCluster);
+				}
+				else
+				{
+					Metadata.ObjectItem->SetFlags(EInternalObjectFlags::ReachableInCluster);
+				}
+				if (bNeedsDoing)
+				{
+					if constexpr (IsParallel(Options))
 					{
-						HandleHistoryReference(ObjectsToSerializeStruct, ReferencingObject, Object, TokenIndex);
+						if (ClearUnreachableInterlocked(RootObjectItem->Flags))
+						{
+							// Make sure all referenced clusters are marked as reachable too
+							MarkReferencedClustersAsReachableThunk<Options>(RootObjectItem->GetClusterIndex(), Context.ObjectsToSerialize);
+						}
 					}
-#endif // ENABLE_GC_HISTORY
+					else if (RootObjectItem->IsUnreachable())
+					{
+						RootObjectItem->ClearFlags(EInternalObjectFlags::Unreachable);
+						// Make sure all referenced clusters are marked as reachable too
+						MarkReferencedClustersAsReachableThunk<Options>(RootObjectItem->GetClusterIndex(), Context.ObjectsToSerialize);
+					}
+				}
+			}
+	#if PERF_DETAILED_PER_CLASS_GC_STATS
+			GCurrentObjectRegularObjectRefs++;
+	#endif
+		}
+
+		return false;
+	}
+};
+
+//////////////////////////////////////////////////////////////////////////
+
+/** Batches up references before dispatching them to the processor, unlike TDirectDispatcher */
+template <class ProcessorType>
+struct TBatchDispatcher
+{
+	UE_NONCOPYABLE(TBatchDispatcher);
+
+	FWorkerContext& Context;
+	const FPermanentObjectPoolExtents PermanentPool;
+	TReferenceBatcher<FMutableReference, FResolvedMutableReference, ProcessorType> KillableBatcher;
+	TReferenceBatcher<FImmutableReference, FImmutableReference, ProcessorType> ImmutableBatcher;
+
+	explicit TBatchDispatcher(FWorkerContext& InContext)
+	: Context(InContext)
+	, KillableBatcher(InContext)
+	, ImmutableBatcher(InContext)
+	{}
+
+	// ProcessObjectArray API
+
+	FORCENOINLINE void HandleReferenceDirectly(const UObject* ReferencingObject, UObject*& Object, FTokenId TokenId, EKillable Killable)
+	{
+		ProcessorType::ProcessReferenceDirectly(Context, PermanentPool, ReferencingObject, Object, TokenId, Killable);
+	}
+
+	FORCEINLINE void HandleReferenceDirectly(const UObject* ReferencingObject, UObject*& Object, FTokenId TokenId, EGCTokenType TokenType, bool bAllowReferenceElimination)
+	{
+		HandleReferenceDirectly(ReferencingObject, Object, TokenId, ProcessorType::MayKill(TokenType, bAllowReferenceElimination));
+	}
+
+	FORCEINLINE void HandleKillableReference(UObject*& Object, FTokenId TokenId, EGCTokenType TokenType)
+	{
+		QueueReference(Context.GetReferencingObject(), Object, TokenId, ProcessorType::MayKill(TokenType, true));
+	}
+
+	FORCEINLINE void HandleImmutableReference(UObject* Object, FTokenId TokenId, EGCTokenType TokenType)
+	{
+		ImmutableBatcher.PushReference(FImmutableReference{Object});
+	}
+
+	FORCEINLINE void HandleKillableArray(TArray<UObject*>& Array, FTokenId TokenId, EGCTokenType TokenType)
+	{
+		QueueArray(Context.GetReferencingObject(), Array, TokenId, ProcessorType::MayKill(TokenType, true));
+	}
+
+	FORCEINLINE void HandleKillableReferences(TArrayView<UObject*> Objects, FTokenId TokenId, EGCTokenType TokenType)
+	{
+		QueueReferences(Context.GetReferencingObject(), Objects, TokenId, ProcessorType::MayKill(TokenType, true));
+	}
+
+	FORCEINLINE void HandleWeakReference(FWeakObjectPtr& WeakPtr, const UObject* ReferencingObject, FTokenId TokenId, EGCTokenType)
+	{
+		UObject* WeakObject = WeakPtr.Get(true);
+		HandleReferenceDirectly(ReferencingObject, WeakObject, TokenId, EKillable::No);
+	}
+
+	FORCENOINLINE void FlushQueuedReferences()
+	{
+		KillableBatcher.FlushQueues();
+		ImmutableBatcher.FlushQueues();
+	}
+
+	// Collector API
+
+	FORCEINLINE_DEBUGGABLE void QueueReference(const UObject* ReferencingObject,  UObject*& Object, FTokenId TokenId, EKillable Killable)
+	{
+		if (Killable == EKillable::Yes)
+		{
+			FPlatformMisc::Prefetch(&Object);
+			KillableBatcher.PushReference(FMutableReference{&Object});
+		}
+		else
+		{
+			ImmutableBatcher.PushReference(FImmutableReference{Object});
+		}
+	}
+
+	FORCEINLINE_DEBUGGABLE void QueueReferences(const UObject* ReferencingObject, TArrayView<UObject*> References, FTokenId TokenId, EKillable Killable)
+	{
+		if (Killable == EKillable::Yes)
+		{
+			KillableBatcher.PushReferences(References);
+		}
+		else
+		{
+			ImmutableBatcher.PushReferences(References);
+		}
+	}
+	
+	FORCEINLINE_DEBUGGABLE void QueueArray(const UObject* ReferencingObject, TArray<UObject*>& Array, FTokenId TokenId, EKillable Killable)
+	{
+		if (Killable == EKillable::Yes)
+		{
+			KillableBatcher.PushArray(FReferenceArray{&Array});
+		}
+		else
+		{
+			ImmutableBatcher.PushArray(FReferenceArray{&Array});
+		}
+	}
+
+	FORCEINLINE_DEBUGGABLE void QueueSet(const UObject* ReferencingObject, TSet<UObject*>& Objects, FTokenId TokenId, EKillable Killable)
+	{
+		if (Killable == EKillable::Yes)
+		{
+			for (UObject*& Object : Objects)
+			{
+				QueueReference(ReferencingObject, Object, TokenId, EKillable::Yes);
+			}
+		}
+		else
+		{
+			for (UObject*& Object : Objects)
+			{
+				QueueReference(ReferencingObject, Object, TokenId, EKillable::No);
+			}
+		}
+	}
+};
+
+//////////////////////////////////////////////////////////////////////////
+
+template <EGCOptions Options>
+class TReachabilityCollectorBase : public FReferenceCollector
+{
+protected:
+	bool bAllowEliminatingReferences = true;
+	EGCTokenType CurrentTokenType = EGCTokenType::Native;
+
+	FORCEINLINE EKillable MayKill() const { return UE::GC::MayKill<Options>(CurrentTokenType, bAllowEliminatingReferences); }
+
+public:
+	virtual bool IsIgnoringArchetypeRef() const override final { return false;}
+	virtual bool IsIgnoringTransient() const override final { return false; }
+	virtual bool NeedsInitialReferences() const override final { return !IsParallel(Options); }
+	virtual void AllowEliminatingReferences(bool bAllow) override final { bAllowEliminatingReferences = bAllow; }
+
+	virtual void SetIsProcessingNativeReferences(bool bIsNative) override final
+	{
+		CurrentTokenType = bIsNative ? EGCTokenType::Native : EGCTokenType::NonNative;
+	}
+
+	virtual bool NeedsPropertyReferencer() const override final { return ENABLE_GC_OBJECT_CHECKS; }
+};
+
+template <EGCOptions Options>
+class TReachabilityCollector final : public TReachabilityCollectorBase<Options>
+{
+	using ProcessorType = TReachabilityProcessor<Options>;
+	using DispatcherType = TBatchDispatcher<ProcessorType>;
+	using Super =  TReachabilityCollectorBase<Options>;
+	using Super::MayKill;
+
+	DispatcherType Dispatcher;
+
+	friend DispatcherType& GetDispatcher(TReachabilityCollector<Options>& Collector, ProcessorType&, FWorkerContext&) { return Collector.Dispatcher; }
+
+public:
+	TReachabilityCollector(ProcessorType&, FWorkerContext& Context) : Dispatcher(Context) {}
+		
+	virtual void HandleObjectReference(UObject*& InObject, const UObject* InReferencingObject, const FProperty* InReferencingProperty) override;
+	virtual void HandleObjectReferences(UObject** InObjects, const int32 ObjectNum, const UObject* InReferencingObject, const FProperty* InReferencingProperty) override;
+	
+	virtual void AddStableReference(UObject** Object) override
+	{
+		Dispatcher.QueueReference(Dispatcher.Context.GetReferencingObject(), *Object, ETokenlessId::Collector, MayKill());
+	}
+	
+	virtual void AddStableReferenceArray(TArray<UObject*>* Objects) override
+	{
+		Dispatcher.QueueArray(Dispatcher.Context.GetReferencingObject(), *Objects, ETokenlessId::Collector, MayKill());
+	}
+
+	virtual void AddStableReferenceSet(TSet<UObject*>* Objects) override
+	{
+		Dispatcher.QueueSet(Dispatcher.Context.GetReferencingObject(), *Objects, ETokenlessId::Collector, MayKill());
+	}
+	
+	virtual bool MarkWeakObjectReferenceForClearing(UObject** WeakReference) override
+	{
+		Dispatcher.Context.WeakReferences.Add(WeakReference);
+		return true;
+	}
+};
 
-					// Objects that are part of a GC cluster should never have the unreachable flag set!
-					checkSlow(ObjectItem->GetOwnerIndex() <= 0);
+#if !UE_BUILD_SHIPPING
 
-					if (!IsWithClusters() || !ObjectItem->HasAnyFlags(EInternalObjectFlags::ClusterRoot))
-					{
-						// Add it to the list of objects to serialize.
-						ObjectsToSerialize.Add(Object);
-					}
-					else
-					{
-						// This is a cluster root reference so mark all referenced clusters as reachable
-						MarkReferencedClustersAsReachable(ObjectItem->GetClusterIndex(), ObjectsToSerialize);
-					}
-				}
-			}
-			else
-			{
-#if ENABLE_GC_DEBUG_OUTPUT
-				// this message is to help track down culprits behind "Object in PIE world still referenced" errors
-				if (GIsEditor && !GIsPlayInEditorWorld && ReferencingObject != nullptr && Object != nullptr && !ReferencingObject->HasAnyFlags(RF_Transient) && Object->RootPackageHasAnyFlags(PKG_PlayInEditor))
-				{
-					UPackage* ReferencingPackage = ReferencingObject->GetOutermost();
-					if (!ReferencingPackage->HasAnyPackageFlags(PKG_PlayInEditor) && !ReferencingPackage->HasAnyFlags(RF_Transient))
-					{
-						UE_LOG(LogGarbage, Warning, TEXT("GC detected illegal reference to PIE object from content [possibly via [todo]]:"));
-						UE_LOG(LogGarbage, Warning, TEXT("      PIE object: %s"), *Object->GetFullName());
-						UE_LOG(LogGarbage, Warning, TEXT("  NON-PIE object: %s"), *ReferencingObject->GetFullName());
-					}
-				}
+#if !ENABLE_GC_HISTORY
+#error TDebugReachabilityProcessor assumes ENABLE_GC_HISTORY is enabled in non-shipping configs
 #endif
 
-#if ENABLE_GC_HISTORY
-				if (bWithHistory)
-				{
-					HandleHistoryReference(ObjectsToSerializeStruct, ReferencingObject, Object, TokenIndex);
-				}
-#endif // ENABLE_GC_HISTORY
-
-				// Mark it as reachable.
-				ObjectItem->ClearUnreachable();
+template <EGCOptions InOptions>
+class TDebugReachabilityProcessor
+{
+public:
+	static constexpr EGCOptions Options = InOptions;
 
-				// Objects that are part of a GC cluster should never have the unreachable flag set!
-				checkSlow(ObjectItem->GetOwnerIndex() <= 0);
+	TDebugReachabilityProcessor()
+	: bTrackGarbage(GGarbageReferenceTrackingEnabled != 0)
+	, bTrackHistory(FGCHistory::Get().IsActive())
+	{}
+	
+	bool TracksHistory() const { return bTrackHistory; }
+	bool TracksGarbage() const { return bTrackGarbage; }
 
-				if (!IsWithClusters() || !ObjectItem->HasAnyFlags(EInternalObjectFlags::ClusterRoot))
+	FORCENOINLINE void HandleTokenStreamObjectReference(FWorkerContext& Context, const UObject* ReferencingObject, UObject*& Object, FTokenId TokenId, EGCTokenType TokenType, bool bAllowReferenceElimination)
+	{
+		if (ValidateReference(Object, PermanentPool, ReferencingObject, TokenId))
+		{
+			FReferenceMetadata Metadata(GUObjectArray.ObjectToIndex(Object));
+			if (bAllowReferenceElimination & Metadata.Has(TReachabilityProcessor<Options>::KillFlag))
+			{
+				if (MayKill<Options>(TokenType, bAllowReferenceElimination) == EKillable::Yes)
 				{
-					// Add it to the list of objects to serialize.
-					ObjectsToSerialize.Add(Object);
+					check(IsPendingKill(Options) || ReferencingObject);
+					checkSlow(Metadata.ObjectItem->GetOwnerIndex() <= 0);
+					KillReference(Object);
+					return;
 				}
-				else
+				else if (!IsPendingKill(Options) && bTrackGarbage)
 				{
-					// This is a cluster root reference so mark all referenced clusters as reachable
-					MarkReferencedClustersAsReachable(ObjectItem->GetClusterIndex(), ObjectsToSerialize);
+					check(TokenType == EGCTokenType::Native);
+					HandleGarbageReference(Context, ReferencingObject, Object, TokenId);
 				}
 			}
+			
+			bool bReachedFirst = TReachabilityProcessor<Options>::HandleValidReference(Context, FImmutableReference{Object}, Metadata);
+
+			if (bReachedFirst && bTrackHistory)
+			{
+				HandleHistoryReference(Context, ReferencingObject, Object, TokenId);
+			}
 		}
-		else if (IsWithClusters() && (ObjectItem->GetOwnerIndex() > 0 && !ObjectItem->HasAnyFlags(EInternalObjectFlags::ReachableInCluster)))
+	}
+
+private:
+	const FPermanentObjectPoolExtents PermanentPool;
+	const bool bTrackGarbage;
+	const bool bTrackHistory;
+
+	FORCENOINLINE static void HandleGarbageReference(FWorkerContext& Context, const UObject* ReferencingObject, UObject*& Object, FTokenId TokenId)
+	{
+		const UObject* Referencer = ReferencingObject ? ReferencingObject : Context.GetReferencingObject();
+		if (IsValid(Referencer))
 		{
-			bool bNeedsDoing = true;
-			if (IsParallel())
+			if (Referencer != FGCObject::GGCObjectReferencer)
 			{
-				bNeedsDoing = ObjectItem->ThisThreadAtomicallySetFlag(EInternalObjectFlags::ReachableInCluster);
+				FName PropertyName = Referencer->GetClass()->ReferenceTokens.GetTokenInfo(TokenId).Name;
+				Context.GarbageReferences.Add(FGarbageReferenceInfo(Referencer, Object, PropertyName));
 			}
 			else
 			{
-				ObjectItem->SetFlags(EInternalObjectFlags::ReachableInCluster);
-			}
-			if (bNeedsDoing)
-			{
-				// Make sure cluster root object is reachable too
-				const int32 OwnerIndex = ObjectItem->GetOwnerIndex();
-				FUObjectItem* RootObjectItem = GUObjectArray.IndexToObjectUnsafeForGC(OwnerIndex);
-				checkSlow(RootObjectItem->HasAnyFlags(EInternalObjectFlags::ClusterRoot));
-				if (IsParallel())
-				{
-					if (RootObjectItem->ThisThreadAtomicallyClearedRFUnreachable())
-					{
-						// Make sure all referenced clusters are marked as reachable too
-						MarkReferencedClustersAsReachable(RootObjectItem->GetClusterIndex(), ObjectsToSerialize);
-					}
-				}
-				else if (RootObjectItem->IsUnreachable())
-				{
-					RootObjectItem->ClearFlags(EInternalObjectFlags::Unreachable);
-					// Make sure all referenced clusters are marked as reachable too
-					MarkReferencedClustersAsReachable(RootObjectItem->GetClusterIndex(), ObjectsToSerialize);
-				}
+				FGCObject* GCObjectReferencer = FGCObject::GGCObjectReferencer->GetCurrentlySerializingObject();
+				Context.GarbageReferences.Add(FGarbageReferenceInfo(GCObjectReferencer, Object));
 			}
 		}
-#if PERF_DETAILED_PER_CLASS_GC_STATS
-		GCurrentObjectRegularObjectRefs++;
-#endif
 	}
 
-	/**
-	* Handles UObject reference from the token stream.
-	*
-	* @param ObjectsToSerialize An array of remaining objects to serialize.
-	* @param ReferencingObject Object referencing the object to process.
-	* @param Object Object being referenced.
-	* @param TokenIndex Index to the token stream where the reference was found.
-	* @param bAllowReferenceElimination True if reference elimination is allowed.
-	*/
-	FORCEINLINE void HandleTokenStreamObjectReference(FGCArrayStruct& ObjectsToSerializeStruct, UObject* ReferencingObject, UObject*& Object, const int32 TokenIndex, const EGCTokenType TokenType, bool bAllowReferenceElimination)
+	FORCENOINLINE static void HandleHistoryReference(FWorkerContext& Context, const UObject* ReferencingObject, UObject*& Object, FTokenId TokenId)
 	{
-#if ENABLE_GC_OBJECT_CHECKS
-		if (Object && IsObjectHandleResolved(*reinterpret_cast<FObjectHandle*>(&Object)))
+		FGCDirectReference Ref(Object);
+		const UObject* Referencer = ReferencingObject ? ReferencingObject : Context.GetReferencingObject();	
+		if (Referencer != FGCObject::GGCObjectReferencer)
 		{
-			if (
-#if DO_POINTER_CHECKS_ON_GC
-				!IsPossiblyAllocatedUObjectPointer(Object) ||
-#endif
-				!Object->IsValidLowLevelFast())
-			{
-				FString TokenDebugInfo;
-				if (UClass *Class = (ReferencingObject ? ReferencingObject->GetClass() : nullptr))
-				{
-					FTokenInfo TokenInfo = Class->ReferenceTokenStream.GetTokenInfo(TokenIndex);
-					TokenDebugInfo = FString::Printf(TEXT("ReferencingObjectClass: %s, Property Name: %s, Offset: %d"),
-						*Class->GetFullName(), *TokenInfo.Name.GetPlainNameString(), TokenInfo.Offset);
-				}
-				else
-				{
-					// This means this objects is most likely being referenced by AddReferencedObjects
-					TokenDebugInfo = TEXT("Native Reference");
-				}
+			Ref.ReferencerName = Referencer->GetClass()->ReferenceTokens.GetTokenInfo(TokenId).Name;
+		}
+		else if (FGCObject::GGCObjectReferencer == Referencer)
+		{
+			Ref.ReferencerName = *FGCObject::GGCObjectReferencer->GetCurrentlySerializingObject()->GetReferencerName();
+		}
 
-				UE_LOG(LogGarbage, Fatal, TEXT("Invalid object in GC: 0x%016llx, ReferencingObject: %s, %s, TokenIndex: %d"),
-					(int64)(PTRINT)Object,
-					ReferencingObject ? *ReferencingObject->GetFullName() : TEXT("NULL"),
-					*TokenDebugInfo, TokenIndex);
-			}
+		TArray<FGCDirectReference>*& DirectReferences = Context.History.FindOrAdd(Referencer);
+		if (!DirectReferences)
+		{
+			DirectReferences = new TArray<FGCDirectReference>();
 		}
-#endif // ENABLE_GC_OBJECT_CHECKS
-		HandleObjectReference(ObjectsToSerializeStruct, ReferencingObject, Object, TokenIndex, TokenType, bAllowReferenceElimination);
+		DirectReferences->Add(Ref);
 	}
 };
 
-template <EFastReferenceCollectorOptions Options>
-FGCCollector<Options>::FGCCollector(FGCReferenceProcessor<Options>& InProcessor, FGCArrayStruct& InObjectArrayStruct)
-		: ReferenceProcessor(InProcessor)
-		, ObjectArrayStruct(InObjectArrayStruct)
-		, bAllowEliminatingReferences(true)
-		, bIsProcessingNativeReferences(true)
+template <EGCOptions Options>
+class TDebugReachabilityCollector final : public TReachabilityCollectorBase<Options>
 {
-}
+	using Super = TReachabilityCollectorBase<Options>;
+	using Super::bAllowEliminatingReferences;
+	using Super::CurrentTokenType;
+	using ProcessorType = TDebugReachabilityProcessor<Options>;
+	ProcessorType& Processor;
+	FWorkerContext& Context;
 
-template <EFastReferenceCollectorOptions Options>
-FORCEINLINE void FGCCollector<Options>::InternalHandleObjectReference(UObject*& Object, const UObject* ReferencingObject, const FProperty* ReferencingProperty, const EGCTokenType InTokenType)
-{
+public:
+	TDebugReachabilityCollector(ProcessorType& InProcessor, FWorkerContext& InContext) : Processor(InProcessor), Context(InContext) {}
+
+	virtual void HandleObjectReference(UObject*& Object, const UObject* ReferencingObject, const FProperty* ReferencingProperty) override
+	{
 #if ENABLE_GC_OBJECT_CHECKS
 		if (Object && !Object->IsValidLowLevelFast())
 		{
@@ -1199,32 +2732,125 @@ FORCEINLINE void FGCCollector<Options>::InternalHandleObjectReference(UObject*&
 				ReferencingProperty ? *ReferencingProperty->GetFullName() : TEXT("NULL"));
 		}
 #endif // ENABLE_GC_OBJECT_CHECKS
-		const int32 InvalidTokenIndex = -1; // this reference is coming from a native AddReferencedObjects so no valid token index is available
-		ReferenceProcessor.HandleObjectReference(ObjectArrayStruct, const_cast<UObject*>(ReferencingObject), Object, InvalidTokenIndex, InTokenType, bAllowEliminatingReferences);
+
+		Processor.HandleTokenStreamObjectReference(Context, ReferencingObject, Object, ETokenlessId::Collector, CurrentTokenType, bAllowEliminatingReferences);
+	}
+
+	virtual void HandleObjectReferences(UObject** Objects, int32 Num, const UObject* ReferencingObject, const FProperty* ReferencingProperty) override
+	{
+		for (UObject*& Object : MakeArrayView(Objects, Num))
+		{
+			HandleObjectReference(Object, ReferencingObject, ReferencingProperty);
+		}
+	}
+
+	virtual void AddStableReference(UObject** Object) override
+	{
+		HandleObjectReference(*Object, Context.GetReferencingObject(), nullptr);
+	}
+	
+	virtual void AddStableReferenceArray(TArray<UObject*>* Objects) override
+	{
+		for (UObject*& Object : *Objects)
+		{
+			HandleObjectReference(Object, Context.GetReferencingObject(), nullptr);
+		}
+	}
+
+	virtual void AddStableReferenceSet(TSet<UObject*>* Objects) override
+	{
+		for (UObject*& Object : *Objects)
+		{
+			HandleObjectReference(Object, Context.GetReferencingObject(), nullptr);
+		}
+	}
+
+	virtual bool MarkWeakObjectReferenceForClearing(UObject** WeakReference) override
+	{
+		Context.WeakReferences.Add(WeakReference);
+		return true;
+	}
+};
+
+#endif // !UE_BUILD_SHIPPING
+
+FORCEINLINE static void CheckReference(UObject*& Object, const UObject* ReferencingObject, const FProperty* ReferencingProperty)
+{
+#if ENABLE_GC_OBJECT_CHECKS
+	if (Object && !Object->IsValidLowLevelFast())
+	{
+		UE_LOG(LogGarbage, Fatal, TEXT("Invalid object in GC: 0x%016llx, ReferencingObject: %s, ReferencingProperty: %s"),
+			(int64)(PTRINT)Object,
+			ReferencingObject ? *ReferencingObject->GetFullName() : TEXT("NULL"),
+			ReferencingProperty ? *ReferencingProperty->GetFullName() : TEXT("NULL"));
+	}
+#endif // ENABLE_GC_OBJECT_CHECKS
 }
 
-template <EFastReferenceCollectorOptions Options>
-void FGCCollector<Options>::HandleObjectReference(UObject*& Object, const UObject* ReferencingObject, const FProperty* ReferencingProperty)
+template <EGCOptions Options>
+void TReachabilityCollector<Options>::HandleObjectReference(UObject*& Object, const UObject* ReferencingObject, const FProperty* ReferencingProperty)
 {
-	const EGCTokenType TokenType = bIsProcessingNativeReferences ? EGCTokenType::Native : EGCTokenType::NonNative;
-	InternalHandleObjectReference(Object, ReferencingObject, ReferencingProperty, TokenType);
+	CheckReference(Object, ReferencingObject, ReferencingProperty);
+	Dispatcher.HandleReferenceDirectly(ReferencingObject, Object, ETokenlessId::Collector, MayKill());
 }
 
-template <EFastReferenceCollectorOptions Options>
-void FGCCollector<Options>::HandleObjectReferences(UObject** InObjects, const int32 ObjectNum, const UObject* InReferencingObject, const FProperty* InReferencingProperty)
+template <EGCOptions Options>
+void TReachabilityCollector<Options>::HandleObjectReferences(UObject** InObjects, int32 ObjectNum, const UObject* ReferencingObject, const FProperty* ReferencingProperty)
 {
-	const EGCTokenType TokenType = bIsProcessingNativeReferences ? EGCTokenType::Native : EGCTokenType::NonNative;
-	for (int32 ObjectIndex = 0; ObjectIndex < ObjectNum; ++ObjectIndex)
+	if (ObjectNum == 0)
+	{
+		return;
+	}
+
+	EKillable Killable = MayKill();
+
+	TArrayView<UObject*> Objects(InObjects, ObjectNum);
+	for (UObject*& Object : Objects)
+	{
+		CheckReference(Object, ReferencingObject, ReferencingProperty);
+	}
+
+	for (UObject*& Object : Objects)
 	{
-		UObject*& Object = InObjects[ObjectIndex];
-		InternalHandleObjectReference(Object, InReferencingObject, InReferencingProperty, TokenType);
+		Dispatcher.HandleReferenceDirectly(ReferencingObject, Object, ETokenlessId::Collector, Killable);
 	}
 }
-#endif	// UE_WITH_GC
 
-/*----------------------------------------------------------------------------
-	FReferenceFinder.
-----------------------------------------------------------------------------*/
+#else // !UE_WITH_GC
+
+static constexpr int32 MaxWorkers = 1;
+static void CollectGarbageInternal(EObjectFlags, bool)				{ unimplemented(); }
+void PadObjectArray(TArray<UObject*>&)								{ unimplemented(); }
+FWorkerContext::FWorkerContext()									{ unimplemented(); }
+FWorkerContext::~FWorkerContext()									{ unimplemented(); }
+FWorkBlockifier::~FWorkBlockifier()									{ unimplemented(); }
+void FWorkBlockifier::FreeOwningBlock(UObject*const*)				{ unimplemented(); }
+void FWorkBlockifier::PushFullBlockSync()							{ unimplemented(); }
+FWorkBlock* FWorkBlockifier::PopFullBlockSync()						{ unimplemented(); return nullptr; }
+FWorkBlock* FWorkBlockifier::PopWipBlock()							{ unimplemented(); return nullptr; }
+void FSlowARO::CallSync(uint32, UObject*, FReferenceCollector&)		{ unimplemented(); }
+void RegisterSlowImplementation(AROFunc, EAROFlags)					{ return; }
+int32 FindSlowImplementation(AROFunc)								{ return INDEX_NONE; }
+
+#endif // !UE_WITH_GC
+
+} // namespace UE::GC
+
+//////////////////////////////////////////////////////////////////////////
+
+void ShutdownGarbageCollection()
+{
+#if UE_WITH_GC
+	UE::GC::FContextPoolScope().Cleanup();
+#endif
+	delete GAsyncPurge;
+	GAsyncPurge = nullptr;
+#if ENABLE_GC_HISTORY
+	FGCHistory::Get().Cleanup();
+#endif
+}
+
+//////////////////////////////////////////////////////////////////////////
 FReferenceFinder::FReferenceFinder(TArray<UObject*>& InObjectArray, UObject* InOuter, bool bInRequireDirectOuter, bool bInShouldIgnoreArchetype, bool bInSerializeRecursively, bool bInShouldIgnoreTransient)
 	: ObjectArray(InObjectArray)
 	, LimitOuter(InOuter)
@@ -1252,8 +2878,7 @@ void FReferenceFinder::FindReferences(UObject* Object, UObject* InReferencingObj
 
 	if (!Object->GetClass()->IsChildOf(UClass::StaticClass()))
 	{
-		FVerySlowReferenceCollectorArchiveScope CollectorScope(GetVerySlowReferenceCollectorArchive(), InReferencingObject, SerializedProperty);
-		Object->SerializeScriptProperties(CollectorScope.GetArchive());
+		AddPropertyReferences(Object->GetClass(), Object, InReferencingObject);
 	}
 	Object->CallAddReferencedObjects(*this);
 }
@@ -1284,15 +2909,13 @@ void FReferenceFinder::HandleObjectReference( UObject*& InObject, const UObject*
 	}
 }
 
-void FGarbageCollectionTracer::PerformReachabilityAnalysisOnObjects(FGCArrayStruct* ArrayStruct, bool bForceSingleThreaded, bool bWithClusters)
+//////////////////////////////////////////////////////////////////////////
+
+#if UE_WITH_GC
+
+namespace UE::GC
 {
-	const EFastReferenceCollectorOptions Options =
-		(bForceSingleThreaded ? EFastReferenceCollectorOptions::None : EFastReferenceCollectorOptions::Parallel) |
-		(bWithClusters ? EFastReferenceCollectorOptions::WithClusters : EFastReferenceCollectorOptions::None) |
-		(UObjectBaseUtility::IsPendingKillEnabled() ? EFastReferenceCollectorOptions::WithPendingKill : EFastReferenceCollectorOptions::None);
 
-	PerformReachabilityAnalysisOnObjects(ArrayStruct, Options);
-}
 
 /**
  * Implementation of parallel realtime garbage collector using recursive subdivision
@@ -1304,62 +2927,93 @@ void FGarbageCollectionTracer::PerformReachabilityAnalysisOnObjects(FGCArrayStru
  * is used to deal with object references from types that aren't supported by the reflectable type system.
  * interface doesn't make sense to implement for.
  */
-
-#if UE_WITH_GC
-
 class FRealtimeGC : public FGarbageCollectionTracer
 {
-	typedef void(FRealtimeGC::*MarkObjectsFn)(TArray<UObject*>&, const EObjectFlags);
-	typedef void(FRealtimeGC::*ReachabilityAnalysisFn)(FGCArrayStruct*);
+	typedef void(FRealtimeGC::*MarkObjectsFn)(EObjectFlags);
+	typedef void(FRealtimeGC::*ReachabilityAnalysisFn)(FWorkerContext&);
 
 	/** Pointers to functions used for Marking objects as unreachable */
 	MarkObjectsFn MarkObjectsFunctions[4];
 	/** Pointers to functions used for Reachability Analysis */
 	ReachabilityAnalysisFn ReachabilityAnalysisFunctions[8];
 
-	template <EFastReferenceCollectorOptions CollectorOptions>
-	void PerformReachabilityAnalysisOnObjectsInternal(FGCArrayStruct* ArrayStruct)
+	UE::Tasks::TTask<void> InitialCollection;
+	TArray<UObject**> InitialReferences;
+	TArray<UObject*> InitialObjects;
+
+	void BeginInitialReferenceCollection(EGCOptions Options)
+	{
+		InitialReferences.Reset();
+
+		if (IsParallel(Options))
+		{
+			InitialCollection = UE::Tasks::Launch(TEXT("CollectInitialReferences"), 
+				[&] () { FGCObject::GGCObjectReferencer->AddInitialReferences(InitialReferences); });
+		}
+	}
+
+	TConstArrayView<UObject**> GetInitialReferences(EGCOptions Options)
+	{
+		if (!!(Options & EGCOptions::Parallel))
+		{
+			InitialCollection.Wait();
+		}
+		else
+		{
+			FGCObject::GGCObjectReferencer->AddInitialReferences(InitialReferences);
+		}
+
+		return InitialReferences;
+	}
+
+	template <EGCOptions Options>
+	void PerformReachabilityAnalysisOnObjectsInternal(FWorkerContext& Context)
 	{
 		TRACE_CPUPROFILER_EVENT_SCOPE(PerformReachabilityAnalysisOnObjectsInternal);
-		FGCReferenceProcessor<CollectorOptions> ReferenceProcessor;
-		// NOTE: we want to run with automatic token stream generation off as it should be already generated at this point,
-		// BUT we want to be ignoring Noop tokens as they're only pointing either at null references or at objects that never get GC'd (native classes)
-		TFastReferenceCollector<
-			FGCReferenceProcessor<CollectorOptions>,
-			FGCCollector<CollectorOptions>,
-			FGCArrayPool,
-			CollectorOptions
-			>  ReferenceCollector(ReferenceProcessor, FGCArrayPool::Get());
-		ReferenceCollector.CollectReferences(*ArrayStruct);
+
+#if !UE_BUILD_SHIPPING		
+		TDebugReachabilityProcessor<Options> DebugProcessor;
+		if (DebugProcessor.TracksHistory() | 
+			DebugProcessor.TracksGarbage() & GRedoReachabilityToTrackGarbage)
+		{
+			CollectReferences<TDebugReachabilityCollector<Options>>(DebugProcessor, Context);
+			GRedoReachabilityToTrackGarbage = false;
+			return;
+		}
+#endif
+		
+		TReachabilityProcessor<Options> Processor;
+		CollectReferences<TReachabilityCollector<Options>>(Processor, Context);
 	}
 
 	/** Calculates GC function index based on current settings */
-	static FORCEINLINE int32 GetGCFunctionIndex(EFastReferenceCollectorOptions InOptions)
+	static FORCEINLINE int32 GetGCFunctionIndex(EGCOptions InOptions)
 	{
-		return (!!(InOptions & EFastReferenceCollectorOptions::Parallel)) |
-			(!!(InOptions & EFastReferenceCollectorOptions::WithClusters) << 1) |
-			(!!(InOptions & EFastReferenceCollectorOptions::WithPendingKill) << 2);
+		return (!!(InOptions & EGCOptions::Parallel)) |
+			(!!(InOptions & EGCOptions::WithClusters) << 1) |
+			(!!(InOptions & EGCOptions::WithPendingKill) << 2);
 	}
 
 public:
 	/** Default constructor, initializing all members. */
 	FRealtimeGC()
 	{
-		MarkObjectsFunctions[GetGCFunctionIndex(EFastReferenceCollectorOptions::None)] = &FRealtimeGC::MarkObjectsAsUnreachable<false, false>;
-		MarkObjectsFunctions[GetGCFunctionIndex(EFastReferenceCollectorOptions::Parallel | EFastReferenceCollectorOptions::None)] = &FRealtimeGC::MarkObjectsAsUnreachable<true, false>;
-		MarkObjectsFunctions[GetGCFunctionIndex(EFastReferenceCollectorOptions::None | EFastReferenceCollectorOptions::WithClusters)] = &FRealtimeGC::MarkObjectsAsUnreachable<false, true>;
-		MarkObjectsFunctions[GetGCFunctionIndex(EFastReferenceCollectorOptions::Parallel | EFastReferenceCollectorOptions::WithClusters)] = &FRealtimeGC::MarkObjectsAsUnreachable<true, true>;
-
-		ReachabilityAnalysisFunctions[GetGCFunctionIndex(EFastReferenceCollectorOptions::None)] = &FRealtimeGC::PerformReachabilityAnalysisOnObjectsInternal<EFastReferenceCollectorOptions::None | EFastReferenceCollectorOptions::None>;
-		ReachabilityAnalysisFunctions[GetGCFunctionIndex(EFastReferenceCollectorOptions::Parallel | EFastReferenceCollectorOptions::None)] = &FRealtimeGC::PerformReachabilityAnalysisOnObjectsInternal<EFastReferenceCollectorOptions::Parallel | EFastReferenceCollectorOptions::None>;
-		ReachabilityAnalysisFunctions[GetGCFunctionIndex(EFastReferenceCollectorOptions::None | EFastReferenceCollectorOptions::WithClusters)] = &FRealtimeGC::PerformReachabilityAnalysisOnObjectsInternal<EFastReferenceCollectorOptions::None | EFastReferenceCollectorOptions::WithClusters>;
-		ReachabilityAnalysisFunctions[GetGCFunctionIndex(EFastReferenceCollectorOptions::Parallel | EFastReferenceCollectorOptions::WithClusters)] = &FRealtimeGC::PerformReachabilityAnalysisOnObjectsInternal<EFastReferenceCollectorOptions::Parallel | EFastReferenceCollectorOptions::WithClusters>;
-
-		ReachabilityAnalysisFunctions[GetGCFunctionIndex(EFastReferenceCollectorOptions::None | EFastReferenceCollectorOptions::WithPendingKill)] = &FRealtimeGC::PerformReachabilityAnalysisOnObjectsInternal<EFastReferenceCollectorOptions::None | EFastReferenceCollectorOptions::None | EFastReferenceCollectorOptions::WithPendingKill>;
-		ReachabilityAnalysisFunctions[GetGCFunctionIndex(EFastReferenceCollectorOptions::Parallel | EFastReferenceCollectorOptions::None | EFastReferenceCollectorOptions::WithPendingKill)] = &FRealtimeGC::PerformReachabilityAnalysisOnObjectsInternal<EFastReferenceCollectorOptions::Parallel | EFastReferenceCollectorOptions::None | EFastReferenceCollectorOptions::WithPendingKill>;
-		ReachabilityAnalysisFunctions[GetGCFunctionIndex(EFastReferenceCollectorOptions::None | EFastReferenceCollectorOptions::WithClusters | EFastReferenceCollectorOptions::WithPendingKill)] = &FRealtimeGC::PerformReachabilityAnalysisOnObjectsInternal<EFastReferenceCollectorOptions::None | EFastReferenceCollectorOptions::WithClusters | EFastReferenceCollectorOptions::WithPendingKill>;
-		ReachabilityAnalysisFunctions[GetGCFunctionIndex(EFastReferenceCollectorOptions::Parallel | EFastReferenceCollectorOptions::WithClusters | EFastReferenceCollectorOptions::WithPendingKill)] = &FRealtimeGC::PerformReachabilityAnalysisOnObjectsInternal<EFastReferenceCollectorOptions::Parallel | EFastReferenceCollectorOptions::WithClusters | EFastReferenceCollectorOptions::WithPendingKill>;
-
+		MarkObjectsFunctions[GetGCFunctionIndex(EGCOptions::None)] = &FRealtimeGC::MarkObjectsAsUnreachable<false, false>;
+		MarkObjectsFunctions[GetGCFunctionIndex(EGCOptions::Parallel | EGCOptions::None)] = &FRealtimeGC::MarkObjectsAsUnreachable<true, false>;
+		MarkObjectsFunctions[GetGCFunctionIndex(EGCOptions::None | EGCOptions::WithClusters)] = &FRealtimeGC::MarkObjectsAsUnreachable<false, true>;
+		MarkObjectsFunctions[GetGCFunctionIndex(EGCOptions::Parallel | EGCOptions::WithClusters)] = &FRealtimeGC::MarkObjectsAsUnreachable<true, true>;
+
+		ReachabilityAnalysisFunctions[GetGCFunctionIndex(EGCOptions::None)] = &FRealtimeGC::PerformReachabilityAnalysisOnObjectsInternal<EGCOptions::None | EGCOptions::None>;
+		ReachabilityAnalysisFunctions[GetGCFunctionIndex(EGCOptions::Parallel | EGCOptions::None)] = &FRealtimeGC::PerformReachabilityAnalysisOnObjectsInternal<EGCOptions::Parallel | EGCOptions::None>;
+		ReachabilityAnalysisFunctions[GetGCFunctionIndex(EGCOptions::None | EGCOptions::WithClusters)] = &FRealtimeGC::PerformReachabilityAnalysisOnObjectsInternal<EGCOptions::None | EGCOptions::WithClusters>;
+		ReachabilityAnalysisFunctions[GetGCFunctionIndex(EGCOptions::Parallel | EGCOptions::WithClusters)] = &FRealtimeGC::PerformReachabilityAnalysisOnObjectsInternal<EGCOptions::Parallel | EGCOptions::WithClusters>;
+
+		ReachabilityAnalysisFunctions[GetGCFunctionIndex(EGCOptions::None | EGCOptions::WithPendingKill)] = &FRealtimeGC::PerformReachabilityAnalysisOnObjectsInternal<EGCOptions::None | EGCOptions::None | EGCOptions::WithPendingKill>;
+		ReachabilityAnalysisFunctions[GetGCFunctionIndex(EGCOptions::Parallel | EGCOptions::None | EGCOptions::WithPendingKill)] = &FRealtimeGC::PerformReachabilityAnalysisOnObjectsInternal<EGCOptions::Parallel | EGCOptions::None | EGCOptions::WithPendingKill>;
+		ReachabilityAnalysisFunctions[GetGCFunctionIndex(EGCOptions::None | EGCOptions::WithClusters | EGCOptions::WithPendingKill)] = &FRealtimeGC::PerformReachabilityAnalysisOnObjectsInternal<EGCOptions::None | EGCOptions::WithClusters | EGCOptions::WithPendingKill>;
+		ReachabilityAnalysisFunctions[GetGCFunctionIndex(EGCOptions::Parallel | EGCOptions::WithClusters | EGCOptions::WithPendingKill)] = &FRealtimeGC::PerformReachabilityAnalysisOnObjectsInternal<EGCOptions::Parallel | EGCOptions::WithClusters | EGCOptions::WithPendingKill>;
+
+		FGCObject::StaticInit();
 	}
 
 	/** 
@@ -1367,7 +3021,7 @@ public:
 	 * This function is a template to speed up the case where we don't need to assemble the token stream (saves about 6ms on PS4)
 	 */
 	template <bool bParallel, bool bWithClusters>
-	void MarkObjectsAsUnreachable(TArray<UObject*>& ObjectsToSerialize, const EObjectFlags KeepFlags)
+	void MarkObjectsAsUnreachable(const EObjectFlags KeepFlags)
 	{
 		TRACE_CPUPROFILER_EVENT_SCOPE(MarkObjectsAsUnreachable);
 		const EInternalObjectFlags FastKeepFlags = EInternalObjectFlags::GarbageCollectionKeepFlags;
@@ -1377,22 +3031,20 @@ public:
 
 		TLockFreePointerListFIFO<FUObjectItem, PLATFORM_CACHE_LINE_SIZE> ClustersToDissolveList;
 		TLockFreePointerListFIFO<FUObjectItem, PLATFORM_CACHE_LINE_SIZE> KeepClusterRefsList;
-		FGCArrayStruct** ObjectsToSerializeArrays = new FGCArrayStruct*[NumThreads];
-		for (int32 ThreadIndex = 0; ThreadIndex < NumThreads; ++ThreadIndex)
-		{
-			ObjectsToSerializeArrays[ThreadIndex] = FGCArrayPool::Get().GetArrayStructFromPool();
-		}
+
+		TArray<TArray<UObject*>, TInlineAllocator<32>> ObjectsToSerializeArrays;
+		ObjectsToSerializeArrays.SetNum(NumThreads);
 
 		// Iterate over all objects. Note that we iterate over the UObjectArray and usually check only internal flags which
 		// are part of the array so we don't suffer from cache misses as much as we would if we were to check ObjectFlags.
-		ParallelFor( TEXT("GarbageCollection.PF"),NumThreads,1, [ObjectsToSerializeArrays, &ClustersToDissolveList, &KeepClusterRefsList, FastKeepFlags, KeepFlags, NumberOfObjectsPerThread, NumThreads, MaxNumberOfObjects](int32 ThreadIndex)
+		ParallelFor( TEXT("GC.MarkUnreachable"),NumThreads,1, [&ObjectsToSerializeArrays, &ClustersToDissolveList, &KeepClusterRefsList, FastKeepFlags, KeepFlags, NumberOfObjectsPerThread, NumThreads, MaxNumberOfObjects](int32 ThreadIndex)
 		{
 			TRACE_CPUPROFILER_EVENT_SCOPE(MarkObjectsAsUnreachableTask);
 			int32 FirstObjectIndex = ThreadIndex * NumberOfObjectsPerThread + GUObjectArray.GetFirstGCIndex();
 			int32 NumObjects = (ThreadIndex < (NumThreads - 1)) ? NumberOfObjectsPerThread : (MaxNumberOfObjects - (NumThreads - 1) * NumberOfObjectsPerThread);
 			int32 LastObjectIndex = FMath::Min(GUObjectArray.GetObjectArrayNum() - 1, FirstObjectIndex + NumObjects - 1);
 			int32 ObjectCountDuringMarkPhase = 0;
-			TArray<UObject*>& LocalObjectsToSerialize = ObjectsToSerializeArrays[ThreadIndex]->ObjectsToSerialize;
+			TArray<UObject*>& LocalObjectsToSerialize = ObjectsToSerializeArrays[ThreadIndex];
 
 			for (int32 ObjectIndex = FirstObjectIndex; ObjectIndex <= LastObjectIndex; ++ObjectIndex)
 			{
@@ -1402,7 +3054,8 @@ public:
 					UObject* Object = (UObject*)ObjectItem->Object;
 
 					// We can't collect garbage during an async load operation and by now all unreachable objects should've been purged.
-					checkf(!ObjectItem->HasAnyFlags(EInternalObjectFlags::Unreachable|EInternalObjectFlags::PendingConstruction|EInternalObjectFlags::PersistentGarbage),
+					checkf(	GRedoReachabilityToTrackGarbage ||
+							!ObjectItem->HasAnyFlags(EInternalObjectFlags::Unreachable|EInternalObjectFlags::PendingConstruction),
 							TEXT("Object: '%s' with ObjectFlags=0x%08x and InternalObjectFlags=0x%08x. ")
 							TEXT("State: IsEngineExitRequested=%d, GIsCriticalError=%d, GExitPurge=%d, GObjPurgeIsRequired=%d, GObjIncrementalPurgeIsInProgress=%d, GObjFinishDestroyHasBeenRoutedToAllObjects=%d, GGCObjectsPendingDestructionCount=%d"),
 							*Object->GetFullName(),
@@ -1501,21 +3154,21 @@ public:
 		
 		// Collect all objects to serialize from all threads and put them into a single array
 		{
-			int32 NumObjectsToSerialize = 0;
-			for (int32 ThreadIndex = 0; ThreadIndex < NumThreads; ++ThreadIndex)
+			int32 NumTotal = 0;
+			for (TArray<UObject*>& Objects : ObjectsToSerializeArrays)
 			{
-				NumObjectsToSerialize += ObjectsToSerializeArrays[ThreadIndex]->ObjectsToSerialize.Num();
+				NumTotal += Objects.Num();
 			}
-			ObjectsToSerialize.Reserve(NumObjectsToSerialize);
-			for (int32 ThreadIndex = 0; ThreadIndex < NumThreads; ++ThreadIndex)
+			InitialObjects.Reserve(InitialObjects.Num() + NumTotal + UE::GC::ObjectLookahead);
+			for (TArray<UObject*>& Objects : ObjectsToSerializeArrays)
 			{
-				ObjectsToSerialize.Append(ObjectsToSerializeArrays[ThreadIndex]->ObjectsToSerialize);
-				FGCArrayPool::Get().ReturnToPool(ObjectsToSerializeArrays[ThreadIndex]);
+				InitialObjects.Append(Objects);
 			}
-			delete[] ObjectsToSerializeArrays;
+
+			ObjectsToSerializeArrays.Empty();
 		}
 
-		if (bWithClusters)
+		if constexpr (bWithClusters)
 		{
 			TArray<FUObjectItem*> ClustersToDissolve;
 			ClustersToDissolveList.PopAll(ClustersToDissolve);
@@ -1529,10 +3182,7 @@ public:
 					GUObjectClusters.SetClustersNeedDissolving();
 				}
 			}
-		}
-
-		if (bWithClusters)
-		{
+		
 			TArray<FUObjectItem*> KeepClusterRefs;
 			KeepClusterRefsList.PopAll(KeepClusterRefs);
 			for (FUObjectItem* ObjectItem : KeepClusterRefs)
@@ -1553,7 +3203,7 @@ public:
 						{
 							RootObjectItem->ClearFlags(EInternalObjectFlags::Unreachable);
 							// Make sure all referenced clusters are marked as reachable too
-							FGCReferenceProcessor<EFastReferenceCollectorOptions::WithClusters>::MarkReferencedClustersAsReachable(RootObjectItem->GetClusterIndex(), ObjectsToSerialize);
+							MarkReferencedClustersAsReachable<EGCOptions::WithClusters>(RootObjectItem->GetClusterIndex(), InitialObjects);
 						}
 					}
 				}
@@ -1562,7 +3212,7 @@ public:
 					checkSlow(ObjectItem->HasAnyFlags(EInternalObjectFlags::ClusterRoot));
 					// this thing is definitely not marked unreachable, so don't test it here
 					// Make sure all referenced clusters are marked as reachable too
-					FGCReferenceProcessor<EFastReferenceCollectorOptions::WithClusters>::MarkReferencedClustersAsReachable(ObjectItem->GetClusterIndex(), ObjectsToSerialize);
+					MarkReferencedClustersAsReachable<EGCOptions::WithClusters>(ObjectItem->GetClusterIndex(), InitialObjects);
 				}
 			}
 		}
@@ -1573,56 +3223,61 @@ public:
 	 *
 	 * @param KeepFlags		Objects with these flags will be kept regardless of being referenced or not
 	 */
-	void PerformReachabilityAnalysis(EObjectFlags KeepFlags, const EFastReferenceCollectorOptions InOptions)
+	void PerformReachabilityAnalysis(EObjectFlags KeepFlags, const EGCOptions Options)
 	{
 		LLM_SCOPE(ELLMTag::GC);
 
-		SCOPED_NAMED_EVENT(FRealtimeGC_PerformReachabilityAnalysis, FColor::Red);
-		DECLARE_SCOPE_CYCLE_COUNTER(TEXT("FRealtimeGC::PerformReachabilityAnalysis"), STAT_FArchiveRealtimeGC_PerformReachabilityAnalysis, STATGROUP_GC);
-
-		/** Growing array of objects that require serialization */
-		FGCArrayStruct* ArrayStruct = FGCArrayPool::Get().GetArrayStructFromPool();
-		TArray<UObject*>& ObjectsToSerialize = ArrayStruct->ObjectsToSerialize;
+		BeginInitialReferenceCollection(Options);
 
 		// Reset object count.
 		GObjectCountDuringLastMarkPhase.Reset();
+		
+		InitialObjects.Reset();
 
 		// Make sure GC referencer object is checked for references to other objects even if it resides in permanent object pool
-		if (FPlatformProperties::RequiresCookedData() && FGCObject::GGCObjectReferencer && GUObjectArray.IsDisregardForGC(FGCObject::GGCObjectReferencer))
+		if (FPlatformProperties::RequiresCookedData() && GUObjectArray.IsDisregardForGC(FGCObject::GGCObjectReferencer))
 		{
-			ObjectsToSerialize.Add(FGCObject::GGCObjectReferencer);
+			InitialObjects.Add(FGCObject::GGCObjectReferencer);
 		}
 
 		{
 			const double StartTime = FPlatformTime::Seconds();
 			// Mark phase doesn't care about PendingKill being enabled or not so there's just fewer compiled in functions
-			const EFastReferenceCollectorOptions OptionsForMarkPhase = InOptions & ~EFastReferenceCollectorOptions::WithPendingKill;
-			(this->*MarkObjectsFunctions[GetGCFunctionIndex(OptionsForMarkPhase)])(ObjectsToSerialize, KeepFlags);
-			UE_LOG(LogGarbage, Verbose, TEXT("%f ms for MarkObjectsAsUnreachable Phase (%d Objects To Serialize)"), (FPlatformTime::Seconds() - StartTime) * 1000, ObjectsToSerialize.Num());
+			const EGCOptions OptionsForMarkPhase = Options & ~EGCOptions::WithPendingKill;
+			(this->*MarkObjectsFunctions[GetGCFunctionIndex(OptionsForMarkPhase)])(KeepFlags);
+			UE_LOG(LogGarbage, Verbose, TEXT("%f ms for MarkObjectsAsUnreachable Phase (%d Objects To Serialize)"), (FPlatformTime::Seconds() - StartTime) * 1000, InitialObjects.Num());
 		}
 
 		{
 			const double StartTime = FPlatformTime::Seconds();
-			PerformReachabilityAnalysisOnObjects(ArrayStruct, InOptions);
+		
+			FContextPoolScope Pool;
+			FWorkerContext* Context = Pool.AllocateFromPool();
+			Context->InitialNativeReferences = GetInitialReferences(Options);
+			Context->SetInitialObjectsUnpadded(InitialObjects);
+
+			PerformReachabilityAnalysisOnObjects(Context, Options);
+
+			Context->ResetInitialObjects();
+			Context->InitialNativeReferences = TConstArrayView<UObject**>();
+			Pool.ReturnToPool(Context);
+
 			UE_LOG(LogGarbage, Verbose, TEXT("%f ms for Reachability Analysis"), (FPlatformTime::Seconds() - StartTime) * 1000);
 		}
-        
+
 		// Allowing external systems to add object roots. This can't be done through AddReferencedObjects
 		// because it may require tracing objects (via FGarbageCollectionTracer) multiple times
-		FCoreUObjectDelegates::TraceExternalRootsForReachabilityAnalysis.Broadcast(*this, KeepFlags, !(InOptions & EFastReferenceCollectorOptions::Parallel));
-
-		FGCArrayPool::Get().ReturnToPool(ArrayStruct);
-
-#if UE_BUILD_DEBUG
-		FGCArrayPool::Get().CheckLeaks();
-#endif
+		FCoreUObjectDelegates::TraceExternalRootsForReachabilityAnalysis.Broadcast(*this, KeepFlags, !(Options & EGCOptions::Parallel));
 	}
 
-	virtual void PerformReachabilityAnalysisOnObjects(FGCArrayStruct* ArrayStruct, const EFastReferenceCollectorOptions InOptions) override
+	virtual void PerformReachabilityAnalysisOnObjects(FWorkerContext* Context, EGCOptions Options) override
 	{
-		(this->*ReachabilityAnalysisFunctions[GetGCFunctionIndex(InOptions)])(ArrayStruct);
+		(this->*ReachabilityAnalysisFunctions[GetGCFunctionIndex(Options)])(*Context);
 	}
 };
+
+} // namespace UE::GC
+
 #endif // UE_WITH_GC
 
 // Allow parallel GC to be overridden to single threaded via console command.
@@ -1638,6 +3293,8 @@ static FAutoConsoleVariableRef CVarAllowParallelGC(
 /** Returns true if Garbage Collection should be forced to run on a single thread */
 static bool ShouldForceSingleThreadedGC()
 {
+//	return true;
+
 	const bool bForceSingleThreadedGC = !FApp::ShouldUseThreadingForPerformance() || !FPlatformProcess::SupportsMultithreading() ||
 #if PLATFORM_SUPPORTS_MULTITHREADED_GC
 	(FPlatformMisc::NumberOfCores() < 2 || GAllowParallelGC == 0 || PERF_DETAILED_PER_CLASS_GC_STATS);
@@ -2177,7 +3834,7 @@ void GatherUnreachableObjects(bool bForceSingleThreaded)
 
 	// Iterate over all objects. Note that we iterate over the UObjectArray and usually check only internal flags which
 	// are part of the array so we don't suffer from cache misses as much as we would if we were to check ObjectFlags.
-	ParallelFor( TEXT("GarbageCollection.PF"),NumThreads,1, [&ClusterItemsToDestroy, NumberOfObjectsPerThread, NumThreads, MaxNumberOfObjects](int32 ThreadIndex)
+	ParallelFor( TEXT("GC.GatherUnreachable"),NumThreads,1, [&ClusterItemsToDestroy, NumberOfObjectsPerThread, NumThreads, MaxNumberOfObjects](int32 ThreadIndex)
 	{
 		TRACE_CPUPROFILER_EVENT_SCOPE(GatherUnreachableObjectsTask);
 		int32 FirstObjectIndex = ThreadIndex * NumberOfObjectsPerThread + (GExitPurge ? 0 : GUObjectArray.GetFirstGCIndex());
@@ -2198,7 +3855,6 @@ void GatherUnreachableObjects(bool bForceSingleThreaded)
 					ThisThreadClusterItemsToDestroy.Add(ObjectItem);
 				}
 			}
-			ObjectItem->ClearFlags(EInternalObjectFlags::PersistentGarbage);
 		}
 		if (ThisThreadUnreachableObjects.Num())
 		{
@@ -2242,20 +3898,127 @@ void GatherUnreachableObjects(bool bForceSingleThreaded)
 		ClusterItemsToDestroy.Num());
 }
 
+namespace UE::GC
+{
+
+static void ClearWeakReferences(TConstArrayView<TUniquePtr<FWorkerContext>> Contexts)
+{
+	TRACE_CPUPROFILER_EVENT_SCOPE(ClearWeakReferences);
+	for (const TUniquePtr<FWorkerContext>& Context : Contexts)
+	{
+		for (UObject** WeakReference : Context->WeakReferences)
+		{
+			UObject*& ReferencedObject = *WeakReference;
+			if (ReferencedObject && ReferencedObject->IsUnreachable())
+			{
+				ReferencedObject = nullptr;
+			}
+		}
+		Context->WeakReferences.Reset();
+	}
+}
+
+static void DumpGarbageReferencers(TConstArrayView<TUniquePtr<FWorkerContext>> Contexts)
+{	
+#if !UE_BUILD_SHIPPING
+	if (GGarbageReferenceTrackingEnabled == 0)
+	{
+		return;
+	}
+
+	// We don't care about leaks when engine exit was requested since the final GC pass will destroy everything anyway
+	// We still want to clear the GarbageReferences array though
+	if (IsEngineExitRequested())
+	{
+		for (const TUniquePtr<FWorkerContext>& Context : Contexts)
+		{
+			Context->GarbageReferences.Reset();
+		}
+		return;
+	}
+
+	// Dump all garbage references in detail
+	if (GGarbageReferenceTrackingEnabled == 1)
+	{
+		int32 NumGarbageReferences = 0;
+		for (const TUniquePtr<FWorkerContext>& Context : Contexts)
+		{
+			for (FGarbageReferenceInfo& GarbageReference : Context->GarbageReferences)
+			{
+				UE_LOG(LogGarbage, Warning, TEXT("Reachable garbage object: %s"), *GarbageReference.GarbageObject->GetFullName());
+				UE_LOG(LogGarbage, Warning, TEXT("Referenced by:            %s"), *GarbageReference.GetReferencingObjectInfo());
+				UE_LOG(LogGarbage, Warning, TEXT(""));
+				NumGarbageReferences++;
+			}
+			Context->GarbageReferences.Reset();
+		}
+		
+		UE_LOG(LogGarbage, Log, TEXT("Reported %d garbage references"), NumGarbageReferences);
+	}
+	// Dump one reference for each source property
+	else if (GGarbageReferenceTrackingEnabled == 2)
+	{
+		typedef TPair<void*, FName> FKey;
+		TSet<FKey> Seen;
+		int32 TotalGarbageReferences = 0, ReportedGarbageReferences = 0;
+		for (const TUniquePtr<FWorkerContext>& Context : Contexts)
+		{
+			for (FGarbageReferenceInfo& GarbageReference : Context->GarbageReferences)
+			{
+				FKey Key(GarbageReference.bReferencerUObject ? (void*)GarbageReference.Referencer.Object->GetClass() : (void*)  GarbageReference.Referencer.GCObject, GarbageReference.PropertyName);
+				if (Seen.Contains(Key) == false)
+				{
+					UE_LOG(LogGarbage, Warning, TEXT("Reachable garbage object: %s"), *GarbageReference.GarbageObject->GetFullName());
+					UE_LOG(LogGarbage, Warning, TEXT("Referenced by:            %s"), *GarbageReference.GetReferencingObjectInfo());
+					UE_LOG(LogGarbage, Warning, TEXT(""));
+					Seen.Add(Key);
+					ReportedGarbageReferences++;
+				}
+				TotalGarbageReferences++;
+			}
+			Context->GarbageReferences.Reset();
+		}
+		UE_LOG(LogGarbage, Log, TEXT("Reported %d/%d garbage references"), ReportedGarbageReferences, TotalGarbageReferences);
+	}
+#endif // !UE_BUILD_SHIPPING
+}
+
+static void UpdateGCHistory(TConstArrayView<TUniquePtr<FWorkerContext>> Contexts)
+{
+#if ENABLE_GC_HISTORY
+	FGCHistory::Get().Update(Contexts);
+
+	for (const TUniquePtr<FWorkerContext>& Context : Contexts)
+	{
+		for (TPair<const UObject*, TArray<FGCDirectReference>*>& DirectReferenceInfos : Context->History)
+		{
+			delete DirectReferenceInfos.Value;
+		}
+		Context->History.Reset();
+	}
+#endif // ENABLE_GC_HISTORY
+}
+
+#if UE_WITH_GC
 /** 
  * Deletes all unreferenced objects, keeping objects that have any of the passed in KeepFlags set
  *
  * @param	KeepFlags			objects with those flags will be kept regardless of being referenced or not
  * @param	bPerformFullPurge	if true, perform a full purge after the mark pass
  */
-void CollectGarbageInternal(EObjectFlags KeepFlags, bool bPerformFullPurge)
+static void CollectGarbageInternal(EObjectFlags KeepFlags, bool bPerformFullPurge)
 {
-#if UE_WITH_GC
+	if (GIsInitialLoad)
+	{
+		// During initial load classes may not yet have their GC token streams assembled
+		UE_LOG(LogGarbage, Log, TEXT("Skipping CollectGarbage() call during initial load. It's not safe."));
+		return;
+	}
+
 	SCOPE_TIME_GUARD(TEXT("Collect Garbage"));
 	SCOPED_NAMED_EVENT(CollectGarbageInternal, FColor::Red);
 	CSV_EVENT_GLOBAL(TEXT("GC"));
 	CSV_SCOPED_TIMING_STAT_EXCLUSIVE(GarbageCollection);
-	LLM_SCOPE(ELLMTag::GC);
 
 	FGCCSyncObject::Get().ResetGCIsWaiting();
 
@@ -2291,6 +4054,8 @@ void CollectGarbageInternal(EObjectFlags KeepFlags, bool bPerformFullPurge)
 	GLastGCFrame = GFrameCounter;
 
 	{
+		LLM_SCOPE(ELLMTag::GC);
+
 		// Set 'I'm garbage collecting' flag - might be checked inside various functions.
 		// This has to be unlocked before we call post GC callbacks
 		FGCLegacyHashTableScopeLock GCLegacyHashTableLock;
@@ -2343,29 +4108,38 @@ void CollectGarbageInternal(EObjectFlags KeepFlags, bool bPerformFullPurge)
 			}
 #endif
 
-			const EFastReferenceCollectorOptions Options = 
+			const EGCOptions Options = 
 				// Fall back to single threaded GC if processor count is 1 or parallel GC is disabled
 				// or detailed per class gc stats are enabled (not thread safe)
-				(ShouldForceSingleThreadedGC() ? EFastReferenceCollectorOptions::None : EFastReferenceCollectorOptions::Parallel) |
+				(ShouldForceSingleThreadedGC() ? EGCOptions::None : EGCOptions::Parallel) |
 				// Run with GC clustering code enabled only if clustering is enabled and there's actual allocated clusters
-				((!!GCreateGCClusters && GUObjectClusters.GetNumAllocatedClusters()) ? EFastReferenceCollectorOptions::WithClusters : EFastReferenceCollectorOptions::None) |
+				((!!GCreateGCClusters && GUObjectClusters.GetNumAllocatedClusters()) ? EGCOptions::WithClusters : EGCOptions::None) |
 				// Toggle between PendingKill enabled or disabled
-				(UObjectBaseUtility::IsPendingKillEnabled() ? EFastReferenceCollectorOptions::WithPendingKill : EFastReferenceCollectorOptions::None);
+				(UObjectBaseUtility::IsPendingKillEnabled() ? EGCOptions::WithPendingKill : EGCOptions::None);
 
 			// Perform reachability analysis.
 			{
+				SCOPED_NAMED_EVENT(FRealtimeGC_PerformReachabilityAnalysis, FColor::Red);
+				DECLARE_SCOPE_CYCLE_COUNTER(TEXT("FRealtimeGC::PerformReachabilityAnalysis"), STAT_FArchiveRealtimeGC_PerformReachabilityAnalysis, STATGROUP_GC);
 				const double StartTime = FPlatformTime::Seconds();
-				FRealtimeGC TagUsedRealtimeGC;
-
-				TagUsedRealtimeGC.PerformReachabilityAnalysis(KeepFlags, Options);
+				FRealtimeGC().PerformReachabilityAnalysis(KeepFlags, Options);
 				UE_LOG(LogGarbage, Log, TEXT("%f ms for GC"), (FPlatformTime::Seconds() - StartTime) * 1000);
 			}
 
-			FGCArrayPool& ArrayPool = FGCArrayPool::Get();
-			TArray<FGCArrayStruct*> AllArrays;
-			ArrayPool.GetAllArrayStructsFromPool(AllArrays);
+			if (GRedoReachabilityToTrackGarbage && GGarbageReferenceTrackingEnabled > 0)
+			{
+				SCOPED_NAMED_EVENT(FRealtimeGC_PerformReachabilityAnalysisRerun, FColor::Orange);
+				DECLARE_SCOPE_CYCLE_COUNTER(TEXT("FRealtimeGC::PerformReachabilityAnalysisRerun"), STAT_FArchiveRealtimeGC_PerformReachabilityAnalysisRerun, STATGROUP_GC);
+				const double StartTime = FPlatformTime::Seconds();
+				FRealtimeGC().PerformReachabilityAnalysis(KeepFlags, Options);
+				UE_LOG(LogGarbage, Log, TEXT("%f ms for GC rerun to track garbage references (gc.GarbageReferenceTrackingEnabled=%d)"), (FPlatformTime::Seconds() - StartTime) * 1000, GGarbageReferenceTrackingEnabled);
+			}
+
+
+			FContextPoolScope ContextPool;
+			TConstArrayView<TUniquePtr<FWorkerContext>> AllContexts = ContextPool.PeekFree();
 			// This needs to happen before clusters get dissolved otherwisise cluster information will be missing from history
-			ArrayPool.UpdateGCHistory(AllArrays);
+			UpdateGCHistory(AllContexts);
 
 			// Reconstruct clusters if needed
 			if (GUObjectClusters.ClustersNeedDissolving())
@@ -2375,31 +4149,19 @@ void CollectGarbageInternal(EObjectFlags KeepFlags, bool bPerformFullPurge)
 				UE_LOG(LogGarbage, Log, TEXT("%f ms for dissolving GC clusters"), (FPlatformTime::Seconds() - StartTime) * 1000);
 			}
 
-			ArrayPool.DumpGarbageReferencers(AllArrays);
+			DumpGarbageReferencers(AllContexts);
+
+			GatherUnreachableObjects(!(Options & EGCOptions::Parallel));
 
-			GatherUnreachableObjects(!(Options & EFastReferenceCollectorOptions::Parallel));
 
 			// This needs to happen after GatherUnreachableObjects since it can mark more objects as unreachable
-			ArrayPool.ClearWeakReferences(AllArrays);
+			ClearWeakReferences(AllContexts);
 
-			// Now return arrays back to the pool and free some memory if requested
-			for (int32 Index = 0; Index < AllArrays.Num(); ++Index)
+			if (bPerformFullPurge)
 			{
-				FGCArrayStruct* ArrayStruct = AllArrays[Index];
-				if (bPerformFullPurge
-					|| Index % 7 == 3) // delete 1/7th of them just to keep things from growing too much between full purges
-				{
-					ArrayPool.FreeArrayStruct(ArrayStruct);
-				}
-				else
-				{
-					ArrayPool.ReturnToPool(ArrayStruct);
-				}
+				ContextPool.Cleanup();
 			}
 
-			// Make sure nothing will be using potentially freed arrays
-			AllArrays.Empty();
-
 			NotifyUnreachableObjects(GUnreachableObjects);
 		}
 		if (GGCLockBehavior == FGCLockBehavior::Default)
@@ -2455,8 +4217,10 @@ void CollectGarbageInternal(EObjectFlags KeepFlags, bool bPerformFullPurge)
 
 	GLastGCTime = FPlatformTime::Seconds();
 	STAT_ADD_CUSTOMMESSAGE_NAME( STAT_NamedMarker, TEXT( "GarbageCollection - End" ) );
-#endif	// UE_WITH_GC
 }
+#endif // !UE_WITH_GC
+
+} // namespace UE::GC
 
 FString FGarbageReferenceInfo::GetReferencingObjectInfo() const
 {
@@ -2475,105 +4239,6 @@ FString FGarbageReferenceInfo::GetReferencingObjectInfo() const
 	}
 }
 
-void FGCArrayPool::DumpGarbageReferencers(TArray<FGCArrayStruct*>& AllArrays)
-{		
-	if (GGarbageReferenceTrackingEnabled == 0)
-	{
-		return;
-	}
-
-#if !UE_BUILD_SHIPPING
-	// We don't care about leaks when engine exit was requested since the final GC pass will destroy everything anyway
-	// We still want to clear the GarbageReferences array though
-	if (IsEngineExitRequested())
-	{
-		for (FGCArrayStruct* ArrayStruct : AllArrays)
-		{
-			ArrayStruct->GarbageReferences.Reset();
-		}
-		return;
-	}
-
-	double StartTime = FPlatformTime::Seconds();
-	// Dump all garbage references in detail
-	if (GGarbageReferenceTrackingEnabled == 1)
-	{
-		int32 NumGarbageReferences = 0;
-		for (FGCArrayStruct* ArrayStruct : AllArrays)
-		{
-			for (FGarbageReferenceInfo& GarbageReference : ArrayStruct->GarbageReferences)
-			{
-				// No need to spam with references that would never get released anyway (even if we still had PendingKill) because they were referenced by a persistent reference 
-				// Disable checking this flag because it's too aggressive at removing references from reporting.
-				// if (!GarbageReference.GarbageObject->HasAnyInternalFlags(EInternalObjectFlags::PersistentGarbage))
-				{
-					UE_LOG(LogGarbage, Warning, TEXT("Reachable garbage object: %s"), *GarbageReference.GarbageObject->GetFullName());
-					UE_LOG(LogGarbage, Warning, TEXT("Referenced by:            %s"), *GarbageReference.GetReferencingObjectInfo());
-					UE_LOG(LogGarbage, Warning, TEXT(""));
-					NumGarbageReferences++;
-				}
-			}
-			if (ArrayStruct->GarbageReferences.Num())
-			{
-				FCoreUObjectDelegates::GetGarbageCollectReportGarbageReferencers().Broadcast(ArrayStruct->GarbageReferences);
-				ArrayStruct->GarbageReferences.Reset();
-			}
-		}
-		UE_CLOG(NumGarbageReferences > 0, LogGarbage, Log, TEXT("Found %d garbage references in %f ms."), NumGarbageReferences, (FPlatformTime::Seconds() - StartTime) * 1000);
-	}
-	// Dump one reference for each source property
-	else if (GGarbageReferenceTrackingEnabled == 2)
-	{
-		typedef TPair<void*, FName> FKey;
-		TSet<FKey> Seen;
-		int32 TotalGarbageReferences = 0, ReportedGarbageReferences = 0;
-		for (FGCArrayStruct* ArrayStruct : AllArrays)
-		{
-			for (FGarbageReferenceInfo& GarbageReference : ArrayStruct->GarbageReferences)
-			{
-				// No need to spam with references that would never get released anyway (even if we still had PendingKill) because they were referenced by a persistent reference 
-				// Disable checking this flag because it's too aggressive at removing references from reporting.
-				// if (!GarbageReference.GarbageObject->HasAnyInternalFlags(EInternalObjectFlags::PersistentGarbage))
-				{
-					FKey Key(GarbageReference.bReferencerUObject ? (void*)GarbageReference.Referencer.Object->GetClass() : (void*)  GarbageReference.Referencer.GCObject, GarbageReference.PropertyName);
-					if (Seen.Contains(Key) == false)
-					{
-						UE_LOG(LogGarbage, Warning, TEXT("Reachable garbage object: %s"), *GarbageReference.GarbageObject->GetFullName());
-						UE_LOG(LogGarbage, Warning, TEXT("Referenced by:            %s"), *GarbageReference.GetReferencingObjectInfo());
-						UE_LOG(LogGarbage, Warning, TEXT(""));
-						Seen.Add(Key);
-						ReportedGarbageReferences++;
-					}
-					TotalGarbageReferences++;
-				}
-			}
-			if (ArrayStruct->GarbageReferences.Num())
-			{
-				FCoreUObjectDelegates::GetGarbageCollectReportGarbageReferencers().Broadcast(ArrayStruct->GarbageReferences);
-				ArrayStruct->GarbageReferences.Reset();
-			}
-		}
-		UE_CLOG(TotalGarbageReferences > 0, LogGarbage, Log, TEXT("Reported %d/%d garbage references in %f ms. (The rest were omitted as duplicate referencing class/property.)"), ReportedGarbageReferences, TotalGarbageReferences, (FPlatformTime::Seconds() - StartTime) * 1000);
-	}
-#endif
-}
-
-void FGCArrayPool::UpdateGCHistory(TArray<FGCArrayStruct*>& AllArrays)
-{
-#if ENABLE_GC_HISTORY
-	FGCHistory::Get().Update(AllArrays);
-
-	for (FGCArrayStruct* ArrayStruct : AllArrays)
-	{
-		for (TPair<UObject*, TArray<FGCDirectReference>*>& DirectReferenceInfos : ArrayStruct->History)
-		{
-			delete DirectReferenceInfos.Value;
-		}
-		ArrayStruct->History.Reset();
-	}
-#endif // ENABLE_GC_HISTORY
-}
-
 double GetLastGCTime()
 {
 	return GLastGCTime;
@@ -2658,21 +4323,11 @@ bool UnhashUnreachableObjects(bool bUseTimeLimit, double TimeLimit)
 
 void CollectGarbage(EObjectFlags KeepFlags, bool bPerformFullPurge)
 {
-#if !UE_WITH_GC
-	return;
-#else
-	if (GIsInitialLoad)
-	{
-		// During initial load classes may not yet have their GC token streams assembled
-		UE_LOG(LogGarbage, Log, TEXT("Skipping CollectGarbage() call during initial load. It's not safe."));
-		return;
-	}
-
 	// No other thread may be performing UObject operations while we're running
 	AcquireGCLock();
 
 	// Perform actual garbage collection
-	CollectGarbageInternal(KeepFlags, bPerformFullPurge);
+	UE::GC::CollectGarbageInternal(KeepFlags, bPerformFullPurge);
 
 	if (GGCLockBehavior == FGCLockBehavior::Legacy)
 	{
@@ -2680,21 +4335,10 @@ void CollectGarbage(EObjectFlags KeepFlags, bool bPerformFullPurge)
 		ReleaseGCLock();
 	}
 	// With the new FGCLockBehavior::Default behavior the lock was released after reachability analysis inside CollectGarbageInternal
-#endif
 }
 
 bool TryCollectGarbage(EObjectFlags KeepFlags, bool bPerformFullPurge)
 {
-#if !UE_WITH_GC
-	return false;
-#else
-	if (GIsInitialLoad)
-	{
-		// During initial load classes may not yet have their GC token streams assembled
-		UE_LOG(LogGarbage, Log, TEXT("Skipping TryCollectGarbage() call during initial load. It's not safe."));
-		return false;
-	}
-
 	// No other thread may be performing UObject operations while we're running
 	bool bCanRunGC = FGCCSyncObject::Get().TryGCLock();
 	if (!bCanRunGC)
@@ -2709,9 +4353,9 @@ bool TryCollectGarbage(EObjectFlags KeepFlags, bool bPerformFullPurge)
 		}
 	}
 	if (bCanRunGC)
-	{
+	{ 
 		// Perform actual garbage collection
-		CollectGarbageInternal(KeepFlags, bPerformFullPurge);
+		UE::GC::CollectGarbageInternal(KeepFlags, bPerformFullPurge);
 
 		if (GGCLockBehavior == FGCLockBehavior::Legacy)
 		{
@@ -2726,7 +4370,6 @@ bool TryCollectGarbage(EObjectFlags KeepFlags, bool bPerformFullPurge)
 	}
 
 	return bCanRunGC;
-#endif
 }
 
 void UObject::CallAddReferencedObjects(FReferenceCollector& Collector)
@@ -2749,10 +4392,8 @@ void UObject::AddReferencedObjects(UObject* This, FReferenceCollector& Collector
 		UObject* LoadOuter = This->GetOuter();
 		UClass* Class = This->GetClass();
 		UPackage* Package = This->GetExternalPackageInternal();
-		Collector.AllowEliminatingReferences(false);
 		Collector.AddReferencedObject(LoadOuter, This);
 		Collector.AddReferencedObject(Package, This);
-		Collector.AllowEliminatingReferences(true);
 		Collector.AddReferencedObject(Class, This);
 	}
 #endif
@@ -2874,10 +4515,12 @@ bool FMulticastDelegateProperty::ContainsObjectReference(TArray<const FStructPro
 	return !!(InReferenceType & EPropertyObjectReferenceType::Weak);
 }
 
+namespace UE::GC
+{
 /**
  * Scope helper structure to emit tokens for fixed arrays in the case of ArrayDim (passed in count) being > 1.
  */
-struct FGCReferenceFixedArrayTokenHelper
+struct FFixedArrayTokenHelper
 {
 	/**
 	 * Constructor, emitting necessary tokens for fixed arrays if count > 1 and also keeping track of count so 
@@ -2889,47 +4532,49 @@ struct FGCReferenceFixedArrayTokenHelper
 	 * @param InStride					array type stride (e.g. sizeof(struct) or sizeof(UObject*))
 	 * @param InProperty                the property this array represents
 	 */
-	FGCReferenceFixedArrayTokenHelper(UClass& OwnerClass, int32 InOffset, int32 InCount, int32 InStride, const FProperty& InProperty, FGCStackSizeHelper& InStackSizeHelper)
-		: ReferenceTokenStream(&OwnerClass.ReferenceTokenStream)
+	FFixedArrayTokenHelper(FTokenStreamBuilder& TokenStream,  int32 InOffset, int32 InCount, int32 InStride, const FProperty& InProperty, FGCStackSizeHelper& InStackSizeHelper)
+		: ReferenceTokenStream(TokenStream)
 		, StackSizeHelper(InStackSizeHelper)
 		, Count(InCount)
 	{
 		if( InCount > 1 )
 		{
-			OwnerClass.EmitObjectReference(InOffset, *StackSizeHelper.GetPropertyPath(), GCRT_FixedArray);
+			ReferenceTokenStream.EmitObjectReference(InOffset, *StackSizeHelper.GetPropertyPath(), GCRT_FixedArray);
 
-			OwnerClass.ReferenceTokenStream.EmitStride(InStride);
-			OwnerClass.ReferenceTokenStream.EmitCount(InCount);
+			ReferenceTokenStream.EmitStride(InStride);
+			ReferenceTokenStream.EmitCount(InCount);
 			// GCRT_FixedArray pushes a new stack frame in TFastReferenceCollector
 			StackSizeHelper.Push(&InProperty);
 		}
 	}
 
 	/** Destructor, emitting return if ArrayDim > 1 */
-	~FGCReferenceFixedArrayTokenHelper()
+	~FFixedArrayTokenHelper()
 	{
 		if( Count > 1 )
 		{
-			ReferenceTokenStream->EmitReturn();
+			ReferenceTokenStream.EmitReturn();
 			StackSizeHelper.Pop();
 		}
 	}
 
 private:
 	/** Reference token stream used to emit to */
-	FGCReferenceTokenStream*	ReferenceTokenStream;
+	FTokenStreamBuilder&	ReferenceTokenStream;
 	/** GC token stream stack size helper */
 	FGCStackSizeHelper&			StackSizeHelper;
 	/** Size of fixed array */
 	int32							Count;
 };
 
+} // namespace UE::GC
+
 
 /**
  * Emits tokens used by realtime garbage collection code to passed in ReferenceTokenStream. The offset emitted is relative
  * to the passed in BaseOffset which is used by e.g. arrays of structs.
  */
-void FProperty::EmitReferenceInfo(UClass& OwnerClass, int32 BaseOffset, TArray<const FStructProperty*>& EncounteredStructProps, FGCStackSizeHelper& StackSizeHelper)
+void FProperty::EmitReferenceInfo(UE::GC::FTokenStreamBuilder& TokenStream, int32 BaseOffset, TArray<const FStructProperty*>& EncounteredStructProps, FGCStackSizeHelper& StackSizeHelper)
 {
 }
 
@@ -2937,43 +4582,43 @@ void FProperty::EmitReferenceInfo(UClass& OwnerClass, int32 BaseOffset, TArray<c
  * Emits tokens used by realtime garbage collection code to passed in OwnerClass' ReferenceTokenStream. The offset emitted is relative
  * to the passed in BaseOffset which is used by e.g. arrays of structs.
  */
-void FObjectProperty::EmitReferenceInfo(UClass& OwnerClass, int32 BaseOffset, TArray<const FStructProperty*>& EncounteredStructProps, FGCStackSizeHelper& StackSizeHelper)
+void FObjectProperty::EmitReferenceInfo(UE::GC::FTokenStreamBuilder& TokenStream, int32 BaseOffset, TArray<const FStructProperty*>& EncounteredStructProps, FGCStackSizeHelper& StackSizeHelper)
 {
-	FGCReferenceFixedArrayTokenHelper FixedArrayHelper(OwnerClass, BaseOffset + GetOffset_ForGC(), ArrayDim, sizeof(UObject*), *this, StackSizeHelper);
-	OwnerClass.EmitObjectReference(BaseOffset + GetOffset_ForGC(), *StackSizeHelper.GetPropertyPath(), GCRT_Object);
+	UE::GC::FFixedArrayTokenHelper FixedArrayHelper(TokenStream, BaseOffset + GetOffset_ForGC(), ArrayDim, sizeof(UObject*), *this, StackSizeHelper);
+	TokenStream.EmitObjectReference(BaseOffset + GetOffset_ForGC(), *StackSizeHelper.GetPropertyPath(), GCRT_Object);
 }
 
-void FWeakObjectProperty::EmitReferenceInfo(UClass& OwnerClass, int32 BaseOffset, TArray<const FStructProperty*>& EncounteredStructProps, FGCStackSizeHelper& StackSizeHelper)
+void FWeakObjectProperty::EmitReferenceInfo(UE::GC::FTokenStreamBuilder& TokenStream, int32 BaseOffset, TArray<const FStructProperty*>& EncounteredStructProps, FGCStackSizeHelper& StackSizeHelper)
 {
-	FGCReferenceFixedArrayTokenHelper FixedArrayHelper(OwnerClass, BaseOffset + GetOffset_ForGC(), ArrayDim, sizeof(FWeakObjectPtr), *this, StackSizeHelper);
-	OwnerClass.EmitObjectReference(BaseOffset + GetOffset_ForGC(), *StackSizeHelper.GetPropertyPath(), GCRT_WeakObject);
+	UE::GC::FFixedArrayTokenHelper FixedArrayHelper(TokenStream, BaseOffset + GetOffset_ForGC(), ArrayDim, sizeof(FWeakObjectPtr), *this, StackSizeHelper);
+	TokenStream.EmitObjectReference(BaseOffset + GetOffset_ForGC(), *StackSizeHelper.GetPropertyPath(), GCRT_WeakObject);
 }
-void FLazyObjectProperty::EmitReferenceInfo(UClass& OwnerClass, int32 BaseOffset, TArray<const FStructProperty*>& EncounteredStructProps, FGCStackSizeHelper& StackSizeHelper)
+void FLazyObjectProperty::EmitReferenceInfo(UE::GC::FTokenStreamBuilder& TokenStream, int32 BaseOffset, TArray<const FStructProperty*>& EncounteredStructProps, FGCStackSizeHelper& StackSizeHelper)
 {
-	FGCReferenceFixedArrayTokenHelper FixedArrayHelper(OwnerClass, BaseOffset + GetOffset_ForGC(), ArrayDim, sizeof(FLazyObjectPtr), *this, StackSizeHelper);
-	OwnerClass.EmitObjectReference(BaseOffset + GetOffset_ForGC(), *StackSizeHelper.GetPropertyPath(), GCRT_LazyObject);
+	UE::GC::FFixedArrayTokenHelper FixedArrayHelper(TokenStream, BaseOffset + GetOffset_ForGC(), ArrayDim, sizeof(FLazyObjectPtr), *this, StackSizeHelper);
+	TokenStream.EmitObjectReference(BaseOffset + GetOffset_ForGC(), *StackSizeHelper.GetPropertyPath(), GCRT_LazyObject);
 }
-void FSoftObjectProperty::EmitReferenceInfo(UClass& OwnerClass, int32 BaseOffset, TArray<const FStructProperty*>& EncounteredStructProps, FGCStackSizeHelper& StackSizeHelper)
+void FSoftObjectProperty::EmitReferenceInfo(UE::GC::FTokenStreamBuilder& TokenStream, int32 BaseOffset, TArray<const FStructProperty*>& EncounteredStructProps, FGCStackSizeHelper& StackSizeHelper)
 {
-	FGCReferenceFixedArrayTokenHelper FixedArrayHelper(OwnerClass, BaseOffset + GetOffset_ForGC(), ArrayDim, sizeof(FSoftObjectPtr), *this, StackSizeHelper);
-	OwnerClass.EmitObjectReference(BaseOffset + GetOffset_ForGC(), *StackSizeHelper.GetPropertyPath(), GCRT_SoftObject);
+	UE::GC::FFixedArrayTokenHelper FixedArrayHelper(TokenStream, BaseOffset + GetOffset_ForGC(), ArrayDim, sizeof(FSoftObjectPtr), *this, StackSizeHelper);
+	TokenStream.EmitObjectReference(BaseOffset + GetOffset_ForGC(), *StackSizeHelper.GetPropertyPath(), GCRT_SoftObject);
 }
-void FDelegateProperty::EmitReferenceInfo(UClass& OwnerClass, int32 BaseOffset, TArray<const FStructProperty*>& EncounteredStructProps, FGCStackSizeHelper& StackSizeHelper)
+void FDelegateProperty::EmitReferenceInfo(UE::GC::FTokenStreamBuilder& TokenStream, int32 BaseOffset, TArray<const FStructProperty*>& EncounteredStructProps, FGCStackSizeHelper& StackSizeHelper)
 {
-	FGCReferenceFixedArrayTokenHelper FixedArrayHelper(OwnerClass, BaseOffset + GetOffset_ForGC(), ArrayDim, this->ElementSize, *this, StackSizeHelper);
-	OwnerClass.EmitObjectReference(BaseOffset + GetOffset_ForGC(), *StackSizeHelper.GetPropertyPath(), GCRT_Delegate);
+	UE::GC::FFixedArrayTokenHelper FixedArrayHelper(TokenStream, BaseOffset + GetOffset_ForGC(), ArrayDim, this->ElementSize, *this, StackSizeHelper);
+	TokenStream.EmitObjectReference(BaseOffset + GetOffset_ForGC(), *StackSizeHelper.GetPropertyPath(), GCRT_Delegate);
 }
-void FMulticastDelegateProperty::EmitReferenceInfo(UClass& OwnerClass, int32 BaseOffset, TArray<const FStructProperty*>& EncounteredStructProps, FGCStackSizeHelper& StackSizeHelper)
+void FMulticastDelegateProperty::EmitReferenceInfo(UE::GC::FTokenStreamBuilder& TokenStream, int32 BaseOffset, TArray<const FStructProperty*>& EncounteredStructProps, FGCStackSizeHelper& StackSizeHelper)
 {
-	FGCReferenceFixedArrayTokenHelper FixedArrayHelper(OwnerClass, BaseOffset + GetOffset_ForGC(), ArrayDim, this->ElementSize, *this, StackSizeHelper);
-	OwnerClass.EmitObjectReference(BaseOffset + GetOffset_ForGC(), *StackSizeHelper.GetPropertyPath(), GCRT_MulticastDelegate);
+	UE::GC::FFixedArrayTokenHelper FixedArrayHelper(TokenStream, BaseOffset + GetOffset_ForGC(), ArrayDim, this->ElementSize, *this, StackSizeHelper);
+	TokenStream.EmitObjectReference(BaseOffset + GetOffset_ForGC(), *StackSizeHelper.GetPropertyPath(), GCRT_MulticastDelegate);
 }
 
 /**
  * Emits tokens used by realtime garbage collection code to passed in OwnerClass' ReferenceTokenStream. The offset emitted is relative
  * to the passed in BaseOffset which is used by e.g. arrays of structs.
  */
-void FArrayProperty::EmitReferenceInfo(UClass& OwnerClass, int32 BaseOffset, TArray<const FStructProperty*>& EncounteredStructProps, FGCStackSizeHelper& StackSizeHelper)
+void FArrayProperty::EmitReferenceInfo(UE::GC::FTokenStreamBuilder& TokenStream, int32 BaseOffset, TArray<const FStructProperty*>& EncounteredStructProps, FGCStackSizeHelper& StackSizeHelper)
 {
 	if (Inner->ContainsObjectReference(EncounteredStructProps, EPropertyObjectReferenceType::Strong | EPropertyObjectReferenceType::Weak))
 	{
@@ -2984,60 +4629,60 @@ void FArrayProperty::EmitReferenceInfo(UClass& OwnerClass, int32 BaseOffset, TAr
 		// which is exactly what is required for nested arrays to work
 		if( Inner->IsA(FObjectProperty::StaticClass()) || Inner->IsA(FObjectPtrProperty::StaticClass()) )
 		{
-			OwnerClass.EmitObjectReference(BaseOffset + GetOffset_ForGC(), *StackSizeHelper.GetPropertyPath(), bUsesFreezableAllocator ? GCRT_ArrayObjectFreezable : GCRT_ArrayObject);
+			TokenStream.EmitObjectReference(BaseOffset + GetOffset_ForGC(), *StackSizeHelper.GetPropertyPath(), bUsesFreezableAllocator ? GCRT_ArrayObjectFreezable : GCRT_ArrayObject);
 		}
 		else if( Inner->IsA(FInterfaceProperty::StaticClass()) )
 		{
-			OwnerClass.EmitObjectReference(BaseOffset + GetOffset_ForGC(), *StackSizeHelper.GetPropertyPath(), bUsesFreezableAllocator ? GCRT_ArrayStructFreezable : GCRT_ArrayStruct);
+			TokenStream.EmitObjectReference(BaseOffset + GetOffset_ForGC(), *StackSizeHelper.GetPropertyPath(), bUsesFreezableAllocator ? GCRT_ArrayStructFreezable : GCRT_ArrayStruct);
 
 			// GCRT_ArrayStruct and GCRT_ArrayStructFreezable push a new stack frame in TFastReferenceCollector
 			FGCStackSizeHelperScope StackSizeScope(StackSizeHelper, Inner);
 
-			OwnerClass.ReferenceTokenStream.EmitStride(Inner->ElementSize);
-			const uint32 SkipIndexIndex = OwnerClass.ReferenceTokenStream.EmitSkipIndexPlaceholder();
+			TokenStream.EmitStride(Inner->ElementSize);
+			const uint32 SkipIndexIndex = TokenStream.EmitSkipIndexPlaceholder();
 
-			OwnerClass.EmitObjectReference(0, *StackSizeHelper.GetPropertyPath(), GCRT_Object);
+			TokenStream.EmitObjectReference(0, *StackSizeHelper.GetPropertyPath(), GCRT_Object);
 
-			const uint32 SkipIndex = OwnerClass.ReferenceTokenStream.EmitReturn();
-			OwnerClass.ReferenceTokenStream.UpdateSkipIndexPlaceholder(SkipIndexIndex, SkipIndex);
+			const uint32 SkipIndex = TokenStream.EmitReturn();
+			TokenStream.UpdateSkipIndexPlaceholder(SkipIndexIndex, SkipIndex);
 		}
 		else if (Inner->IsA(FFieldPathProperty::StaticClass()))
 		{
-			OwnerClass.EmitObjectReference(BaseOffset + GetOffset_ForGC(), *StackSizeHelper.GetPropertyPath(), GCRT_ArrayAddFieldPathReferencedObject);
+			TokenStream.EmitObjectReference(BaseOffset + GetOffset_ForGC(), *StackSizeHelper.GetPropertyPath(), GCRT_ArrayAddFieldPathReferencedObject);
 		}
 		else if (Inner->IsA(FWeakObjectProperty::StaticClass()))
 		{
-			OwnerClass.EmitObjectReference(BaseOffset + GetOffset_ForGC(), *StackSizeHelper.GetPropertyPath(), GCRT_ArrayWeakObject);
+			TokenStream.EmitObjectReference(BaseOffset + GetOffset_ForGC(), *StackSizeHelper.GetPropertyPath(), GCRT_ArrayWeakObject);
 		}
 		else if (Inner->IsA(FLazyObjectProperty::StaticClass()))
 		{
-			OwnerClass.EmitObjectReference(BaseOffset + GetOffset_ForGC(), *StackSizeHelper.GetPropertyPath(), GCRT_ArrayLazyObject);
+			TokenStream.EmitObjectReference(BaseOffset + GetOffset_ForGC(), *StackSizeHelper.GetPropertyPath(), GCRT_ArrayLazyObject);
 		}
 		else if (Inner->IsA(FSoftObjectProperty::StaticClass()))
 		{
-			OwnerClass.EmitObjectReference(BaseOffset + GetOffset_ForGC(), *StackSizeHelper.GetPropertyPath(), GCRT_ArraySoftObject);
+			TokenStream.EmitObjectReference(BaseOffset + GetOffset_ForGC(), *StackSizeHelper.GetPropertyPath(), GCRT_ArraySoftObject);
 		}
 		else if (Inner->IsA(FDelegateProperty::StaticClass()))
 		{
-			OwnerClass.EmitObjectReference(BaseOffset + GetOffset_ForGC(), *StackSizeHelper.GetPropertyPath(), GCRT_ArrayDelegate);
+			TokenStream.EmitObjectReference(BaseOffset + GetOffset_ForGC(), *StackSizeHelper.GetPropertyPath(), GCRT_ArrayDelegate);
 		}
 		else if (Inner->IsA(FMulticastDelegateProperty::StaticClass()))
 		{
-			OwnerClass.EmitObjectReference(BaseOffset + GetOffset_ForGC(), *StackSizeHelper.GetPropertyPath(), GCRT_ArrayMulticastDelegate);
+			TokenStream.EmitObjectReference(BaseOffset + GetOffset_ForGC(), *StackSizeHelper.GetPropertyPath(), GCRT_ArrayMulticastDelegate);
 		}
 		else
 		{
 			// In the general case, emit a GCRT_ArrayStruct that can handle any inner property type.
-			OwnerClass.EmitObjectReference(BaseOffset + GetOffset_ForGC(), *StackSizeHelper.GetPropertyPath(), bUsesFreezableAllocator ? GCRT_ArrayStructFreezable : GCRT_ArrayStruct);
+			TokenStream.EmitObjectReference(BaseOffset + GetOffset_ForGC(), *StackSizeHelper.GetPropertyPath(), bUsesFreezableAllocator ? GCRT_ArrayStructFreezable : GCRT_ArrayStruct);
 
 			// GCRT_ArrayStruct and GCRT_ArrayStructFreezable push a new stack frame in TFastReferenceCollector
 			FGCStackSizeHelperScope StackSizeScope(StackSizeHelper, Inner);
 
-			OwnerClass.ReferenceTokenStream.EmitStride(Inner->ElementSize);
-			const uint32 SkipIndexIndex = OwnerClass.ReferenceTokenStream.EmitSkipIndexPlaceholder();
-			Inner->EmitReferenceInfo(OwnerClass, 0, EncounteredStructProps, StackSizeHelper);
-			const uint32 SkipIndex = OwnerClass.ReferenceTokenStream.EmitReturn();
-			OwnerClass.ReferenceTokenStream.UpdateSkipIndexPlaceholder(SkipIndexIndex, SkipIndex);
+			TokenStream.EmitStride(Inner->ElementSize);
+			const uint32 SkipIndexIndex = TokenStream.EmitSkipIndexPlaceholder();
+			Inner->EmitReferenceInfo(TokenStream, 0, EncounteredStructProps, StackSizeHelper);
+			const uint32 SkipIndex = TokenStream.EmitReturn();
+			TokenStream.UpdateSkipIndexPlaceholder(SkipIndexIndex, SkipIndex);
 		}
 	}
 }
@@ -3047,30 +4692,30 @@ void FArrayProperty::EmitReferenceInfo(UClass& OwnerClass, int32 BaseOffset, TAr
  * Emits tokens used by realtime garbage collection code to passed in OwnerClass' ReferenceTokenStream. The offset emitted is relative
  * to the passed in BaseOffset which is used by e.g. arrays of structs.
  */
-void FMapProperty::EmitReferenceInfo(UClass& OwnerClass, int32 BaseOffset, TArray<const FStructProperty*>& EncounteredStructProps, FGCStackSizeHelper& StackSizeHelper)
+void FMapProperty::EmitReferenceInfo(UE::GC::FTokenStreamBuilder& TokenStream, int32 BaseOffset, TArray<const FStructProperty*>& EncounteredStructProps, FGCStackSizeHelper& StackSizeHelper)
 {
 	if (ContainsObjectReference(EncounteredStructProps, EPropertyObjectReferenceType::Strong | EPropertyObjectReferenceType::Weak))
 	{
 		// TMap reference tokens are processed by GC in a similar way to an array of structs
-		OwnerClass.EmitObjectReference(BaseOffset + GetOffset_ForGC(), *StackSizeHelper.GetPropertyPath(), GCRT_AddTMapReferencedObjects);
-		OwnerClass.ReferenceTokenStream.EmitPointer((const void*)this);
-		const uint32 SkipIndexIndex = OwnerClass.ReferenceTokenStream.EmitSkipIndexPlaceholder();
+		TokenStream.EmitObjectReference(BaseOffset + GetOffset_ForGC(), *StackSizeHelper.GetPropertyPath(), GCRT_AddTMapReferencedObjects);
+		TokenStream.EmitPointer((const void*)this);
+		const uint32 SkipIndexIndex = TokenStream.EmitSkipIndexPlaceholder();
 
 		if (KeyProp->ContainsObjectReference(EncounteredStructProps, EPropertyObjectReferenceType::Strong | EPropertyObjectReferenceType::Weak))
 		{
 			// GCRT_AddTMapReferencedObjects pushes a new stack frame in TFastReferenceCollector
 			FGCStackSizeHelperScope StackSizeScope(StackSizeHelper, KeyProp);
-			KeyProp->EmitReferenceInfo(OwnerClass, 0, EncounteredStructProps, StackSizeHelper);
+			KeyProp->EmitReferenceInfo(TokenStream, 0, EncounteredStructProps, StackSizeHelper);
 		}
 		if (ValueProp->ContainsObjectReference(EncounteredStructProps, EPropertyObjectReferenceType::Strong | EPropertyObjectReferenceType::Weak))
 		{
 			// GCRT_AddTMapReferencedObjects pushes a new stack frame in TFastReferenceCollector
 			FGCStackSizeHelperScope StackSizeScope(StackSizeHelper, ValueProp);
-			ValueProp->EmitReferenceInfo(OwnerClass, 0, EncounteredStructProps, StackSizeHelper);
+			ValueProp->EmitReferenceInfo(TokenStream, 0, EncounteredStructProps, StackSizeHelper);
 		}
 
-		const uint32 SkipIndex = OwnerClass.ReferenceTokenStream.EmitReturn();
-		OwnerClass.ReferenceTokenStream.UpdateSkipIndexPlaceholder(SkipIndexIndex, SkipIndex);
+		const uint32 SkipIndex = TokenStream.EmitReturn();
+		TokenStream.UpdateSkipIndexPlaceholder(SkipIndexIndex, SkipIndex);
 	}
 }
 
@@ -3078,21 +4723,21 @@ void FMapProperty::EmitReferenceInfo(UClass& OwnerClass, int32 BaseOffset, TArra
 * Emits tokens used by realtime garbage collection code to passed in OwnerClass' ReferenceTokenStream. The offset emitted is relative
 * to the passed in BaseOffset which is used by e.g. arrays of structs.
 */
-void FSetProperty::EmitReferenceInfo(UClass& OwnerClass, int32 BaseOffset, TArray<const FStructProperty*>& EncounteredStructProps, FGCStackSizeHelper& StackSizeHelper)
+void FSetProperty::EmitReferenceInfo(UE::GC::FTokenStreamBuilder& TokenStream, int32 BaseOffset, TArray<const FStructProperty*>& EncounteredStructProps, FGCStackSizeHelper& StackSizeHelper)
 {
 	if (ContainsObjectReference(EncounteredStructProps, EPropertyObjectReferenceType::Strong | EPropertyObjectReferenceType::Weak))
 	{
 		// TSet reference tokens are processed by GC in a similar way to an array of structs
-		OwnerClass.EmitObjectReference(BaseOffset + GetOffset_ForGC(), *StackSizeHelper.GetPropertyPath(), GCRT_AddTSetReferencedObjects);
-		OwnerClass.ReferenceTokenStream.EmitPointer((const void*)this);
+		TokenStream.EmitObjectReference(BaseOffset + GetOffset_ForGC(), *StackSizeHelper.GetPropertyPath(), GCRT_AddTSetReferencedObjects);
+		TokenStream.EmitPointer((const void*)this);
 
 		// GCRT_AddTSetReferencedObjects pushes a new stack frame in TFastReferenceCollector
 		FGCStackSizeHelperScope StackSizeScope(StackSizeHelper, ElementProp);
 
-		const uint32 SkipIndexIndex = OwnerClass.ReferenceTokenStream.EmitSkipIndexPlaceholder();
-		ElementProp->EmitReferenceInfo(OwnerClass, 0, EncounteredStructProps, StackSizeHelper);
-		const uint32 SkipIndex = OwnerClass.ReferenceTokenStream.EmitReturn();
-		OwnerClass.ReferenceTokenStream.UpdateSkipIndexPlaceholder(SkipIndexIndex, SkipIndex);
+		const uint32 SkipIndexIndex = TokenStream.EmitSkipIndexPlaceholder();
+		ElementProp->EmitReferenceInfo(TokenStream, 0, EncounteredStructProps, StackSizeHelper);
+		const uint32 SkipIndex = TokenStream.EmitReturn();
+		TokenStream.UpdateSkipIndexPlaceholder(SkipIndexIndex, SkipIndex);
 	}
 }
 
@@ -3101,19 +4746,19 @@ void FSetProperty::EmitReferenceInfo(UClass& OwnerClass, int32 BaseOffset, TArra
  * Emits tokens used by realtime garbage collection code to passed in ReferenceTokenStream. The offset emitted is relative
  * to the passed in BaseOffset which is used by e.g. arrays of structs.
  */
-void FStructProperty::EmitReferenceInfo(UClass& OwnerClass, int32 BaseOffset, TArray<const FStructProperty*>& EncounteredStructProps, FGCStackSizeHelper& StackSizeHelper)
+void FStructProperty::EmitReferenceInfo(UE::GC::FTokenStreamBuilder& TokenStream, int32 BaseOffset, TArray<const FStructProperty*>& EncounteredStructProps, FGCStackSizeHelper& StackSizeHelper)
 {
 	check(Struct);
 	if (Struct->StructFlags & STRUCT_AddStructReferencedObjects)
 	{
 		UScriptStruct::ICppStructOps* CppStructOps = Struct->GetCppStructOps();
 		check(CppStructOps); // else should not have STRUCT_AddStructReferencedObjects
-		FGCReferenceFixedArrayTokenHelper FixedArrayHelper(OwnerClass, BaseOffset + GetOffset_ForGC(), ArrayDim, ElementSize, *this, StackSizeHelper);
+		UE::GC::FFixedArrayTokenHelper FixedArrayHelper(TokenStream, BaseOffset + GetOffset_ForGC(), ArrayDim, ElementSize, *this, StackSizeHelper);
 
-		OwnerClass.EmitObjectReference(BaseOffset + GetOffset_ForGC(), *StackSizeHelper.GetPropertyPath(), GCRT_AddStructReferencedObjects);
+		TokenStream.EmitObjectReference(BaseOffset + GetOffset_ForGC(), *StackSizeHelper.GetPropertyPath(), GCRT_AddStructReferencedObjects);
 
 		void *FunctionPtr = (void*)CppStructOps->AddStructReferencedObjects();
-		OwnerClass.ReferenceTokenStream.EmitPointer(FunctionPtr);
+		TokenStream.EmitPointer(FunctionPtr);
 	}
 
 	// Check if the struct has any properties that reference UObjects
@@ -3131,13 +4776,13 @@ void FStructProperty::EmitReferenceInfo(UClass& OwnerClass, int32 BaseOffset, TA
 	// If the struct has UObject properties (and only if) emit tokens for them
 	if (bHasPropertiesWithObjectReferences)
 	{
-		FGCReferenceFixedArrayTokenHelper FixedArrayHelper(OwnerClass, BaseOffset + GetOffset_ForGC(), ArrayDim, ElementSize, *this, StackSizeHelper);
+		UE::GC::FFixedArrayTokenHelper FixedArrayHelper(TokenStream, BaseOffset + GetOffset_ForGC(), ArrayDim, ElementSize, *this, StackSizeHelper);
 
 		FProperty* Property = Struct->PropertyLink;
 		while (Property)
 		{
 			FGCStackSizeHelperScope StackSizeScope(StackSizeHelper, Property);
-			Property->EmitReferenceInfo(OwnerClass, BaseOffset + GetOffset_ForGC(), EncounteredStructProps, StackSizeHelper);
+			Property->EmitReferenceInfo(TokenStream, BaseOffset + GetOffset_ForGC(), EncounteredStructProps, StackSizeHelper);
 			Property = Property->PropertyLinkNext;
 		}
 	}
@@ -3147,39 +4792,101 @@ void FStructProperty::EmitReferenceInfo(UClass& OwnerClass, int32 BaseOffset, TA
  * Emits tokens used by realtime garbage collection code to passed in ReferenceTokenStream. The offset emitted is relative
  * to the passed in BaseOffset which is used by e.g. arrays of structs.
  */
-void FInterfaceProperty::EmitReferenceInfo(UClass& OwnerClass, int32 BaseOffset, TArray<const FStructProperty*>& EncounteredStructProps, FGCStackSizeHelper& StackSizeHelper)
+void FInterfaceProperty::EmitReferenceInfo(UE::GC::FTokenStreamBuilder& TokenStream, int32 BaseOffset, TArray<const FStructProperty*>& EncounteredStructProps, FGCStackSizeHelper& StackSizeHelper)
 {
-	FGCReferenceFixedArrayTokenHelper FixedArrayHelper(OwnerClass, BaseOffset + GetOffset_ForGC(), ArrayDim, sizeof(FScriptInterface), *this, StackSizeHelper);
+	UE::GC::FFixedArrayTokenHelper FixedArrayHelper(TokenStream, BaseOffset + GetOffset_ForGC(), ArrayDim, sizeof(FScriptInterface), *this, StackSizeHelper);
 
-	OwnerClass.EmitObjectReference(BaseOffset + GetOffset_ForGC(), *StackSizeHelper.GetPropertyPath(), GCRT_Object);
+	TokenStream.EmitObjectReference(BaseOffset + GetOffset_ForGC(), *StackSizeHelper.GetPropertyPath(), GCRT_Object);
 }
 
-void FFieldPathProperty::EmitReferenceInfo(UClass& OwnerClass, int32 BaseOffset, TArray<const FStructProperty*>& EncounteredStructProps, FGCStackSizeHelper& StackSizeHelper)
+void FFieldPathProperty::EmitReferenceInfo(UE::GC::FTokenStreamBuilder& TokenStream, int32 BaseOffset, TArray<const FStructProperty*>& EncounteredStructProps, FGCStackSizeHelper& StackSizeHelper)
 {
 	static_assert(sizeof(FFieldPath) == sizeof(TFieldPath<FProperty>), "TFieldPath should have the same size as the underlying FFieldPath");
-	FGCReferenceFixedArrayTokenHelper FixedArrayHelper(OwnerClass, BaseOffset + GetOffset_ForGC(), ArrayDim, sizeof(FFieldPath), *this, StackSizeHelper);
-	OwnerClass.EmitObjectReference(BaseOffset + GetOffset_ForGC(), *StackSizeHelper.GetPropertyPath(), GCRT_AddFieldPathReferencedObject);
+	UE::GC::FFixedArrayTokenHelper FixedArrayHelper(TokenStream, BaseOffset + GetOffset_ForGC(), ArrayDim, sizeof(FFieldPath), *this, StackSizeHelper);
+	TokenStream.EmitObjectReference(BaseOffset + GetOffset_ForGC(), *StackSizeHelper.GetPropertyPath(), GCRT_AddFieldPathReferencedObject);
+}
+
+namespace UE::GC
+{
+
+static int32 GNumTokens;
+static int32 GNumTokenStreams;
+static void UpdateGlobalTokenCount(int32 Delta)
+{
+#if !UE_BUILD_SHIPPING
+	GNumTokens += Delta;
+	GNumTokenStreams += FMath::Sign(Delta);
+#endif
+}
+
+struct FTokenStreamHeader
+{
+	std::atomic<int32> RefCount{1};
+};
+
+static FTokenStreamHeader* GetStreamHeader(FTokenStreamView View)
+{
+	check(View.GetTokenData());
+	return reinterpret_cast<FTokenStreamHeader*>(View.GetTokenData()) - 1;
+}
+
+static void DropStreamReference(FTokenStreamView View)
+{
+	if (View.GetTokenData())
+	{
+		if (GetStreamHeader(View)->RefCount.fetch_sub(1) == 1)
+		{
+			UpdateGlobalTokenCount(-int32(View.NumTokens()));
+			FMemory::Free(GetStreamHeader(View));
+		}
+	}
+}
+
+static void AddStreamReference(FTokenStreamView View)
+{
+	if (View.GetTokenData())
+	{
+		GetStreamHeader(View)->RefCount.fetch_add(1);
+	}
+}
+
+static void UpdateStreamReference(FTokenStreamView& Out, FTokenStreamView New)
+{
+	if (Out.GetTokenData() != New.GetTokenData())
+	{
+		DropStreamReference(Out);
+		AddStreamReference(New);
+		Out = New;
+	}
 }
 
-void UClass::EmitObjectReference(int32 Offset, const FName& DebugName, EGCReferenceType Kind)
+struct  
+{
+	FGCReferenceInfo Token{GCRT_EndOfStream, 0};
+	FName DebugName{"EndOfStreamToken"};
+} GEndOfStream;
+
+uint32 UE::GC::FTokenStreamBuilder::MaxStackSize;
+
+
+void FTokenStreamBuilder::EmitObjectReference(int32 Offset, FName DebugName, EGCReferenceType Kind)
 {
 	FGCReferenceInfo ObjectReference(Kind, Offset);
-	ReferenceTokenStream.EmitReferenceInfo(ObjectReference, DebugName);
+	int32 Index = EmitReferenceInfo(ObjectReference, DebugName);
 }
 
-void UClass::EmitObjectArrayReference(int32 Offset, const FName& DebugName)
+void FTokenStreamBuilder::EmitObjectArrayReference(int32 Offset, FName DebugName)
 {
-	check(HasAnyClassFlags(CLASS_Intrinsic));
+	check(DebugOwner.HasAnyClassFlags(CLASS_Intrinsic));
 	EmitObjectReference(Offset, DebugName, GCRT_ArrayObject);
 }
 
-uint32 UClass::EmitStructArrayBegin(int32 Offset, const FName& DebugName, int32 Stride)
+uint32 FTokenStreamBuilder::EmitStructArrayBegin(int32 Offset, FName DebugName, int32 Stride)
 {
-	check(HasAnyClassFlags(CLASS_Intrinsic));
+	check(DebugOwner.HasAnyClassFlags(CLASS_Intrinsic));
 	EmitObjectReference(Offset, DebugName, GCRT_ArrayStruct);
-	ReferenceTokenStream.EmitStride(Stride);
-	const uint32 SkipIndexIndex = ReferenceTokenStream.EmitSkipIndexPlaceholder();
-	return SkipIndexIndex;
+	EmitStride(Stride);
+	return EmitSkipIndexPlaceholder();
 }
 
 /**
@@ -3189,38 +4896,140 @@ uint32 UClass::EmitStructArrayBegin(int32 Offset, const FName& DebugName, int32
  *
  * @param SkipIndexIndex
  */
-void UClass::EmitStructArrayEnd( uint32 SkipIndexIndex )
+void FTokenStreamBuilder::EmitStructArrayEnd(uint32 SkipIndexIndex)
 {
-	check( HasAnyClassFlags( CLASS_Intrinsic ) );
-	const uint32 SkipIndex = ReferenceTokenStream.EmitReturn();
-	ReferenceTokenStream.UpdateSkipIndexPlaceholder( SkipIndexIndex, SkipIndex );
+	check(DebugOwner.HasAnyClassFlags(CLASS_Intrinsic));
+	uint32 SkipIndex = EmitReturn();
+	UpdateSkipIndexPlaceholder( SkipIndexIndex, SkipIndex );
 }
 
-void UClass::EmitFixedArrayBegin(int32 Offset, const FName& DebugName, int32 Stride, int32 Count)
+void FTokenStreamBuilder::EmitFixedArrayBegin(int32 Offset, FName DebugName, int32 Stride, int32 Count)
 {
-	check(HasAnyClassFlags(CLASS_Intrinsic));
+	check(DebugOwner.HasAnyClassFlags(CLASS_Intrinsic));
 	EmitObjectReference(Offset, DebugName, GCRT_FixedArray);
-	ReferenceTokenStream.EmitStride(Stride);
-	ReferenceTokenStream.EmitCount(Count);
+	EmitStride(Stride);
+	EmitCount(Count);
 }
 
-/**
- * Realtime garbage collection helper function used to indicated the end of a fixed array.
- */
-void UClass::EmitFixedArrayEnd()
+void FTokenStreamBuilder::EmitFixedArrayEnd()
 {
-	check( HasAnyClassFlags( CLASS_Intrinsic ) );
-	ReferenceTokenStream.EmitReturn();
+	check(DebugOwner.HasAnyClassFlags(CLASS_Intrinsic));
+	EmitReturn();
 }
 
-void UClass::EmitExternalPackageReference()
+void FTokenStreamBuilder::EmitExternalPackageReference()
 {
 #if WITH_EDITOR
 	static const FName TokenName("ExternalPackageToken");
-	ReferenceTokenStream.EmitReferenceInfo(FGCReferenceInfo(GCRT_ExternalPackage, 0), TokenName);
+	EmitReferenceInfo(FGCReferenceInfo(GCRT_ExternalPackage, 0), TokenName);
+#endif
+}
+
+// FTokenStreamBuilderIterator helper
+struct FNamedReferenceInfo
+{
+	FGCReferenceInfo ReferenceInfo;
+	FName DebugName;
+};
+
+// Helps iterate over a range of FTokenStreamBuilder tokens and debug names
+class FTokenStreamBuilder::FConstIterator
+{
+public:
+	FConstIterator(const FTokenStreamBuilder& InBuilder)
+	: Builder(InBuilder)
+	, Idx(0)
+	, Num(InBuilder.Tokens.Num())
+	{}
+
+	explicit operator bool() const { return Idx < Num; }
+	
+	uint32 Pos() const { return Idx; }
+	uint32 EndPos() const { return Num; }
+
+	uint32 PeekReturnCount()
+	{
+		return FGCReferenceInfo(GetToken()).ReturnCount;
+	}
+
+	FNamedReferenceInfo ConsumeReferenceInfo()
+	{
+		FNamedReferenceInfo Out = { FGCReferenceInfo(GetToken()), GetDebugName() };
+		++Idx;
+		return Out;
+	}
+
+	// @return Index to skip to
+	uint32 ConsumeSkipInfo()
+	{
+#if ENABLE_GC_TOKEN_DEBUG_INFO
+		static const FName TokenName("SkipIndexPlaceholder");
+		check(GetDebugName() == TokenName);
+#endif
+		FGCSkipInfo SkipInfo(ConsumeToken());
+		return Idx - 1 + SkipInfo.SkipIndex;
+	}
+	
+	uint32 ConsumeStride()
+	{
+#if ENABLE_GC_TOKEN_DEBUG_INFO
+		static const FName TokenName("StrideToken");
+		check(GetDebugName() == TokenName);
+#endif
+		return ConsumeToken();
+	}
+
+	uint32 ConsumeCount()
+	{
+#if ENABLE_GC_TOKEN_DEBUG_INFO
+		static const FName TokenName("CountToken");
+		check(GetDebugName() == TokenName);
+#endif
+		return ConsumeToken();
+	}
+		
+	void* ConsumePointer(uint32& OutReturnCount)
+	{
+		void* Out = Builder.ReadPointer(/* in-out */ Idx);
+#if ENABLE_GC_TOKEN_DEBUG_INFO
+		static const FName TokenName("EndOfPointerToken");
+		check(GetDebugName() == TokenName);
+#endif
+		FGCReferenceInfo EndOfPointer(ConsumeToken());
+		OutReturnCount = EndOfPointer.ReturnCount;
+		return Out;
+	}
+
+	const FTokenStreamBuilder& DebugGetBuilder() { return Builder; }
+private:
+	uint32 GetToken() const { return Builder.Tokens[Idx]; }
+	uint32 ConsumeToken() { return Builder.Tokens[Idx++]; }
+
+#if ENABLE_GC_TOKEN_DEBUG_INFO
+	FName GetDebugName() const { return Builder.DebugNames[Idx]; }
+#else
+	FName GetDebugName() const { return FName(); }
+#endif
+
+	const FTokenStreamBuilder& Builder;
+	uint32 Idx = 0;
+	uint32 Num = 0;
+};
+
+static AROFunc GetAROFunc(UClass* Class)
+{
+	AROFunc ARO = Class->CppClassStaticFunctions.GetAddReferencedObjects();
+	check(ARO != nullptr);
+
+#if WITH_EDITOR
+	return ARO;
+#else
+	return ARO != &UObject::AddReferencedObjects ? ARO : nullptr;
 #endif
 }
 
+} // namespace UE::GC
+
 struct FScopeLockIfNotNative
 {
 	FCriticalSection& ScopeCritical;
@@ -3243,318 +5052,417 @@ struct FScopeLockIfNotNative
 	}
 };
 
+/** Token stream can assemble code can sometimes be called from two threads throuh a web of async loading calls. */
+static FCriticalSection GReferenceTokenStreamCritical;
+int32 GCreatedMixedTokens;
+int32 GCreatedStrongTokens;
+
 void UClass::AssembleReferenceTokenStream(bool bForce)
 {
 	LLM_SCOPE(ELLMTag::GC);
 
-	// Lock for non-native classes
-	FScopeLockIfNotNative ReferenceTokenStreamLock(ReferenceTokenStreamCritical, !(ClassFlags & CLASS_Native));
-
 	UE_CLOG(!IsInGameThread() && !IsGarbageCollectionLocked(), LogGarbage, Fatal, TEXT("AssembleReferenceTokenStream for %s called on a non-game thread while GC is not locked."), *GetFullName());
 
+	if (ClassFlags & CLASS_Native)
+	{
+		AssembleReferenceTokenStreamInternal(bForce);
+	}
+	else
+	{
+		FScopeLock NonNativeLock(&GReferenceTokenStreamCritical);
+		AssembleReferenceTokenStreamInternal(bForce);
+	}
+}
+
+void UClass::AssembleReferenceTokenStreamInternal(bool bForce)
+{
+	using namespace UE::GC;
+
 	if (!HasAnyClassFlags(CLASS_TokenStreamAssembled) || bForce)
 	{
+		FTokenStreamBuilder StackBuilder(*this);
+		TUniquePtr<FTokenStreamBuilder> IntrinsicBuilder;
+		if (ClassFlags & CLASS_Intrinsic)
+		{
+			IntrinsicBuilder = FIntrinsicClassTokens::ConsumeBuilder(this);
+		}
+		FTokenStreamBuilder& MixedTokens = IntrinsicBuilder ? *IntrinsicBuilder : StackBuilder;
+
 		if (bForce)
 		{
-			ReferenceTokenStream.Empty();
 			ClassFlags &= ~CLASS_TokenStreamAssembled;
 		}
-		TArray<const FStructProperty*> EncounteredStructProps;
-		
+
 		{
 			FGCStackSizeHelper StackSizeHelper;
+			TArray<const FStructProperty*> EncounteredStructProps;
 
 			// Iterate over properties defined in this class
 			for (TFieldIterator<FProperty> It(this, EFieldIteratorFlags::ExcludeSuper); It; ++It)
 			{
 				FProperty* Property = *It;
 				FGCStackSizeHelperScope StackSizeScope(StackSizeHelper, Property);
-				Property->EmitReferenceInfo(*this, 0, EncounteredStructProps, StackSizeHelper);
+				Property->EmitReferenceInfo(MixedTokens, 0, EncounteredStructProps, StackSizeHelper);
 			}
 
-			// The fact we set min to 1 is because we always emit noexport (non FProperty exposed) references for 
-			// UObject's Outer and ClassPrivate members in EmitBaseReferences (see below)
-			ReferenceTokenStream.SetStackSize(FMath::Max(1, StackSizeHelper.GetMaxStackSize()));
+			MixedTokens.SetStackSize(StackSizeHelper.GetMaxStackSize());
 		}
-
+				
+		FTokenStreamView SuperMixedTokens;
+		FTokenStreamView SuperStrongTokens;
 		if (UClass* SuperClass = GetSuperClass())
 		{
-			// We also need to lock the super class stream in case something (like PostLoad) wants to reconstruct it on GameThread
-			FScopeLockIfNotNative SuperClassReferenceTokenStreamLock(SuperClass->ReferenceTokenStreamCritical, !(SuperClass->ClassFlags & CLASS_Native));
-			
 			// Make sure super class has valid token stream.
-			SuperClass->AssembleReferenceTokenStream();
-			if (!SuperClass->ReferenceTokenStream.IsEmpty())
-			{
-				// Prepend super's stream. This automatically handles removing the EOS token.
-				ReferenceTokenStream.PrependStream(SuperClass->ReferenceTokenStream);
-			}
+			SuperClass->AssembleReferenceTokenStreamInternal();
+			SuperMixedTokens  = FTokenStreamBuilder::DropFinalTokens(SuperClass->ReferenceTokens.Mixed,  GetAROFunc(SuperClass));
+			SuperStrongTokens = FTokenStreamBuilder::DropFinalTokens(SuperClass->ReferenceTokens.Strong, GetAROFunc(SuperClass));
 		}
 		else
 		{
-			UObjectBase::EmitBaseReferences(this);
+			UObjectBase::EmitBaseReferences(MixedTokens);
 		}
 
+		// Make sure all Blueprint properties are marked as non-native
+		MixedTokens.SetTokenType(GetClass()->HasAnyClassFlags(CLASS_NeedsDeferredDependencyLoading) ? EGCTokenType::NonNative : EGCTokenType::Native);
+		MixedTokens.EmitFinalTokens(GetAROFunc(this));
+
+		// Allocate flat token stream with current and super class tokens
+		FTokenStreamBuilder::Merge(/* out */ ReferenceTokens.Mixed, MixedTokens, SuperMixedTokens);
+		
+		// Create subset of strong reference tokens for this class
+		FTokenStreamBuilder StrongTokens(*this);
+		StrongTokens.SetStackSize(MixedTokens.GetStackSize());
+		StrongTokens.SetTokenType(MixedTokens.GetTokenType());
+		FTokenStreamBuilder::CopyStrongTokens(FTokenStreamBuilder::FConstIterator(MixedTokens), /* out */ StrongTokens);
+
+		// Reuse mixed tokens if possible
+		if (MixedTokens.Num() == StrongTokens.Num() && SuperMixedTokens.NumTokens() == SuperStrongTokens.NumTokens())
+		{
+			UpdateStreamReference(/* out */ ReferenceTokens.Strong, ReferenceTokens.Mixed);
+		}
+		else
+		{
+			FTokenStreamBuilder::Merge(/* out */ ReferenceTokens.Strong, StrongTokens, SuperStrongTokens);
+		}
+
+		check(!HasAnyClassFlags(CLASS_TokenStreamAssembled)); // recursion here is probably bad
+		ClassFlags |= CLASS_TokenStreamAssembled;
+	}
+}
+
+namespace UE::GC
+{
+
+//////////////////////////////////////////////////////////////////////////
+
+void FTokenStreamBuilder::EmitFinalTokens(AROFunc ARO)
+{
+	if (ARO)
+	{
+		static const FName TokenName("AROToken");
+		int32 SlowIndex = FindSlowImplementation(ARO);
+		if (SlowIndex == INDEX_NONE)
+		{
+			EmitReferenceInfo(FGCReferenceInfo(GCRT_AddReferencedObjects, 0), TokenName);
+			EmitPointer(reinterpret_cast<const void*>(ARO));
+		}
+		else
 		{
-			check(CppClassStaticFunctions.GetAddReferencedObjects() != nullptr);
-			const bool bKeepOuter = true;//GetFName() != NAME_Package;
-			const bool bKeepClass = true;//!HasAnyInternalFlags(EInternalObjectFlags::Native);
+			EmitReferenceInfo(FGCReferenceInfo(GCRT_SlowAddReferencedObjects, static_cast<uint32>(SlowIndex)), TokenName);
+		}
+	}
+
+	// Emit end of stream token.
+	EmitReferenceInfo(GEndOfStream.Token, GEndOfStream.DebugName);
+}
+
+ FTokenStreamView FTokenStreamBuilder::DropFinalTokens(FTokenStreamView Stream, AROFunc DropARO)
+ {
+	check(Stream.DroppedNum == 0)
+	if (Stream.Num >= 5 &&
+		FGCReferenceInfo(Stream.Tokens[Stream.Num - 5]).Type == GCRT_AddReferencedObjects &&
+		FGCReferenceInfo(Stream.Tokens[Stream.Num - 2]).Type == GCRT_EndOfPointer &&
+		FGCReferenceInfo(Stream.Tokens[Stream.Num - 1]).Type == GCRT_EndOfStream)
+	{
+		checkf(DropARO, TEXT("Super class token stream assembled with ClassAddReferencedObjects set and now it has been reset."
+							" Super class token stream is outdated and should have been regenerated. "));
+		Stream.Num -= 5;
+		Stream.DroppedNum = 5;
+	}
+	else if (Stream.Num >= 2 &&
+		FGCReferenceInfo(Stream.Tokens[Stream.Num - 2]).Type == GCRT_SlowAddReferencedObjects &&
+		FGCReferenceInfo(Stream.Tokens[Stream.Num - 1]).Type == GCRT_EndOfStream)
+	{
+		checkf(DropARO, TEXT("Super class token stream assembled with ClassAddReferencedObjects set and now it has been reset."
+							" Super class token stream is outdated and should have been regenerated. "));
+		Stream.Num -= 2;
+		Stream.DroppedNum = 2;
+	}
+	else if (Stream.Num >= 1)
+	{
+		check(FGCReferenceInfo(Stream.Tokens[Stream.Num - 1]).Type == GCRT_EndOfStream);
+		Stream.Num -= 1;
+		Stream.DroppedNum = 1;
+	}
 
-			FUObjectCppClassStaticFunctions::AddReferencedObjectsType AddReferencedObjectsFn = nullptr;
-#if !WITH_EDITOR
-			// In no-editor builds UObject::ARO is empty, thus only classes
-			// which implement their own ARO function need to have the ARO token generated.
-			if (CppClassStaticFunctions.GetAddReferencedObjects() != &UObject::AddReferencedObjects)
-			{
-				AddReferencedObjectsFn = CppClassStaticFunctions.GetAddReferencedObjects();
-			}
-#else
-			AddReferencedObjectsFn = CppClassStaticFunctions.GetAddReferencedObjects();
-#endif
-			ReferenceTokenStream.Fixup(AddReferencedObjectsFn, bKeepOuter, bKeepClass);
+	return Stream;
+}
 
-			// Make sure all Blueprint properties are marked as non-native
-			const EGCTokenType TokenType = GetClass()->HasAnyClassFlags(CLASS_NeedsDeferredDependencyLoading) ? EGCTokenType::NonNative : EGCTokenType::Native;
-			ReferenceTokenStream.SetTokenType(TokenType);
-		}
+void FTokenStreamBuilder::Merge(FTokenStreamView& Out, const FTokenStreamBuilder& Class, FTokenStreamView Super) 
+{
+	check(Class.Tokens.Num() && FGCReferenceInfo(Class.Tokens.Last()).Type == GCRT_EndOfStream);
 
-		if (ReferenceTokenStream.IsEmpty())
-		{
-			return;
-		}
+	uint32 Num = Class.Num() + Super.Num;
+	if (Num == 1)
+	{
+		// Return empty stream if there's only an EndOfStream terminator	
+		check(Super.IsEmpty() && Class.Tokens[0] == GEndOfStream.Token);
+		UpdateStreamReference(Out, FTokenStreamView());
+	}
+	else if (Num == Out.Num &&
+		!FMemory::Memcmp(Out.Tokens,				Super.Tokens,			sizeof(uint32) * Super.Num) &&
+		!FMemory::Memcmp(Out.Tokens + Super.Num,	Class.Tokens.GetData(),	sizeof(uint32) * Class.Num()))
+	{
+		// Nothing changed, no need to update
+		check(Out.IsNonNative == (Class.TokenType == EGCTokenType::NonNative));
+	}
+	else if (Class.Tokens.Num() == Super.DroppedNum && 
+		!FMemory::Memcmp(Super.Tokens + Super.Num,	Class.Tokens.GetData(), sizeof(uint32) * Class.Num()))
+	{
+		// Reuse Super tokens since Class only contain Super's dropped final tokens
+		FTokenStreamView New = Super;
+		New.Num += New.DroppedNum;
+		New.DroppedNum = 0;
+		New.IsNonNative = Class.TokenType == EGCTokenType::NonNative;
+		UpdateStreamReference(Out, New);
+	}
+	else
+	{
+		// Allocate new flat stream
+		LLM_SCOPE(ELLMTag::GC);
+	
+		static_assert(alignof(FTokenStreamHeader) == alignof(uint32));
+		static_assert(alignof(FName) == alignof(uint32));
+		SIZE_T NumBytes = sizeof(FTokenStreamHeader) + Num * sizeof(uint32) + (ENABLE_GC_TOKEN_DEBUG_INFO ? Num * sizeof(FName) : 0);
+		FTokenStreamHeader* Header = new (FMemory::Malloc(NumBytes)) FTokenStreamHeader;
+		uint32* Tokens = reinterpret_cast<uint32*>(Header + 1);
 
-		// Emit end of stream token.
-		static const FName EOSDebugName("EndOfStreamToken");
-		ReferenceTokenStream.EmitReferenceInfo(FGCReferenceInfo::EndOfStreamToken, EOSDebugName);
+		FMemory::Memcpy(Tokens,						Super.Tokens,					sizeof(uint32) * Super.Num);
+		FMemory::Memcpy(Tokens + Super.Num,			Class.Tokens.GetData(),			sizeof(uint32) * Class.Num());
 
-		// Shrink reference token stream to proper size.
-		ReferenceTokenStream.Shrink();
+	#if ENABLE_GC_TOKEN_DEBUG_INFO
+		check(Class.DebugNames.Num() == Class.Tokens.Num());
+		FName* DebugNames = reinterpret_cast<FName*>(Tokens + Num);
+		FMemory::Memcpy(DebugNames,					Super.GetDebugNames(),			sizeof(FName) * Super.Num);
+		FMemory::Memcpy(DebugNames + Super.Num,		Class.DebugNames.GetData(),		sizeof(FName) * Class.Num());
+	#endif
 
-		check(!HasAnyClassFlags(CLASS_TokenStreamAssembled)); // recursion here is probably bad
-		ClassFlags |= CLASS_TokenStreamAssembled;
+		UpdateGlobalTokenCount(int32(Num));
 
 		// Update the global maximum stack size after assembling the token stream of each class
-		FGCReferenceTokenStream::SetMaxStackSize(FMath::Max(FGCReferenceTokenStream::GetMaxStackSize(), ReferenceTokenStream.GetStackSize()));
+		uint32 StackSize = FMath::Max(Class.StackSize, Super.StackSize);
+		MaxStackSize = FMath::Max(MaxStackSize, StackSize);
+		check(StackSize < (1u << FTokenStreamView::NumStackSizeBits));
+
+		FTokenStreamView New;
+		New.Tokens = Tokens;
+		New.Num = Num;
+		New.StackSize = StackSize;
+		New.IsNonNative = Class.TokenType == EGCTokenType::NonNative;
+
+		DropStreamReference(Out);
+		Out = New;
 	}
 }
 
+static uint32 PickReturnCount(uint32 OuterReturnCount, uint32 InnerReturnCount)
+{
+	if (InnerReturnCount > 1)
+	{
+		check(OuterReturnCount == 0);
+		return InnerReturnCount - 1;
+	}
+
+	return OuterReturnCount;
+}
 
-/**
- * Prepends passed in stream to existing one.
- *
- * @param Other	stream to concatenate
- */
-void FGCReferenceTokenStream::PrependStream( const FGCReferenceTokenStream& Other )
+static bool CopyStrongTokensUntilPos(FTokenStreamBuilder::FConstIterator& /* in-out */ MixedIt, uint32 MixedEndPos, FTokenStreamBuilder& Out, uint32& OutReturnCount)
 {
-	// Remove embedded EOS token if needed.
-	int32 NumTokensToPrepend = (Other.Tokens.Num() && Other.Tokens.Last() == FGCReferenceInfo::EndOfStreamToken) ? (Other.Tokens.Num() - 1) : Other.Tokens.Num();
+	bool bOut = false;
+	uint32 LastReturnCount = 0;
+	while (MixedIt.Pos() < MixedEndPos)
+	{
+		bOut |= FTokenStreamBuilder::CopyNextStrongToken(MixedIt, Out, LastReturnCount);
+	}
 
-	TArray<uint32> TempTokens;
-	TempTokens.Reserve(NumTokensToPrepend + Tokens.Num());
+	OutReturnCount = PickReturnCount(OutReturnCount, LastReturnCount);
 
-#if ENABLE_GC_TOKEN_DEBUG_INFO
-	check(TokenDebugInfo.Num() == Tokens.Num());
-	check(Other.TokenDebugInfo.Num() == Other.Tokens.Num());
-	TArray<FName> TempTokenDebugInfo;
-	TempTokenDebugInfo.Reserve(NumTokensToPrepend + TokenDebugInfo.Num());
-#endif // ENABLE_GC_TOKEN_DEBUG_INFO
+	return bOut;
+}
+
+static bool CopyStrongTokensUntilReturn(FTokenStreamBuilder::FConstIterator& /* in-out */ MixedIt, FTokenStreamBuilder& Out, uint32& OutReturnCount)
+{
+	uint32 FooStartPos = MixedIt.Pos();
 
-	for (int32 TokenIndex = 0; TokenIndex < NumTokensToPrepend; ++TokenIndex)
+	uint32 InnerReturnCount = 0;
+	bool bOut = false;
+	while (InnerReturnCount == 0)
 	{
-		TempTokens.Add(Other.Tokens[TokenIndex]);
-#if ENABLE_GC_TOKEN_DEBUG_INFO
-		TempTokenDebugInfo.Add(Other.TokenDebugInfo[TokenIndex]);
-#endif // ENABLE_GC_TOKEN_DEBUG_INFO
+		bOut |= FTokenStreamBuilder::CopyNextStrongToken(/* in-out */ MixedIt, Out, InnerReturnCount);
 	}
+	
+	OutReturnCount = PickReturnCount(OutReturnCount, InnerReturnCount);
 
-	TempTokens.Append(Tokens);
-	Tokens = MoveTemp(TempTokens);
-
-#if ENABLE_GC_TOKEN_DEBUG_INFO
-	TempTokenDebugInfo.Append(TokenDebugInfo);
-	TokenDebugInfo = MoveTemp(TempTokenDebugInfo);
-#endif // ENABLE_GC_TOKEN_DEBUG_INFO
+	return bOut;
+}
 
-	StackSize = FMath::Max(StackSize, Other.GetStackSize());
+void FTokenStreamBuilder::CopyStrongTokens(FConstIterator MixedIt, FTokenStreamBuilder& Out)
+{
+	uint32 ReturnCount = 0;
+	CopyStrongTokensUntilPos(MixedIt, MixedIt.EndPos(), Out, ReturnCount);
+	checkf(Out.Num() < MixedIt.EndPos() || Out == MixedIt.DebugGetBuilder(), 
+		TEXT("No weak tokens stripped, copied tokens should be bitwise identical to the original mixed tokens"));
 }
 
-void FGCReferenceTokenStream::Fixup(void (*AddReferencedObjectsPtr)(UObject*, class FReferenceCollector&), bool bKeepOuterToken, bool bKeepClassToken)
+bool FTokenStreamBuilder::CopyNextStrongToken(FConstIterator& /* in-out */ MixedIt, FTokenStreamBuilder& Out, uint32& OutReturnCount)
 {
-	bool bReplacedARO = false;
+	check(MixedIt);
+	
+	const uint32 RollbackNum = Out.Tokens.Num();
+	FNamedReferenceInfo Token = MixedIt.ConsumeReferenceInfo();
+
+	// Strip old return count and emit new returns since we might strip out entire fixed arrays/struct arrays/maps/sets/etc
+	OutReturnCount = Token.ReferenceInfo.ReturnCount;
+	Token.ReferenceInfo.ReturnCount = 0;
+
+	switch (Token.ReferenceInfo.Type)
+	{
+	case GCRT_Object:
+	case GCRT_ArrayObject:
+	case GCRT_ArrayObjectFreezable:
+	case GCRT_ExternalPackage:
+	case GCRT_AddFieldPathReferencedObject:
+	case GCRT_ArrayAddFieldPathReferencedObject:
+	case GCRT_DynamicallyTypedValue: // Always keep for now
+	case GCRT_SlowAddReferencedObjects:
+	case GCRT_EndOfStream:
+		Out.EmitReferenceInfo(Token.ReferenceInfo, Token.DebugName);
+		return true;
 
-	// Try to find exiting ARO pointer and replace it (to avoid removing and readding tokens).
-	for (int32 TokenStreamIndex = 0; TokenStreamIndex < Tokens.Num(); ++TokenStreamIndex)
+	case GCRT_AddStructReferencedObjects:
+	case GCRT_AddReferencedObjects:
+		Out.EmitReferenceInfo(Token.ReferenceInfo, Token.DebugName);
+		Out.EmitPointer(MixedIt.ConsumePointer(OutReturnCount));
+		return true;
+
+	case GCRT_WeakObject:
+	case GCRT_ArrayWeakObject:
+	case GCRT_LazyObject:
+	case GCRT_ArrayLazyObject:
+	case GCRT_SoftObject:
+	case GCRT_ArraySoftObject:
+	case GCRT_Delegate:
+	case GCRT_ArrayDelegate:
+	case GCRT_MulticastDelegate:
+	case GCRT_ArrayMulticastDelegate:
+		return false;
+
+	// Below property types get their Out tokens rolled back if nested properties are entirely weak
+	case GCRT_Optional:
+	case GCRT_ArrayStructFreezable:
+	case GCRT_ArrayStruct:
 	{
-		uint32 TokenIndex = (uint32)TokenStreamIndex;
-		FGCReferenceInfo Token = Tokens[TokenIndex];
-		// Read token type and skip additional data if present.
-		switch (Token.Type)
+		Out.EmitReferenceInfo(Token.ReferenceInfo, Token.DebugName);
+		Out.EmitStride(MixedIt.ConsumeStride());
+
+		uint32 MixedSkipPos = MixedIt.ConsumeSkipInfo();
+		uint32 SkipSkipIndex = Out.EmitSkipIndexPlaceholder();
+
+		if (CopyStrongTokensUntilPos(MixedIt, MixedSkipPos, Out, OutReturnCount))
 		{
-		case GCRT_ArrayStruct:
-		case GCRT_ArrayStructFreezable:
-			{
-				// Skip stride and move to Skip Info
-				TokenIndex += 2;
-				const FGCSkipInfo SkipInfo = ReadSkipInfo(TokenIndex);
-				// Set the TokenIndex to the skip index - 1 because we're going to
-				// increment in the for loop anyway.
-				TokenIndex = SkipInfo.SkipIndex - 1;
-			}
-			break;
-		case GCRT_FixedArray:
-			{
-				// Skip stride
-				TokenIndex++; 
-				// Skip count
-				TokenIndex++; 
-			}
-			break;
-		case GCRT_AddStructReferencedObjects:
-			{
-				// Skip pointer
-				TokenIndex += GNumTokensPerPointer;
-			}
-			break;
-		case GCRT_AddReferencedObjects:
-			{
-				// Store the pointer after the ARO token.
-				if (AddReferencedObjectsPtr)
-				{
-					StorePointer(&Tokens[TokenIndex + 1], (const void*)AddReferencedObjectsPtr);					
-				}
-				bReplacedARO = true;
-				TokenIndex += GNumTokensPerPointer;
-			}
-			break;
-		case GCRT_AddTMapReferencedObjects:
-		case GCRT_AddTSetReferencedObjects:
-			{
-				// Skip pointer
-				TokenIndex += GNumTokensPerPointer;
-				TokenIndex += 1; // GCRT_EndOfPointer;
-				// Move to Skip Info
-				TokenIndex += 1;
-				const FGCSkipInfo SkipInfo = ReadSkipInfo(TokenIndex);
-				// Set the TokenIndex to the skip index - 1 because we're going to
-				// increment in the for loop anyway.
-				TokenIndex = SkipInfo.SkipIndex - 1;
-			}
-			break;
-		case GCRT_Class:
-		case GCRT_NoopClass:
-			{
-				if (bKeepClassToken)
-				{
-					Token.Type = GCRT_Class;
-				}
-				else
-				{
-					Token.Type = GCRT_NoopClass;
-				}
-				Tokens[TokenIndex] = Token;
-			}
-			break;
-		case GCRT_PersistentObject:
-		case GCRT_NoopPersistentObject:
-			{
-				if (bKeepOuterToken)
-				{
-					Token.Type = GCRT_PersistentObject;
-				}
-				else
-				{
-					Token.Type = GCRT_NoopPersistentObject;
-				}
-				Tokens[TokenIndex] = Token;
-			}
-			break;
-		case GCRT_Optional:
-			{
-				// Move to Skip Info
-				TokenIndex++;
-				const FGCSkipInfo SkipInfo = ReadSkipInfo(TokenIndex);
-				// Set the TokenIndex to the skip index - 1 because we're going to
-				// increment in the for loop anyway.
-				TokenIndex = SkipInfo.SkipIndex - 1;
-			}
-			break;
-		case GCRT_None:
-		case GCRT_Object:
-		case GCRT_ExternalPackage:
-		case GCRT_ArrayObject:
-		case GCRT_ArrayObjectFreezable:
-		case GCRT_AddFieldPathReferencedObject:
-		case GCRT_ArrayAddFieldPathReferencedObject:
-		case GCRT_EndOfPointer:
-		case GCRT_EndOfStream:
-		case GCRT_WeakObject:
-		case GCRT_ArrayWeakObject:
-		case GCRT_LazyObject:
-		case GCRT_ArrayLazyObject:
-		case GCRT_SoftObject:
-		case GCRT_ArraySoftObject:
-		case GCRT_Delegate:
-		case GCRT_ArrayDelegate:
-		case GCRT_MulticastDelegate:
-		case GCRT_ArrayMulticastDelegate:
-		case GCRT_DynamicallyTypedValue:
-			break;
-		default:
-			UE_LOG(LogGarbage, Fatal, TEXT("Unknown token type (%u) when trying to add ARO token."), (uint32)Token.Type);
-			break;
-		};
-		TokenStreamIndex = (int32)TokenIndex;
+			uint32 SkipIndex = Out.EmitReturn();
+			Out.UpdateSkipIndexPlaceholder(SkipSkipIndex, SkipIndex);
+			return true;
+		}
 	}
-	// ARO is not in the token stream yet.
-	if (!bReplacedARO && AddReferencedObjectsPtr)
+	break;
+	case GCRT_FixedArray:
 	{
-		static const FName TokenName("AROToken");
-		EmitReferenceInfo(FGCReferenceInfo(GCRT_AddReferencedObjects, 0), TokenName);
-		EmitPointer((const void*)AddReferencedObjectsPtr);
+		Out.EmitReferenceInfo(Token.ReferenceInfo, Token.DebugName);
+		Out.EmitStride(MixedIt.ConsumeStride());
+		Out.EmitCount(MixedIt.ConsumeCount());
+
+		if (CopyStrongTokensUntilReturn(/* in-out */ MixedIt, Out, OutReturnCount))
+		{
+			Out.EmitReturn();
+			return true;
+		}
+		break;
+	}
+	case GCRT_AddTSetReferencedObjects:
+	case GCRT_AddTMapReferencedObjects:
+	{
+		Out.EmitReferenceInfo(Token.ReferenceInfo, Token.DebugName);
+		uint32 NoReturnCount = 0;
+		FProperty* Property = (FProperty*)MixedIt.ConsumePointer(/* Out */ NoReturnCount);
+		check(NoReturnCount == 0);
+		check(Property->IsA<FMapProperty>() || Property->IsA<FSetProperty>());
+		Out.EmitPointer(Property);
+
+		uint32 MixedSkipPos = MixedIt.ConsumeSkipInfo();
+		uint32 SkipSkipIndex = Out.EmitSkipIndexPlaceholder();
+	
+		if (CopyStrongTokensUntilPos(MixedIt, MixedSkipPos, Out, OutReturnCount))
+		{
+			uint32 SkipIndex = Out.EmitReturn();
+			Out.UpdateSkipIndexPlaceholder(SkipSkipIndex, SkipIndex);
+			return true;
+		}
+		break;
 	}
+	case GCRT_EndOfPointer:
+		checkf(false, TEXT("ConsumePointer() should consume all GCRT_EndOfPointer tokens"));
+	break;
+	default:
+		checkf(false, TEXT("Illegal token type %d %s"), Token.ReferenceInfo.Type, *Token.DebugName.ToString());
+	break;
+	} // switch
+
+	Out.Tokens.SetNum(RollbackNum, /* shrink */ false);		
+#if ENABLE_GC_TOKEN_DEBUG_INFO
+	Out.DebugNames.SetNum(RollbackNum, /* shrink */ false);		
+#endif
+	return false;
 }
 
-int32 FGCReferenceTokenStream::EmitReferenceInfo(FGCReferenceInfo ReferenceInfo, const FName& DebugName)
+int32 FTokenStreamBuilder::EmitReferenceInfo(FGCReferenceInfo ReferenceInfo, FName DebugName)
 {
 	int32 TokenIndex = Tokens.Add(ReferenceInfo);
 #if ENABLE_GC_TOKEN_DEBUG_INFO
-	check(TokenDebugInfo.Num() == TokenIndex);
-	check(DebugName != NAME_None);
-	TokenDebugInfo.Add(DebugName);
+	check(DebugNames.Num() == TokenIndex);
+	check(!DebugName.IsNone());
+	DebugNames.Add(DebugName);
 #endif
 	return TokenIndex;
 }
 
-/**
- * Emit placeholder for aray skip index, updated in UpdateSkipIndexPlaceholder
- *
- * @return the index of the skip index, used later in UpdateSkipIndexPlaceholder
- */
-uint32 FGCReferenceTokenStream::EmitSkipIndexPlaceholder()
+uint32 FTokenStreamBuilder::EmitSkipIndexPlaceholder()
 {
-	uint32 TokenIndex = Tokens.Add(E_GCSkipIndexPlaceholder);
+	uint32 TokenIndex = Tokens.Add(SkipIndexPlaceholderMagic);
 #if ENABLE_GC_TOKEN_DEBUG_INFO
 	static const FName TokenName("SkipIndexPlaceholder");
-	check(TokenDebugInfo.Num() == TokenIndex);
-	TokenDebugInfo.Add(TokenName);
+	check(DebugNames.Num() == TokenIndex);
+	DebugNames.Add(TokenName);
 #endif
 	return TokenIndex;
 }
 
-/**
- * Updates skip index place holder stored and passed in skip index index with passed
- * in skip index. The skip index is used to skip over tokens in the case of an empty 
- * dynamic array.
- * 
- * @param SkipIndexIndex index where skip index is stored at.
- * @param SkipIndex index to store at skip index index
- */
-void FGCReferenceTokenStream::UpdateSkipIndexPlaceholder( uint32 SkipIndexIndex, uint32 SkipIndex )
+void FTokenStreamBuilder::UpdateSkipIndexPlaceholder( uint32 SkipIndexIndex, uint32 SkipIndex )
 {
 	check( SkipIndex > 0 && SkipIndex <= (uint32)Tokens.Num() );
-	const FGCReferenceInfo& ReferenceInfo = Tokens[SkipIndex-1];
+	FGCReferenceInfo ReferenceInfo(Tokens[SkipIndex-1]);
 	check( ReferenceInfo.Type != GCRT_None );
-	check( Tokens[SkipIndexIndex] == E_GCSkipIndexPlaceholder );
+	check( Tokens[SkipIndexIndex] == SkipIndexPlaceholderMagic );
 	check( SkipIndexIndex < SkipIndex );
 	check( ReferenceInfo.ReturnCount >= 1 );
 	FGCSkipInfo SkipInfo;
@@ -3564,34 +5472,31 @@ void FGCReferenceTokenStream::UpdateSkipIndexPlaceholder( uint32 SkipIndexIndex,
 	Tokens[SkipIndexIndex]		= SkipInfo;
 }
 
-/**
- * Emit count
- *
- * @param Count count to emit
- */
-int32 FGCReferenceTokenStream::EmitCount( uint32 Count )
+int32 FTokenStreamBuilder::EmitCount( uint32 Count )
 {
 	int32 TokenIndex = Tokens.Add( Count );
 #if ENABLE_GC_TOKEN_DEBUG_INFO
 	static const FName TokenName("CountToken");
-	check(TokenDebugInfo.Num() == TokenIndex);
-	TokenDebugInfo.Add(TokenName);
+	check(DebugNames.Num() == TokenIndex);
+	DebugNames.Add(TokenName);
 #endif
 	return TokenIndex;
 }
 
-int32 FGCReferenceTokenStream::EmitPointer( void const* Ptr )
+int32 FTokenStreamBuilder::EmitPointer( void const* Ptr )
 {
+	static constexpr uint32 NumTokensPerPointer = sizeof(void*) / sizeof(uint32); //-V514
+
 	const int32 StoreIndex = Tokens.Num();
-	Tokens.AddUninitialized(GNumTokensPerPointer);
+	Tokens.AddUninitialized(NumTokensPerPointer);
 	StorePointer(&Tokens[StoreIndex], Ptr);
 
 #if ENABLE_GC_TOKEN_DEBUG_INFO
 	static const FName TokenName("PointerToken");
-	check(TokenDebugInfo.Num() == StoreIndex);
-	for (int32 PointerTokenIndex = 0; PointerTokenIndex < GNumTokensPerPointer; ++PointerTokenIndex)
+	check(DebugNames.Num() == StoreIndex);
+	for (int32 PointerTokenIndex = 0; PointerTokenIndex < NumTokensPerPointer; ++PointerTokenIndex)
 	{
-		TokenDebugInfo.Add(TokenName);
+		DebugNames.Add(TokenName);
 	}
 #endif
 
@@ -3603,65 +5508,382 @@ int32 FGCReferenceTokenStream::EmitPointer( void const* Ptr )
 	return StoreIndex;
 }
 
-/**
- * Emit stride
- *
- * @param Stride stride to emit
- */
-int32 FGCReferenceTokenStream::EmitStride( uint32 Stride )
+int32 FTokenStreamBuilder::EmitStride( uint32 Stride )
 {
 	int32 TokenIndex = Tokens.Add( Stride );
-
 #if ENABLE_GC_TOKEN_DEBUG_INFO
 	static const FName TokenName("StrideToken");
-	check(TokenDebugInfo.Num() == TokenIndex);
-	TokenDebugInfo.Add(TokenName);
+	check(DebugNames.Num() == TokenIndex);
+	DebugNames.Add(TokenName);
 #endif
-
 	return TokenIndex;
 }
 
-/**
- * Increase return count on last token.
- *
- * @return index of next token
- */
-uint32 FGCReferenceTokenStream::EmitReturn()
+uint32 FTokenStreamBuilder::EmitReturn()
 {
-	FGCReferenceInfo ReferenceInfo = Tokens.Last();
+	FGCReferenceInfo ReferenceInfo(Tokens.Last());
 	check(ReferenceInfo.Type != GCRT_None);
 	ReferenceInfo.ReturnCount++;
 	Tokens.Last() = ReferenceInfo;
 	return Tokens.Num();
 }
 
+//////////////////////////////////////////////////////////////////////////
+
+static FName ToName(ETokenlessId Id)
+{
+	static_assert(ETokenlessId::Collector == static_cast<ETokenlessId>(1u));
+	static_assert(ETokenlessId::Max == ETokenlessId::InitialReference);
+
+	static FName Names[] =
+	{
+		FName("Invalid"),
+		FName("Collector"), // ETokenlessId::Collector
+		FName("Class"), // ETokenlessId::Class
+		FName("Outer"), // ETokenlessId::Outer
+		FName("ClassOuter"), // ETokenlessId::ClassOuter
+		FName("Cluster"), // ETokenlessId::Cluster
+		FName("InitialReference"), // ETokenlessId::InitialReference
+	};
+
+	check(static_cast<uint32>(Id) < UE_ARRAY_COUNT(Names));
+	return Names[static_cast<uint32>(Id)];
+}
+
+void FTokenId::StaticAssert()
+{
+	static_assert(1u << TokenlessIdBits >= static_cast<uint32>(ETokenlessId::Max), "Need to bump TokenlessIdBits");
+}
+
+//////////////////////////////////////////////////////////////////////////
 
-FTokenInfo FGCReferenceTokenStream::GetTokenInfo(int32 TokenIndex) const
+FTokenInfo FTokenStreamOwner::GetTokenInfo(FTokenId Id) const
 {
+	if (Id.IsTokenless())
+	{
+		// Technically offsetof(UObjectBase, ClassPrivate/OuterPrivate) could be exposed via some new API or friend declaration
+		return {-1, ToName(Id.AsTokenless()) };
+	}
+
+	FTokenStreamView Stream = Id.IsMixed() ? Mixed : Strong;
+	uint32 Idx = Id.GetIndex();
+
+	check(Idx < Stream.Num);
 	FTokenInfo DebugInfo;
-	DebugInfo.Offset = FGCReferenceInfo(Tokens[TokenIndex]).Offset;
+	DebugInfo.Offset = FGCReferenceInfo(Stream.Tokens[Idx]).Offset;
 #if ENABLE_GC_TOKEN_DEBUG_INFO
-	DebugInfo.Name = TokenDebugInfo[TokenIndex];
+	DebugInfo.Name = Stream.GetDebugNames()[Idx];
 #else
-	DebugInfo.Name = FName(TEXT("Token"), TokenIndex);
+	static const FName DefaultTokenName("Token");	
+	DebugInfo.Name = DefaultTokenName; 
+#if !UE_FNAME_OUTLINE_NUMBER
+	DebugInfo.Name.SetNumber(Idx);
+#endif
 #endif
 	return DebugInfo;
 }
 
+void FTokenStreamOwner::Reset()
+{
+	UpdateStreamReference(/* out */ Strong, FTokenStreamView());
+	UpdateStreamReference(/* out */ Mixed, FTokenStreamView());
+}
+
+int64 FTokenStreamOwner::GetTokenAllocatedSize() const
+{
+	uint32 NumTokens = Mixed.Tokens == Strong.Tokens ? Mixed.Num : Mixed.Num + Strong.Num;
+	return NumTokens * sizeof(uint32);
+}
 
-FGCArrayPool* FGCArrayPool::GetGlobalSingleton()
+int64 FTokenStreamOwner::GetDebugInfoAllocatedSize() const
 {
-	static FAutoConsoleCommandWithOutputDevice GCDumpPoolCommand(
-		TEXT("gc.DumpPoolStats"),
-		TEXT("Dumps count and size of GC Pools"),
-		FConsoleCommandWithOutputDeviceDelegate::CreateStatic(&FGCArrayPool::DumpStats)
-	);
+#if ENABLE_GC_TOKEN_DEBUG_INFO
+	return GetTokenAllocatedSize() * sizeof(FName) / sizeof(uint32);
+#else
+	return 0;
+#endif
+}
+
+TMap<UClass*, TUniquePtr<FTokenStreamBuilder>> GIntrinsicClassTokens;
+
+FTokenStreamBuilder& FIntrinsicClassTokens::AllocateBuilder(UClass* Class)
+{
+	TUniquePtr<FTokenStreamBuilder>& Builder = GIntrinsicClassTokens.FindOrAdd(Class);
+	checkf(!Builder, TEXT("Allocated multiple builders for same class"));
+	Builder.Reset(new FTokenStreamBuilder(*Class));
+	return *Builder;
+}
+
+TUniquePtr<FTokenStreamBuilder> FIntrinsicClassTokens::ConsumeBuilder(UClass* Class)
+{
+	TUniquePtr<FTokenStreamBuilder>* Builder = GIntrinsicClassTokens.Find(Class);
+	return Builder && *Builder ? MoveTemp(*Builder) : TUniquePtr<FTokenStreamBuilder>();
+}
+
+//////////////////////////////////////////////////////////////////////////
+
+#if UE_WITH_GC
+
+void DumpMemoryStats(FOutputDevice& OutputDevice)
+{
+	FContextPoolScope Pool;
+	int32 TotalWeakCapacity = 0;
+	int32 NumContexts = Pool.PeekFree().Num();
+	for (const TUniquePtr<FWorkerContext>& Context : Pool.PeekFree())
+	{
+		TotalWeakCapacity += Context->WeakReferences.Max();
+	}
+
+	int32 ContextsKB = (NumContexts * sizeof(FWorkerContext) + TotalWeakCapacity * sizeof(UObject**))/1024;
+	int32 ScratchKB = static_cast<int32>(GScratchPages.CountBytes() / 1024);
+	int32 GlobalsKB = (sizeof(GScratchPages) + sizeof(GAROBlocks) + sizeof(GWorkstealingManager) + sizeof(GSlowARO)) / 1024;
+	// Excluding debug names in non-shipping configs
+	int32 TokensKB = (GNumTokens * sizeof(uint32) + GNumTokenStreams * sizeof(FTokenStreamHeader)) / 1024; 
+	int32 TotalKB = ContextsKB + ScratchKB + GlobalsKB + TokensKB;
+
+	OutputDevice.Logf(TEXT("GC mem usage: %dKB total, %dKB scratch, %dKB for %d contexts, %dKB globals (code size) and %dKB for %d tokens in %d streams"),
+						TotalKB, ScratchKB, ContextsKB, NumContexts, GlobalsKB, TokensKB, GNumTokens, GNumTokenStreams);
+}
+
+//////////////////////////////////////////////////////////////////////////
+
+// Coordinates worker starting, tail spinning and stopping
+class FWorkCoordinator
+{
+	struct alignas(PLATFORM_CACHE_LINE_SIZE) FStealableContext : public std::atomic<FWorkerContext*> {};
+	const TArrayView<FStealableContext> Contexts;
+	alignas(PLATFORM_CACHE_LINE_SIZE) std::atomic<int32> NumUsedContexts {0};
+	std::atomic<int32> NumWorkless{0}; // Number of workless workers spinning to steal work
+	std::atomic<int32> NumStopDirectly; // Stop a few workless workers directly to pick up non-GC tasks
+	std::atomic<int32> NumStopped{0}; // Synchronize when all workers are done
+
+
+public:
+	FWorkCoordinator(TArrayView<FWorkerContext*> InContexts, int32 NumTaskgraphWorkers)
+	: Contexts(reinterpret_cast<FStealableContext*>(FMemory::Malloc(sizeof(FStealableContext) * InContexts.Num())), InContexts.Num())
+	, NumStopDirectly(NumTaskgraphWorkers > MaxWorkers ? 0 : 1 + (InContexts.Num() > 5))
+	{
+		for (FStealableContext& Context : Contexts)
+		{
+			Context.store(InContexts[&Context - Contexts.GetData()], std::memory_order_relaxed);
+		}
+	}
+
+	~FWorkCoordinator()
+	{
+		check(NumUsedContexts == Contexts.Num());
+		for (FStealableContext& Context : Contexts)
+		{
+			check(Context.load(std::memory_order_relaxed) == nullptr);
+		}
+
+		FMemory::Free(Contexts.GetData());
+	}
+
+	// @return nullptr if late to the party and initial objects / references are stolen by another worker
+	// @see StealContext
+	FWorkerContext* TryStartWorking(int32 WorkerIndex)
+	{
+		if (FWorkerContext* Out = Contexts[WorkerIndex].exchange(nullptr))
+		{
+			++NumUsedContexts;
+			return Out;
+		}
+
+		return nullptr;
+	}
+
+	// Allows one thread to finish GC alone if task workers are busy processing long tasks
+	FWorkerContext* StealContext()
+	{
+		if (NumUsedContexts.load(std::memory_order_relaxed) < Contexts.Num())
+		{
+			for (FStealableContext& Context : Contexts)
+			{
+				if (FWorkerContext* Out = Context.exchange(nullptr))
+				{
+					++NumUsedContexts;
+					++NumWorkless;
+					++NumStopped;
+					--NumStopDirectly;
+					return Out;
+				}
+			}	
+		}
+
+		return nullptr;	
+	}
+
+	// @return if worker should start spinning
+	// 
+	// Some worker threads stop immediately to reduce context switching with render flipping and audio mixing threads
+	bool ReportOutOfWork()
+	{
+		++NumWorkless;
+
+		// Keep main thread spinning to reduce risk that it is swapped out when all work is done
+		if (IsInGameThread() || --NumStopDirectly < 0)
+		{
+			return true;
+		}
+		
+		++NumStopped;
+		return false;
+	}
+
+	// Notify that current work successfully stole some work
+	void ReportBackToWork()
+	{
+		--NumWorkless;
+	}
+	
+	// There's an acceptable race condition when some worker steals from the last worker, 
+	// which reports out-of-work before the thief reports back-to-work. Other workers
+	// might then stop prematurely.
+	// 
+	// This is fairly unlikely, would likely happen very late when little work remains
+	// and the remaining workers might not yield more stealable work anyway.
+	bool KeepSpinning()
+	{
+		if (NumWorkless.load(std::memory_order_acquire) < Contexts.Num())
+		{
+			return true;
+		}
+
+		++NumStopped;
+		return false;
+	}
+
+	FORCENOINLINE void SpinUntilAllStopped()
+	{
+		while (NumStopped.load() < Contexts.Num())
+		{
+			FPlatformProcess::Yield();
+		}
+	}
+};
+
+ELoot StealWork(FWorkerContext& Context, FReferenceCollector& Collector, FWorkBlock*& OutBlock)
+{
+	FWorkBlockifier& RemainingObjects = Context.ObjectsToSerialize;
+	FWorkCoordinator& Tailspin = *Context.Coordinator;
+
+	if (OutBlock = RemainingObjects.StealFullBlock(); OutBlock)
+	{
+		return ELoot::Block;
+	}
+	else if (FSlowARO::ProcessAllCalls(Context, Collector))
+	{
+		return ELoot::ARO;
+	}
+	else if (FWorkerContext* StolenContext = Tailspin.StealContext())
+	{
+		Context.InitialNativeReferences = StolenContext->InitialNativeReferences;
+		Context.SetInitialObjectsPrepadded(StolenContext->GetInitialObjects());
+		return ELoot::Context;
+	}
+	else if (Tailspin.ReportOutOfWork())
+	{
+		while (Tailspin.KeepSpinning())
+		{
+			FPlatformProcess::Yield();
+			OutBlock = RemainingObjects.StealFullBlock();
+			if (OutBlock || FSlowARO::ProcessAllCalls(Context, Collector))
+			{
+				Tailspin.ReportBackToWork();
+				return OutBlock ? ELoot::Block : ELoot::ARO;
+			}
+		}
+	}
+
+	return ELoot::Nothing;
+}
+
+void ProcessAsync(void (*ProcessSync)(void*, FWorkerContext&), void* Processor, FWorkerContext& InContext)
+{
+	checkf(InContext.ObjectsToSerialize.IsUnused(), TEXT("Use InitialObjects instead, ObjectsToSerialize.Add() may only be called during reference processing"));
+
+	TConstArrayView<UObject*> InitialObjects = InContext.GetInitialObjects();
+	TConstArrayView<UObject**> InitialReferences = InContext.InitialNativeReferences;
+
+	int32 NumTaskgraphWorkers = FTaskGraphInterface::Get().GetNumWorkerThreads();
+	const int32 NumWorkers = FMath::Clamp(NumTaskgraphWorkers, 1, MaxWorkers);
+	const int32 ObjPerWorker =  (InitialObjects.Num() + NumWorkers - 1) / NumWorkers;
+	const int32 RefPerWorker =  (InitialReferences.Num() + NumWorkers - 1) / NumWorkers;
+
+	// Allocate contexts
+	FContextPoolScope ContextPool;
+	checkf(ContextPool.NumAllocated() == 1, TEXT("Other contexts forbidden during parallel reference collection. Work-stealing from all live contexts. "));
+	
+	FWorkerContext* ContextArray[MaxWorkers];
+	TArrayView<FWorkerContext*> Contexts = MakeArrayView(ContextArray, NumWorkers);
+	Contexts[0] = &InContext;
+	for (FWorkerContext*& Context : Contexts.RightChop(1))
+	{
+		Context = ContextPool.AllocateFromPool();
+	}
+
+	// Setup work-stealing queues and distribute initial workload across worker contexts
+	TSharedRef<FWorkCoordinator> Coordinator = MakeShared<FWorkCoordinator>(Contexts, NumTaskgraphWorkers);
+	GSlowARO.GetPostInit().SetupWorkerQueues(NumWorkers);
+	for (FWorkerContext* Context : Contexts)
+	{
+		int32 Idx = Context->GetWorkerIndex();
+		check(Idx >= 0 && Idx < NumWorkers);
+		Context->ObjectsToSerialize.SetAsyncQueue(GWorkstealingManager.Queues[Idx]);
+		// Initial objects is already padded at the end and its safe to prefetch in the middle too
+		Context->SetInitialObjectsPrepadded(InitialObjects.Mid(Idx * ObjPerWorker, ObjPerWorker));
+		Context->InitialNativeReferences = InitialReferences.Mid(Idx * RefPerWorker, RefPerWorker);
+		Context->Coordinator = &Coordinator.Get();
+	}
+	
+	// Kick workers
+	for (int32 Idx = 1; Idx < NumWorkers; ++Idx)
+	{
+		Tasks::Launch(TEXT("CollectReferences"), [=]() 
+			{
+				if (FWorkerContext* Context = Coordinator->TryStartWorking(Idx))
+				{
+					ProcessSync(Processor, *Context);
+				}
+			});
+	}
+
+	// Start working ourselves
+	if (FWorkerContext* Context = Coordinator->TryStartWorking(0))
+	{
+		ProcessSync(Processor, *Context);
+	}
+
+	// Wait until all work is complete. Current thread can steal and complete everything 
+	// alone if task workers are busy with long-running tasks.
+	Coordinator->SpinUntilAllStopped();
 
-	static FGCArrayPool* GlobalSingleton = nullptr;
+	// Tear down contexts and work-stealing queues
+	for (FWorkerContext* Context : Contexts)
+	{
+		Context->Coordinator = nullptr;
+		Context->InitialNativeReferences = TConstArrayView<UObject**>();
+		Context->ResetInitialObjects();
+		GWorkstealingManager.Queues[Context->GetWorkerIndex()].CheckEmpty();
+		Context->ObjectsToSerialize.ResetAsyncQueue();
+	}
+	
+	GSlowARO.GetPostInit().ResetWorkerQueues();
 
-	if (!GlobalSingleton)
+	for (FWorkerContext* Context : Contexts.RightChop(1))
 	{
-		GlobalSingleton = new FGCArrayPool();
+		ContextPool.ReturnToPool(Context);
 	}
-	return GlobalSingleton;
+}
+
+//////////////////////////////////////////////////////////////////////////
+
+#endif //UE_WITH_GC
+
+} // namespace UE::GC
+ 
+int32 GetNumCollectReferenceWorkers()
+{
+	return FMath::Clamp(FTaskGraphInterface::Get().GetNumWorkerThreads(), 1, UE::GC::MaxWorkers);
 }
