GC reachability optimizations. ~1.5x speedup and ~2.5MB memory saved for an internal title.

diff --git a/Engine/Source/Runtime/CoreUObject/Private/UObject/GarbageCollectionHistory.cpp b/Engine/Source/Runtime/CoreUObject/Private/UObject/GarbageCollectionHistory.cpp
--- a/Engine/Source/Runtime/CoreUObject/Private/UObject/GarbageCollectionHistory.cpp
+++ b/Engine/Source/Runtime/CoreUObject/Private/UObject/GarbageCollectionHistory.cpp
@@ -9,6 +9,7 @@
 #include "Misc/TimeGuard.h"
 #include "HAL/IConsoleManager.h"
 #include "Misc/App.h"
+#include "UObject/FastReferenceCollector.h"
 #include "UObject/UObjectAllocator.h"
 #include "UObject/UObjectBase.h"
 #include "UObject/Object.h"
@@ -70,7 +71,7 @@ void FGCHistory::Cleanup(FGCSnapshot& InSnapshot)
 		delete DirectReferenceInfos.Value;
 	}
 	InSnapshot.DirectReferences.Reset();
-	for (TPair<UObject*, FGCObjectInfo*>& ObjectToInfoPair : InSnapshot.ObjectToInfoMap)
+	for (TPair<const UObject*, FGCObjectInfo*>& ObjectToInfoPair : InSnapshot.ObjectToInfoMap)
 	{
 		delete ObjectToInfoPair.Value;
 	}
@@ -100,9 +101,9 @@ void FGCHistory::SetHistorySize(int32 HistorySize)
 	}
 }
 
-void FGCHistory::MergeArrayStructHistory(TMap<UObject*, TArray<FGCDirectReference>*>& History, FGCSnapshot& Snapshot)
+void FGCHistory::MergeArrayStructHistory(TMap<const UObject*, TArray<FGCDirectReference>*>& History, FGCSnapshot& Snapshot)
 {
-	for (TPair<UObject*, TArray<FGCDirectReference>*>& ReferencePair : History)
+	for (TPair<const UObject*, TArray<FGCDirectReference>*>& ReferencePair : History)
 	{
 		FGCObjectInfo* ReferencerInfo = FGCObjectInfo::FindOrAddInfoHelper(ReferencePair.Key, Snapshot.ObjectToInfoMap);
 		TArray<FGCDirectReferenceInfo>*& DirectReferencesInfoArray = Snapshot.DirectReferences.FindOrAdd(ReferencerInfo);
@@ -119,7 +120,7 @@ void FGCHistory::MergeArrayStructHistory(TMap<UObject*, TArray<FGCDirectReferenc
 	}
 }
 
-void FGCHistory::Update(TArray<FGCArrayStruct*>& AllArrays)
+void FGCHistory::Update(TConstArrayView<TUniquePtr<UE::GC::FWorkerContext>> Contexts)
 {
 	if (IsActive())
 	{
@@ -133,9 +134,9 @@ void FGCHistory::Update(TArray<FGCArrayStruct*>& AllArrays)
 		}
 
 		// Copy information gather in all of the GC struct arrays
-		for (FGCArrayStruct* ArrayStruct : AllArrays)
+		for (const TUniquePtr<UE::GC::FWorkerContext>& Context : Contexts)
 		{
-			MergeArrayStructHistory(ArrayStruct->History, CurrentSnapshot);
+			MergeArrayStructHistory(Context->History, CurrentSnapshot);
 		}
 	}
 }
diff --git a/Engine/Source/Runtime/CoreUObject/Private/UObject/GarbageCollectionVerification.cpp b/Engine/Source/Runtime/CoreUObject/Private/UObject/GarbageCollectionVerification.cpp
--- a/Engine/Source/Runtime/CoreUObject/Private/UObject/GarbageCollectionVerification.cpp
+++ b/Engine/Source/Runtime/CoreUObject/Private/UObject/GarbageCollectionVerification.cpp
@@ -52,7 +52,7 @@ public:
 	{
 		return NumErrors.GetValue();
 	}
-	FORCEINLINE_DEBUGGABLE void HandleTokenStreamObjectReference(FGCArrayStruct& ObjectsToSerializeStruct, UObject* ReferencingObject, UObject*& Object, const int32 TokenIndex, const EGCTokenType TokenType, bool bAllowReferenceElimination)
+	FORCEINLINE_DEBUGGABLE void HandleTokenStreamObjectReference(FGCArrayStruct& ObjectsToSerializeStruct, UObject* ReferencingObject, UObject*& Object, UE::GC::FTokenId TokenIndex, const EGCTokenType TokenType, bool bAllowReferenceElimination)
 	{
 		if (Object)
 		{
@@ -66,7 +66,7 @@ public:
 				FString TokenDebugInfo;
 				if (UClass *Class = (ReferencingObject ? ReferencingObject->GetClass() : nullptr))
 				{
-					FTokenInfo TokenInfo = Class->ReferenceTokenStream.GetTokenInfo(TokenIndex);
+					FTokenInfo TokenInfo = Class->ReferenceTokens.GetTokenInfo(TokenIndex);
 					TokenDebugInfo = FString::Printf(TEXT("ReferencingObjectClass: %s, Property Name: %s, Offset: %d"),
 						*Class->GetFullName(), *TokenInfo.Name.GetPlainNameString(), TokenInfo.Offset);
 				}
@@ -79,7 +79,7 @@ public:
 				UE_LOG(LogGarbage, Fatal, TEXT("Invalid object while verifying Disregard for GC assumptions: 0x%016llx, ReferencingObject: %s, %s, TokenIndex: %d"),
 					(int64)(PTRINT)Object,
 					ReferencingObject ? *ReferencingObject->GetFullName() : TEXT("NULL"),
-					*TokenDebugInfo, TokenIndex);
+					*TokenDebugInfo, TokenIndex.AsPrintableIndex());
 				}
 #endif // ENABLE_GC_OBJECT_CHECKS
 
@@ -96,44 +96,38 @@ public:
 		}
 	}
 };
-typedef TDefaultReferenceCollector<FDisregardSetReferenceProcessor> FDisregardSetReferenceCollector;
 
 void VerifyGCAssumptions()
 {	
 	int32 MaxNumberOfObjects = GUObjectArray.GetObjectArrayNumPermanent();
 
 	FDisregardSetReferenceProcessor Processor;
-	TFastReferenceCollector<
-		FDisregardSetReferenceProcessor, 
-		FDisregardSetReferenceCollector, 
-		FGCArrayPool, 
-		EFastReferenceCollectorOptions::AutogenerateTokenStream | EFastReferenceCollectorOptions::ProcessNoOpTokens
-	> ReferenceCollector(Processor, FGCArrayPool::Get());
 
-	int32 NumThreads = FMath::Max(1, FTaskGraphInterface::Get().GetNumWorkerThreads());
+	int32 NumThreads = GetNumCollectReferenceWorkers();
 	int32 NumberOfObjectsPerThread = (MaxNumberOfObjects / NumThreads) + 1;
-	FGCArrayStruct* ArrayStructs = new FGCArrayStruct[NumThreads];
-
-	ParallelFor( TEXT("GarbageCollection.PF"),NumThreads,1, [&ReferenceCollector, ArrayStructs, NumberOfObjectsPerThread, NumThreads, MaxNumberOfObjects](int32 ThreadIndex)
+	
+	ParallelFor( TEXT("GC.VerifyAssumptions"),NumThreads,1, [&Processor, NumberOfObjectsPerThread, NumThreads, MaxNumberOfObjects](int32 ThreadIndex)
 	{
 		int32 FirstObjectIndex = ThreadIndex * NumberOfObjectsPerThread;
 		int32 NumObjects = (ThreadIndex < (NumThreads - 1)) ? NumberOfObjectsPerThread : (MaxNumberOfObjects - (NumThreads - 1)*NumberOfObjectsPerThread);
-		FGCArrayStruct& ArrayStruct = ArrayStructs[ThreadIndex];
-		ArrayStruct.ObjectsToSerialize.Reserve(NumberOfObjectsPerThread);
+		
+		TArray<UObject*> ObjectsToSerialize;
+		ObjectsToSerialize.Reserve(NumberOfObjectsPerThread + UE::GC::ObjectLookahead);
 
 		for (int32 ObjectIndex = 0; ObjectIndex < NumObjects && (FirstObjectIndex + ObjectIndex) < MaxNumberOfObjects; ++ObjectIndex)
 		{
 			FUObjectItem& ObjectItem = GUObjectArray.GetObjectItemArrayUnsafe()[FirstObjectIndex + ObjectIndex];
 			if (ObjectItem.Object && ObjectItem.Object != FGCObject::GGCObjectReferencer)
 			{
-				ArrayStruct.ObjectsToSerialize.Add(static_cast<UObject*>(ObjectItem.Object));
+				ObjectsToSerialize.Add(static_cast<UObject*>(ObjectItem.Object));
 			}
 		}
-		ReferenceCollector.CollectReferences(ArrayStruct);
+		
+		UE::GC::FWorkerContext Context;
+		Context.SetInitialObjectsUnpadded(ObjectsToSerialize);
+		CollectReferences(Processor, Context);
 	});
 
-	delete[] ArrayStructs;
-
 	UE_CLOG(Processor.GetErrorCount() > 0, LogGarbage, Fatal, TEXT("Encountered %d object(s) breaking Disregard for GC assumptions. Please check log for details."), Processor.GetErrorCount());
 }
 
@@ -177,7 +171,7 @@ public:
 	* @param TokenIndex Index to the token stream where the reference was found.
 	* @param bAllowReferenceElimination True if reference elimination is allowed (ignored when constructing clusters).
 	*/
-	FORCEINLINE_DEBUGGABLE void HandleTokenStreamObjectReference(FGCArrayStruct& ObjectsToSerializeStruct, UObject* ReferencingObject, UObject*& Object, const int32 TokenIndex, const EGCTokenType TokenType, bool bAllowReferenceElimination)
+	FORCEINLINE_DEBUGGABLE void HandleTokenStreamObjectReference(FGCArrayStruct& ObjectsToSerializeStruct, UObject* ReferencingObject, UObject*& Object, UE::GC::FTokenId TokenIndex, const EGCTokenType TokenType, bool bAllowReferenceElimination)
 	{
 		if (Object)
 		{
@@ -197,7 +191,7 @@ public:
 				FString TokenDebugInfo;
 				if (UClass *Class = (ReferencingObject ? ReferencingObject->GetClass() : nullptr))
 				{
-					FTokenInfo TokenInfo = Class->ReferenceTokenStream.GetTokenInfo(TokenIndex);
+					FTokenInfo TokenInfo = Class->ReferenceTokens.GetTokenInfo(TokenIndex);
 					TokenDebugInfo = FString::Printf(TEXT("ReferencingObjectClass: %s, Property Name: %s, Offset: %d"),
 						*Class->GetFullName(), *TokenInfo.Name.GetPlainNameString(), TokenInfo.Offset);
 				}
@@ -214,7 +208,7 @@ public:
 				UE_LOG(LogGarbage, Fatal, TEXT("Invalid object while verifying cluster assumptions: 0x%016llx, ReferencingObject: %s, %s, TokenIndex: %d"),
 					(int64)(PTRINT)Object,
 					ReferencingObject ? *ReferencingObject->GetFullName() : TEXT("NULL"),
-					*TokenDebugInfo, TokenIndex);
+					*TokenDebugInfo, TokenIndex.AsPrintableIndex());
 			}
 #endif // ENABLE_GC_OBJECT_CHECKS
 
@@ -283,57 +277,51 @@ public:
 		}
 	}
 };
-typedef TDefaultReferenceCollector<FClusterVerifyReferenceProcessor> FClusterVerifyReferenceCollector;
 
 void VerifyClustersAssumptions()
 {
 	int32 MaxNumberOfClusters = GUObjectClusters.GetClustersUnsafe().Num();
-	int32 NumThreads = FMath::Max(1, FTaskGraphInterface::Get().GetNumWorkerThreads());
+	int32 NumThreads = GetNumCollectReferenceWorkers();
 	int32 NumberOfClustersPerThread = (MaxNumberOfClusters / NumThreads) + 1;
-	FGCArrayStruct* ArrayStructs = new FGCArrayStruct[NumThreads];
+	
 	FThreadSafeCounter NumErrors(0);
 
-	ParallelFor( TEXT("GarbageCollection.PF"),NumThreads,1, [&NumErrors, ArrayStructs, NumberOfClustersPerThread, NumThreads, MaxNumberOfClusters](int32 ThreadIndex)
+	ParallelFor( TEXT("GC.VerifyClusterAssumptions"),NumThreads,1, [&NumErrors, NumberOfClustersPerThread, NumThreads, MaxNumberOfClusters](int32 ThreadIndex)
 	{
 		int32 FirstClusterIndex = ThreadIndex * NumberOfClustersPerThread;
 		int32 NumClusters = (ThreadIndex < (NumThreads - 1)) ? NumberOfClustersPerThread : (MaxNumberOfClusters - (NumThreads - 1) * NumberOfClustersPerThread);
-		FGCArrayStruct& ArrayStruct = ArrayStructs[ThreadIndex];
-		
+				
 		FClusterVerifyReferenceProcessor Processor;
-		TFastReferenceCollector<
-			FClusterVerifyReferenceProcessor, 
-			FClusterVerifyReferenceCollector, 
-			FGCArrayPool, 
-			EFastReferenceCollectorOptions::AutogenerateTokenStream | EFastReferenceCollectorOptions::ProcessNoOpTokens
-		> ReferenceCollector(Processor, FGCArrayPool::Get());
 
+		TArray<UObject*> ObjectsToSerialize;
 		for (int32 ClusterIndex = 0; ClusterIndex < NumClusters && (FirstClusterIndex + ClusterIndex) < MaxNumberOfClusters; ++ClusterIndex)
 		{
 			FUObjectCluster& Cluster = GUObjectClusters.GetClustersUnsafe()[FirstClusterIndex + ClusterIndex];
 			if (Cluster.RootIndex >= 0 && Cluster.Objects.Num())
 			{
-				ArrayStruct.ObjectsToSerialize.Reset();
-				ArrayStruct.ObjectsToSerialize.Reserve(Cluster.Objects.Num() + 1);
+				ObjectsToSerialize.Reset(Cluster.Objects.Num() + 1 + UE::GC::ObjectLookahead);
 				{
 					FUObjectItem* RootItem = GUObjectArray.IndexToObject(Cluster.RootIndex);
 					check(RootItem);
 					check(RootItem->Object);
-					ArrayStruct.ObjectsToSerialize.Add(static_cast<UObject*>(RootItem->Object));
+					ObjectsToSerialize.Add(static_cast<UObject*>(RootItem->Object));
 				}
 				for (int32 ObjectIndex : Cluster.Objects)
 				{
 					FUObjectItem* ObjectItem = GUObjectArray.IndexToObject(ObjectIndex);
 					check(ObjectItem);
 					check(ObjectItem->Object);
-					ArrayStruct.ObjectsToSerialize.Add(static_cast<UObject*>(ObjectItem->Object));
+					ObjectsToSerialize.Add(static_cast<UObject*>(ObjectItem->Object));
 				}
-				ReferenceCollector.CollectReferences(ArrayStruct);
+
+				UE::GC::FWorkerContext Context;
+				Context.SetInitialObjectsUnpadded(ObjectsToSerialize);
+				CollectReferences(Processor, Context);
 				NumErrors.Add(Processor.GetErrorCount());
 			}			
 		}		
 	});
 
-	delete[] ArrayStructs;
 
 	UE_CLOG(NumErrors.GetValue() > 0, LogGarbage, Fatal, TEXT("Encountered %d object(s) breaking GC Clusters assumptions. Please check log for details."), NumErrors.GetValue());
 }
@@ -345,7 +333,7 @@ void VerifyObjectFlagMirroring()
 	int32 NumberOfObjectsPerThread = (MaxNumberOfObjects / NumThreads) + 1;
 	std::atomic<uint32> NumErrors(0);
 
-	ParallelFor( TEXT("GarbageCollection.PF"),NumThreads,1, [&NumErrors, NumberOfObjectsPerThread, NumThreads, MaxNumberOfObjects](int32 ThreadIndex)
+	ParallelFor( TEXT("GC.VerifyFlagMirroring"),NumThreads,1, [&NumErrors, NumberOfObjectsPerThread, NumThreads, MaxNumberOfObjects](int32 ThreadIndex)
 	{
 		int32 FirstObjectIndex = ThreadIndex * NumberOfObjectsPerThread;
 		int32 NumObjects = (ThreadIndex < (NumThreads - 1)) ? NumberOfObjectsPerThread : (MaxNumberOfObjects - (NumThreads - 1) * NumberOfObjectsPerThread);
diff --git a/Engine/Source/Runtime/CoreUObject/Private/UObject/Obj.cpp b/Engine/Source/Runtime/CoreUObject/Private/UObject/Obj.cpp
--- a/Engine/Source/Runtime/CoreUObject/Private/UObject/Obj.cpp
+++ b/Engine/Source/Runtime/CoreUObject/Private/UObject/Obj.cpp
@@ -1765,6 +1765,12 @@ public:
 	{
 		if( !ReferencingObject->GetClass()->IsChildOf(UClass::StaticClass()) )
 		{
+			// Didn't dare switching from SerializeScriptProperties to new and faster AddPropertyReferencers.
+			// This collector IsIgnoringTransient and SerializeScriptProperties will skip
+			// transient default objects but AddPropertyReferencers / SerializeBin will skip all transient 
+			// properties, including default ones. Not sure if this matters.
+			// 
+			// See FReferenceFinder::FindReferences whose collector doesn't ignore transient
 			FVerySlowReferenceCollectorArchiveScope CollectorScope(GetVerySlowReferenceCollectorArchive(), ReferencingObject);
 			ReferencingObject->SerializeScriptProperties(CollectorScope.GetArchive());
 		}
diff --git a/Engine/Source/Runtime/CoreUObject/Private/UObject/ObjectRedirector.cpp b/Engine/Source/Runtime/CoreUObject/Private/UObject/ObjectRedirector.cpp
--- a/Engine/Source/Runtime/CoreUObject/Private/UObject/ObjectRedirector.cpp
+++ b/Engine/Source/Runtime/CoreUObject/Private/UObject/ObjectRedirector.cpp
@@ -102,6 +102,7 @@ bool UObjectRedirector::GetNativePropertyValues( TMap<FString,FString>& out_Prop
 
 IMPLEMENT_CORE_INTRINSIC_CLASS(UObjectRedirector, UObject,
 	{
-		Class->EmitObjectReference(STRUCT_OFFSET(UObjectRedirector, DestinationObject), TEXT("DestinationObject"));
+		UE::GC::FTokenStreamBuilder& Builder = UE::GC::FIntrinsicClassTokens::AllocateBuilder(Class);
+		Builder.EmitObjectReference(STRUCT_OFFSET(UObjectRedirector, DestinationObject), TEXT("DestinationObject"));
 	}
 );
diff --git a/Engine/Source/Runtime/CoreUObject/Private/UObject/Package.cpp b/Engine/Source/Runtime/CoreUObject/Private/UObject/Package.cpp
--- a/Engine/Source/Runtime/CoreUObject/Private/UObject/Package.cpp
+++ b/Engine/Source/Runtime/CoreUObject/Private/UObject/Package.cpp
@@ -361,7 +361,8 @@ void UPackage::SetLoadedByEditorPropertiesOnly(bool bIsEditorOnly, bool bRecursi
 PRAGMA_DISABLE_DEPRECATION_WARNINGS
 IMPLEMENT_CORE_INTRINSIC_CLASS(UPackage, UObject,
 	{
-		Class->EmitObjectReference(STRUCT_OFFSET(UPackage, MetaData), TEXT("MetaData"));
+		UE::GC::FTokenStreamBuilder& Builder = UE::GC::FIntrinsicClassTokens::AllocateBuilder(Class);
+		Builder.EmitObjectReference(STRUCT_OFFSET(UPackage, MetaData), TEXT("MetaData"));
 	}
 );
 PRAGMA_ENABLE_DEPRECATION_WARNINGS
diff --git a/Engine/Source/Runtime/CoreUObject/Private/UObject/PropertyObject.cpp b/Engine/Source/Runtime/CoreUObject/Private/UObject/PropertyObject.cpp
--- a/Engine/Source/Runtime/CoreUObject/Private/UObject/PropertyObject.cpp
+++ b/Engine/Source/Runtime/CoreUObject/Private/UObject/PropertyObject.cpp
@@ -101,10 +101,12 @@ void FObjectProperty::SerializeItem( FStructuredArchive::FSlot Slot, void* Value
 		TObjectPtr<UObject>* ObjectPtr = GetPropertyValuePtr(Value);
 		Slot << (*ObjectPtr);
 
+#if !(UE_BUILD_TEST || UE_BUILD_SHIPPING) 
 		if(!UnderlyingArchive.IsSaving())
 		{
 			CheckValidObject(ObjectPtr, *ObjectPtr);
 		}
+#endif
 	}
 	else
 	{
diff --git a/Engine/Source/Runtime/CoreUObject/Private/UObject/PropertyObjectPtr.cpp b/Engine/Source/Runtime/CoreUObject/Private/UObject/PropertyObjectPtr.cpp
--- a/Engine/Source/Runtime/CoreUObject/Private/UObject/PropertyObjectPtr.cpp
+++ b/Engine/Source/Runtime/CoreUObject/Private/UObject/PropertyObjectPtr.cpp
@@ -62,10 +62,12 @@ void FObjectPtrProperty::StaticSerializeItem(const FObjectPropertyBase* ObjectPr
 	{
 		Slot << *ObjectPtr;
 
+#if !(UE_BUILD_TEST || UE_BUILD_SHIPPING) 
 		if(!UnderlyingArchive.IsSaving() && IsObjectHandleResolved(ObjectPtr->GetHandle()))
 		{
 			ObjectProperty->CheckValidObject(ObjectPtr, CurrentValue);
 		}
+#endif
 	}
 	else
 	{
diff --git a/Engine/Source/Runtime/CoreUObject/Private/UObject/PropertyPrivate.cpp b/Engine/Source/Runtime/CoreUObject/Private/UObject/PropertyPrivate.cpp
--- a/Engine/Source/Runtime/CoreUObject/Private/UObject/PropertyPrivate.cpp
+++ b/Engine/Source/Runtime/CoreUObject/Private/UObject/PropertyPrivate.cpp
@@ -172,8 +172,9 @@ namespace UEEnumProperty_Private
 
 IMPLEMENT_CORE_INTRINSIC_CLASS(UEnumProperty, UProperty,
 	{
-		Class->EmitObjectReference(UEEnumProperty_Private::FEnumPropertyFriend::EnumOffset, TEXT("Enum"));
-		Class->EmitObjectReference(UEEnumProperty_Private::FEnumPropertyFriend::UnderlyingPropOffset, TEXT("UnderlyingProp"));
+		UE::GC::FTokenStreamBuilder& Builder = UE::GC::FIntrinsicClassTokens::AllocateBuilder(Class);
+		Builder.EmitObjectReference(UEEnumProperty_Private::FEnumPropertyFriend::EnumOffset, TEXT("Enum"));
+		Builder.EmitObjectReference(UEEnumProperty_Private::FEnumPropertyFriend::UnderlyingPropOffset, TEXT("UnderlyingProp"));
 	}
 );
 
@@ -192,7 +193,8 @@ void UArrayProperty::AddReferencedObjects(UObject* InThis, FReferenceCollector&
 
 IMPLEMENT_CORE_INTRINSIC_CLASS(UArrayProperty, UProperty,
 	{
-		Class->EmitObjectReference(STRUCT_OFFSET(UArrayProperty, Inner), TEXT("Inner"));
+		UE::GC::FTokenStreamBuilder& Builder = UE::GC::FIntrinsicClassTokens::AllocateBuilder(Class);
+		Builder.EmitObjectReference(STRUCT_OFFSET(UArrayProperty, Inner), TEXT("Inner"));
 
 		// Ensure that TArray and FScriptArray are interchangeable, as FScriptArray will be used to access a native array property
 		// from script that is declared as a TArray in C++.
@@ -254,7 +256,8 @@ void UObjectPropertyBase::AddReferencedObjects(UObject* InThis, FReferenceCollec
 
 IMPLEMENT_CORE_INTRINSIC_CLASS(UObjectPropertyBase, UProperty,
 	{
-		Class->EmitObjectReference(STRUCT_OFFSET(UObjectProperty, PropertyClass), TEXT("PropertyClass"));
+		UE::GC::FTokenStreamBuilder& Builder = UE::GC::FIntrinsicClassTokens::AllocateBuilder(Class);
+		Builder.EmitObjectReference(STRUCT_OFFSET(UObjectProperty, PropertyClass), TEXT("PropertyClass"));
 	}
 );
 
@@ -369,7 +372,8 @@ void UByteProperty::AddReferencedObjects(UObject* InThis, FReferenceCollector& C
 
 IMPLEMENT_CORE_INTRINSIC_CLASS(UByteProperty, UNumericProperty,
 	{
-		Class->EmitObjectReference(STRUCT_OFFSET(UByteProperty, Enum), TEXT("Enum"));
+		UE::GC::FTokenStreamBuilder& Builder = UE::GC::FIntrinsicClassTokens::AllocateBuilder(Class);
+		Builder.EmitObjectReference(STRUCT_OFFSET(UByteProperty, Enum), TEXT("Enum"));
 	}
 );
 
@@ -438,7 +442,8 @@ void UClassProperty::AddReferencedObjects(UObject* InThis, FReferenceCollector&
 
 IMPLEMENT_CORE_INTRINSIC_CLASS(UClassProperty, UObjectProperty,
 	{
-		Class->EmitObjectReference(STRUCT_OFFSET(UClassProperty, MetaClass), TEXT("MetaClass"));
+		UE::GC::FTokenStreamBuilder& Builder = UE::GC::FIntrinsicClassTokens::AllocateBuilder(Class);
+		Builder.EmitObjectReference(STRUCT_OFFSET(UClassProperty, MetaClass), TEXT("MetaClass"));
 	}
 );
 
@@ -473,7 +478,8 @@ void UDelegateProperty::BeginDestroy()
 
 IMPLEMENT_CORE_INTRINSIC_CLASS(UDelegateProperty, UProperty,
 	{
-		Class->EmitObjectReference(STRUCT_OFFSET(UDelegateProperty, SignatureFunction), TEXT("SignatureFunction"));
+		UE::GC::FTokenStreamBuilder& Builder = UE::GC::FIntrinsicClassTokens::AllocateBuilder(Class);
+		Builder.EmitObjectReference(STRUCT_OFFSET(UDelegateProperty, SignatureFunction), TEXT("SignatureFunction"));
 	}
 );
 
@@ -574,7 +580,8 @@ void UInterfaceProperty::AddReferencedObjects(UObject* InThis, FReferenceCollect
 
 IMPLEMENT_CORE_INTRINSIC_CLASS(UInterfaceProperty, UProperty,
 	{
-		Class->EmitObjectReference(STRUCT_OFFSET(UInterfaceProperty, InterfaceClass), TEXT("InterfaceClass"));
+		UE::GC::FTokenStreamBuilder& Builder = UE::GC::FIntrinsicClassTokens::AllocateBuilder(Class);
+		Builder.EmitObjectReference(STRUCT_OFFSET(UInterfaceProperty, InterfaceClass), TEXT("InterfaceClass"));
 	}
 );
 
@@ -611,8 +618,9 @@ void UMapProperty::AddReferencedObjects(UObject* InThis, FReferenceCollector& Co
 
 IMPLEMENT_CORE_INTRINSIC_CLASS(UMapProperty, UProperty,
 	{
-		Class->EmitObjectReference(STRUCT_OFFSET(UMapProperty, KeyProp),   TEXT("KeyProp"));
-		Class->EmitObjectReference(STRUCT_OFFSET(UMapProperty, ValueProp), TEXT("ValueProp"));
+		UE::GC::FTokenStreamBuilder& Builder = UE::GC::FIntrinsicClassTokens::AllocateBuilder(Class);
+		Builder.EmitObjectReference(STRUCT_OFFSET(UMapProperty, KeyProp),   TEXT("KeyProp"));
+		Builder.EmitObjectReference(STRUCT_OFFSET(UMapProperty, ValueProp), TEXT("ValueProp"));
 
 		// Ensure that TArray and FScriptMap are interchangeable, as FScriptMap will be used to access a native array property
 		// from script that is declared as a TArray in C++.
@@ -650,7 +658,8 @@ void UMulticastDelegateProperty::BeginDestroy()
 
 IMPLEMENT_CORE_INTRINSIC_CLASS(UMulticastDelegateProperty, UProperty,
 	{
-		Class->EmitObjectReference(STRUCT_OFFSET(UMulticastDelegateProperty, SignatureFunction), TEXT("SignatureFunction"));
+		UE::GC::FTokenStreamBuilder& Builder = UE::GC::FIntrinsicClassTokens::AllocateBuilder(Class);
+		Builder.EmitObjectReference(STRUCT_OFFSET(UMulticastDelegateProperty, SignatureFunction), TEXT("SignatureFunction"));
 	}
 );
 
@@ -704,7 +713,8 @@ void USetProperty::AddReferencedObjects(UObject* InThis, FReferenceCollector& Co
 
 IMPLEMENT_CORE_INTRINSIC_CLASS(USetProperty, UProperty,
 	{
-		Class->EmitObjectReference(STRUCT_OFFSET(USetProperty, ElementProp), TEXT("ElementProp"));
+		UE::GC::FTokenStreamBuilder& Builder = UE::GC::FIntrinsicClassTokens::AllocateBuilder(Class);
+		Builder.EmitObjectReference(STRUCT_OFFSET(USetProperty, ElementProp), TEXT("ElementProp"));
 
 		// Ensure that TArray and FScriptMap are interchangeable, as FScriptMap will be used to access a native array property
 		// from script that is declared as a TArray in C++.
@@ -778,7 +788,8 @@ void USoftClassProperty::AddReferencedObjects(UObject* InThis, FReferenceCollect
 
 IMPLEMENT_CORE_INTRINSIC_CLASS(USoftClassProperty, USoftObjectProperty,
 	{
-		Class->EmitObjectReference(STRUCT_OFFSET(USoftClassProperty, MetaClass), TEXT("MetaClass"));
+		UE::GC::FTokenStreamBuilder& Builder = UE::GC::FIntrinsicClassTokens::AllocateBuilder(Class);
+		Builder.EmitObjectReference(STRUCT_OFFSET(USoftClassProperty, MetaClass), TEXT("MetaClass"));
 	}
 );
 
@@ -879,7 +890,8 @@ void UStructProperty::AddReferencedObjects(UObject* InThis, FReferenceCollector&
 
 IMPLEMENT_CORE_INTRINSIC_CLASS(UStructProperty, UProperty,
 	{
-		Class->EmitObjectReference(STRUCT_OFFSET(UStructProperty, Struct), TEXT("Struct"));
+		UE::GC::FTokenStreamBuilder& Builder = UE::GC::FIntrinsicClassTokens::AllocateBuilder(Class);
+		Builder.EmitObjectReference(STRUCT_OFFSET(UStructProperty, Struct), TEXT("Struct"));
 	}
 );
 
diff --git a/Engine/Source/Runtime/CoreUObject/Private/UObject/ReferenceChainSearch.cpp b/Engine/Source/Runtime/CoreUObject/Private/UObject/ReferenceChainSearch.cpp
--- a/Engine/Source/Runtime/CoreUObject/Private/UObject/ReferenceChainSearch.cpp
+++ b/Engine/Source/Runtime/CoreUObject/Private/UObject/ReferenceChainSearch.cpp
@@ -35,7 +35,7 @@ FReferenceChainSearch::FGraphNode* FReferenceChainSearch::FindOrAddNode(FGCObjec
 	return ObjectNode;
 }
 
-FReferenceChainSearch::FGraphNode* FReferenceChainSearch::FindOrAddNode(UObject* InObjectToFindNodeFor)
+FReferenceChainSearch::FGraphNode* FReferenceChainSearch::FindOrAddNode(const UObject* InObjectToFindNodeFor)
 {
 	return FindOrAddNode(FGCObjectInfo::FindOrAddInfoHelper(InObjectToFindNodeFor, ObjectToInfoMap));
 }
@@ -587,11 +587,11 @@ class FDirectReferenceProcessor : public FSimpleReferenceProcessorBase
 {	
 protected:
 	TSet<FReferenceChainSearch::FObjectReferenceInfo> ReferencedObjects;
-	TMap<UObject*, FGCObjectInfo*>& ObjectToInfoMap;
+	TMap<const UObject*, FGCObjectInfo*>& ObjectToInfoMap;
 
 public:
 
-	FDirectReferenceProcessor(TMap<UObject*, FGCObjectInfo*>& InObjectToInfoMap)
+	FDirectReferenceProcessor(TMap<const UObject*, FGCObjectInfo*>& InObjectToInfoMap)
 	: ObjectToInfoMap(InObjectToInfoMap)
 	{
 	}
@@ -606,7 +606,7 @@ public:
 		return ReferencedObjects;
 	}
 
-	void HandleTokenStreamObjectReference(FGCArrayStruct& ObjectsToSerializeStruct, UObject* ReferencingObject, UObject*& Object, const int32 TokenIndex, const EGCTokenType TokenType, bool bAllowReferenceElimination)
+	void HandleTokenStreamObjectReference(FGCArrayStruct& ObjectsToSerializeStruct, UObject* ReferencingObject, UObject*& Object, UE::GC::FTokenId TokenIndex, EGCTokenType TokenType, bool bAllowReferenceElimination)
 	{		
 		if (Object && Object != ReferencingObject) // Skip self-references just in case 
 		{
@@ -615,9 +615,9 @@ public:
 			FReferenceChainSearch::FObjectReferenceInfo RefInfo(ObjectInfo);
 			if (!ReferencedObjects.Contains(RefInfo))
 			{
-				if (TokenIndex >= 0)
+				if (TokenIndex != UE::GC::ETokenlessId::Collector)
 				{
-					FTokenInfo TokenInfo = ReferencingObject->GetClass()->ReferenceTokenStream.GetTokenInfo(TokenIndex);
+					FTokenInfo TokenInfo = ReferencingObject->GetClass()->ReferenceTokens.GetTokenInfo(TokenIndex);
 					RefInfo.ReferencerName = TokenInfo.Name;
 					RefInfo.Type = FReferenceChainSearch::EReferenceType::Property;
 				}
@@ -667,7 +667,7 @@ class FMinimalReferenceProcessor : public FDirectReferenceProcessor
 public:
 	FMinimalReferenceProcessor(
 		TSet<UObject*>& InTargetObjects,
-		TMap<UObject*, FGCObjectInfo*>& InObjectToInfoMap
+		TMap<const UObject*, FGCObjectInfo*>& InObjectToInfoMap
 	)
 	: FDirectReferenceProcessor(InObjectToInfoMap)
 	, TargetObjects(InTargetObjects)
@@ -679,14 +679,7 @@ public:
 		bSearchInners = InValue;
 	}
 
-	void HandleTokenStreamObjectReference(
-		FGCArrayStruct& ObjectsToSerializeStruct,
-		UObject* ReferencingObject,
-		UObject*& Object,
-		const int32 TokenIndex,
-		const EGCTokenType TokenType,
-		bool bAllowReferenceElimination
-	)
+	void HandleTokenStreamObjectReference(FGCArrayStruct& ObjectsToSerializeStruct, UObject* ReferencingObject, UObject*& Object, UE::GC::FTokenId TokenIndex, EGCTokenType TokenType, bool bAllowReferenceElimination)
 	{		
 		if (Object && Object != ReferencingObject) // Skip self-references just in case 
 		{
@@ -720,14 +713,11 @@ public:
 	}
 };
 
-template<typename ProcessorType>
-class TReferenceSearchReferenceCollector : public TDefaultReferenceCollector<ProcessorType>
+using FDirectReferenceCollectorBase = UE::GC::TDefaultCollector<FDirectReferenceProcessor>;
+class FDirectReferenceCollector : public FDirectReferenceCollectorBase
 {
 public:
-	TReferenceSearchReferenceCollector (ProcessorType& InProcessor, FGCArrayStruct& InObjectArrayStruct)
-		: TDefaultReferenceCollector<ProcessorType>(InProcessor, InObjectArrayStruct)
-	{
-	}
+	using FDirectReferenceCollectorBase::FDirectReferenceCollectorBase;
 
 	virtual bool MarkWeakObjectReferenceForClearing(UObject** WeakReference) override
 	{
@@ -808,7 +798,7 @@ void FReferenceChainSearch::PerformSearch()
 {
 	TArray<FGraphNode*> GraphNodesToFind;
 	GraphNodesToFind.Reserve(ObjectsToFindReferencesTo.Num());
-	for (UObject* Obj : ObjectsToFindReferencesTo)
+	for (const UObject* Obj : ObjectsToFindReferencesTo)
 	{
 		FGraphNode* ObjectNodeToFindReferencesTo = FindOrAddNode(Obj);
 		check(ObjectNodeToFindReferencesTo);
@@ -836,7 +826,7 @@ void FReferenceChainSearch::PerformSearch()
 #if ENABLE_GC_HISTORY
 void FReferenceChainSearch::PerformSearchFromGCSnapshot(UObject* InObjectToFindReferencesTo, FGCSnapshot& InSnapshot)
 {
-	PerformSearchFromGCSnapshot(TConstArrayView<UObject*>(&InObjectToFindReferencesTo, 1), InSnapshot);
+	PerformSearchFromGCSnapshot(MakeArrayView(&InObjectToFindReferencesTo, 1), InSnapshot);
 }
 
 void FReferenceChainSearch::PerformSearchFromGCSnapshot(TConstArrayView<UObject*> InObjectsToFindReferencesTo, FGCSnapshot& InSnapshot)
@@ -850,7 +840,7 @@ void FReferenceChainSearch::PerformSearchFromGCSnapshot(TConstArrayView<UObject*
 
 	ObjectsToFindReferencesTo = InObjectsToFindReferencesTo;
 	ObjectInfosToFindReferencesTo.Reset();
-	for (UObject* Obj : InObjectsToFindReferencesTo)
+	for (const UObject* Obj : InObjectsToFindReferencesTo)
 	{
 		ObjectInfosToFindReferencesTo.Add(FGCObjectInfo::FindOrAddInfoHelper(Obj, ObjectToInfoMap));
 	}
@@ -894,19 +884,11 @@ template<typename ProcessorType>
 struct TReferenceSearchHelper
 {
 	ProcessorType Processor;
-	TFastReferenceCollector<
-		ProcessorType, 
-		TReferenceSearchReferenceCollector<ProcessorType>, 
-		FGCArrayPool, 
-		EFastReferenceCollectorOptions::AutogenerateTokenStream | EFastReferenceCollectorOptions::ProcessNoOpTokens
-	> ReferenceCollector;
-	FGCArrayStruct ArrayStruct;
+	TArray<UObject*> ObjectsToProcess;
 
 	TReferenceSearchHelper(ProcessorType&& InProcessor)
 	: Processor(MoveTemp(InProcessor))
-	, ReferenceCollector(Processor, FGCArrayPool::Get())
-	{
-	}
+	{}
 
 	static bool DoNotFilterObjects(UObject* Object) { return true; }
 
@@ -915,7 +897,6 @@ struct TReferenceSearchHelper
 	template<typename HANDLE_REFERENCES, typename FILTER_OBJECTS = decltype(DoNotFilterObjects)>
 	void CollectAllReferences(bool bGCOnly, HANDLE_REFERENCES HandleReferences, FILTER_OBJECTS FilterObjects = &DoNotFilterObjects)
 	{
-		TArray<UObject*>& ObjectsToProcess = ArrayStruct.ObjectsToSerialize;
 		for (FRawObjectIterator It; It; ++It)
 		{
 			FUObjectItem* ObjItem = *It;
@@ -934,9 +915,10 @@ struct TReferenceSearchHelper
 
 			// Find direct references
 			Processor.Reset();
-			ObjectsToProcess.Reset();
-			ObjectsToProcess.Add(Object);
-			ReferenceCollector.CollectReferences(ArrayStruct);
+			UE::GC::FWorkerContext Context;
+			ObjectsToProcess = {Object};
+			Context.SetInitialObjectsUnpadded(ObjectsToProcess);
+			CollectReferences(Processor, Context);
 
 			if (!HandleReferences(Object, Processor.GetReferencedObjects()))
 			{
@@ -1102,7 +1084,7 @@ int32 FReferenceChainSearch::PrintResults(TFunctionRef<bool(FCallbackParams& Par
 	{
 		if (TargetObject)
 		{
-			FGCObjectInfo* ObjInfo = FGCObjectInfo::FindOrAddInfoHelper(TargetObject, *const_cast<TMap<UObject*, FGCObjectInfo*>*>(&ObjectToInfoMap));
+			FGCObjectInfo* ObjInfo = FGCObjectInfo::FindOrAddInfoHelper(TargetObject, const_cast<TMap<const UObject*, FGCObjectInfo*>&>(ObjectToInfoMap));
 			check(ObjInfo);
 			UE_LOG(LogReferenceChain, Log, TEXT("%s%s is not currently reachable."),
 				*GetObjectFlags(ObjInfo),
@@ -1174,7 +1156,7 @@ void FReferenceChainSearch::Cleanup()
 	}
 	AllNodes.Empty();
 
-	for (TPair<UObject*, FGCObjectInfo*>& ObjectToInfoPair : ObjectToInfoMap)
+	for (TPair<const UObject*, FGCObjectInfo*>& ObjectToInfoPair : ObjectToInfoMap)
 	{
 		delete ObjectToInfoPair.Value;
 	}
diff --git a/Engine/Source/Runtime/CoreUObject/Private/UObject/ReferencerFinder.cpp b/Engine/Source/Runtime/CoreUObject/Private/UObject/ReferencerFinder.cpp
--- a/Engine/Source/Runtime/CoreUObject/Private/UObject/ReferencerFinder.cpp
+++ b/Engine/Source/Runtime/CoreUObject/Private/UObject/ReferencerFinder.cpp
@@ -21,7 +21,7 @@ public:
 		, Flags(InFlags)
 	{
 	}
-	FORCEINLINE_DEBUGGABLE void HandleTokenStreamObjectReference(FGCArrayStruct& ObjectsToSerializeStruct, UObject* ReferencingObject, UObject*& Object, const int32 TokenIndex, const EGCTokenType TokenType, bool bAllowReferenceElimination)
+	FORCEINLINE_DEBUGGABLE void HandleTokenStreamObjectReference(FGCArrayStruct& ObjectsToSerializeStruct, UObject* ReferencingObject, UObject*& Object, UE::GC::FTokenId, EGCTokenType, bool)
 	{
 		if (!ReferencingObject)
 		{
@@ -49,15 +49,11 @@ public:
 	}
 };
 
-class FAllReferencesCollector : public TDefaultReferenceCollector<FAllReferencesProcessor>
+class FAllReferencesCollector : public UE::GC::TDefaultCollector<FAllReferencesProcessor>
 {
-	FAllReferencesProcessor& Processor;
+	using Super = UE::GC::TDefaultCollector<FAllReferencesProcessor>;
 public:
-	FAllReferencesCollector(FAllReferencesProcessor& InProcessor, FGCArrayStruct& InObjectArrayStruct)
-		: TDefaultReferenceCollector<FAllReferencesProcessor>(InProcessor, InObjectArrayStruct)
-		, Processor(InProcessor)
-	{
-	}
+	using Super::Super;
 
 	virtual bool MarkWeakObjectReferenceForClearing(UObject** WeakReference) override
 	{
@@ -98,22 +94,15 @@ TArray<UObject*> FReferencerFinder::GetAllReferencers(const TSet<UObject*>& Refe
 		GUObjectArray.LockInternalArray();
 
 		const int32 MaxNumberOfObjects = GUObjectArray.GetObjectArrayNum();
-		const int32 NumThreads = FMath::Max(1, FTaskGraphInterface::Get().GetNumWorkerThreads());
+		const int32 NumThreads = GetNumCollectReferenceWorkers();
 		const int32 NumberOfObjectsPerThread = (MaxNumberOfObjects / NumThreads) + 1;
 
 		ParallelFor(NumThreads, [&Referencees, ObjectsToIgnore, &ResultCritical, &Ret, NumberOfObjectsPerThread, NumThreads, MaxNumberOfObjects, Flags](int32 ThreadIndex)
 		{
 			TSet<UObject*> ThreadResult;
 			FAllReferencesProcessor Processor(Referencees, Flags, ThreadResult);
-			TFastReferenceCollector<
-				FAllReferencesProcessor, 
-				FAllReferencesCollector, 
-				FGCArrayPool, 
-				EFastReferenceCollectorOptions::AutogenerateTokenStream | EFastReferenceCollectorOptions::ProcessNoOpTokens
-			> ReferenceCollector(Processor, FGCArrayPool::Get());
-			FGCArrayStruct ArrayStruct;
-
-			ArrayStruct.ObjectsToSerialize.Reserve(NumberOfObjectsPerThread);
+			TArray<UObject*> ObjectsToSerialize;
+			ObjectsToSerialize.Reserve(NumberOfObjectsPerThread);
 
 			const int32 FirstObjectIndex = ThreadIndex * NumberOfObjectsPerThread;
 			const int32 NumObjects = (ThreadIndex < (NumThreads - 1)) ? NumberOfObjectsPerThread : (MaxNumberOfObjects - (NumThreads - 1)*NumberOfObjectsPerThread);
@@ -132,17 +121,20 @@ TArray<UObject*> FReferencerFinder::GetAllReferencers(const TSet<UObject*>& Refe
 
 					if (!Referencees.Contains(PotentialReferencer))
 					{
-						ArrayStruct.ObjectsToSerialize.Add(PotentialReferencer);
+						ObjectsToSerialize.Add(PotentialReferencer);
 					}
 				}
 			}
+
+			FGCArrayStruct ArrayStruct;
+			ArrayStruct.SetInitialObjectsUnpadded(ObjectsToSerialize);
 			
 			{
 				// Since ReferenceCollector is configured to automatically assemble reference token streams
 				// for classes that require it, make sure GC is locked because UClass::AssembleReferenceTokenStream requires it
 				FGCScopeGuard GCGuard;
 				// Now check if any of the potential referencers is referencing any of the referencees
-				ReferenceCollector.CollectReferences(ArrayStruct);
+				CollectReferences(Processor, ArrayStruct);
 			}
 
 			if (ThreadResult.Num())
diff --git a/Engine/Source/Runtime/CoreUObject/Private/UObject/UObjectAllocator.cpp b/Engine/Source/Runtime/CoreUObject/Private/UObject/UObjectAllocator.cpp
--- a/Engine/Source/Runtime/CoreUObject/Private/UObject/UObjectAllocator.cpp
+++ b/Engine/Source/Runtime/CoreUObject/Private/UObject/UObjectAllocator.cpp
@@ -100,7 +100,7 @@ void FUObjectAllocator::FreeUObject(UObjectBase *Object) const
 {
 	check(Object);
 	// Only free memory if it was allocated directly from allocator and not from permanent object pool.
-	if( ResidesInPermanentPool(Object) == false )
+	if (FPermanentObjectPoolExtents(*this).Contains(Object) == false)
 	{
 		FMemory::Free(Object);
 	}
diff --git a/Engine/Source/Runtime/CoreUObject/Private/UObject/UObjectBase.cpp b/Engine/Source/Runtime/CoreUObject/Private/UObject/UObjectBase.cpp
--- a/Engine/Source/Runtime/CoreUObject/Private/UObject/UObjectBase.cpp
+++ b/Engine/Source/Runtime/CoreUObject/Private/UObject/UObjectBase.cpp
@@ -384,16 +384,9 @@ bool UObjectBase::IsValidLowLevelFast(bool bRecursive /*= true*/) const
 	return true;
 }
 
-void UObjectBase::EmitBaseReferences(UClass *RootClass)
+void UObjectBase::EmitBaseReferences(UE::GC::FTokenStreamBuilder& TokenStream)
 {
-	static const FName ClassPropertyName(TEXT("Class"));
-	static const FName OuterPropertyName(TEXT("Outer"));
-	// Mark UObject class reference as persistent object reference so that it (ClassPrivate) doesn't get nulled when a class
-	// is marked as pending kill. Nulling ClassPrivate may leave the object in a broken state if it doesn't get GC'd in the same
-	// GC call as its class. And even if it gets GC'd in the same call as its class it may break inside of GC (for example when traversing TMap references)
-	RootClass->EmitObjectReference(STRUCT_OFFSET(UObjectBase, ClassPrivate), ClassPropertyName, GCRT_PersistentObject);
-	RootClass->EmitObjectReference(STRUCT_OFFSET(UObjectBase, OuterPrivate), OuterPropertyName, GCRT_PersistentObject);
-	RootClass->EmitExternalPackageReference();
+	TokenStream.EmitExternalPackageReference();
 }
 
 #if USE_PER_MODULE_UOBJECT_BOOTSTRAP
diff --git a/Engine/Source/Runtime/CoreUObject/Private/UObject/UObjectClusters.cpp b/Engine/Source/Runtime/CoreUObject/Private/UObject/UObjectClusters.cpp
--- a/Engine/Source/Runtime/CoreUObject/Private/UObject/UObjectClusters.cpp
+++ b/Engine/Source/Runtime/CoreUObject/Private/UObject/UObjectClusters.cpp
@@ -629,7 +629,7 @@ public:
 		check(Obj->CanBeInCluster());
 		if (ObjectIndex != ClusterRootIndex && ObjectItem->GetOwnerIndex() == 0 && !GUObjectArray.IsDisregardForGC(Obj) && !Obj->IsRooted())
 		{
-			ObjectsToSerializeStruct.ObjectsToSerialize.Add(Obj);
+			ObjectsToSerializeStruct.ObjectsToSerialize.Add<Options>(Obj);
 			check(!ObjectItem->HasAnyFlags(EInternalObjectFlags::ClusterRoot));
 			ObjectItem->SetOwnerIndex(ClusterRootIndex);
 			Cluster.Objects.Add(ObjectIndex);
@@ -639,14 +639,14 @@ public:
 				UObject* ObjOuter = Obj->GetOuter();
 				if (ObjOuter)
 				{
-					HandleTokenStreamObjectReference(ObjectsToSerializeStruct, Obj, ObjOuter, INDEX_NONE, EGCTokenType::Native, true);
+					HandleTokenStreamObjectReference(ObjectsToSerializeStruct, Obj, ObjOuter, UE::GC::ETokenlessId::Outer, EGCTokenType::Native, true);
 				}
 				if (!Obj->GetClass()->HasAllClassFlags(CLASS_Native))
 				{
 					UObject* ObjectClass = Obj->GetClass();
-					HandleTokenStreamObjectReference(ObjectsToSerializeStruct, Obj, ObjectClass, INDEX_NONE, EGCTokenType::Native, true);
+					HandleTokenStreamObjectReference(ObjectsToSerializeStruct, Obj, ObjectClass, UE::GC::ETokenlessId::Class, EGCTokenType::Native, true);
 					UObject* ObjectClassOuter = Obj->GetClass()->GetOuter();
-					HandleTokenStreamObjectReference(ObjectsToSerializeStruct, Obj, ObjectClassOuter, INDEX_NONE, EGCTokenType::Native, true);
+					HandleTokenStreamObjectReference(ObjectsToSerializeStruct, Obj, ObjectClassOuter, UE::GC::ETokenlessId::ClassOuter, EGCTokenType::Native, true);
 				}
 			}
 		}
@@ -660,7 +660,7 @@ public:
 	* @param TokenIndex Index to the token stream where the reference was found.
 	* @param bAllowReferenceElimination True if reference elimination is allowed (ignored when constructing clusters).
 	*/
-	FORCEINLINE void HandleTokenStreamObjectReference(FGCArrayStruct& ObjectsToSerializeStruct, UObject* ReferencingObject, UObject*& Object, const int32 TokenIndex, const EGCTokenType TokenType, bool bAllowReferenceElimination)
+	FORCEINLINE void HandleTokenStreamObjectReference(FGCArrayStruct& ObjectsToSerializeStruct, UObject* ReferencingObject, UObject*& Object, UE::GC::FTokenId TokenIndex, const EGCTokenType TokenType, bool bAllowReferenceElimination)
 	{
 		if (Object)
 		{
@@ -748,19 +748,10 @@ void UObjectBaseUtility::AddToCluster(UObjectBaseUtility* ClusterRootOrObjectFro
 		if (!bAddAsMutableObject)
 		{
 			FClusterReferenceProcessor Processor(ClusterRootIndex, *Cluster);
-			TFastReferenceCollector<
-				FClusterReferenceProcessor, 
-				TDefaultReferenceCollector<FClusterReferenceProcessor>, 
-				FGCArrayPool, 
-				EFastReferenceCollectorOptions::AutogenerateTokenStream | EFastReferenceCollectorOptions::ProcessNoOpTokens
-			> ReferenceCollector(Processor, FGCArrayPool::Get());
-			FGCArrayStruct ArrayStruct;
 			UObject* ThisObject = static_cast<UObject*>(this);
-			Processor.HandleTokenStreamObjectReference(ArrayStruct, static_cast<UObject*>(ClusterRootOrObjectFromCluster), ThisObject, INDEX_NONE, EGCTokenType::Native, true);
-			if (ArrayStruct.ObjectsToSerialize.Num())
-			{
-				ReferenceCollector.CollectReferences(ArrayStruct);
-			}
+			FGCArrayStruct ArrayStruct;
+			ArrayStruct.InitialNativeReferences = {&ThisObject};
+			CollectReferences(Processor, ArrayStruct);
 
 #if UE_GCCLUSTER_VERBOSE_LOGGING
 			UObject* ClusterRootObject = static_cast<UObject*>(GUObjectArray.IndexToObjectUnsafeForGC(Cluster->RootIndex)->Object);
@@ -820,19 +811,10 @@ void UObjectBaseUtility::CreateCluster()
 
 	// Collect all objects referenced by cluster root and by all objects it's referencing
 	FClusterReferenceProcessor Processor(InternalIndex, Cluster);
-	TFastReferenceCollector<
-		FClusterReferenceProcessor, 
-		TDefaultReferenceCollector<FClusterReferenceProcessor>, 
-		FGCArrayPool, 
-		EFastReferenceCollectorOptions::AutogenerateTokenStream | EFastReferenceCollectorOptions::ProcessNoOpTokens
-	> ReferenceCollector(Processor, FGCArrayPool::Get());
 	FGCArrayStruct ArrayStruct;
-	TArray<UObject*>& ObjectsToProcess = ArrayStruct.ObjectsToSerialize;
-	ObjectsToProcess.Add(static_cast<UObject*>(this));
-	ReferenceCollector.CollectReferences(ArrayStruct);
-#if UE_BUILD_DEBUG
-	FGCArrayPool::Get().CheckLeaks();
-#endif
+	TArray<UObject*> ObjectsToProcess = {static_cast<UObject*>(this)};
+	ArrayStruct.SetInitialObjectsUnpadded(ObjectsToProcess);
+	CollectReferences(Processor, ArrayStruct);
 
 	check(RootItem->GetOwnerIndex() == 0);
 	RootItem->SetClusterIndex(ClusterIndex);
diff --git a/Engine/Source/Runtime/CoreUObject/Private/UObject/UObjectGlobals.cpp b/Engine/Source/Runtime/CoreUObject/Private/UObject/UObjectGlobals.cpp
--- a/Engine/Source/Runtime/CoreUObject/Private/UObject/UObjectGlobals.cpp
+++ b/Engine/Source/Runtime/CoreUObject/Private/UObject/UObjectGlobals.cpp
@@ -29,6 +29,7 @@
 #include "UObject/GarbageCollection.h"
 #include "UObject/Class.h"
 #include "UObject/CoreRedirects.h"
+#include "UObject/FastReferenceCollector.h"
 #include "UObject/UObjectIterator.h"
 #include "UObject/Package.h"
 #include "Templates/Casts.h"
@@ -4442,34 +4443,22 @@ void ConstructorHelpers::StripObjectClass( FString& PathName, bool bAssertOnBadP
 	}
 }
 
-/*----------------------------------------------------------------------------
-FSimpleObjectReferenceCollectorArchive.
-----------------------------------------------------------------------------*/
-class FSimpleObjectReferenceCollectorArchive : public FReferenceCollectorArchive
-{
-public:
+//////////////////////////////////////////////////////////////////////////
 
-	/**
-	* Constructor
-	*
-	* @param	InObjectArray			Array to add object references to
-	*/
-	FSimpleObjectReferenceCollectorArchive(UObject* InSerializingObject, FReferenceCollector& InCollector)
-		: FReferenceCollectorArchive(InSerializingObject, InCollector)
-	{
-		ArIsObjectReferenceCollector = true;
-		this->SetIsPersistent(InCollector.IsIgnoringTransient());
-		ArIgnoreArchetypeRef = InCollector.IsIgnoringArchetypeRef();
-	}
+FReferenceCollectorArchive::FReferenceCollectorArchive(const UObject* InSerializingObject, FReferenceCollector& InCollector)
+: SerializingObject(InSerializingObject)
+, Collector(InCollector)
+{
+	ArIsObjectReferenceCollector = true;
+	this->SetIsPersistent(InCollector.IsIgnoringTransient());
+	ArIgnoreArchetypeRef = InCollector.IsIgnoringArchetypeRef();
+}
 
-protected:
+class FPropertyTrackingReferenceCollectorArchive : public FReferenceCollectorArchive
+{
+public:
+	using FReferenceCollectorArchive::FReferenceCollectorArchive;
 
-	/**
-	* UObject serialize operator implementation
-	*
-	* @param Object	reference to Object reference
-	* @return reference to instance of this class
-	*/
 	virtual FArchive& operator<<(UObject*& Object) override
 	{
 		if (Object)
@@ -4497,167 +4486,436 @@ protected:
 
 };
 
-class FPersistentFrameCollectorArchive : public FSimpleObjectReferenceCollectorArchive
+void FReferenceCollector::AddStableReference(UObject** Object)
 {
-public:
-	FPersistentFrameCollectorArchive(UObject* InSerializingObject, FReferenceCollector& InCollector)
-		: FSimpleObjectReferenceCollectorArchive(InSerializingObject, InCollector)
-	{	}
+	AddReferencedObject(*Object);
+}
 
-protected:
-	virtual FArchive& operator<<(UObject*& Object) override
+void FReferenceCollector::AddStableReferenceArray(TArray<UObject*>* Array)
+{
+	AddReferencedObjects(*Array); 
+}
+
+void FReferenceCollector::AddStableReferenceSet(TSet<UObject*>* Objects)
+{
+	AddReferencedObjects(*Objects);
+}
+
+void FReferenceCollector::AddReferencedObjects(const UScriptStruct*& ScriptStruct, void* StructMemory, const UObject* ReferencingObject /*= nullptr*/, const FProperty* ReferencingProperty /*= nullptr*/)
+{
+	check(ScriptStruct != nullptr);
+	check(StructMemory != nullptr);
+
+	AddReferencedObject(ScriptStruct, ReferencingObject, ReferencingProperty);
+
+	// If the script struct explicitly provided an implementation of AddReferencedObjects, make sure to capture its referenced objects
+	if (ScriptStruct->StructFlags & STRUCT_AddStructReferencedObjects)
 	{
-#if !(UE_BUILD_TEST || UE_BUILD_SHIPPING)
-		const bool bIsValidObjectReference = (Object == nullptr || Object->IsValidLowLevelFast());
-		if (!bIsValidObjectReference)
-		{
-			if (const UFunction* UberGraphFunction = Cast<UFunction>(GetSerializingObject()))
-			{
-				const int32 PersistentFrameDataSize = UberGraphFunction->GetStructureSize();
-				if (const uint8* PersistentFrameDataAddr = (const uint8*)GetSerializedDataPtr())
-				{
-					FString PersistentFrameDataText;
-					const int32 MaxBytesToDisplayPerLine = 32;
-					PersistentFrameDataText.Reserve(PersistentFrameDataSize * 2 + PersistentFrameDataSize / MaxBytesToDisplayPerLine);
-					for (int32 PersistentFrameDataIdx = 0; PersistentFrameDataIdx < PersistentFrameDataSize; ++PersistentFrameDataIdx)
-					{
-						if (PersistentFrameDataIdx % MaxBytesToDisplayPerLine == 0)
-						{
-							PersistentFrameDataText += TEXT("\n");
-						}
+		ScriptStruct->GetCppStructOps()->AddStructReferencedObjects()(StructMemory, *this);
+	}
 
-						PersistentFrameDataText += FString::Printf(TEXT("%02x "), PersistentFrameDataAddr[PersistentFrameDataIdx]);
-					}
+	AddPropertyReferences(ScriptStruct, StructMemory, ReferencingObject);
+}
 
-					UE_LOG(LogUObjectGlobals, Log, TEXT("PersistentFrame: Addr=0x%016llx, Size=%d%s"),
-						(int64)(PTRINT)PersistentFrameDataAddr,
-						PersistentFrameDataSize,
-						*PersistentFrameDataText);
-				}
-			}
+void FReferenceCollector::HandleObjectReferences(FObjectPtr* InObjects, const int32 ObjectNum, const UObject* InReferencingObject, const FProperty* InReferencingProperty)
+{
+	for (int32 ObjectIndex = 0; ObjectIndex < ObjectNum; ++ObjectIndex)
+	{
+		FObjectPtr& Object = InObjects[ObjectIndex];
+		if (IsObjectHandleResolved(Object.GetHandle()))
+		{
+			HandleObjectReference(*reinterpret_cast<UObject**>(&Object), InReferencingObject, InReferencingProperty);
 		}
+	}
+}
+
+//////////////////////////////////////////////////////////////////////////
+
+enum class EPropertyCollectFlags : uint32
+{
+	None				= 0,
+	SkipTransient		= 1 << 0,		
+	NeedsReferencer		= 1 << 1,
+	CallStructARO		= 1 << 2,
+	OnlyObjectProperty	= 1 << 3,
+};
+ENUM_CLASS_FLAGS(EPropertyCollectFlags);
+
+static constexpr EPropertyCollectFlags AllCollectorFlags =  EPropertyCollectFlags::SkipTransient | EPropertyCollectFlags::NeedsReferencer;
+FORCEINLINE EPropertyCollectFlags GetCollectorPropertyFlags(FReferenceCollector& Collector)
+{
+	return (Collector.IsIgnoringTransient() ? EPropertyCollectFlags::SkipTransient : EPropertyCollectFlags::None)
+		 | (Collector.NeedsPropertyReferencer() ? EPropertyCollectFlags::NeedsReferencer : EPropertyCollectFlags::None);
+}
+
+FORCEINLINE constexpr EPropertyFlags GetPropertyFlagsToSkip(EPropertyCollectFlags CollectFlags)
+{
+	return CPF_SkipSerialization | (EnumHasAnyFlags(CollectFlags, EPropertyCollectFlags::SkipTransient) ? CPF_Transient : CPF_None);
+}
+
+
+/** Core property types with weak references */
+static constexpr EClassCastFlags WeakCastFlags =			CASTCLASS_FWeakObjectProperty | 
+															CASTCLASS_FLazyObjectProperty |
+															CASTCLASS_FSoftObjectProperty |
+															CASTCLASS_FDelegateProperty |
+															CASTCLASS_FMulticastDelegateProperty;
+
+static constexpr EClassCastFlags ObjectCastFlags =			CASTCLASS_FObjectProperty |
+															CASTCLASS_FObjectPtrProperty;
+
+static constexpr EClassCastFlags OtherStrongCastFlags =		CASTCLASS_FInterfaceProperty |
+															CASTCLASS_FFieldPathProperty;
+
+/** Core property types with strong references */
+static constexpr EClassCastFlags StrongCastFlags =			ObjectCastFlags | OtherStrongCastFlags;
+
+/** Core property types with neither weak nor strong references */
+static constexpr EClassCastFlags UnreferencingCastFlags =	CASTCLASS_FByteProperty |
+															CASTCLASS_FInt8Property |
+															CASTCLASS_FIntProperty |
+															CASTCLASS_FFloatProperty |
+															CASTCLASS_FUInt64Property |
+															CASTCLASS_FUInt32Property |
+															CASTCLASS_FNameProperty |
+															CASTCLASS_FStrProperty |
+															CASTCLASS_FBoolProperty |
+															CASTCLASS_FUInt16Property |
+															CASTCLASS_FInt64Property |
+															CASTCLASS_FNumericProperty |
+															CASTCLASS_FTextProperty |
+															CASTCLASS_FInt16Property |
+															CASTCLASS_FDoubleProperty |
+															CASTCLASS_FEnumProperty |
+															CASTCLASS_FLargeWorldCoordinatesRealProperty;
+
+FORCEINLINE constexpr EClassCastFlags GetCastFlagsToSkip(EPropertyCollectFlags CollectFlags)
+{
+	return WeakCastFlags | (EnumHasAnyFlags(CollectFlags, EPropertyCollectFlags::OnlyObjectProperty) ? OtherStrongCastFlags : CASTCLASS_None);
+}
+
+FORCEINLINE constexpr bool MayContainStrongReference(EClassCastFlags CastFlags)
+{
+	return !EnumHasAnyFlags(CastFlags, UnreferencingCastFlags | WeakCastFlags);
+}
+
+FORCEINLINE constexpr bool MayContainStrongReference(const FProperty& Property)
+{
+	return MayContainStrongReference(static_cast<EClassCastFlags>(Property.GetClass()->GetCastFlags()));
+}
 
-		auto GetBlueprintObjectNameLambda = [](const UObject* InSerializingObject) -> FString
+template<EPropertyCollectFlags CollectFlags>
+void CollectPropertyReferences(FReferenceCollector& Collector, FProperty& Property, void* Instance, const UObject* Referencer);
+
+template<EPropertyCollectFlags CollectFlags, /* UStruct or UScriptStruct */ class StructType>
+static void CollectStructReferences(FReferenceCollector& Collector, const StructType* Struct, void* Instance, const UObject* Referencer)
+{
+	// The FProperty instance might start in the middle of a cache line	
+	static constexpr uint32 ExtraPrefetchBytes = PLATFORM_CACHE_LINE_SIZE - /* min alignment */ 16;
+	// Prefetch vtable, PropertyFlags and NextRef. NextRef comes last.
+	static constexpr uint32 PropertyPrefetchBytes = offsetof(FProperty, NextRef) + ExtraPrefetchBytes;
+
+	FPlatformMisc::PrefetchBlock(Struct->RefLink, PropertyPrefetchBytes);
+	
+	if constexpr (EnumHasAnyFlags(CollectFlags, EPropertyCollectFlags::CallStructARO))
+	{
+		if (Struct->StructFlags & STRUCT_AddStructReferencedObjects)
 		{
-			if (InSerializingObject)
-			{
-				if (const UClass* BPGC = InSerializingObject->GetTypedOuter<UClass>())
-				{
-#if WITH_EDITORONLY_DATA
-					if (BPGC->ClassGeneratedBy)
-					{
-						return BPGC->ClassGeneratedBy->GetFullName();
-					}
-					else
-#endif
-					{
-						return BPGC->GetFullName();
-					}
-				}
-			}
+			Struct->GetCppStructOps()->AddStructReferencedObjects()(Instance, Collector);
+		}
+	}
 
-			return TEXT("NULL");
-		};
+	for (FProperty* It = Struct->RefLink; It; It = It->NextRef)
+	{
+		FPlatformMisc::PrefetchBlock(It->NextRef, PropertyPrefetchBytes);
+		CollectPropertyReferences<CollectFlags>(Collector, *It, Instance, Referencer);
+	}
+}
 
-		if (!ensureMsgf(bIsValidObjectReference
-			, TEXT("Invalid object referenced by the PersistentFrame: 0x%016llx (Blueprint object: %s, ReferencingProperty: %s, Instance: %s, Address: 0x%016llx) - If you have a reliable repro for this, please contact the development team with it.")
-			, (int64)(PTRINT)Object
-			, *GetBlueprintObjectNameLambda(GetSerializingObject())
-			, GetSerializedProperty() ? *GetSerializedProperty()->GetFullName() : TEXT("NULL")
-			, GetSerializedDataContainer() ? *GetSerializedDataContainer()->GetFullName() : TEXT("NULL")
-			, (int64)(PTRINT)&Object))
+template<EPropertyCollectFlags CollectFlags>
+void CollectArrayReferences(FReferenceCollector& Collector, FArrayProperty& Property, void* Instance, const UObject* Referencer)
+{
+	FProperty& InnerProperty = *Property.Inner;
+	EClassCastFlags InnerCastFlags = static_cast<EClassCastFlags>(InnerProperty.GetClass()->GetCastFlags());
+	if (MayContainStrongReference(InnerCastFlags))
+	{
+		bool bIsReferenceArray = EnumHasAnyFlags(InnerCastFlags, ObjectCastFlags) &
+								!EnumHasAnyFlags(Property.ArrayFlags, EArrayPropertyFlags::UsesMemoryImageAllocator);
+		if (bIsReferenceArray && !EnumHasAnyFlags(CollectFlags, EPropertyCollectFlags::NeedsReferencer))
 		{
-			// clear the property value (it's garbage)... the ubergraph-frame
-			// has just lost a reference to whatever it was attempting to hold onto
-			Object = nullptr;
+			Collector.AddStableReferenceArray(reinterpret_cast<TArray<UObject*>*>(Instance));
 		}
-#endif
-		if (Object)
+		else if (FScriptArrayHelper Helper(&Property, Instance); int32 Num = Helper.Num())
 		{
-			bool bWeakRef = false;
-
-			// If the property that serialized us is not an object property we are in some native serializer, we have to treat these as strong
-			if (!Object->HasAnyFlags(RF_StrongRefOnFrame))
+			if (bIsReferenceArray)
 			{
-				FObjectProperty* ObjectProperty = CastField<FObjectProperty>(GetSerializedProperty());
-
-				if (ObjectProperty)
+				if (UE_WITH_OBJECT_HANDLE_LATE_RESOLVE && EnumHasAnyFlags(InnerCastFlags, CASTCLASS_FObjectPtrProperty))
 				{
-					// This was a raw UObject* serialized by FObjectProperty, so just save the address
-					bWeakRef = true;
+					Collector.AddReferencedObjects(*reinterpret_cast<TArray<TObjectPtr<UObject>>*>(Instance), Referencer, &Property);
+				}
+				else
+				{
+					Collector.AddReferencedObjects(*reinterpret_cast<TArray<UObject*>*>(Instance), Referencer, &Property);
+				}
+			}
+			else if (EnumHasAnyFlags(InnerCastFlags, CASTCLASS_FStructProperty))
+			{
+				for (int32 Idx = 0; Idx < Num; ++Idx)
+				{
+					CollectStructReferences<CollectFlags>(Collector, static_cast<FStructProperty&>(InnerProperty).Struct, Helper.GetRawPtr(Idx), Referencer);
 				}
 			}
+			else
+			{
+				for (int32 Idx = 0; Idx < Num; ++Idx)
+				{
+					CollectPropertyReferences<CollectFlags>(Collector, InnerProperty, Helper.GetRawPtr(Idx), Referencer);
+				}
+			}
+			
+		}
+	}
+}
 
-			// Try to handle it as a weak ref, if it returns false treat it as a strong ref instead
-			bWeakRef = bWeakRef && GetCollector().MarkWeakObjectReferenceForClearing(&Object);
+template<EPropertyCollectFlags CollectFlags>
+static void CollectMapReferences(FReferenceCollector& Collector, FMapProperty& Property, void* Instance, const UObject* Referencer)
+{
+	FScriptMapHelper MapHelper(&Property, Instance);
+	const int32 Num = MapHelper.Num();
+	const bool bCollectKeys = MayContainStrongReference(*MapHelper.GetKeyProperty());
+	const bool bCollectValues = MayContainStrongReference(*MapHelper.GetValueProperty());
+
+	if (Num == 0)
+	{
+		return;	
+	}
 
-			if (!bWeakRef)
+	if (bCollectKeys)
+	{
+		for (int32 Idx = 0; Idx < Num; ++Idx)
+		{
+			if (MapHelper.IsValidIndex(Idx))
 			{
-				// This is a hard reference or we don't know what's serializing it, so serialize it normally
-				return FSimpleObjectReferenceCollectorArchive::operator<<(Object);
+				CollectPropertyReferences<CollectFlags>(Collector, *MapHelper.GetKeyProperty(), MapHelper.GetPairPtr(Idx), Referencer);
 			}
 		}
-
-		return *this;
 	}
-};
 
+	if (bCollectValues)
+	{
+		for (int32 Idx = 0; Idx < Num; ++Idx)
+		{
+			if (MapHelper.IsValidIndex(Idx))
+			{
+				CollectPropertyReferences<CollectFlags>(Collector, *MapHelper.GetValueProperty(), MapHelper.GetPairPtr(Idx), Referencer);
+			}
+		}
+	}
+}
 
-FReferenceCollector::FReferenceCollector()
-	: DefaultReferenceCollectorArchive(nullptr)
-	, PersistentFrameReferenceCollectorArchive(nullptr)
+template<EPropertyCollectFlags CollectFlags>
+void CollectSetReferences(FReferenceCollector& Collector, FSetProperty& Property, void* Instance, const UObject* Referencer)
 {
+	FScriptSetHelper SetHelper(&Property, Instance);
+	if (MayContainStrongReference(*SetHelper.GetElementProperty()))
+	{
+		for (int32 Idx = 0, Num = SetHelper.Num(); Idx < Num; ++Idx)
+		{	
+			if (SetHelper.IsValidIndex(Idx))
+			{
+				CollectPropertyReferences<CollectFlags>(Collector, *SetHelper.GetElementProperty(), SetHelper.GetElementPtr(Idx), Referencer);
+			}
+		}
+	}
 }
 
-FReferenceCollector::~FReferenceCollector()
+// Process FObjectProperty or FObjectPtrProperty reference
+template<EPropertyCollectFlags CollectFlags>
+FORCEINLINE_DEBUGGABLE void CollectObjectReference(FReferenceCollector& Collector, FProperty& Property, void* Value, const UObject* Referencer)
 {
-	delete DefaultReferenceCollectorArchive;
-	delete PersistentFrameReferenceCollectorArchive;
+	UObject*& Reference = *reinterpret_cast<UObject**>(Value);
+	if constexpr (EnumHasAnyFlags(CollectFlags, EPropertyCollectFlags::NeedsReferencer))
+	{
+		// Sync reference processors will inspect Reference immediately so might as well avoid virtual call
+		if (!!Reference & IsObjectHandleResolved(*reinterpret_cast<FObjectHandle*>(Value)))
+		{
+			Collector.AddReferencedObject(Reference, Referencer, &Property);
+		}
+	}
+	else
+	{
+		// Allows batch reference processor to queue up Reference and prefetch before accessing it
+		Collector.AddStableReference(&Reference);
+	}
 }
 
-void FReferenceCollector::AddReferencedObjects(const UScriptStruct*& ScriptStruct, void* StructMemory, const UObject* ReferencingObject /*= nullptr*/, const FProperty* ReferencingProperty /*= nullptr*/)
+// Process stack reference synchronously and return true if reference got nulled out
+FORCEINLINE_DEBUGGABLE bool CollectStackReference(FReferenceCollector& Collector, FProperty& Property, UObject*& Reference, const UObject* Referencer)
 {
-	check(ScriptStruct != nullptr);
-	check(StructMemory != nullptr);
-
-	AddReferencedObject(ScriptStruct, ReferencingObject, ReferencingProperty);
+	if (Reference)
+	{
+		Collector.AddReferencedObject(Reference, Referencer, &Property);
+		return !Reference;
+	}
+	return false;
+}
 
-	// If the script struct explicitly provided an implementation of AddReferencedObjects, make sure to capture its referenced objects
-	if (ScriptStruct->StructFlags & STRUCT_AddStructReferencedObjects)
+FORCENOINLINE static void CollectInterfaceReference(FReferenceCollector& Collector, FInterfaceProperty& Property, FScriptInterface& Interface, const UObject* Referencer)
+{
+	// Handle reference synchronously and update interface if reference was nulled out
+	UObject*& Ref = Interface.GetObjectRef();
+	if (CollectStackReference(Collector, Property, Ref, Referencer))
 	{
-		ScriptStruct->GetCppStructOps()->AddStructReferencedObjects()(StructMemory, *this);
+		Interface.SetInterface(nullptr);
 	}
+}
 
-	// Iterate through all object properties within the struct (will also search through structs within the struct)
-	for (TPropertyValueIterator<const FObjectProperty> ObjectPropertyIter(ScriptStruct, StructMemory); ObjectPropertyIter; ++ObjectPropertyIter)
+FORCENOINLINE static void CollectFieldPathReference(FReferenceCollector& Collector, FFieldPathProperty& Property, FFieldPath& FieldPath, const UObject* Referencer)
+{
+	if (FUObjectItem* FieldOwnerItem = FGCInternals::GetResolvedOwner(FieldPath))
 	{
-		TObjectPtr<UObject>& ObjectPtrRef = ObjectPropertyIter.Key()->GetObjectPtrPropertyValueRef(ObjectPropertyIter.Value());
-		AddReferencedObject(ObjectPtrRef, ReferencingObject, ReferencingProperty);
+		// Handle reference synchronously and update field path if reference was nulled out
+		UObject* Owner = static_cast<UObject*>(FieldOwnerItem->Object);		
+		if (CollectStackReference(Collector, Property, Owner, Referencer))
+		{
+			FGCInternals::ClearCachedField(FieldPath);
+		}
 	}
 }
 
-void FReferenceCollector::HandleObjectReferences(FObjectPtr* InObjects, const int32 ObjectNum, const UObject* InReferencingObject, const FProperty* InReferencingProperty)
+template<EPropertyCollectFlags CollectFlags>
+void CollectPropertyReferences(FReferenceCollector& Collector, FProperty& Property, void* Instance, const UObject* Referencer)
 {
-	for (int32 ObjectIndex = 0; ObjectIndex < ObjectNum; ++ObjectIndex)
+	FFieldClass* Class = Property.GetClass();
+	const int32 ArrayDim = Property.ArrayDim;
+	EPropertyFlags PropertyFlags = Property.GetPropertyFlags();
+	const EClassCastFlags CastFlags = static_cast<EClassCastFlags>(Class->GetCastFlags());
+	
+	if (EnumHasAnyFlags(CastFlags, GetCastFlagsToSkip(CollectFlags)) |
+		EnumHasAnyFlags(PropertyFlags, GetPropertyFlagsToSkip(CollectFlags)))
 	{
-		FObjectPtr& Object = InObjects[ObjectIndex];
-		if (IsObjectHandleResolved(Object.GetHandle()))
+		return;
+	}
+
+	int32 Idx = 0;
+	do
+	{
+		void* Value = Property.ContainerPtrToValuePtr<void>(Instance, Idx);
+
+		if (EnumHasAnyFlags(CastFlags, ObjectCastFlags))
 		{
-			HandleObjectReference(*reinterpret_cast<UObject**>(&Object), InReferencingObject, InReferencingProperty);
+			CollectObjectReference<CollectFlags>(Collector, Property, Value, Referencer);
+		}
+		else if (EnumHasAnyFlags(CastFlags, CASTCLASS_FArrayProperty))
+		{
+			CollectArrayReferences<CollectFlags>(Collector, static_cast<FArrayProperty&>(Property), Value, Referencer);
+		}
+		else if (EnumHasAnyFlags(CastFlags, CASTCLASS_FStructProperty))
+		{
+			CollectStructReferences<CollectFlags>(Collector, static_cast<FStructProperty&>(Property).Struct, Value, Referencer);
+		}
+		else if (EnumHasAnyFlags(CastFlags, CASTCLASS_FMapProperty))
+		{	
+			CollectMapReferences<CollectFlags>(Collector, static_cast<FMapProperty&>(Property), Value, Referencer);
+		}
+		else if (EnumHasAnyFlags(CastFlags, CASTCLASS_FSetProperty))
+		{	
+			CollectSetReferences<CollectFlags>(Collector, static_cast<FSetProperty&>(Property), Value, Referencer);
+		}
+		else if (EnumHasAnyFlags(CastFlags, CASTCLASS_FFieldPathProperty))
+		{	
+			CollectFieldPathReference(Collector, static_cast<FFieldPathProperty&>(Property), *reinterpret_cast<FFieldPath*>(Value), Referencer);
+		}
+		else if (EnumHasAnyFlags(CastFlags, CASTCLASS_FInterfaceProperty))
+		{	
+			CollectInterfaceReference(Collector, static_cast<FInterfaceProperty&>(Property), *reinterpret_cast<FScriptInterface*>(Value), Referencer);
+		}
+		else
+		{
+			// Fallback to virtual SerializeItem dispatch inside SerializeBin
+			// for certain Epic plugins that actually add new FProperty types (not recommended)
+			checkf(MayContainStrongReference(CastFlags), TEXT("Missing code to collect references from %s properties (%llx). "
+				"Core property types part of RefLink chain / overloading ContainsObjectReference should be handled above."),
+				*Class->GetFName().ToString(), CastFlags);
+
+			FReferenceCollectorArchive& Archive = Collector.GetVerySlowReferenceCollectorArchive();
+
+			if constexpr (EnumHasAnyFlags(CollectFlags, EPropertyCollectFlags::NeedsReferencer))
+			{
+				FVerySlowReferenceCollectorArchiveScope CollectorScope(Archive, Referencer, &Property);
+				Property.SerializeItem(FStructuredArchiveFromArchive(Archive).GetSlot(), Value, /* defaults */ nullptr);
+			}
+			else
+			{
+				Property.SerializeItem(FStructuredArchiveFromArchive(Archive).GetSlot(), Value, /* defaults */ nullptr);
+			}			
 		}
 	}
+	while (++Idx < ArrayDim);
+}
+
+//////////////////////////////////////////////////////////////////////////
+
+template<EPropertyCollectFlags NonCollectorFlag, /* UStruct or UScriptStruct */ class StructType>
+FORCEINLINE_DEBUGGABLE void CallCollectStructReferences(FReferenceCollector& Collector, const StructType* Struct, void* Instance, const UObject* Referencer)
+{
+	static_assert(!EnumHasAnyFlags(NonCollectorFlag, AllCollectorFlags));
+	static_assert(AllCollectorFlags == EPropertyCollectFlags(3));
+
+	EPropertyCollectFlags CollectorFlags = GetCollectorPropertyFlags(Collector);
+
+	using Func = void (*)(FReferenceCollector&, const StructType*, void*, const UObject*);
+	static constexpr Func Funcs[] = 
+	{
+		&CollectStructReferences<NonCollectorFlag | EPropertyCollectFlags(0), StructType>,
+		&CollectStructReferences<NonCollectorFlag | EPropertyCollectFlags(1), StructType>,
+		&CollectStructReferences<NonCollectorFlag | EPropertyCollectFlags(2), StructType>,
+		&CollectStructReferences<NonCollectorFlag | EPropertyCollectFlags(3), StructType>,
+	};
+
+	uint32 Idx = static_cast<uint32>(CollectorFlags);
+	check(Idx < UE_ARRAY_COUNT(Funcs));
+	(*Funcs[Idx])(Collector, Struct, Instance, Referencer);
+}
+
+void FReferenceCollector::AddPropertyReferences(const UStruct* Struct, void* Instance, const UObject* ReferencingObject)
+{
+	CallCollectStructReferences<EPropertyCollectFlags::None>(*this, Struct, Instance, ReferencingObject);
+}
+
+void FReferenceCollector::AddPropertyReferencesWithStructARO(const UScriptStruct* Struct, void* Instance, const UObject* ReferencingObject)
+{
+	CallCollectStructReferences<EPropertyCollectFlags::CallStructARO>(*this, Struct, Instance, ReferencingObject);
+}
+
+void FReferenceCollector::AddPropertyReferencesLimitedToObjectProperties(const UStruct* Struct, void* Instance, const UObject* ReferencingObject)
+{
+	CallCollectStructReferences<EPropertyCollectFlags::OnlyObjectProperty>(*this, Struct, Instance, ReferencingObject);
 }
 
 void FReferenceCollector::CreateVerySlowReferenceCollectorArchive()
 {
 	check(DefaultReferenceCollectorArchive == nullptr);
-	DefaultReferenceCollectorArchive = new FSimpleObjectReferenceCollectorArchive(nullptr, *this);
+	if (NeedsPropertyReferencer())
+	{
+		DefaultReferenceCollectorArchive.Reset(new FPropertyTrackingReferenceCollectorArchive(nullptr, *this));
+	}
+	else
+	{
+		DefaultReferenceCollectorArchive.Reset(new FReferenceCollectorArchive(nullptr, *this));
+	}
+}
+
+FArchive& FReferenceCollectorArchive::operator<<(UObject*& Object)
+{
+	Collector.AddStableReference(&Object);
+	return *this;
 }
 
-void FReferenceCollector::CreatePersistentFrameReferenceCollectorArchive()
+FArchive& FReferenceCollectorArchive::operator<<(FObjectPtr& Object)
 {
-	check(PersistentFrameReferenceCollectorArchive == nullptr);
-	PersistentFrameReferenceCollectorArchive = new FPersistentFrameCollectorArchive(nullptr, *this);
+	Collector.AddStableReference(reinterpret_cast<UObject**>(&Object));
+	return *this;
 }
 
 /**
@@ -4758,7 +5016,7 @@ private:
 
 		if( !Object->GetClass()->IsChildOf(UClass::StaticClass()) )
 		{
-			FSimpleObjectReferenceCollectorArchive CollectorArchive( Object, *this );
+			FPropertyTrackingReferenceCollectorArchive CollectorArchive( Object, *this );
 			Object->SerializeScriptProperties( CollectorArchive );
 		}
 		Object->CallAddReferencedObjects(*this);
@@ -5697,7 +5955,7 @@ namespace UECodeGen_Private
 		if ((NewClass->ClassFlags & CLASS_Intrinsic) != CLASS_Intrinsic)
 		{
 			check((NewClass->ClassFlags & CLASS_TokenStreamAssembled) != CLASS_TokenStreamAssembled);
-			NewClass->ReferenceTokenStream.Empty();
+			NewClass->ReferenceTokens.Reset();
 		}
 		NewClass->CreateLinkAndAddChildFunctionsToMap(Params.FunctionLinkArray, Params.NumFunctions);
 
diff --git a/Engine/Source/Runtime/CoreUObject/Private/UObject/UnrealType.cpp b/Engine/Source/Runtime/CoreUObject/Private/UObject/UnrealType.cpp
--- a/Engine/Source/Runtime/CoreUObject/Private/UObject/UnrealType.cpp
+++ b/Engine/Source/Runtime/CoreUObject/Private/UObject/UnrealType.cpp
@@ -24,7 +24,8 @@ FPropertyValueIterator::FPropertyValueIterator(
 	if (Entry.ValueArray.Num() > 0)
 	{
 		PropertyIteratorStack.Emplace(MoveTemp(Entry));
-		IterateToNext();
+
+		while (NextValue(InRecursionFlags));
 	}
 }
 
@@ -33,22 +34,11 @@ FPropertyValueIterator::EPropertyValueFlags FPropertyValueIterator::GetPropertyV
 	EPropertyValueFlags Flags = EPropertyValueFlags::None;
 	if (RecursionFlags == EPropertyValueIteratorFlags::FullRecursion)
 	{
-		if (Property->IsA(FArrayProperty::StaticClass()))
-		{
-			Flags = EPropertyValueFlags::IsArray;
-		}
-		else if (Property->IsA(FMapProperty::StaticClass()))
-		{
-			Flags = EPropertyValueFlags::IsMap;
-		}
-		else if (Property->IsA(FSetProperty::StaticClass()))
-		{
-			Flags = EPropertyValueFlags::IsSet;
-		}
-		else if (Property->IsA(FStructProperty::StaticClass()))
-		{
-			Flags = EPropertyValueFlags::IsStruct;
-		}
+		uint64 CastFlags = Property->GetClass()->GetCastFlags();
+		Flags = EPropertyValueFlags(  !!(CastFlags & CASTCLASS_FArrayProperty)	* uint32(EPropertyValueFlags::IsArray)
+									| !!(CastFlags & CASTCLASS_FMapProperty)	* uint32(EPropertyValueFlags::IsMap)
+									| !!(CastFlags & CASTCLASS_FSetProperty)	* uint32(EPropertyValueFlags::IsSet)
+									| !!(CastFlags & CASTCLASS_FStructProperty)	* uint32(EPropertyValueFlags::IsStruct));
 	}
 	if (bMatchAll || Property->IsA(PropertyClass))
 	{
@@ -77,7 +67,7 @@ void FPropertyValueIterator::FillStructProperties(const UStruct* Struct, FProper
 	}
 }
 
-bool FPropertyValueIterator::NextValue(EPropertyValueIteratorFlags InRecursionFlags)
+FORCEINLINE_DEBUGGABLE bool FPropertyValueIterator::NextValue(EPropertyValueIteratorFlags InRecursionFlags)
 {
 	check(PropertyIteratorStack.Num() > 0)
 	FPropertyValueStackEntry& Entry = PropertyIteratorStack.Last();
@@ -200,19 +190,18 @@ bool FPropertyValueIterator::NextValue(EPropertyValueIteratorFlags InRecursionFl
 
 void FPropertyValueIterator::IterateToNext()
 {
-	EPropertyValueIteratorFlags LocalRecursionFlags = RecursionFlags;
-
 	if (bSkipRecursionOnce)
 	{
-		LocalRecursionFlags = EPropertyValueIteratorFlags::NoRecursion;
 		bSkipRecursionOnce = false;
+		
+		if (!NextValue(EPropertyValueIteratorFlags::NoRecursion))
+		{
+			return;
+		}
 	}
 
-	while (NextValue(LocalRecursionFlags))
-	{
-		// Reset recursion override as we've skipped the first property
-		LocalRecursionFlags = RecursionFlags;
-	}
+	EPropertyValueIteratorFlags LocalRecursionFlags = RecursionFlags;	
+	while (NextValue(LocalRecursionFlags));
 }
 
 void FPropertyValueIterator::GetPropertyChain(TArray<const FProperty*>& PropertyChain) const
diff --git a/Engine/Source/Runtime/CoreUObject/Public/UObject/Class.h b/Engine/Source/Runtime/CoreUObject/Public/UObject/Class.h
--- a/Engine/Source/Runtime/CoreUObject/Public/UObject/Class.h
+++ b/Engine/Source/Runtime/CoreUObject/Public/UObject/Class.h
@@ -2785,9 +2785,7 @@ public:
 	TArray<FImplementedInterface> Interfaces;
 
 	/** Reference token stream used by realtime garbage collector, finalized in AssembleReferenceTokenStream */
-	FGCReferenceTokenStream ReferenceTokenStream;
-	/** CS for the token stream. Token stream can assemble code can sometimes be called from two threads throuh a web of async loading calls. */
-	FCriticalSection ReferenceTokenStreamCritical;
+	UE::GC::FTokenStreamOwner ReferenceTokens;
 
 	/** This class's native functions. */
 	TArray<FNativeFunctionLookup> NativeFunctionLookupTable;
@@ -3218,68 +3216,6 @@ public:
 
 	FString GetDescription() const;
 
-	/**
-	 * Realtime garbage collection helper function used to emit token containing information about a 
-	 * direct UObject reference at the passed in offset.
-	 *
-	 * @param Offset	Offset into object at which object reference is stored.
-	 * @param DebugName	DebugName for this objects token. Only used in non-shipping builds.
-	 * @param Kind		Optional parameter the describe the type of the reference.
-	 */
-	void EmitObjectReference(int32 Offset, const FName& DebugName, EGCReferenceType Kind = GCRT_Object);
-
-	/**
-	 * Realtime garbage collection helper function used to emit token containing information about a 
-	 * an array of UObject references at the passed in offset. Handles both TArray and TTransArray.
-	 *
-	 * @param Offset	Offset into object at which array of objects is stored.
-	 * @param DebugName	DebugName for this objects token. Only used in non-shipping builds.
-	 */
-	void EmitObjectArrayReference(int32 Offset, const FName& DebugName);
-
-	/**
-	 * Realtime garbage collection helper function used to indicate an array of structs at the passed in 
-	 * offset.
-	 *
-	 * @param Offset	Offset into object at which array of structs is stored
-	 * @param DebugName	DebugName for this objects token. Only used in non-shipping builds.
-	 * @param Stride	Size/stride of struct
-	 * @return	Index into token stream at which later on index to next token after the array is stored
-	 *			which is used to skip over empty dynamic arrays
-	 */
-	uint32 EmitStructArrayBegin(int32 Offset, const FName& DebugName, int32 Stride);
-
-	/**
-	 * Realtime garbage collection helper function used to indicate the end of an array of structs. The
-	 * index following the current one will be written to the passed in SkipIndexIndex in order to be
-	 * able to skip tokens for empty dynamic arrays.
-	 *
-	 * @param SkipIndexIndex
-	 */
-	void EmitStructArrayEnd(uint32 SkipIndexIndex);
-
-	/**
-	 * Realtime garbage collection helper function used to indicate the beginning of a fixed array.
-	 * All tokens issues between Begin and End will be replayed Count times.
-	 *
-	 * @param Offset	Offset at which fixed array starts.
-	 * @param DebugName	DebugName for this objects token. Only used in non-shipping builds.
-	 * @param Stride	Stride of array element, e.g. sizeof(struct) or sizeof(UObject*).
-	 * @param Count		Fixed array count.
-	 */
-	void EmitFixedArrayBegin(int32 Offset, const FName& DebugName, int32 Stride, int32 Count);
-	
-	/**
-	 * Realtime garbage collection helper function used to indicated the end of a fixed array.
-	 */
-	void EmitFixedArrayEnd();
-
-	/**
-	 * Realtime garbage collection helper function used to emit token containing information about an
-	 * external package reference.
-	 */
-	void EmitExternalPackageReference();
-
 	/**
 	 * Assembles the token stream for realtime garbage collection by combining the per class only
 	 * token stream for each class in the class hierarchy. This is only done once and duplicate
@@ -3446,6 +3382,7 @@ private:
 	 * Tests if all properties tagged with Replicate were registered in GetLifetimeReplicatedProps
 	 */
 	void ValidateRuntimeReplicationData();
+	void AssembleReferenceTokenStreamInternal(bool bForce = false);
 	void InternalCreateDefaultObjectWrapper() const;
 protected:
 	/**
