GC reachability optimizations. ~1.5x speedup and ~2.5MB memory saved for an internal title.

diff --git a/Engine/Source/Runtime/CoreUObject/Public/UObject/FastReferenceCollector.h b/Engine/Source/Runtime/CoreUObject/Public/UObject/FastReferenceCollector.h
--- a/Engine/Source/Runtime/CoreUObject/Public/UObject/FastReferenceCollector.h
+++ b/Engine/Source/Runtime/CoreUObject/Public/UObject/FastReferenceCollector.h
@@ -14,771 +14,569 @@
 #include "UObject/FieldPath.h"
 #include "Async/ParallelFor.h"
 #include "UObject/UObjectArray.h"
-#include "UObject/FastReferenceCollectorOptions.h"
 #include "UObject/DynamicallyTypedValue.h"
 
-struct FStackEntry;
-
 /** Token stream stack overflow checks are not enabled in test and shipping configs */
 #define UE_ENABLE_TOKENSTREAM_STACKOVERFLOW_CHECKS (!(UE_BUILD_TEST || UE_BUILD_SHIPPING) || 0)
 
+namespace UE::GC { class FWorkCoordinator; }
+
 /*=============================================================================
 	FastReferenceCollector.h: Unreal realtime garbage collection helpers
 =============================================================================*/
 
-/** Base class for TFastReferenceCollector array pools. */
-template <typename ArrayStructType>
-class TFastReferenceCollectorArrayPool
+enum class EGCOptions : uint32
 {
-	static_assert(std::is_base_of_v<FGCArrayStruct, ArrayStructType>, "ArrayStructType must derive from FGCArrayStruct.");
-public:
-	/**
-	 * Gets an event from the pool or creates one if necessary.
-	 *
-	 * @return The array.
-	 * @see ReturnToPool
-	 */
-	FORCEINLINE ArrayStructType* GetArrayStructFromPool()
-	{
-		ArrayStructType* Result = Pool.Pop();
-		if (!Result)
-		{
-			Result = new ArrayStructType();
-		}
-		check(Result);
-#if UE_BUILD_DEBUG
-		NumberOfUsedArrays++;
-#endif // UE_BUILD_DEBUG
-		return Result;
-	}
-
-	/**
-	 * Returns an array to the pool.
-	 *
-	 * @param Array The array to return.
-	 * @see GetArrayFromPool
-	 */
-	FORCEINLINE void ReturnToPool(ArrayStructType* ArrayStruct)
-	{
-#if UE_BUILD_DEBUG
-		const int32 CheckUsedArrays = --NumberOfUsedArrays;
-		checkSlow(CheckUsedArrays >= 0);
-#endif // UE_BUILD_DEBUG
-		check(ArrayStruct);
-		ArrayStruct->ObjectsToSerialize.Reset();
-		Pool.Push(ArrayStruct);
-	}
-
-	/**
-	 * Performs manual memory cleanup.
-	 * Generally the pools will be cleaned up when ClearWeakReferences is called on a full GC purge
-	 */
-	void Cleanup()
-	{
-#if UE_BUILD_DEBUG
-		const int32 CheckUsedArrays = NumberOfUsedArrays;
-		checkSlow(CheckUsedArrays == 0);
-#endif // UE_BUILD_DEBUG
-
-		SIZE_T FreedMemory = 0;
-		TArray<ArrayStructType*> AllArrays;
-		Pool.PopAll(AllArrays);
-		for (ArrayStructType* ArrayStruct : AllArrays)
-		{
-			// If we are cleaning up with active weak references the weak references will get corrupted
-			checkSlow(ArrayStruct->WeakReferences.Num() == 0);
-			FreedMemory += ArrayStruct->GetAllocatedSize();
-			delete ArrayStruct;
-		}
-		UE_LOG(LogGarbage, Log, TEXT("Freed %" SIZE_T_FMT "b from %d GC array pools."), FreedMemory, AllArrays.Num());
-	}
+	None = 0,
+	Parallel = 1 << 0,					// Use all task workers to collect references, must be started on main thread
+	AutogenerateTokenStream = 1 << 1,	// Lazily generate token streams for new UClasses
+	WithClusters = 1 << 2,				// Use clusters, see FUObjectCluster
+	ProcessWeakReferences = 1 << 3,		// Collect both weak and strong references instead of just strong
+	WithPendingKill = 1 << 4,			// Internal flag used by reachability analysis
+};
+ENUM_CLASS_FLAGS(EGCOptions);
 
-#if UE_BUILD_DEBUG
-	void CheckLeaks()
-	{
-		// This function is called after GC has finished so at this point there should be no
-		// arrays used by GC and all should be returned to the pool
-		const int32 LeakedGCPoolArrays = NumberOfUsedArrays;
-		checkSlow(LeakedGCPoolArrays == 0);
-	}
-#endif
+inline constexpr bool IsParallel(EGCOptions Options) { return !!(Options & EGCOptions::Parallel); }
+inline constexpr bool IsPendingKill(EGCOptions Options) { return !!(Options & EGCOptions::WithPendingKill); }
+inline constexpr bool ShouldProcessWeakReferences(EGCOptions Options) { return !!(Options & EGCOptions::ProcessWeakReferences); }
 
-protected:
+using EFastReferenceCollectorOptions UE_DEPRECATED(5.2, "Use EGCOptions instead") = EGCOptions;
 
-	/** Holds the collection of recycled arrays. */
-	TLockFreePointerListLIFO< ArrayStructType > Pool;
+/** Helper to give GC internals friend access to certain core classes */
+struct FGCInternals
+{
+	FORCEINLINE static FUObjectItem* GetResolvedOwner(FFieldPath& Path) { return Path.GetResolvedOwnerItemInternal(); }
+	FORCEINLINE static void ClearCachedField(FFieldPath& Path) { Path.ClearCachedFieldInternal(); }
 
-#if UE_BUILD_DEBUG
-	/** Number of arrays currently acquired from the pool by GC */
-	std::atomic<int32> NumberOfUsedArrays;
-#endif // UE_BUILD_DEBUG
+	template<class ObjectID>
+	FORCEINLINE static FWeakObjectPtr& AccessWeakPtr(TPersistentObjectPtr<ObjectID>& Ptr) { return Ptr.WeakPtr; }
 };
 
-/**
- * Pool for reducing GC allocations
- */
-class FGCArrayPool : public TFastReferenceCollectorArrayPool<FGCArrayStruct>
+/** Interface to allow external systems to trace additional object references, used for bridging GCs */
+class FGarbageCollectionTracer
 {
-private:
-	// allows sharing a singleton between all compilation units while still having an inlined getter
-	COREUOBJECT_API static FGCArrayPool* GetGlobalSingleton();
 public:
+	virtual ~FGarbageCollectionTracer() {}
+	virtual void PerformReachabilityAnalysisOnObjects(UE::GC::FWorkerContext* Context, EGCOptions Options) = 0;
+};
 
-	/**
-	 * Gets the singleton instance of the FObjectArrayPool
-	 * @return Pool singleton.
-	 */
-	FORCEINLINE static FGCArrayPool& Get()
-	{
-		static FGCArrayPool* Singleton = nullptr;
-
-		if (!Singleton)
-		{
-			Singleton = GetGlobalSingleton();
-		}
-		return *Singleton;
-	}
-
-	/**
-	 * Writes out info about the makeup of the pool. called by 'gc.DumpPoolStats'
-	 *
-	 * @param Array The array to return.
-	 * @see GetArrayFromPool
-	 */
-	static void DumpStats(FOutputDevice& OutputDevice)
-	{
-		FGCArrayPool& Instance = Get();
+//////////////////////////////////////////////////////////////////////////
 
-		TArray<FGCArrayStruct*> PoppedItems;
+namespace UE::GC
+{
 
-		TMap<int32, int32> Buckets;
+static constexpr uint32 ObjectLookahead = 16;
 
-		int32 TotalSize = 0;
-		int32 MaxSize = 0;
-		int32 TotalItems = 0;
+// Prefetches ClassPrivate/OuterPrivate/Class' ReferenceTokens/TokenData while iterating over an object array
+//
+// Tuned on a Gen5 console using an internal game replay and an in-game GC pass
+class FPrefetchingObjectIterator
+{
+public:
+	// Objects must be padded with PadObjects
+	explicit FPrefetchingObjectIterator(TConstArrayView<UObject*> Objects)
+	: It(Objects.begin())
+	, End(Objects.end())
+	, PrefetchedTokenStream(Objects.Num() ? &It[1]->GetClass()->ReferenceTokens.Strong : nullptr)
+	{}
 
-		do
-		{
-			FGCArrayStruct* Item = Instance.Pool.Pop();
+	FORCEINLINE_DEBUGGABLE void Advance()
+	{	
+		FPlatformMisc::Prefetch(PrefetchedTokenStream->GetTokenData());
+		PrefetchedTokenStream = &It[2]->GetClass()->ReferenceTokens.Strong;
 
-			if (Item)
-			{
-				PoppedItems.Push(Item);
+		UObjectBase::PrefetchOuter(It[6]);
+		FPlatformMisc::Prefetch(It[6]->GetClass(), offsetof(UClass, ReferenceTokens));
+		UObjectBase::PrefetchClass(It[ObjectLookahead]);
 
-				// Inc our bucket
-				Buckets.FindOrAdd(Item->ObjectsToSerialize.Max()) += 1;
+		++It;
+	}
+	
+	bool HasMore() const { return It != End; }
+	UObject* GetCurrentObject() { return *It; }
 
-				TotalSize += Item->ObjectsToSerialize.Max();
-				TotalSize += Item->WeakReferences.Max();
-				TotalItems++;
-			}
-			else
-			{
-				break;
-			}
+private:
+	UObject*const* It;
+	UObject*const* End;
+	FTokenStreamView* PrefetchedTokenStream;
+};
 
-		} while (true);
+// Pad object array for FPrefetchingObjectIterator use
+COREUOBJECT_API void PadObjectArray(TArray<UObject*>& Objects);
 
-		// return everything to the pool
-		while (PoppedItems.Num())
-		{
-			Instance.Pool.Push(PoppedItems.Pop());
-		}
+//////////////////////////////////////////////////////////////////////////
 
-		// One of these lists is used by the main GC and is huge, so remove it so that it doesn't
-		// pollute the stats and we can accurately see what the task pools are using.
-		int32 TotalSizeKB = (TotalSize * sizeof(UObject*)) / 1024;
+/** Fixed block of reachable objects waiting to be processed */
+struct FWorkBlock
+{
+	static constexpr uint32 ObjectCapacity = 512 - /* Previous */ 1 - ObjectLookahead;
 
-		OutputDevice.Logf(TEXT("GCPoolStats: %d Pools totaling %d KB. Avg: Objs=%d, Size=%d KB."),
-			TotalItems,
-			TotalSizeKB,
-			TotalSize / FMath::Max(TotalItems, 1),
-			TotalSizeKB / FMath::Max(TotalItems, 1));
+	FWorkBlock* Previous;
+	UObject* Objects[ObjectCapacity + ObjectLookahead];
 
-		// long form output...
-		TArray<int32> Keys;
+	TArrayView<UObject*> GetObjects() { return MakeArrayView(Objects, ObjectCapacity); }
+	TArrayView<UObject*> GetPadding() { return MakeArrayView(Objects + ObjectCapacity, ObjectLookahead); }
+};
 
-		Buckets.GetKeys(Keys);
+class FWorkstealingQueue;
 
-		Keys.Sort([&](int32 lhs, int32 rhs) {
-			return lhs > rhs;
-		});
+/** Reachable objects waiting to be processed. Type-erases parallel/serial queue. */
+class FWorkBlockifier
+{
+public:
+	UE_NONCOPYABLE(FWorkBlockifier);
+	FWorkBlockifier() = default;
+	COREUOBJECT_API ~FWorkBlockifier();
+	
+	void Init() { AllocateWipBlock(); }
+	void SetAsyncQueue(FWorkstealingQueue& Queue) { AsyncQueue = &Queue; }
+	void ResetAsyncQueue();
 
-		for (int Key : Keys)
+	template<EGCOptions Options>
+	FORCEINLINE_DEBUGGABLE void Add(UObject* Object)
+	{
+		*WipIt = Object;
+		if (++WipIt == Wip->GetPadding().GetData())
 		{
-			const int32 Value = Buckets[Key];
-			int32 ItemSize = (Key * sizeof(UObject*)) / 1024;
-			OutputDevice.Logf(TEXT("\t%d\t\t(%d Items @ %d KB = %d KB)"), Key, Value, ItemSize, Value * ItemSize);
+			if constexpr (IsParallel(Options))
+			{
+				PushFullBlockAsync();
+			}
+			else
+			{
+				PushFullBlockSync();
+			}
 		}
 	}
 
-	/** 
-	 * Clears weak references for everything in the pool. 
-	 * @param AllArrays Arrays to process 
-	 */
-	void ClearWeakReferences(TArray<FGCArrayStruct*>& AllArrays)
+	FORCEINLINE_DEBUGGABLE FWorkBlock* PopPartialBlock(int32& OutNum)
 	{
-		TRACE_CPUPROFILER_EVENT_SCOPE(ClearWeakReferences);
-		for (FGCArrayStruct* ArrayStruct : AllArrays)
+		if (int32 Num = PartialNum())
 		{
-			for (UObject** WeakReference : ArrayStruct->WeakReferences)
-			{
-				UObject*& ReferencedObject = *WeakReference;
-				if (ReferencedObject && ReferencedObject->IsUnreachable())
-				{
-					ReferencedObject = nullptr;
-				}
-			}
-			ArrayStruct->WeakReferences.Reset();
+			OutNum = Num;
+			return PopWipBlock();
 		}
+
+		return nullptr;
 	}
 
-	/**
-	 * Dumps garbage references to the log.
-	 * @param AllArrays Arrays to process
-	 */
-	void DumpGarbageReferencers(TArray<FGCArrayStruct*>& AllArrays);
-
-	/**
-	 * Updates GC history after the last GC run.
-	 * @param AllArrays Arrays to process
-	 */
-	void UpdateGCHistory(TArray<FGCArrayStruct*>& AllArrays);
-
-	/**
-	 * Grabs all arrays from the pool 
-	 */
-	void GetAllArrayStructsFromPool(TArray<FGCArrayStruct*>& AllArrays)
+	template<EGCOptions Options>
+	FORCEINLINE FWorkBlock* PopFullBlock()
 	{
-		Pool.PopAll(AllArrays);
-#if UE_BUILD_DEBUG
-		NumberOfUsedArrays += AllArrays.Num();
-#endif // UE_BUILD_DEBUG
+		return IsParallel(Options) ? PopFullBlockAsync() : PopFullBlockSync();
 	}
 
-	/**
-	 * Frees an allocated array
-	 */
-	void FreeArrayStruct(FGCArrayStruct* InArray)
+	FORCEINLINE FWorkBlock* StealFullBlock() const
 	{
-		delete InArray;
-#if UE_BUILD_DEBUG
-		const int32 CheckUsedArrays = --NumberOfUsedArrays;
-		checkSlow(CheckUsedArrays >= 0);
-#endif // UE_BUILD_DEBUG
+		return StealAsyncBlock();
 	}
+
+	COREUOBJECT_API void FreeOwningBlock(UObject*const* BlockObjects);
+
+	FORCEINLINE bool IsUnused() const
+	{
+		return PartialNum() == 0 && SyncQueue == nullptr;
+	}
+
+	void SetWorkerIndex(int32 Idx) { WorkerIndex = Idx; }
+	int32 GetWorkerIndex() const { return WorkerIndex; }
+
+	
+private:
+	UObject** WipIt; // Wip->Objects cursor
+	FWorkBlock* Wip;
+	union
+	{
+		FWorkBlock* SyncQueue = nullptr;
+		FWorkstealingQueue* AsyncQueue;
+	};
+	int32 WorkerIndex = INDEX_NONE;
+	
+	void AllocateWipBlock();
+	COREUOBJECT_API void PushFullBlockSync();
+	COREUOBJECT_API void PushFullBlockAsync();
+	COREUOBJECT_API FWorkBlock* PopFullBlockSync();
+	COREUOBJECT_API FWorkBlock* PopFullBlockAsync();
+	COREUOBJECT_API FWorkBlock* PopWipBlock();
+	COREUOBJECT_API FWorkBlock* StealAsyncBlock() const;
+	FORCEINLINE int32 PartialNum() const { return WipIt - Wip->Objects; }
 };
 
-/**
- * Helper class that looks for UObject references by traversing UClass token stream and calls AddReferencedObjects.
- * Provides a generic way of processing references that is used by Unreal Engine garbage collection.
- * Can be used for fast (does not use serialization) reference collection purposes.
- * 
- * IT IS CRITICAL THIS CLASS DOES NOT CHANGE WITHOUT CONSIDERING PERFORMANCE IMPACT OF SAID CHANGES
- *
- * This class depends on three components: ReferenceProcessor, ReferenceCollector and ArrayPool.
- * The assumptions for each of those components are as follows:
- *
-   class FSampleReferenceProcessor
-   {
-   public:
-     int32 GetMinDesiredObjectsPerSubTask() const;
-		 void HandleTokenStreamObjectReference(FGCArrayStruct& ObjectsToSerializeStruct, UObject* ReferencingObject, UObject*& Object, const int32 TokenIndex, const EGCTokenType TokenType, bool bAllowReferenceElimination);
-		 void UpdateDetailedStats(UObject* CurrentObject, uint32 DeltaCycles);
-		 void LogDetailedStatsSummary();
-	 };
-
-	 class FSampleCollector : public FReferenceCollector 
-	 {
-	   // Needs to implement FReferenceCollector pure virtual functions
-	 };
-   
-	 class FSampleArrayPool
-	 {
-	   static FSampleArrayPool& Get();
-		 FGCArrayStruct* GetArrayStryctFromPool();
-		 void ReturnToPool(FGCArrayStruct* ArrayStruct);
-	 };
- */
+//////////////////////////////////////////////////////////////////////////
+
+struct FSlowARO
+{
+	COREUOBJECT_API static void CallSync(uint32 SlowAROIndex, UObject* Object, FReferenceCollector& Collector);
+	COREUOBJECT_API static bool TryQueueCall(uint32 SlowAROIndex, UObject* Object, FWorkerContext& Context);
+	COREUOBJECT_API static void ProcessUnbalancedCalls(FWorkerContext& Context, FReferenceCollector& Collector);
+	// @return if any calls were made
+	COREUOBJECT_API static bool ProcessAllCalls(FWorkerContext& Context, FReferenceCollector& Collector);
+};
+
+//////////////////////////////////////////////////////////////////////////
 
-template <typename ReferenceProcessorType, typename CollectorType, typename ArrayPoolType, EFastReferenceCollectorOptions Options = EFastReferenceCollectorOptions::None>
-class TFastReferenceCollector
+/** Thread-local context containing initial objects and references to collect */
+struct alignas(PLATFORM_CACHE_LINE_SIZE) FWorkerContext
 {
 private:
+	template <typename ProcessorType, typename CollectorType>
+	friend class TFastReferenceCollector;
+	
+	// This is set by GC when processing references from the current referencing object
+	UObject* ReferencingObject = nullptr;
+	TConstArrayView<UObject*> InitialObjects;
+public:
+	UE_NONCOPYABLE(FWorkerContext);
+	COREUOBJECT_API FWorkerContext();
+	COREUOBJECT_API ~FWorkerContext();
+
+	FWorkBlockifier ObjectsToSerialize;
+	TConstArrayView<UObject**> InitialNativeReferences;
+	FWorkCoordinator* Coordinator = nullptr;
+	TArray<UObject**> WeakReferences;
+#if !UE_BUILD_SHIPPING
+	TArray<FGarbageReferenceInfo> GarbageReferences;
+	bool bDetectedGarbageReference = false;
+#endif
+#if ENABLE_GC_HISTORY
+	TMap<const UObject*, TArray<FGCDirectReference>*> History;
+#endif
 
-	constexpr FORCEINLINE bool IsParallel() const
-	{
-		return !!(Options & EFastReferenceCollectorOptions::Parallel);
-	}
-	constexpr FORCEINLINE bool CanAutogenerateTokenStream() const
+	FORCEINLINE UObject* GetReferencingObject()	{ return ReferencingObject;	}
+
+	TConstArrayView<UObject*> GetInitialObjects() { return InitialObjects; }
+	void ResetInitialObjects() { InitialObjects = {}; }
+
+	/** @param Objects must outlive this context. It's data is padded by repeating the last object to allow prefetching past the end. */
+	void SetInitialObjectsUnpadded(TArray<UObject*>& Objects)
 	{
-		return !!(Options & EFastReferenceCollectorOptions::AutogenerateTokenStream);
+		PadObjectArray(Objects);
+		SetInitialObjectsPrepadded(Objects);
 	}
-	constexpr FORCEINLINE bool ShouldProcessNoOpTokens() const
+
+	/** @param PaddedObjects must already be padded to allow reading valid objects past the end */
+	void SetInitialObjectsPrepadded(TConstArrayView<UObject*> PaddedObjects)
 	{
-		return !!(Options & EFastReferenceCollectorOptions::ProcessNoOpTokens);
+		check(PaddedObjects.IsEmpty() || PaddedObjects.GetData()[PaddedObjects.Num() + ObjectLookahead - 1]->IsValidLowLevel() );
+		InitialObjects = PaddedObjects;
 	}
-	constexpr FORCEINLINE bool ShouldProcessWeakReferences() const
+
+	/** Returns the size of memory allocated by internal arrays */
+	int64 GetAllocatedSize() const
 	{
-		return !!(Options & EFastReferenceCollectorOptions::ProcessWeakReferences);
+		return WeakReferences.GetAllocatedSize() + sizeof(FWorkBlock);
 	}
 	
-	class FCollectorTaskQueue
-	{
-		TFastReferenceCollector*	Owner;
-		ArrayPoolType& ArrayPool;
-		TLockFreePointerListUnordered<FGCArrayStruct, PLATFORM_CACHE_LINE_SIZE> Tasks;
-
-		FCriticalSection WaitingThreadsLock;
-		TArray<FEvent*> WaitingThreads;
-		bool bDone;
-		int32 NumThreadsStarted;
-	public:
-
-		FCollectorTaskQueue(TFastReferenceCollector* InOwner, ArrayPoolType& InArrayPool)
-			: Owner(InOwner)
-			, ArrayPool(InArrayPool)
-			, bDone(false)
-			, NumThreadsStarted(0)
-		{
-		}
+	FORCEINLINE int32 GetWorkerIndex() const { return ObjectsToSerialize.GetWorkerIndex(); }
+	void AllocateWorkerIndex();
+	void FreeWorkerIndex();
+};
 
-		void CheckDone()
-		{
-			FScopeLock Lock(&WaitingThreadsLock);
-			check(bDone);
-			check(!Tasks.Pop());
-			check(!WaitingThreads.Num());
-			check(NumThreadsStarted);
-		}
+//////////////////////////////////////////////////////////////////////////
 
-		FORCENOINLINE void AddTask(const TArray<UObject*>* InObjectsToSerialize, int32 StartIndex, int32 NumObjects)
-		{
-			FGCArrayStruct* ArrayStruct = ArrayPool.GetArrayStructFromPool();
-			ArrayStruct->ObjectsToSerialize.AddUninitialized(NumObjects);
-			FMemory::Memcpy(ArrayStruct->ObjectsToSerialize.GetData(), InObjectsToSerialize->GetData() + StartIndex, NumObjects * sizeof(UObject*));
-			Tasks.Push(ArrayStruct);
+/* Debug id for references that lack a token */
+enum class ETokenlessId
+{
+	Collector = 1,
+	Class,
+	Outer,
+	ClassOuter,
+	Cluster,
+	InitialReference,
+	Max = InitialReference
+};
 
-			FEvent* WaitingThread = nullptr;
-			{
-				FScopeLock Lock(&WaitingThreadsLock);
-				check(!bDone);
-				if (WaitingThreads.Num())
-				{
-					WaitingThread = WaitingThreads.Pop();
-				}
-			}
-			if (WaitingThread)
-			{
-				WaitingThread->Trigger();
-			}
-		}
+/** Debug identifier for a token stream index or a tokenless reference */
+struct FTokenId
+{
+public:
+	FORCEINLINE FTokenId(ETokenlessId Tokenless) : Index(0), TokenlessId((uint32)Tokenless), Mixed(0) {}
+	FORCEINLINE FTokenId(uint32 Idx, bool bIsMixed) : Index(Idx), TokenlessId(0), Mixed(bIsMixed) {}
+	
+	bool IsMixed() const { return Mixed; }
+	bool IsTokenless() const { return TokenlessId != 0; }
+	uint32 GetIndex() const { check(!IsTokenless()); return Index; }
+	int32 AsPrintableIndex() const { return IsTokenless() ? -int32(TokenlessId) : Index; }
+	ETokenlessId AsTokenless() const { check(IsTokenless()); return (ETokenlessId)TokenlessId;}
 
-		FORCENOINLINE void DoTask()
-		{
-			TRACE_CPUPROFILER_EVENT_SCOPE(FCollectorTaskQueue::DoTask);
-			{
-				FScopeLock Lock(&WaitingThreadsLock);
-				if (bDone)
-				{
-					return;
-				}
-				NumThreadsStarted++;
-			}
-			while (true)
-			{
-				FGCArrayStruct* ArrayStruct = Tasks.Pop();
-				while (!ArrayStruct)
-				{
-					if (bDone)
-					{
-						return;
-					}
-					FEvent* WaitEvent = nullptr;
-					{
-						FScopeLock Lock(&WaitingThreadsLock);
-						if (bDone)
-						{
-							return;
-						}
-						ArrayStruct = Tasks.Pop();
-						if (!ArrayStruct)
-						{
-							if (WaitingThreads.Num() + 1 == NumThreadsStarted)
-							{
-								bDone = true;
-								FPlatformMisc::MemoryBarrier();
-								for (FEvent* WaitingThread : WaitingThreads)
-								{
-									WaitingThread->Trigger();
-								}
-								WaitingThreads.Empty();
-								return;
-							}
-							else
-							{
-								WaitEvent = FPlatformProcess::GetSynchEventFromPool(false);
-								WaitingThreads.Push(WaitEvent);
-							}
-						}
-					}
-					if (ArrayStruct)
-					{
-						check(!WaitEvent);
-					}
-					else
-					{
-						TRACE_CPUPROFILER_EVENT_SCOPE(FCollectorTaskQueue::WaitEvent);
-						check(WaitEvent);
-						WaitEvent->Wait();
-						FPlatformProcess::ReturnSynchEventToPool(WaitEvent);
-						ArrayStruct = Tasks.Pop();
-						check(!ArrayStruct || !bDone);
-					}
-				}
-				Owner->ProcessObjectArray(*ArrayStruct, FGraphEventRef());
-				ArrayPool.ReturnToPool(ArrayStruct);
-			}
-		}
-	};
+	friend bool operator==(FTokenId A, FTokenId B) { return A.All == B.All; }
+	friend bool operator!=(FTokenId A, FTokenId B) { return A.All != B.All; }
+
+private:
+	static constexpr uint32 TokenlessIdBits = 3;
+	static constexpr uint32 IndexBits = 32 - TokenlessIdBits - 1;
+	void StaticAssert();
 
-	/** Task graph task responsible for processing UObject array */
-	class FCollectorTaskProcessorTask
+	union
 	{
-		FCollectorTaskQueue& TaskQueue;
-		ENamedThreads::Type DesiredThread;
-	public:
-		FCollectorTaskProcessorTask(FCollectorTaskQueue& InTaskQueue, ENamedThreads::Type InDesiredThread)
-			: TaskQueue(InTaskQueue)
-			, DesiredThread(InDesiredThread)
-		{
-		}
-		FORCEINLINE TStatId GetStatId() const
-		{
-			RETURN_QUICK_DECLARE_CYCLE_STAT(FCollectorTaskProcessorTask, STATGROUP_TaskGraphTasks);
-		}
-		ENamedThreads::Type GetDesiredThread()
-		{
-			return DesiredThread;
-		}
-		static ESubsequentsMode::Type GetSubsequentsMode()
-		{
-			return ESubsequentsMode::TrackSubsequents;
-		}
-		void DoTask(ENamedThreads::Type CurrentThread, FGraphEventRef& MyCompletionGraphEvent)
+		struct  
 		{
-			TaskQueue.DoTask();
-		}
+			uint32 Index : IndexBits;
+			uint32 TokenlessId : TokenlessIdBits;
+			uint32 Mixed : 1;
+		};
+		uint32 All;
 	};
+};
+
+//////////////////////////////////////////////////////////////////////////
 
-	/** Task graph task responsible for processing UObject array */
-	class FCollectorTask
+/** Helper struct for stack based approach */
+struct FStackEntry
+{
+	/** Current data pointer, incremented by stride */
+	uint8* Data;
+	/** Current container property for data pointer. DO NOT rely on its value being initialized. Instead check ContainerType first. */
+	FProperty* ContainerProperty;
+	/** Pointer to the container being processed by GC. DO NOT rely on its value being initialized. Instead check ContainerType first. */
+	void* ContainerPtr;
+	/** Current index within the container. DO NOT rely on its value being initialized. Instead check ContainerType first. */
+	int32	ContainerIndex;
+	/** Current container helper type */
+	uint32	ContainerType : 5; // The number of bits needs to match FGCReferenceInfo::Type
+	/** Current stride */
+	uint32	Stride : 27; // This will always be bigger (8 bits more) than FGCReferenceInfo::Ofset which is the max offset GC can handle
+	/** Current loop count, decremented each iteration */
+	int32	Count;
+	/** First token index in loop */
+	int32	LoopStartIndex;
+
+	FORCEINLINE bool MoveToNextContainerElementAndCheckIfValid()
 	{
-		TFastReferenceCollector*	Owner;
-		FGCArrayStruct*	ArrayStruct;
-		ArrayPoolType& ArrayPool;
-
-	public:
-		FCollectorTask(TFastReferenceCollector* InOwner, const TArray<UObject*>* InObjectsToSerialize, int32 StartIndex, int32 NumObjects, ArrayPoolType& InArrayPool)
-			: Owner(InOwner)
-			, ArrayStruct(InArrayPool.GetArrayStructFromPool())
-			, ArrayPool(InArrayPool)
-		{
-			ArrayStruct->ObjectsToSerialize.AddUninitialized(NumObjects);
-			FMemory::Memcpy(ArrayStruct->ObjectsToSerialize.GetData(), InObjectsToSerialize->GetData() + StartIndex, NumObjects * sizeof(UObject*));
-		}
-		~FCollectorTask()
-		{
-			ArrayPool.ReturnToPool(ArrayStruct);
-		}
-		FORCEINLINE TStatId GetStatId() const
-		{
-			RETURN_QUICK_DECLARE_CYCLE_STAT(FCollectorTask, STATGROUP_TaskGraphTasks);
-		}
-		static ENamedThreads::Type GetDesiredThread()
-		{
-			return FPlatformProcess::GetDesiredThreadForUObjectReferenceCollector();
-		}
-		static ESubsequentsMode::Type GetSubsequentsMode()
+		if (ContainerType == GCRT_AddTMapReferencedObjects)
 		{
-			return ESubsequentsMode::TrackSubsequents;
+			return ((FMapProperty*)ContainerProperty)->IsValidIndex(ContainerPtr, ++ContainerIndex);
 		}
-		void DoTask(ENamedThreads::Type CurrentThread, const FGraphEventRef& MyCompletionGraphEvent)
+		else if (ContainerType == GCRT_AddTSetReferencedObjects)
 		{
-			Owner->ProcessObjectArray(*ArrayStruct, MyCompletionGraphEvent);
+			return ((FSetProperty*)ContainerProperty)->IsValidIndex(ContainerPtr, ++ContainerIndex);
 		}
-	};
-
-	/** Object that handles all UObject references */
-	ReferenceProcessorType& ReferenceProcessor;
-	/** Custom TArray allocator */
-	ArrayPoolType& ArrayPool;
+		
+		return true;
+	}
+};
 
-	FCollectorTaskQueue TaskQueue;
+//////////////////////////////////////////////////////////////////////////
 
-	/** Helper struct for stack based approach */
-	struct FStackEntry
-	{
-		/** Current data pointer, incremented by stride */
-		uint8* Data;
-		/** Current container property for data pointer. DO NOT rely on its value being initialized. Instead check ContainerType first. */
-		FProperty* ContainerProperty;
-		/** Pointer to the container being processed by GC. DO NOT rely on its value being initialized. Instead check ContainerType first. */
-		void* ContainerPtr;
-		/** Current index within the container. DO NOT rely on its value being initialized. Instead check ContainerType first. */
-		int32	ContainerIndex;
-		/** Current container helper type */
-		uint32	ContainerType : 5; // The number of bits needs to match FGCReferenceInfo::Type
-		/** Current stride */
-		uint32	Stride : 27; // This will always be bigger (8 bits more) than FGCReferenceInfo::Ofset which is the max offset GC can handle
-		/** Current loop count, decremented each iteration */
-		int32	Count;
-		/** First token index in loop */
-		int32	LoopStartIndex;
-	};
+/** Helper class to manage GC token stream stack  */
+class FTokenStreamStack
+{
+	/** Allocated stack memory */
+	TArray<FStackEntry> Stack;
+	/** Current stack frame */
+	FStackEntry* CurrentFrame = nullptr;
+#if UE_ENABLE_TOKENSTREAM_STACKOVERFLOW_CHECKS
+	/** Number of used stack frames (for debugging) */
+	int32 NumberOfUsedStackFrames = 0;
+#endif // UE_ENABLE_TOKENSTREAM_STACKOVERFLOW_CHECKS
 
 public:
-	/** Default constructor, initializing all members. */
-	TFastReferenceCollector(ReferenceProcessorType& InReferenceProcessor, ArrayPoolType& InArrayPool)
-		: ReferenceProcessor(InReferenceProcessor)
-		, ArrayPool(InArrayPool)
-		, TaskQueue(this, InArrayPool)
-	{}
-
-	/**
-	* Performs reachability analysis.
-	*
-	* @param ObjectsToCollectReferencesFor List of objects which references should be collected
-	* @param bForceSingleThreaded Collect references on a single thread
-	*/
-	void CollectReferences(FGCArrayStruct& ArrayStruct)
+	FTokenStreamStack()
 	{
-		TArray<UObject*>& ObjectsToCollectReferencesFor = ArrayStruct.ObjectsToSerialize;
-		if (ObjectsToCollectReferencesFor.Num())
-		{
-			if (!IsParallel())
-			{
-				FGraphEventRef InvalidRef;
-				ProcessObjectArray(ArrayStruct, InvalidRef);
-			}
-			else
-			{
-				int32 NumThreads = FTaskGraphInterface::Get().GetNumWorkerThreads() + 1; // +1 for GT
-				// Avoid going over 26 threads as it may cause race conditions in TLockFreePointerListUnordered
-				int32 NumTasks = FMath::Min(NumThreads, 26);
-
-				check(NumTasks > 0);
-				int32 NumPerChunk = ObjectsToCollectReferencesFor.Num() / NumTasks;
-				int32 StartIndex = 0;
-				for (int32 Chunk = 0; Chunk < NumTasks; Chunk++)
-				{
-					if (Chunk + 1 == NumTasks)
-					{
-						NumPerChunk = ObjectsToCollectReferencesFor.Num() - StartIndex; // last chunk takes all remaining items
-					}
-					TaskQueue.AddTask(&ObjectsToCollectReferencesFor, StartIndex, NumPerChunk);
-					StartIndex += NumPerChunk;
-				}
-
-				// Use ParallelFor to ensure we're making progress in this thread instead of just waiting for tasks to complete.
-				// This is especially important for good game-thread latency when the taskgraph is already filled and to avoid
-				// potential deadlock with workers making use of FGCScopeGuard.
-				ParallelForTemplate(NumTasks, [this](int32) { TaskQueue.DoTask(); }, EParallelForFlags::Unbalanced);
-				TaskQueue.CheckDone();
-			}
-		}
+		Stack.AddUninitialized(FTokenStreamBuilder::GetMaxStackSize());
 	}
 
-private:
-	FORCEINLINE void ConditionalHandleTokenStreamObjectReference(FGCArrayStruct& ObjectsToSerializeStruct, UObject* ReferencingObject, UObject*& Object, const int32 TokenIndex, const EGCTokenType TokenType, bool bAllowReferenceElimination)
+	/** Initializes the stack and returns its first frame */
+	FORCEINLINE FStackEntry* Initialize()
 	{
-		if (IsObjectHandleResolved(*reinterpret_cast<FObjectHandle*>(&Object)))
-		{
-			ReferenceProcessor.HandleTokenStreamObjectReference(ObjectsToSerializeStruct, ReferencingObject, Object, TokenIndex, TokenType, bAllowReferenceElimination);
-		}
+#if UE_ENABLE_TOKENSTREAM_STACKOVERFLOW_CHECKS
+		NumberOfUsedStackFrames = 1;
+#endif // UE_ENABLE_TOKENSTREAM_STACKOVERFLOW_CHECKS
+		// Grab te first frame. From now on we'll be using it instead of TArray to move to the next one and back
+		// in order to skip any extra internal TArray checks and overhead
+		CurrentFrame = GetBottomFrame();
+		return CurrentFrame;
 	}
 
-	FORCEINLINE bool MoveToNextContainerElementAndCheckIfValid(FStackEntry* StackEntry) const
+	/** Returns the frame at the bottom of the stack (the first one) */
+	FORCEINLINE FStackEntry* GetBottomFrame()
 	{
-		switch (StackEntry->ContainerType)
-		{
-		case GCRT_AddTMapReferencedObjects:
-		{
-			FMapProperty* MapProperty = (FMapProperty*)StackEntry->ContainerProperty;
-			return MapProperty->IsValidIndex(StackEntry->ContainerPtr, ++StackEntry->ContainerIndex);
-		}
-		case GCRT_AddTSetReferencedObjects:
-		{
-			FSetProperty* SetProperty = (FSetProperty*)StackEntry->ContainerProperty;
-			return SetProperty->IsValidIndex(StackEntry->ContainerPtr, ++StackEntry->ContainerIndex);
-		}
-		default:
-		{
-			return true;
-		}
-		}
+		return Stack.GetData();
 	}
 
-	/**
-	 * Handles weak object pointer references
-	 * @param WeakPtr weak object pointer
-	 * @param NewObjectsToSerialize List of new objects to process as a result of processing this reference
-	 * @param CurrentObject current object being processed (owner of the weak object pointer)
-	 * @param ReferenceTokenStreamIndex GC token stream index (for debugging)
-	 */
-	FORCEINLINE void HandleWeakObjectPtr(FWeakObjectPtr& WeakPtr, FGCArrayStruct& NewObjectsToSerializeStruct, UObject* CurrentObject, int32 ReferenceTokenStreamIndex, const EGCTokenType TokenType)
+	/** Advances to the next frame on the stack */
+	FORCEINLINE FStackEntry* Push()
 	{
-		UObject* WeakObject = WeakPtr.Get(true);
-		ReferenceProcessor.HandleTokenStreamObjectReference(NewObjectsToSerializeStruct, CurrentObject, WeakObject, ReferenceTokenStreamIndex, TokenType, true);
+		checkSlow(NumberOfUsedStackFrames > 0); // sanity check to make sure Push() gets called after Initialize()
+#if UE_ENABLE_TOKENSTREAM_STACKOVERFLOW_CHECKS
+		NumberOfUsedStackFrames++;
+		UE_CLOG(NumberOfUsedStackFrames > Stack.Num(), LogGarbage, Fatal, TEXT("Ran out of stack space for GC Token Stream. FGCReferenceTokenStream::GetMaxStackSize() = %d must be miscalculated. Verify EmitReferenceInfo code."), FTokenStreamBuilder::GetMaxStackSize());
+#endif // UE_ENABLE_TOKENSTREAM_STACKOVERFLOW_CHECKS
+		return ++CurrentFrame;
 	}
 
-	/**
-	 * Helper class to manage GC token stram stack 
-	 */
-	class FTokenStreamStack
+	/** Pops back to the previous frame on the stack */
+	FORCEINLINE FStackEntry* Pop()
 	{
-		/** Allocated stack memory */
-		TArray<FStackEntry> Stack;
-		/** Current stack frame */
-		FStackEntry* CurrentFrame = nullptr;
 #if UE_ENABLE_TOKENSTREAM_STACKOVERFLOW_CHECKS
-		/** Number of used stack frames (for debugging) */
-		int32 NumberOfUsedStackFrames = 0;
+		NumberOfUsedStackFrames--;
+		UE_CLOG(NumberOfUsedStackFrames < 1, LogGarbage, Fatal, TEXT("GC token stream stack Pop() was called too many times. Probably a Push() call is missing for one of the tokens."));
 #endif // UE_ENABLE_TOKENSTREAM_STACKOVERFLOW_CHECKS
+		return --CurrentFrame;
+	}
+};
 
-	public:
+//////////////////////////////////////////////////////////////////////////
 
-		FTokenStreamStack()
-		{
-			Stack.AddUninitialized(FGCReferenceTokenStream::GetMaxStackSize());
-		}
+/** Forwards references directly to ProcessorType::HandleTokenStreamObjectReference(), unlike TBatchDispatcher */
+template<class ProcessorType, class ContextType>
+struct TDirectDispatcher
+{
+	ProcessorType& Processor;
+	ContextType& Context;
 
-		/** Initializes the stack and returns its first frame */
-		FORCEINLINE FStackEntry* Initialize()
+	FORCEINLINE void HandleReferenceDirectly(UObject* ReferencingObject, UObject*& Object, FTokenId TokenId, EGCTokenType TokenType, bool bAllowReferenceElimination) const
+	{
+		if (IsObjectHandleResolved(*reinterpret_cast<FObjectHandle*>(&Object)))
 		{
-#if UE_ENABLE_TOKENSTREAM_STACKOVERFLOW_CHECKS
-			NumberOfUsedStackFrames = 1;
-#endif // UE_ENABLE_TOKENSTREAM_STACKOVERFLOW_CHECKS
-			// Grab te first frame. From now on we'll be using it instead of TArray to move to the next one and back
-			// in order to skip any extra internal TArray checks and overhead
-			CurrentFrame = GetBottomFrame();
-			return CurrentFrame;
+			Processor.HandleTokenStreamObjectReference(Context, ReferencingObject, Object, TokenId, TokenType, bAllowReferenceElimination);
 		}
+	}
+	
+	FORCEINLINE void HandleKillableReference(UObject*& Object, FTokenId TokenId, EGCTokenType TokenType) const
+	{
+		HandleReferenceDirectly(Context.GetReferencingObject(), Object, TokenId, TokenType, true);
+	}
 
-		/** Returns the frame at the bottom of the stack (the first one) */
-		FORCEINLINE FStackEntry* GetBottomFrame()
-		{
-			return Stack.GetData();
-		}
+	FORCEINLINE void HandleImmutableReference(UObject* Object, FTokenId TokenId, EGCTokenType TokenType) const
+	{
+		HandleReferenceDirectly(Context.GetReferencingObject(), Object, TokenId, TokenType, false);
+	}
 
-		/** Advances to the next frame on the stack */
-		FORCEINLINE FStackEntry* Push()
+	FORCEINLINE void HandleKillableArray(TArray<UObject*>& Array, FTokenId TokenId, EGCTokenType TokenType) const
+	{
+		for (UObject*& Object : Array)
 		{
-			checkSlow(NumberOfUsedStackFrames > 0); // sanity check to make sure Push() gets called after Initialize()
-#if UE_ENABLE_TOKENSTREAM_STACKOVERFLOW_CHECKS
-			NumberOfUsedStackFrames++;
-			UE_CLOG(NumberOfUsedStackFrames > Stack.Num(), LogGarbage, Fatal, TEXT("Ran out of stack space for GC Token Stream. FGCReferenceTokenStream::GetMaxStackSize() = %d must be miscalculated. Verify EmitReferenceInfo code."), FGCReferenceTokenStream::GetMaxStackSize());
-#endif // UE_ENABLE_TOKENSTREAM_STACKOVERFLOW_CHECKS
-			return ++CurrentFrame;
+			HandleReferenceDirectly(Context.GetReferencingObject(), Object, TokenId, TokenType, true);
 		}
+	}
 
-		/** Pops back to the previous frame on the stack */
-		FORCEINLINE FStackEntry* Pop()
+	FORCEINLINE void HandleKillableReferences(TArrayView<UObject*> Objects, FTokenId TokenId, EGCTokenType TokenType) const
+	{
+		for (UObject*& Object : Objects)
 		{
-#if UE_ENABLE_TOKENSTREAM_STACKOVERFLOW_CHECKS
-			NumberOfUsedStackFrames--;
-			UE_CLOG(NumberOfUsedStackFrames < 1, LogGarbage, Fatal, TEXT("GC token stream stack Pop() was called too many times. Probably a Push() call is missing for one of the tokens."));
-#endif // UE_ENABLE_TOKENSTREAM_STACKOVERFLOW_CHECKS
-			return --CurrentFrame;
+			HandleReferenceDirectly(Context.GetReferencingObject(), Object, TokenId, TokenType, true);
 		}
-	};
+	}
 
-	/**
-	 * Traverses UObject token stream to find existing references
-	 *
-	 * @param InObjectsToSerializeArray Objects to process
-	 * @param MyCompletionGraphEvent Task graph event
-	 */
-	void ProcessObjectArray(FGCArrayStruct& InObjectsToSerializeStruct, const FGraphEventRef& MyCompletionGraphEvent)
+	FORCEINLINE void HandleWeakReference(FWeakObjectPtr& WeakPtr, UObject* ReferencingObject, FTokenId TokenId, EGCTokenType TokenType) const
 	{
-		UObject* CurrentObject = nullptr;
+		UObject* WeakObject = WeakPtr.Get(true);
+		HandleReferenceDirectly(ReferencingObject, WeakObject, TokenId, TokenType, false);
+	}
 
-		const int32 MinDesiredObjectsPerSubTask = ReferenceProcessor.GetMinDesiredObjectsPerSubTask(); // sometimes there will be less, a lot less
+	FORCEINLINE void FlushQueuedReferences() const {}
+};
 
-		/** Growing array of objects that require serialization */
-		FGCArrayStruct&	NewObjectsToSerializeStruct = *ArrayPool.GetArrayStructFromPool();
+// Default implementation is to create new direct dispatcher
+template<class CollectorType, class ProcessorType, class ContextType>
+TDirectDispatcher<ProcessorType, ContextType> GetDispatcher(CollectorType&, ProcessorType& Processor, ContextType& Context)
+{
+	return { Processor, Context };
+}
 
-		// Ping-pong between these two arrays if there's not enough objects to spawn a new task
-		TArray<UObject*>& ObjectsToSerialize = InObjectsToSerializeStruct.ObjectsToSerialize;
-		TArray<UObject*>& NewObjectsToSerialize = NewObjectsToSerializeStruct.ObjectsToSerialize;
+//////////////////////////////////////////////////////////////////////////
 
-		// Presized "recursion" stack for handling arrays and structs.
-		FTokenStreamStack Stack;
+enum class ELoot { Nothing, Block, ARO, Context };
+COREUOBJECT_API ELoot StealWork(FWorkerContext& Context, FReferenceCollector& Collector, FWorkBlock*& OutBlock);
 
-		// it is necessary to have at least one extra item in the array memory block for the iffy prefetch code, below
-		ObjectsToSerialize.Reserve(ObjectsToSerialize.Num() + 1);
+/** Allocates contexts and coordinator, kicks worker tasks that also call ProcessSync. Processor is type-erased to void* to avoid templated code. */
+COREUOBJECT_API void ProcessAsync(void (*ProcessSync)(void*, FWorkerContext&), void* Processor, FWorkerContext& InitialContext);
 
-		// Keep serializing objects till we reach the end of the growing array at which point
-		// we are done.
-		int32 CurrentIndex = 0;
-		do
-		{
-			CollectorType ReferenceCollector(ReferenceProcessor, NewObjectsToSerializeStruct);
-			while (CurrentIndex < ObjectsToSerialize.Num())
-			{
-#if PERF_DETAILED_PER_CLASS_GC_STATS
-				uint32 StartCycles = FPlatformTime::Cycles();
-#endif
-				CurrentObject = ObjectsToSerialize[CurrentIndex++];
-				checkSlow(CurrentObject);
+//////////////////////////////////////////////////////////////////////////
 
-				// GetData() used to avoiding bounds checking (min and max)
-				// FMath::Min used to avoid out of bounds (without branching) on last iteration. Though anything can be passed into PrefetchBlock, 
-				// reading ObjectsToSerialize out of bounds is not safe since ObjectsToSerialize[Num()] may be an unallocated/unsafe address.
-				const UObject * const NextObject = ObjectsToSerialize.GetData()[FMath::Min<int32>(CurrentIndex, ObjectsToSerialize.Num() - 1)];
+/**
+ * Helper class that looks for UObject references by traversing UClass token stream and calls AddReferencedObjects.
+ * Provides a generic way of processing references that is used by Unreal Engine garbage collection.
+ * Can be used for fast (does not use serialization) reference collection purposes.
+ * 
+ * IT IS CRITICAL THIS CLASS DOES NOT CHANGE WITHOUT CONSIDERING PERFORMANCE IMPACT OF SAID CHANGES
+ *
+ * @see FSimpleReferenceProcessorBase and TDefaultCollector for documentation on required APIs
+ */
+template <typename ProcessorType, typename CollectorType>
+class TFastReferenceCollector : public FGCInternals
+{
+	static constexpr EGCOptions Options = ProcessorType::Options;
 
-				// Prefetch the next object assuming that the property size of the next object is the same as the current one.
-				// This allows us to avoid a branch here.
-				FPlatformMisc::PrefetchBlock(NextObject, CurrentObject->GetClass()->GetPropertiesSize());
+	static constexpr FORCEINLINE bool IsParallel()					{ return !!(Options & EGCOptions::Parallel); }
+	static constexpr FORCEINLINE bool ShouldProcessWeakReferences()	{ return !!(Options & EGCOptions::ProcessWeakReferences); }
 
-				//@todo rtgc: we need to handle object references in struct defaults
+	static FORCEINLINE FTokenId MakeTokenId(uint32 Index)
+	{
+		return FTokenId(Index, ShouldProcessWeakReferences());
+	}
 
-				// Make sure that token stream has been assembled at this point as the below code relies on it.
-				if (!IsParallel() && CanAutogenerateTokenStream())
-				{
-					UClass* ObjectClass = CurrentObject->GetClass();
-					if (!ObjectClass->HasAnyClassFlags(CLASS_TokenStreamAssembled))
-					{
-						ObjectClass->AssembleReferenceTokenStream();
-					}
-				}
-#if DO_CHECK
-				if (!CurrentObject->GetClass()->HasAnyClassFlags(CLASS_TokenStreamAssembled))
+	static FORCEINLINE UE::GC::FTokenStreamView& GetTokenStream(UClass* Class)
+	{
+		if (!!(Options & EGCOptions::AutogenerateTokenStream) && !Class->HasAnyClassFlags(CLASS_TokenStreamAssembled))
+		{			
+			Class->AssembleReferenceTokenStream();
+		}
+
+		return ShouldProcessWeakReferences() ? Class->ReferenceTokens.Mixed : Class->ReferenceTokens.Strong;
+	}
+
+	ProcessorType& Processor;
+
+public:
+	TFastReferenceCollector(ProcessorType& InProcessor) : Processor(InProcessor) {}
+
+	void ProcessObjectArray(FWorkerContext& Context)
+	{
+		static_assert(!EnumHasAllFlags(Options, EGCOptions::Parallel | EGCOptions::AutogenerateTokenStream), "Can't assemble token streams in parallel");
+
+		// Presized "recursion" stack for handling arrays and structs.
+		FTokenStreamStack Stack;
+		
+		CollectorType Collector(Processor, Context);
+
+		// Either TDirectDispatcher living on the stack or TBatchDispatcher reference owned by Collector
+		decltype(GetDispatcher(Collector, Processor, Context)) Dispatcher = GetDispatcher(Collector, Processor, Context);
+
+StoleContext:
+		// Process initial references first
+		for (UObject** InitialReference : Context.InitialNativeReferences)
+		{
+			Dispatcher.HandleKillableReference(*InitialReference, ETokenlessId::Collector, EGCTokenType::Native);
+		}
+
+		TConstArrayView<UObject*> CurrentObjects = Context.InitialObjects;
+		while (true)
+		{
+			for (FPrefetchingObjectIterator It(CurrentObjects); It.HasMore(); It.Advance())
+			{
+				UObject* CurrentObject = It.GetCurrentObject();
+				UClass* Class = CurrentObject->GetClass();
+				UObject* Outer = CurrentObject->GetOuter();
+
+				FTokenStreamView TokenStream = GetTokenStream(Class);
+				Context.ReferencingObject = CurrentObject;
+
+				// Emit base references
+				Dispatcher.HandleImmutableReference(Class, ETokenlessId::Class, EGCTokenType::Native);
+				Dispatcher.HandleImmutableReference(Outer, ETokenlessId::Outer, EGCTokenType::Native);
+
+				if (TokenStream.IsEmpty())
 				{
-					UE_LOG(LogGarbage, Fatal, TEXT("%s does not yet have a token stream assembled."), *GetFullNameSafe(CurrentObject->GetClass()));
+					continue;
 				}
-#endif
-				// Store the referencing object within the FGCArrayStruct so that we can always fall back to it
-				// when AddReferencedObjects functions don't provide one
-				NewObjectsToSerializeStruct.ReferencingObject = CurrentObject;
 
-				// Get pointer to token stream and jump to the start.
-				FGCReferenceTokenStream* RESTRICT TokenStream = &CurrentObject->GetClass()->ReferenceTokenStream;
 				uint32 TokenStreamIndex = 0;
 				// Keep track of index to reference info. Used to avoid LHSs.
 				uint32 ReferenceTokenStreamIndex = 0;
 
 				// Create stack entry and initialize sane values.
-				FStackEntry* RESTRICT StackEntry = Stack.Initialize();
+				FStackEntry* StackEntry = Stack.Initialize();
 				uint8* StackEntryData = (uint8*)CurrentObject;
 				StackEntry->Data = StackEntryData;
 				StackEntry->ContainerType = GCRT_None;
 				StackEntry->Stride = 0;
 				StackEntry->Count = -1;
 				StackEntry->LoopStartIndex = -1;
-
+				
 				// Keep track of token return count in separate integer as arrays need to fiddle with it.
-				int32 TokenReturnCount = 0;
+				uint8 TokenReturnCount = 0;
 
 				// Parse the token stream.
 				while (true)
@@ -787,7 +585,7 @@ private:
 					ReferenceTokenStreamIndex = TokenStreamIndex;
 
 					// Handle returning from an array of structs, array of structs of arrays of ... (yadda yadda)
-					for (int32 ReturnCount = 0; ReturnCount<TokenReturnCount; ReturnCount++)
+					for (uint8 ReturnCount = 0; ReturnCount < TokenReturnCount; ReturnCount++)
 					{
 						// Make sure there's no stack underflow.
 						check(StackEntry->Count != -1);
@@ -810,7 +608,7 @@ private:
 								{
 									StackEntryData = StackEntry->Data + StackEntry->Stride;
 									StackEntry->Data = StackEntryData;
-								} while (!MoveToNextContainerElementAndCheckIfValid(StackEntry));
+								} while (!StackEntry->MoveToNextContainerElementAndCheckIfValid());
 							}
 
 							// Jump back to the beginning of the loop.
@@ -828,622 +626,588 @@ private:
 					}
 
 					TokenStreamIndex++;
-					const FGCReferenceInfo ReferenceInfo = TokenStream->AccessReferenceInfo(ReferenceTokenStreamIndex);
-					const EGCTokenType TokenType = TokenStream->GetTokenType();
 
-					switch(ReferenceInfo.Type)
+					struct FUnpackedReferenceInfo
 					{
-					case GCRT_Object:
-					case GCRT_Class:
+						explicit FUnpackedReferenceInfo(FGCReferenceInfo In)
+						: ReturnCount(In.ReturnCount)
+						, Type(static_cast<EGCReferenceType>(In.Type))
+						, Offset(In.Offset)
+						{}
+
+						uint8 ReturnCount;
+						EGCReferenceType Type;
+						uint32 Offset;
+					};
+					const FUnpackedReferenceInfo ReferenceInfo(TokenStream.AccessReferenceInfo(ReferenceTokenStreamIndex));
+							
+					const EGCTokenType TokenType = TokenStream.GetTokenType();
+
+					switch (ReferenceInfo.Type)
 					{
-						// We're dealing with an object reference (this code should be identical to GCRT_NoopClass if ShouldProcessNoOpTokens())
-						UObject**	ObjectPtr = (UObject**)(StackEntryData + ReferenceInfo.Offset);
-						UObject*&	Object = *ObjectPtr;
-						TokenReturnCount = ReferenceInfo.ReturnCount;
-						ConditionalHandleTokenStreamObjectReference(NewObjectsToSerializeStruct, CurrentObject, Object, ReferenceTokenStreamIndex, TokenType, true);
-					}
-					break;
-					case GCRT_ArrayObject:
-					{
-						// We're dealing with an array of object references.
-						TArray<UObject*>& ObjectArray = *((TArray<UObject*>*)(StackEntryData + ReferenceInfo.Offset));
-						TokenReturnCount = ReferenceInfo.ReturnCount;
-						for (int32 ObjectIndex = 0, ObjectNum = ObjectArray.Num(); ObjectIndex < ObjectNum; ++ObjectIndex)
-						{
-							ConditionalHandleTokenStreamObjectReference(NewObjectsToSerializeStruct, CurrentObject, ObjectArray[ObjectIndex], ReferenceTokenStreamIndex, TokenType, true);
+						case GCRT_Object:
+						{								
+							UObject** ObjectPtr = (UObject**)(StackEntryData + ReferenceInfo.Offset);
+							TokenReturnCount = ReferenceInfo.ReturnCount;
+							Dispatcher.HandleKillableReference(*ObjectPtr, MakeTokenId(ReferenceTokenStreamIndex), TokenType);
 						}
-					}
-					break;
-					case GCRT_ArrayObjectFreezable:
-					{
-						// We're dealing with an array of object references.
-						TArray<UObject*, FMemoryImageAllocator>& ObjectArray = *((TArray<UObject*, FMemoryImageAllocator>*)(StackEntryData + ReferenceInfo.Offset));
-						TokenReturnCount = ReferenceInfo.ReturnCount;
-						for (int32 ObjectIndex = 0, ObjectNum = ObjectArray.Num(); ObjectIndex < ObjectNum; ++ObjectIndex)
+						break;
+						case GCRT_ArrayObject:
 						{
-							ConditionalHandleTokenStreamObjectReference(NewObjectsToSerializeStruct, CurrentObject, ObjectArray[ObjectIndex], ReferenceTokenStreamIndex, TokenType, true);
+							TArray<UObject*>& ObjectArray = *((TArray<UObject*>*)(StackEntryData + ReferenceInfo.Offset));
+							TokenReturnCount = ReferenceInfo.ReturnCount;
+							Dispatcher.HandleKillableArray(ObjectArray, MakeTokenId(ReferenceTokenStreamIndex), TokenType);
 						}
-					}
-					break;
-					case GCRT_ArrayStruct:
-					{
-						// We're dealing with a dynamic array of structs.
-						const FScriptArray& Array = *((FScriptArray*)(StackEntryData + ReferenceInfo.Offset));
-						StackEntry = Stack.Push();
-						StackEntryData = (uint8*)Array.GetData();
-						StackEntry->Data = StackEntryData;
-						StackEntry->Stride = TokenStream->ReadStride(TokenStreamIndex);
-						StackEntry->Count = Array.Num();
-						StackEntry->ContainerType = GCRT_None;
-
-						const FGCSkipInfo SkipInfo = TokenStream->ReadSkipInfo(TokenStreamIndex);
-						StackEntry->LoopStartIndex = TokenStreamIndex;
-
-						if (StackEntry->Count == 0)
+						break;
+						case GCRT_ArrayObjectFreezable:
 						{
-							// Skip empty array by jumping to skip index and set return count to the one about to be read in.
-							TokenStreamIndex = SkipInfo.SkipIndex;
-							TokenReturnCount = TokenStream->GetSkipReturnCount(SkipInfo);
+							TArray<UObject*, FMemoryImageAllocator>& ObjectArray = *((TArray<UObject*, FMemoryImageAllocator>*)(StackEntryData + ReferenceInfo.Offset));
+							TokenReturnCount = ReferenceInfo.ReturnCount;
+							Dispatcher.HandleKillableReferences(ObjectArray, MakeTokenId(ReferenceTokenStreamIndex), TokenType);
 						}
-						else
+						break;
+						case GCRT_ArrayStruct:
 						{
-							// Loop again.
-							check(StackEntry->Data);
-							TokenReturnCount = 0;
+							const FScriptArray& Array = *((FScriptArray*)(StackEntryData + ReferenceInfo.Offset));
+							StackEntry = Stack.Push();
+							StackEntryData = (uint8*)Array.GetData();
+							StackEntry->Data = StackEntryData;
+							StackEntry->Stride = TokenStream.ReadStride(TokenStreamIndex);
+							StackEntry->Count = Array.Num();
+							StackEntry->ContainerType = GCRT_None;
+
+							const FGCSkipInfo SkipInfo = TokenStream.ReadSkipInfo(TokenStreamIndex);
+							StackEntry->LoopStartIndex = TokenStreamIndex;
+
+							if (StackEntry->Count == 0)
+							{
+								// Skip empty array by jumping to skip index and set return count to the one about to be read in.
+								TokenStreamIndex = SkipInfo.SkipIndex;
+								TokenReturnCount = TokenStream.GetSkipReturnCount(SkipInfo);
+							}
+							else
+							{
+								// Loop again.
+								check(StackEntry->Data);
+								TokenReturnCount = 0;
+							}
 						}
-					}
-					break;
-					case GCRT_ArrayStructFreezable:
-					{
-						// We're dealing with a dynamic array of structs.
-						const FFreezableScriptArray& Array = *((FFreezableScriptArray*)(StackEntryData + ReferenceInfo.Offset));
-						StackEntry = Stack.Push();
-						StackEntryData = (uint8*)Array.GetData();
-						StackEntry->Data = StackEntryData;
-						StackEntry->Stride = TokenStream->ReadStride(TokenStreamIndex);
-						StackEntry->Count = Array.Num();
-						StackEntry->ContainerType = GCRT_None;
-
-						const FGCSkipInfo SkipInfo = TokenStream->ReadSkipInfo(TokenStreamIndex);
-						StackEntry->LoopStartIndex = TokenStreamIndex;
-
-						if (StackEntry->Count == 0)
+						break;
+						case GCRT_ArrayStructFreezable:
 						{
-							// Skip empty array by jumping to skip index and set return count to the one about to be read in.
-							TokenStreamIndex = SkipInfo.SkipIndex;
-							TokenReturnCount = TokenStream->GetSkipReturnCount(SkipInfo);
+							const FFreezableScriptArray& Array = *((FFreezableScriptArray*)(StackEntryData + ReferenceInfo.Offset));
+							StackEntry = Stack.Push();
+							StackEntryData = (uint8*)Array.GetData();
+							StackEntry->Data = StackEntryData;
+							StackEntry->Stride = TokenStream.ReadStride(TokenStreamIndex);
+							StackEntry->Count = Array.Num();
+							StackEntry->ContainerType = GCRT_None;
+
+							const FGCSkipInfo SkipInfo = TokenStream.ReadSkipInfo(TokenStreamIndex);
+							StackEntry->LoopStartIndex = TokenStreamIndex;
+
+							if (StackEntry->Count == 0)
+							{
+								// Skip empty array by jumping to skip index and set return count to the one about to be read in.
+								TokenStreamIndex = SkipInfo.SkipIndex;
+								TokenReturnCount = TokenStream.GetSkipReturnCount(SkipInfo);
+							}
+							else
+							{
+								// Loop again.
+								check(StackEntry->Data);
+								TokenReturnCount = 0;
+							}
 						}
-						else
+						break;
+						case GCRT_AddReferencedObjects:
 						{
-							// Loop again.
-							check(StackEntry->Data);
-							TokenReturnCount = 0;
+							void(*AddReferencedObjects)(UObject*, FReferenceCollector&) = (void(*)(UObject*, FReferenceCollector&))TokenStream.ReadPointer(TokenStreamIndex);
+							TokenReturnCount = ReferenceInfo.ReturnCount;
+							AddReferencedObjects(CurrentObject, Collector);
+
+							// ARO is always last and an implicit terminator
+							goto TokensDone;
 						}
-					}
-					break;
-					case GCRT_PersistentObject:
-					{
-						// We're dealing with an object reference (this code should be identical to GCRT_NoopPersistentObject if ShouldProcessNoOpTokens())
-						UObject**	ObjectPtr = (UObject**)(StackEntryData + ReferenceInfo.Offset);
-						UObject*&	Object = *ObjectPtr;
-						TokenReturnCount = ReferenceInfo.ReturnCount;
-						ConditionalHandleTokenStreamObjectReference(NewObjectsToSerializeStruct, CurrentObject, Object, ReferenceTokenStreamIndex, TokenType, false);
-					}
-					break;
-					case GCRT_ExternalPackage:
-					{
-						// We're dealing with the external package reference.
-						TokenReturnCount = ReferenceInfo.ReturnCount;
-						// Test if the object isn't itself, since currently package are their own external and tracking that reference is pointless
-						UObject* Object = CurrentObject->GetExternalPackageInternal();
-						Object = Object != CurrentObject ? Object : nullptr;
-						ReferenceProcessor.HandleTokenStreamObjectReference(NewObjectsToSerializeStruct, CurrentObject, Object, ReferenceTokenStreamIndex, TokenType, false);
-					}
-					break;
-					case GCRT_FixedArray:
-					{
-						// We're dealing with a fixed size array
-						uint8* PreviousData = StackEntryData;
-						StackEntry = Stack.Push();
-						StackEntryData = PreviousData;
-						StackEntry->Data = PreviousData;
-						StackEntry->Stride = TokenStream->ReadStride(TokenStreamIndex);
-						StackEntry->Count = TokenStream->ReadCount(TokenStreamIndex);
-						StackEntry->LoopStartIndex = TokenStreamIndex;
-						StackEntry->ContainerType = GCRT_None;
-						TokenReturnCount = 0;
-					}
-					break;
-					case GCRT_AddStructReferencedObjects:
-					{
-						// We're dealing with a function call
-						void* StructPtr = (void*)(StackEntryData + ReferenceInfo.Offset);
-						TokenReturnCount = ReferenceInfo.ReturnCount;
-						UScriptStruct::ICppStructOps::TPointerToAddStructReferencedObjects Func = (UScriptStruct::ICppStructOps::TPointerToAddStructReferencedObjects) TokenStream->ReadPointer(TokenStreamIndex);
-						Func(StructPtr, ReferenceCollector);
-					}
-					break;
-					case GCRT_AddReferencedObjects:
-					{
-						// Static AddReferencedObjects function call.
-						void(*AddReferencedObjects)(UObject*, FReferenceCollector&) = (void(*)(UObject*, FReferenceCollector&))TokenStream->ReadPointer(TokenStreamIndex);
-						TokenReturnCount = ReferenceInfo.ReturnCount;
-						AddReferencedObjects(CurrentObject, ReferenceCollector);
-					}
-					break;
-					case GCRT_AddTMapReferencedObjects:
-					{
-						void* MapPtr = StackEntryData + ReferenceInfo.Offset;
-						FMapProperty* MapProperty = (FMapProperty*)TokenStream->ReadPointer(TokenStreamIndex);
-						TokenStreamIndex++; // GCRT_EndOfPointer
-
-						StackEntry = Stack.Push();
-						StackEntry->ContainerType = GCRT_AddTMapReferencedObjects;
-						StackEntry->ContainerIndex = 0;
-						StackEntry->ContainerProperty = MapProperty;
-						StackEntry->ContainerPtr = MapPtr;
-						StackEntry->Stride = MapProperty->GetPairStride();
-						StackEntry->Count = MapProperty->GetNum(MapPtr);
-
-						const FGCSkipInfo SkipInfo = TokenStream->ReadSkipInfo(TokenStreamIndex);
-						StackEntry->LoopStartIndex = TokenStreamIndex;
-
-						if (StackEntry->Count == 0)
+						break;
+						case GCRT_EndOfStream:
 						{
-							// The map is empty
-							StackEntryData = nullptr;
-							StackEntry->Data = StackEntryData;
-
-							// Skip empty map by jumping to skip index and set return count to the one about to be read in.
-							TokenStreamIndex = SkipInfo.SkipIndex;
-							TokenReturnCount = TokenStream->GetSkipReturnCount(SkipInfo);
+							// Break out of loop
+							goto TokensDone;
 						}
-						else
+						break;
+						case GCRT_SlowAddReferencedObjects:
 						{
-							// Skip any initial invalid entries in the map. We need a valid index for MapProperty->GetPairPtr()
-							int32 FirstValidIndex = 0;
-							while (!MapProperty->IsValidIndex(MapPtr, FirstValidIndex))
+							if (!IsParallel() ||
+								!FSlowARO::TryQueueCall(/* ARO index */ ReferenceInfo.Offset, CurrentObject, Context))
 							{
-								FirstValidIndex++;
+								FSlowARO::CallSync(/* ARO index */ ReferenceInfo.Offset, CurrentObject, Collector);
 							}
-
-							StackEntry->ContainerIndex = FirstValidIndex;
-							StackEntryData = MapProperty->GetPairPtr(MapPtr, FirstValidIndex);
-							StackEntry->Data = StackEntryData;
-
-							// Loop again.
+								
+							// ARO is always last and an implicit terminator
+							goto TokensDone;
+						}
+						case GCRT_ExternalPackage:
+						{
+							TokenReturnCount = ReferenceInfo.ReturnCount;
+							// Test if the object isn't itself, since currently package are their own external and tracking that reference is pointless
+							UObject* Object = CurrentObject->GetExternalPackageInternal();
+							Object = Object != CurrentObject ? Object : nullptr;
+							Dispatcher.HandleImmutableReference( Object, MakeTokenId(ReferenceTokenStreamIndex), TokenType);
+						}
+						break;
+						case GCRT_FixedArray:
+						{
+							uint8* PreviousData = StackEntryData;
+							StackEntry = Stack.Push();
+							StackEntryData = PreviousData;
+							StackEntry->Data = PreviousData;
+							StackEntry->Stride = TokenStream.ReadStride(TokenStreamIndex);
+							StackEntry->Count = TokenStream.ReadCount(TokenStreamIndex);
+							StackEntry->LoopStartIndex = TokenStreamIndex;
+							StackEntry->ContainerType = GCRT_None;
 							TokenReturnCount = 0;
 						}
-					}
-					break;
-					case GCRT_AddTSetReferencedObjects:
-					{
-						void* SetPtr = StackEntryData + ReferenceInfo.Offset;
-						FSetProperty* SetProperty = (FSetProperty*)TokenStream->ReadPointer(TokenStreamIndex);
-						TokenStreamIndex++; // GCRT_EndOfPointer
-
-						StackEntry = Stack.Push();
-						StackEntry->ContainerProperty = SetProperty;
-						StackEntry->ContainerPtr = SetPtr;
-						StackEntry->ContainerType = GCRT_AddTSetReferencedObjects;
-						StackEntry->ContainerIndex = 0;
-						StackEntry->Stride = SetProperty->GetStride();
-						StackEntry->Count = SetProperty->GetNum(SetPtr);
-
-						const FGCSkipInfo SkipInfo = TokenStream->ReadSkipInfo(TokenStreamIndex);
-						StackEntry->LoopStartIndex = TokenStreamIndex;
-
-						if (StackEntry->Count == 0)
+						break;
+						case GCRT_AddStructReferencedObjects:
 						{
-							// The set is empty or it doesn't contain any valid elements
-							StackEntryData = nullptr;
-							StackEntry->Data = StackEntryData;
-
-							// Skip empty set by jumping to skip index and set return count to the one about to be read in.
-							TokenStreamIndex = SkipInfo.SkipIndex;
-							TokenReturnCount = TokenStream->GetSkipReturnCount(SkipInfo);
+							void* StructPtr = (void*)(StackEntryData + ReferenceInfo.Offset);
+							TokenReturnCount = ReferenceInfo.ReturnCount;
+							UScriptStruct::ICppStructOps::TPointerToAddStructReferencedObjects Func = (UScriptStruct::ICppStructOps::TPointerToAddStructReferencedObjects)TokenStream.ReadPointer(TokenStreamIndex);
+							Func(StructPtr, Collector);
 						}
-						else
+						break;
+						case GCRT_AddTMapReferencedObjects:
 						{
-							// Skip any initial invalid entries in the set. We need a valid index for SetProperty->GetElementPtr()
-							int32 FirstValidIndex = 0;
-							while (!SetProperty->IsValidIndex(SetPtr, FirstValidIndex))
+							void* MapPtr = StackEntryData + ReferenceInfo.Offset;
+							FMapProperty* MapProperty = (FMapProperty*)TokenStream.ReadPointer(TokenStreamIndex);
+							TokenStreamIndex++; // GCRT_EndOfPointer
+
+							StackEntry = Stack.Push();
+							StackEntry->ContainerType = GCRT_AddTMapReferencedObjects;
+							StackEntry->ContainerIndex = 0;
+							StackEntry->ContainerProperty = MapProperty;
+							StackEntry->ContainerPtr = MapPtr;
+							StackEntry->Stride = MapProperty->GetPairStride();
+							StackEntry->Count = MapProperty->GetNum(MapPtr);
+
+							const FGCSkipInfo SkipInfo = TokenStream.ReadSkipInfo(TokenStreamIndex);
+							StackEntry->LoopStartIndex = TokenStreamIndex;
+
+							if (StackEntry->Count == 0)
 							{
-								FirstValidIndex++;
+								// The map is empty
+								StackEntryData = nullptr;
+								StackEntry->Data = StackEntryData;
+
+								// Skip empty map by jumping to skip index and set return count to the one about to be read in.
+								TokenStreamIndex = SkipInfo.SkipIndex;
+								TokenReturnCount = TokenStream.GetSkipReturnCount(SkipInfo);
 							}
+							else
+							{
+								// Skip any initial invalid entries in the map. We need a valid index for MapProperty->GetPairPtr()
+								int32 FirstValidIndex = 0;
+								while (!MapProperty->IsValidIndex(MapPtr, FirstValidIndex))
+								{
+									FirstValidIndex++;
+								}
 
-							StackEntry->ContainerIndex = FirstValidIndex;
-							StackEntryData = SetProperty->GetElementPtr(SetPtr, FirstValidIndex);
-							StackEntry->Data = StackEntryData;
+								StackEntry->ContainerIndex = FirstValidIndex;
+								StackEntryData = MapProperty->GetPairPtr(MapPtr, FirstValidIndex);
+								StackEntry->Data = StackEntryData;
 
-							// Loop again.
-							TokenReturnCount = 0;
+								// Loop again.
+								TokenReturnCount = 0;
+							}
 						}
-					}
-					break;
-					case GCRT_AddFieldPathReferencedObject:
-					{
-						FFieldPath*	FieldPathPtr = (FFieldPath*)(StackEntryData + ReferenceInfo.Offset);
-						FUObjectItem* FieldOwnerItem = FieldPathPtr->GetResolvedOwnerItemInternal();
-						TokenReturnCount = ReferenceInfo.ReturnCount;
-						if (FieldOwnerItem)
+						break;
+						case GCRT_AddTSetReferencedObjects:
 						{
-							UObject* OwnerObject = static_cast<UObject*>(FieldOwnerItem->Object);
-							UObject* PreviousOwner = OwnerObject;
-							ConditionalHandleTokenStreamObjectReference(NewObjectsToSerializeStruct, CurrentObject, OwnerObject, ReferenceTokenStreamIndex, TokenType, true);
-							// Handle reference elimination (PendingKill owner)
-							if (PreviousOwner && !OwnerObject)
+							void* SetPtr = StackEntryData + ReferenceInfo.Offset;
+							FSetProperty* SetProperty = (FSetProperty*)TokenStream.ReadPointer(TokenStreamIndex);
+							TokenStreamIndex++; // GCRT_EndOfPointer
+
+							StackEntry = Stack.Push();
+							StackEntry->ContainerProperty = SetProperty;
+							StackEntry->ContainerPtr = SetPtr;
+							StackEntry->ContainerType = GCRT_AddTSetReferencedObjects;
+							StackEntry->ContainerIndex = 0;
+							StackEntry->Stride = SetProperty->GetStride();
+							StackEntry->Count = SetProperty->GetNum(SetPtr);
+
+							const FGCSkipInfo SkipInfo = TokenStream.ReadSkipInfo(TokenStreamIndex);
+							StackEntry->LoopStartIndex = TokenStreamIndex;
+
+							if (StackEntry->Count == 0)
+							{
+								// The set is empty or it doesn't contain any valid elements
+								StackEntryData = nullptr;
+								StackEntry->Data = StackEntryData;
+
+								// Skip empty set by jumping to skip index and set return count to the one about to be read in.
+								TokenStreamIndex = SkipInfo.SkipIndex;
+								TokenReturnCount = TokenStream.GetSkipReturnCount(SkipInfo);
+							}
+							else
 							{
-								FieldPathPtr->ClearCachedFieldInternal();
+								// Skip any initial invalid entries in the set. We need a valid index for SetProperty->GetElementPtr()
+								int32 FirstValidIndex = 0;
+								while (!SetProperty->IsValidIndex(SetPtr, FirstValidIndex))
+								{
+									FirstValidIndex++;
+								}
+
+								StackEntry->ContainerIndex = FirstValidIndex;
+								StackEntryData = SetProperty->GetElementPtr(SetPtr, FirstValidIndex);
+								StackEntry->Data = StackEntryData;
+
+								// Loop again.
+								TokenReturnCount = 0;
 							}
 						}
-					}
-					break;
-					case GCRT_ArrayAddFieldPathReferencedObject:
-					{
-						// We're dealing with an array of object references.
-						TArray<FFieldPath>& FieldArray = *((TArray<FFieldPath>*)(StackEntryData + ReferenceInfo.Offset));
-						TokenReturnCount = ReferenceInfo.ReturnCount;
-						for (int32 FieldIndex = 0, FieldNum = FieldArray.Num(); FieldIndex < FieldNum; ++FieldIndex)
+						break;
+						case GCRT_AddFieldPathReferencedObject:
 						{
-							FUObjectItem* FieldOwnerItem = FieldArray[FieldIndex].GetResolvedOwnerItemInternal();
+							FFieldPath* FieldPathPtr = (FFieldPath*)(StackEntryData + ReferenceInfo.Offset);
+							FUObjectItem* FieldOwnerItem = GetResolvedOwner(*FieldPathPtr);
+							TokenReturnCount = ReferenceInfo.ReturnCount;
 							if (FieldOwnerItem)
 							{
 								UObject* OwnerObject = static_cast<UObject*>(FieldOwnerItem->Object);
 								UObject* PreviousOwner = OwnerObject;
-								ConditionalHandleTokenStreamObjectReference(NewObjectsToSerializeStruct, CurrentObject, OwnerObject, ReferenceTokenStreamIndex, TokenType, true);
+								Dispatcher.HandleReferenceDirectly(CurrentObject, OwnerObject, MakeTokenId(ReferenceTokenStreamIndex), TokenType, true);
+								
 								// Handle reference elimination (PendingKill owner)
 								if (PreviousOwner && !OwnerObject)
 								{
-									FieldArray[FieldIndex].ClearCachedFieldInternal();
+									ClearCachedField(*FieldPathPtr);
 								}
 							}
 						}
-					}
-					break;
-					case GCRT_Optional:
-					{
-						const FGCSkipInfo SkipInfo = TokenStream->ReadSkipInfo(TokenStreamIndex);
-						uint32 ValueSize = TokenStream->ReadStride(TokenStreamIndex); // Size of value in bytes. This is also the offset to the bIsSet variable stored thereafter.
-						const bool& bIsSet = *((bool*)(StackEntryData + ReferenceInfo.Offset + ValueSize));
-						if (bIsSet)
+						break;
+						case GCRT_ArrayAddFieldPathReferencedObject:
 						{
-							// It's set - push a stack entry for processing the value
-							// This is somewhat suboptimal since there is only ever just one value, but this approach avoids any changes to the surrounding code
-							StackEntry = Stack.Push();
-							StackEntryData += ReferenceInfo.Offset;
-							StackEntry->Data = StackEntryData;
-							StackEntry->Stride = ValueSize;
-							StackEntry->Count = 1;
-							StackEntry->LoopStartIndex = TokenStreamIndex;
+							TArray<FFieldPath>& FieldArray = *((TArray<FFieldPath>*)(StackEntryData + ReferenceInfo.Offset));
+							TokenReturnCount = ReferenceInfo.ReturnCount;
+							for (int32 FieldIndex = 0, FieldNum = FieldArray.Num(); FieldIndex < FieldNum; ++FieldIndex)
+							{
+								if (FUObjectItem* FieldOwnerItem = GetResolvedOwner(FieldArray[FieldIndex]))
+								{
+									UObject* OwnerObject = static_cast<UObject*>(FieldOwnerItem->Object);
+									UObject* PreviousOwner = OwnerObject;
+									Dispatcher.HandleReferenceDirectly(CurrentObject, OwnerObject, MakeTokenId(ReferenceTokenStreamIndex), TokenType, true);
+
+									// Handle reference elimination (PendingKill owner)
+									if (PreviousOwner && !OwnerObject)
+									{
+										ClearCachedField(FieldArray[FieldIndex]);
+									}
+								}
+							}
 						}
-						else
+						break;
+						case GCRT_Optional:
 						{
-							// It's unset - keep going by jumping to skip index
-							TokenStreamIndex = SkipInfo.SkipIndex;
+							uint32 ValueSize = TokenStream.ReadStride(TokenStreamIndex); // Size of value in bytes. This is also the offset to the bIsSet variable stored thereafter.
+							const FGCSkipInfo SkipInfo = TokenStream.ReadSkipInfo(TokenStreamIndex);
+							const bool& bIsSet = *((bool*)(StackEntryData + ReferenceInfo.Offset + ValueSize));
+							if (bIsSet)
+							{
+								// It's set - push a stack entry for processing the value
+								// This is somewhat suboptimal since there is only ever just one value, but this approach avoids any changes to the surrounding code
+								StackEntry = Stack.Push();
+								StackEntryData += ReferenceInfo.Offset;
+								StackEntry->Data = StackEntryData;
+								StackEntry->Stride = ValueSize;
+								StackEntry->Count = 1;
+								StackEntry->LoopStartIndex = TokenStreamIndex;
+							}
+							else
+							{
+								// It's unset - keep going by jumping to skip index
+								TokenStreamIndex = SkipInfo.SkipIndex;
+							}
+							TokenReturnCount = 0;
 						}
-						TokenReturnCount = 0;
-					}
-					break;
-					case GCRT_EndOfPointer:
-					{
-						TokenReturnCount = ReferenceInfo.ReturnCount;
-					}
-					break;
-					case GCRT_NoopPersistentObject:
-					{
-						TokenReturnCount = ReferenceInfo.ReturnCount;
-						if (ShouldProcessNoOpTokens())
+						break;
+						case GCRT_EndOfPointer:
 						{
-							// We're dealing with an object reference (this code should be identical to GCRT_PersistentObject)
-							UObject**	ObjectPtr = (UObject**)(StackEntryData + ReferenceInfo.Offset);
-							UObject*&	Object = *ObjectPtr;
-							ConditionalHandleTokenStreamObjectReference(NewObjectsToSerializeStruct, CurrentObject, Object, ReferenceTokenStreamIndex, TokenType, false);
+							TokenReturnCount = ReferenceInfo.ReturnCount;
 						}
-					}
-					break;
-					case GCRT_NoopClass:
-					{
-						TokenReturnCount = ReferenceInfo.ReturnCount;
-						if (ShouldProcessNoOpTokens())
+						break;
+						case GCRT_WeakObject:
 						{
-							// We're dealing with an object reference (this code should be identical to GCRT_Object and GCRT_Class)
-							UObject**	ObjectPtr = (UObject**)(StackEntryData + ReferenceInfo.Offset);
-							UObject*&	Object = *ObjectPtr;
-							ConditionalHandleTokenStreamObjectReference(NewObjectsToSerializeStruct, CurrentObject, Object, ReferenceTokenStreamIndex, TokenType, true);
-						}
-					}
-					break;
-					case GCRT_WeakObject:
-					{
-						TokenReturnCount = ReferenceInfo.ReturnCount;
-						if (ShouldProcessWeakReferences())
-						{							
-							FWeakObjectPtr& WeakPtr = *(FWeakObjectPtr*)(StackEntryData + ReferenceInfo.Offset);
-							HandleWeakObjectPtr(WeakPtr, NewObjectsToSerializeStruct, CurrentObject, ReferenceTokenStreamIndex, TokenType);
-						}
-					}
-					break;
-					case GCRT_ArrayWeakObject:
-					{
-						TokenReturnCount = ReferenceInfo.ReturnCount;
-						if (ShouldProcessWeakReferences())
-						{							
-							TArray<FWeakObjectPtr>& WeakPtrArray = *((TArray<FWeakObjectPtr>*)(StackEntryData + ReferenceInfo.Offset));
-							for (FWeakObjectPtr& WeakPtr : WeakPtrArray)
+							TokenReturnCount = ReferenceInfo.ReturnCount;
+							if (ShouldProcessWeakReferences())
 							{
-								HandleWeakObjectPtr(WeakPtr, NewObjectsToSerializeStruct, CurrentObject, ReferenceTokenStreamIndex, TokenType);
+								FWeakObjectPtr& WeakPtr = *(FWeakObjectPtr*)(StackEntryData + ReferenceInfo.Offset);
+								Dispatcher.HandleWeakReference(WeakPtr, CurrentObject, MakeTokenId(ReferenceTokenStreamIndex), TokenType);
 							}
 						}
-					}
-					break;
-					case GCRT_LazyObject:
-					{
-						TokenReturnCount = ReferenceInfo.ReturnCount;
-						if (ShouldProcessWeakReferences())
-						{							
-							FLazyObjectPtr& LazyPtr = *(FLazyObjectPtr*)(StackEntryData + ReferenceInfo.Offset);
-							FWeakObjectPtr& WeakPtr = LazyPtr.WeakPtr;
-							HandleWeakObjectPtr(WeakPtr, NewObjectsToSerializeStruct, CurrentObject, ReferenceTokenStreamIndex, TokenType);
+						break;
+						case GCRT_ArrayWeakObject:
+						{
+							TokenReturnCount = ReferenceInfo.ReturnCount;
+							if (ShouldProcessWeakReferences())
+							{
+								TArray<FWeakObjectPtr>& WeakPtrArray = *((TArray<FWeakObjectPtr>*)(StackEntryData + ReferenceInfo.Offset));
+								for (FWeakObjectPtr& WeakPtr : WeakPtrArray)
+								{
+									Dispatcher.HandleWeakReference(WeakPtr, CurrentObject, MakeTokenId(ReferenceTokenStreamIndex), TokenType);
+								}
+							}
 						}
-					}
-					break;
-					case GCRT_ArrayLazyObject:
-					{
-						TokenReturnCount = ReferenceInfo.ReturnCount;
-						if (ShouldProcessWeakReferences())
-						{							
-							TArray<FLazyObjectPtr>& LazyPtrArray = *((TArray<FLazyObjectPtr>*)(StackEntryData + ReferenceInfo.Offset));
+						break;
+						case GCRT_LazyObject:
+						{
 							TokenReturnCount = ReferenceInfo.ReturnCount;
-							for (FLazyObjectPtr& LazyPtr : LazyPtrArray)
+							if (ShouldProcessWeakReferences())
 							{
-								FWeakObjectPtr& WeakPtr = LazyPtr.WeakPtr;
-								HandleWeakObjectPtr(WeakPtr, NewObjectsToSerializeStruct, CurrentObject, ReferenceTokenStreamIndex, TokenType);
+								FLazyObjectPtr& LazyPtr = *(FLazyObjectPtr*)(StackEntryData + ReferenceInfo.Offset);
+								FWeakObjectPtr& WeakPtr = AccessWeakPtr(LazyPtr);
+								Dispatcher.HandleWeakReference(WeakPtr, CurrentObject, MakeTokenId(ReferenceTokenStreamIndex), TokenType);
 							}
 						}
-					}
-					break;
-					case GCRT_SoftObject:
-					{
-						TokenReturnCount = ReferenceInfo.ReturnCount;
-						if (ShouldProcessWeakReferences())
-						{							
-							FSoftObjectPtr& SoftPtr = *(FSoftObjectPtr*)(StackEntryData + ReferenceInfo.Offset);
-							FWeakObjectPtr& WeakPtr = SoftPtr.WeakPtr;
-							HandleWeakObjectPtr(WeakPtr, NewObjectsToSerializeStruct, CurrentObject, ReferenceTokenStreamIndex, TokenType);
+						break;
+						case GCRT_ArrayLazyObject:
+						{
+							TokenReturnCount = ReferenceInfo.ReturnCount;
+							if (ShouldProcessWeakReferences())
+							{
+								TArray<FLazyObjectPtr>& LazyPtrArray = *((TArray<FLazyObjectPtr>*)(StackEntryData + ReferenceInfo.Offset));
+								TokenReturnCount = ReferenceInfo.ReturnCount;
+								for (FLazyObjectPtr& LazyPtr : LazyPtrArray)
+								{
+									FWeakObjectPtr& WeakPtr = AccessWeakPtr(LazyPtr);
+									Dispatcher.HandleWeakReference(WeakPtr, CurrentObject, MakeTokenId(ReferenceTokenStreamIndex), TokenType);
+								}
+							}
 						}
-					}
-					break;
-					case GCRT_ArraySoftObject:
-					{
-						TokenReturnCount = ReferenceInfo.ReturnCount;
-						if (ShouldProcessWeakReferences())
-						{							
-							TArray<FSoftObjectPtr>& SoftPtrArray = *((TArray<FSoftObjectPtr>*)(StackEntryData + ReferenceInfo.Offset));
-							for (FSoftObjectPtr& SoftPtr : SoftPtrArray)
+						break;
+						case GCRT_SoftObject:
+						{
+							TokenReturnCount = ReferenceInfo.ReturnCount;
+							if (ShouldProcessWeakReferences())
 							{
-								FWeakObjectPtr& WeakPtr = SoftPtr.WeakPtr;
-								HandleWeakObjectPtr(WeakPtr, NewObjectsToSerializeStruct, CurrentObject, ReferenceTokenStreamIndex, TokenType);
+								FSoftObjectPtr& SoftPtr = *(FSoftObjectPtr*)(StackEntryData + ReferenceInfo.Offset);
+								FWeakObjectPtr& WeakPtr = AccessWeakPtr(SoftPtr);
+								Dispatcher.HandleWeakReference(WeakPtr, CurrentObject, MakeTokenId(ReferenceTokenStreamIndex), TokenType);
 							}
 						}
-					}
-					break;
-					case GCRT_Delegate:
-					{
-						TokenReturnCount = ReferenceInfo.ReturnCount;
-						if (ShouldProcessWeakReferences())
-						{							
-							FScriptDelegate& Delegate = *(FScriptDelegate*)(StackEntryData + ReferenceInfo.Offset);
-							UObject* DelegateObject = Delegate.GetUObject();
-							ReferenceProcessor.HandleTokenStreamObjectReference(NewObjectsToSerializeStruct, CurrentObject, DelegateObject, ReferenceTokenStreamIndex, TokenType, false);
+						break;
+						case GCRT_ArraySoftObject:
+						{
+							TokenReturnCount = ReferenceInfo.ReturnCount;
+							if (ShouldProcessWeakReferences())
+							{
+								TArray<FSoftObjectPtr>& SoftPtrArray = *((TArray<FSoftObjectPtr>*)(StackEntryData + ReferenceInfo.Offset));
+								for (FSoftObjectPtr& SoftPtr : SoftPtrArray)
+								{
+									FWeakObjectPtr& WeakPtr = AccessWeakPtr(SoftPtr);
+									Dispatcher.HandleWeakReference(WeakPtr, CurrentObject, MakeTokenId(ReferenceTokenStreamIndex), TokenType);
+								}
+							}
 						}
-					}
-					break;
-					case GCRT_ArrayDelegate:
-					{
-						TokenReturnCount = ReferenceInfo.ReturnCount;
-						if (ShouldProcessWeakReferences())
-						{							
-							TArray<FScriptDelegate>& DelegateArray = *((TArray<FScriptDelegate>*)(StackEntryData + ReferenceInfo.Offset));
-							for (FScriptDelegate& Delegate : DelegateArray)
+						break;
+						case GCRT_Delegate:
+						{
+							TokenReturnCount = ReferenceInfo.ReturnCount;
+							if (ShouldProcessWeakReferences())
 							{
+								FScriptDelegate& Delegate = *(FScriptDelegate*)(StackEntryData + ReferenceInfo.Offset);
 								UObject* DelegateObject = Delegate.GetUObject();
-								ReferenceProcessor.HandleTokenStreamObjectReference(NewObjectsToSerializeStruct, CurrentObject, DelegateObject, ReferenceTokenStreamIndex, TokenType, false);
+								Dispatcher.HandleReferenceDirectly(CurrentObject, DelegateObject, MakeTokenId(ReferenceTokenStreamIndex), TokenType, false);
 							}
 						}
-					}
-					break;
-					case GCRT_MulticastDelegate:
-					{
-						TokenReturnCount = ReferenceInfo.ReturnCount;
-						if (ShouldProcessWeakReferences())
-						{							
-							FMulticastScriptDelegate& Delegate = *(FMulticastScriptDelegate*)(StackEntryData + ReferenceInfo.Offset);
-							TArray<UObject*> DelegateObjects(Delegate.GetAllObjects());
-							for (UObject* DelegateObject : DelegateObjects)
+						break;
+						case GCRT_ArrayDelegate:
+						{
+							TokenReturnCount = ReferenceInfo.ReturnCount;
+							if (ShouldProcessWeakReferences())
 							{
-								ReferenceProcessor.HandleTokenStreamObjectReference(NewObjectsToSerializeStruct, CurrentObject, DelegateObject, ReferenceTokenStreamIndex, TokenType, false);
+								TArray<FScriptDelegate>& DelegateArray = *((TArray<FScriptDelegate>*)(StackEntryData + ReferenceInfo.Offset));
+								for (FScriptDelegate& Delegate : DelegateArray)
+								{
+									UObject* DelegateObject = Delegate.GetUObject();
+									Dispatcher.HandleReferenceDirectly(CurrentObject, DelegateObject, MakeTokenId(ReferenceTokenStreamIndex), TokenType, false);
+								}
 							}
 						}
-					}
-					break;
-					case GCRT_ArrayMulticastDelegate:
-					{
-						TokenReturnCount = ReferenceInfo.ReturnCount;
-						if (ShouldProcessWeakReferences())
-						{							
-							TArray<FMulticastScriptDelegate>& DelegateArray = *((TArray<FMulticastScriptDelegate>*)(StackEntryData + ReferenceInfo.Offset));
-							for (FMulticastScriptDelegate& Delegate : DelegateArray)
+						break;
+						case GCRT_MulticastDelegate:
+						{
+							TokenReturnCount = ReferenceInfo.ReturnCount;
+							if (ShouldProcessWeakReferences())
 							{
+								FMulticastScriptDelegate& Delegate = *(FMulticastScriptDelegate*)(StackEntryData + ReferenceInfo.Offset);
 								TArray<UObject*> DelegateObjects(Delegate.GetAllObjects());
 								for (UObject* DelegateObject : DelegateObjects)
 								{
-									ReferenceProcessor.HandleTokenStreamObjectReference(NewObjectsToSerializeStruct, CurrentObject, DelegateObject, ReferenceTokenStreamIndex, TokenType, false);
+									Dispatcher.HandleReferenceDirectly(CurrentObject, DelegateObject, MakeTokenId(ReferenceTokenStreamIndex), TokenType, false);
 								}
 							}
 						}
-					}
-					break;
-					case GCRT_DynamicallyTypedValue:
-					{
-						TokenReturnCount = ReferenceInfo.ReturnCount;
-						UE::FDynamicallyTypedValue* DynamicallyTypedValue = (UE::FDynamicallyTypedValue*)(StackEntryData + ReferenceInfo.Offset);
-						DynamicallyTypedValue->GetType().MarkReachable(ReferenceCollector);
-						if (DynamicallyTypedValue->GetType().GetContainsReferences() != UE::FDynamicallyTypedValueType::EContainsReferences::DoesNot)
+						break;
+						case GCRT_ArrayMulticastDelegate:
 						{
-							DynamicallyTypedValue->GetType().MarkValueReachable(DynamicallyTypedValue->GetDataPointer(), ReferenceCollector);
+							TokenReturnCount = ReferenceInfo.ReturnCount;
+							if (ShouldProcessWeakReferences())
+							{
+								TArray<FMulticastScriptDelegate>& DelegateArray = *((TArray<FMulticastScriptDelegate>*)(StackEntryData + ReferenceInfo.Offset));
+								for (FMulticastScriptDelegate& Delegate : DelegateArray)
+								{
+									TArray<UObject*> DelegateObjects(Delegate.GetAllObjects());
+									for (UObject* DelegateObject : DelegateObjects)
+									{
+										Dispatcher.HandleReferenceDirectly(CurrentObject, DelegateObject, MakeTokenId(ReferenceTokenStreamIndex), TokenType, false);
+									}
+								}
+							}
 						}
-					}
-					break;
-					case GCRT_EndOfStream:
-					{
-						// Break out of loop.
-						goto EndLoop;
-					}
-					break;
-					default:
-					{
-						UE_LOG(LogGarbage, Fatal, TEXT("Unknown token. Type:%d ReferenceTokenStreamIndex:%d Class:%s Obj:%s"), ReferenceInfo.Type, ReferenceTokenStreamIndex, CurrentObject ? *GetNameSafe(CurrentObject->GetClass()) : TEXT("Unknown"), *GetPathNameSafe(CurrentObject));
 						break;
-					}
-				}
-				}
-EndLoop:
-				check(StackEntry == Stack.GetBottomFrame());
-				check(NewObjectsToSerializeStruct.ReferencingObject == CurrentObject);
-
-				if (IsParallel() && NewObjectsToSerialize.Num() >= MinDesiredObjectsPerSubTask)
-				{
-					// This will start queueing task with objects from the end of array until there's less objects than worth to queue
-					const int32 ObjectsPerSubTask = FMath::Max<int32>(MinDesiredObjectsPerSubTask, NewObjectsToSerialize.Num() / FTaskGraphInterface::Get().GetNumWorkerThreads());
-					while (NewObjectsToSerialize.Num() >= MinDesiredObjectsPerSubTask)
-					{
-						const int32 StartIndex = FMath::Max(0, NewObjectsToSerialize.Num() - ObjectsPerSubTask);
-						const int32 NumThisTask = NewObjectsToSerialize.Num() - StartIndex;
-						if (MyCompletionGraphEvent.GetReference())
+						case GCRT_DynamicallyTypedValue:
 						{
-							MyCompletionGraphEvent->DontCompleteUntil(TGraphTask< FCollectorTask >::CreateTask().ConstructAndDispatchWhenReady(this, &NewObjectsToSerialize, StartIndex, NumThisTask, ArrayPool));
+							TokenReturnCount = ReferenceInfo.ReturnCount;
+							UE::FDynamicallyTypedValue* DynamicallyTypedValue = (UE::FDynamicallyTypedValue*)(StackEntryData + ReferenceInfo.Offset);
+							DynamicallyTypedValue->GetType().MarkReachable(Collector);
+							if (DynamicallyTypedValue->GetType().GetContainsReferences() != UE::FDynamicallyTypedValueType::EContainsReferences::DoesNot)
+							{
+								DynamicallyTypedValue->GetType().MarkValueReachable(DynamicallyTypedValue->GetDataPointer(), Collector);
+							}
 						}
-						else
+						break;
+						default:
 						{
-							TaskQueue.AddTask(&NewObjectsToSerialize, StartIndex, NumThisTask);
+							UE_LOG(LogGarbage, Fatal, TEXT("Unknown token. Type:%d ReferenceTokenStreamIndex:%d Class:%s Obj:%s"), ReferenceInfo.Type, ReferenceTokenStreamIndex, CurrentObject ? *GetNameSafe(CurrentObject->GetClass()) : TEXT("Unknown"), *GetPathNameSafe(CurrentObject));
+							break;
 						}
-						NewObjectsToSerialize.SetNumUnsafeInternal(StartIndex);
-					}
-				}
+					} // switch (Type)
+				} // while (true)
+TokensDone:
+				check(StackEntry == Stack.GetBottomFrame());
+			} // for (FPrefetchingObjectIterator)
 
-#if PERF_DETAILED_PER_CLASS_GC_STATS
-				// Detailed per class stats should not be performed when parallel GC is running
-				check(!IsParallel());
-				ReferenceProcessor.UpdateDetailedStats(CurrentObject, FPlatformTime::Cycles() - StartCycles);
-#endif
+			// Free finished work block
+			if (CurrentObjects.GetData() != Context.InitialObjects.GetData())
+			{
+				Context.ObjectsToSerialize.FreeOwningBlock(CurrentObjects.GetData());
 			}
 
-			if (IsParallel() && NewObjectsToSerialize.Num() >= MinDesiredObjectsPerSubTask)
+			int32 BlockSize = FWorkBlock::ObjectCapacity;
+			FWorkBlockifier& RemainingObjects = Context.ObjectsToSerialize;
+			FWorkBlock* Block = RemainingObjects.PopFullBlock<Options>();
+			if (!Block)
 			{
-				const int32 ObjectsPerSubTask = FMath::Max<int32>(MinDesiredObjectsPerSubTask, NewObjectsToSerialize.Num() / FTaskGraphInterface::Get().GetNumWorkerThreads());
-				int32 StartIndex = 0;
-				while (StartIndex < NewObjectsToSerialize.Num())
+				if constexpr (IsParallel())
 				{
-					const int32 NumThisTask = FMath::Min<int32>(ObjectsPerSubTask, NewObjectsToSerialize.Num() - StartIndex);
-					if (MyCompletionGraphEvent.GetReference())
-					{
-						MyCompletionGraphEvent->DontCompleteUntil(TGraphTask< FCollectorTask >::CreateTask().ConstructAndDispatchWhenReady(this, &NewObjectsToSerialize, StartIndex, NumThisTask, ArrayPool));
-					}
-					else
+					FSlowARO::ProcessUnbalancedCalls(Context, Collector);
+				}
+
+StoleARO:		
+				Dispatcher.FlushQueuedReferences();
+
+				if (	 Block = RemainingObjects.PopFullBlock<Options>(); Block);
+				else if (Block = RemainingObjects.PopPartialBlock(/* out if successful */ BlockSize); Block);
+				else if (IsParallel()) // if constexpr yields MSVC unreferenced label warning
+				{
+					switch (StealWork(/* in-out */ Context, Collector, /* out */ Block))
 					{
-						TaskQueue.AddTask(&NewObjectsToSerialize, StartIndex, NumThisTask);
+						case ELoot::Nothing:	break;				// Done, stop working
+						case ELoot::Block:		break;				// Stole full block, process it
+						case ELoot::ARO:		goto StoleARO;		// Stole and made ARO calls that feed into Dispatcher queues and RemainingObjects
+						case ELoot::Context:	goto StoleContext;	// Stole initial references and initial objects worker that hasn't started working
 					}
-					StartIndex += NumThisTask;
 				}
-				NewObjectsToSerialize.SetNumUnsafeInternal(0);
-			}
-			else if (NewObjectsToSerialize.Num())
-			{
-				// Don't spawn a new task, continue in the current one
-				// To avoid allocating and moving memory around swap ObjectsToSerialize and NewObjectsToSerialize arrays
-				Exchange(ObjectsToSerialize, NewObjectsToSerialize);
-				// Empty but don't free allocated memory
-				NewObjectsToSerialize.SetNumUnsafeInternal(0);
 
-				CurrentIndex = 0;
+				if (!Block)
+				{
+					break;
+				}
 			}
-		}
-		while (CurrentIndex < ObjectsToSerialize.Num());
 
-#if PERF_DETAILED_PER_CLASS_GC_STATS
+			CurrentObjects = MakeArrayView(Block->Objects, BlockSize);
+		} // while (true)
+		
+	#if PERF_DETAILED_PER_CLASS_GC_STATS
 		// Detailed per class stats should not be performed when parallel GC is running
 		check(!IsParallel());
-		ReferenceProcessor.LogDetailedStatsSummary();
-#endif
-
-		ArrayPool.ReturnToPool(&NewObjectsToSerializeStruct);
+		Processor.LogDetailedStatsSummary();
+	#endif
 	}
 };
 
+//////////////////////////////////////////////////////////////////////////
 
-/** Default implementation for reference collector that can be used with TFastReferenceCollector */
-template <typename ReferenceProcessorType, bool bIgnoringArchetypeRef = false, bool bIgnoringTransient = false>
-class TDefaultReferenceCollector : public FReferenceCollector
+/** Default reference collector for CollectReferences() */
+template <typename ProcessorType>
+class TDefaultCollector : public FReferenceCollector
 {
-	ReferenceProcessorType& Processor;
-	FGCArrayStruct& ObjectArrayStruct;
+protected:
+	ProcessorType& Processor;
+	UE::GC::FWorkerContext& Context;
 
 public:
-	TDefaultReferenceCollector(ReferenceProcessorType& InProcessor, FGCArrayStruct& InObjectArrayStruct)
-		: Processor(InProcessor)
-		, ObjectArrayStruct(InObjectArrayStruct)
-	{
-	}
+	TDefaultCollector(ProcessorType& InProcessor, FGCArrayStruct& InContext)
+	: Processor(InProcessor)
+	, Context(InContext)
+	{}
+
 	virtual void HandleObjectReference(UObject*& Object, const UObject* ReferencingObject, const FProperty* ReferencingProperty) override
 	{
-		Processor.HandleTokenStreamObjectReference(ObjectArrayStruct, const_cast<UObject*>(ReferencingObject), Object, INDEX_NONE, EGCTokenType::Native, false);
+		Processor.HandleTokenStreamObjectReference(Context, const_cast<UObject*>(ReferencingObject), Object, ETokenlessId::Collector, EGCTokenType::Native, false);
 	}
 	virtual void HandleObjectReferences(UObject** InObjects, const int32 ObjectNum, const UObject* ReferencingObject, const FProperty* InReferencingProperty) override
 	{
 		for (int32 ObjectIndex = 0; ObjectIndex < ObjectNum; ++ObjectIndex)
 		{
 			UObject*& Object = InObjects[ObjectIndex];
-			Processor.HandleTokenStreamObjectReference(ObjectArrayStruct, const_cast<UObject*>(ReferencingObject), Object, INDEX_NONE, EGCTokenType::Native, false);
+			Processor.HandleTokenStreamObjectReference(Context, const_cast<UObject*>(ReferencingObject), Object, ETokenlessId::Collector, EGCTokenType::Native, false);
 		}
 	}
-	virtual bool IsIgnoringArchetypeRef() const override
-	{
-		return bIgnoringArchetypeRef;
-	}
-	virtual bool IsIgnoringTransient() const override
-	{
-		return bIgnoringTransient;
-	}
+
+	virtual bool IsIgnoringArchetypeRef() const override { return false;}
+	virtual bool IsIgnoringTransient() const override {	return false; }
 };
 
-/** Simple single-threaded base implementation for reference processor that can be used with FFastReferenceCollector */
+inline constexpr EGCOptions DefaultOptions = EGCOptions::AutogenerateTokenStream;
+
+} // namespace UE::GC
+
+/** Simple single-threaded reference processor base class for use with CollectReferences() */
 class FSimpleReferenceProcessorBase
 {
 public:
-	FORCEINLINE int32 GetMinDesiredObjectsPerSubTask() const
-	{
-		// We only support single-threaded processing at the moment.
-		return 0;
-	}
-	FORCEINLINE volatile bool IsRunningMultithreaded() const
-	{
-		// We only support single-threaded processing at the moment.
-		return false;
-	}
-	FORCEINLINE void SetIsRunningMultithreaded(bool bIsParallel)
-	{
-		// We only support single-threaded processing at the moment.
-		check(!bIsParallel);
-	}
-	void UpdateDetailedStats(UObject* CurrentObject, uint32 DeltaCycles)
+	static constexpr EGCOptions Options = UE::GC::DefaultOptions;
+	void UpdateDetailedStats(UObject* CurrentObject, uint32 DeltaCycles) {}
+	void LogDetailedStatsSummary() {}
+
+	// Implement this in your derived class, don't make this virtual as it will affect performance!
+	//FORCEINLINE void HandleTokenStreamObjectReference(UE::GC::FWorkerContext& Context, UObject* ReferencingObject, UObject*& Object, UE::GC::FTokenId TokenIndex, EGCTokenType TokenType, bool bAllowReferenceElimination)
+};
+
+
+template<class CollectorType, class ProcessorType>
+FORCEINLINE static void CollectReferences(ProcessorType& Processor, UE::GC::FWorkerContext& Context)
+{
+	using namespace UE::GC;
+	using FastReferenceCollector = TFastReferenceCollector<ProcessorType, CollectorType>;
+	
+	if constexpr (IsParallel(ProcessorType::Options))
 	{
-		// Do nothing
+		ProcessAsync([](void* P, FWorkerContext& C) { FastReferenceCollector(*reinterpret_cast<ProcessorType*>(P)).ProcessObjectArray(C); }, &Processor, Context);
 	}
-	void LogDetailedStatsSummary()
+	else
 	{
-		// Do nothing
+		FastReferenceCollector(Processor).ProcessObjectArray(Context);
 	}
-	// Implement this in your derived class, don't make this virtual as it will affect performance!
-	//FORCEINLINE void HandleTokenStreamObjectReference(FGCArrayStruct& ObjectsToSerializeStruct, UObject* ReferencingObject, UObject*& Object, const int32 TokenIndex, const EGCTokenType TokenType, bool bAllowReferenceElimination);
-};
+}
+
+template<class ProcessorType>
+FORCEINLINE static void CollectReferences(ProcessorType& Processor, UE::GC::FWorkerContext& Context)
+{
+	CollectReferences<UE::GC::TDefaultCollector<ProcessorType>, ProcessorType>(Processor, Context);
+}
+
+// Get number of workers to use when calling CollectReferences in parallel
+COREUOBJECT_API int32 GetNumCollectReferenceWorkers();
\ No newline at end of file
diff --git a/Engine/Source/Runtime/CoreUObject/Public/UObject/FieldPath.h b/Engine/Source/Runtime/CoreUObject/Public/UObject/FieldPath.h
--- a/Engine/Source/Runtime/CoreUObject/Public/UObject/FieldPath.h
+++ b/Engine/Source/Runtime/CoreUObject/Public/UObject/FieldPath.h
@@ -18,7 +18,6 @@
 #include "Templates/PointerIsConvertibleFromTo.h"
 #include "Templates/TypeHash.h"
 #include "Templates/UnrealTemplate.h"
-#include "UObject/FastReferenceCollectorOptions.h"
 #include "UObject/Field.h"
 #include "UObject/NameTypes.h"
 #include "UObject/UObjectArray.h"
@@ -36,9 +35,7 @@ template <typename T> struct TIsZeroConstructType;
 
 struct COREUOBJECT_API FFieldPath
 {
-	// GC needs access to GetResolvedOwnerItemInternal and ClearCachedFieldInternal
-	template <typename ReferenceProcessorType, typename CollectorType, typename ArrayPoolType, EFastReferenceCollectorOptions Options>
-	friend class TFastReferenceCollector;
+	friend struct FGCInternals;
 
 	// TWeakFieldPtr needs access to ClearCachedField
 	template<class T>
diff --git a/Engine/Source/Runtime/CoreUObject/Public/UObject/FieldPathProperty.h b/Engine/Source/Runtime/CoreUObject/Public/UObject/FieldPathProperty.h
--- a/Engine/Source/Runtime/CoreUObject/Public/UObject/FieldPathProperty.h
+++ b/Engine/Source/Runtime/CoreUObject/Public/UObject/FieldPathProperty.h
@@ -21,7 +21,7 @@ class UObject;
 class UStruct;
 namespace UECodeGen_Private { struct FFieldPathPropertyParams; }
 struct FPropertyTag;
-
+namespace UE::GC { class FTokenStreamBuilder; }
 // need to break this out a different type so that the DECLARE_CASTED_CLASS_INTRINSIC macro can digest the comma
 typedef TProperty<FFieldPath, FProperty> FFieldPathProperty_Super;
 
@@ -76,7 +76,7 @@ PRAGMA_ENABLE_DEPRECATION_WARNINGS
 	virtual const TCHAR* ImportText_Internal(const TCHAR* Buffer, void* ContainerOrPropertyPtr, EPropertyPointerType PropertyPointerType, UObject* OwnerObject, int32 PortFlags, FOutputDevice* ErrorText) const override;
 	virtual void Serialize(FArchive& Ar) override;
 	virtual bool ContainsObjectReference(TArray<const FStructProperty*>& EncounteredStructProps, EPropertyObjectReferenceType InReferenceType = EPropertyObjectReferenceType::Strong) const override;
-	virtual void EmitReferenceInfo(UClass& OwnerClass, int32 BaseOffset, TArray<const FStructProperty*>& EncounteredStructProps, FGCStackSizeHelper& StackSizeHelper) override;
+	virtual void EmitReferenceInfo(UE::GC::FTokenStreamBuilder& TokenStream, int32 BaseOffset, TArray<const FStructProperty*>& EncounteredStructProps, FGCStackSizeHelper& StackSizeHelper) override;
 	virtual bool SupportsNetSharedSerialization() const override;
 	// End of FProperty interface
 
