GC reachability optimizations. ~1.5x speedup and ~2.5MB memory saved for an internal title.

diff --git a/Engine/Source/Runtime/CoreUObject/Public/UObject/GCObject.h b/Engine/Source/Runtime/CoreUObject/Public/UObject/GCObject.h
--- a/Engine/Source/Runtime/CoreUObject/Public/UObject/GCObject.h
+++ b/Engine/Source/Runtime/CoreUObject/Public/UObject/GCObject.h
@@ -14,48 +14,42 @@
 
 class FGCObject;
 
-class COREUOBJECT_API FGCObject;
-
 /**
  * This nested class is used to provide a UObject interface between non
  * UObject classes and the UObject system. It handles forwarding all
  * calls of AddReferencedObjects() to objects/ classes that register with it.
  */
-class COREUOBJECT_API UGCObjectReferencer : public UObject
+class UGCObjectReferencer : public UObject
 {
-	/**
-	 * This is the list of objects that are referenced
-	 */
-	TArray<FGCObject*> ReferencedObjects;
-	/** Critical section used when adding and removing objects */
-	FCriticalSection ReferencedObjectsCritical;
-	/** True if we are currently inside AddReferencedObjects */
-	bool bIsAddingReferencedObjects = false;
+	struct FImpl;
+	TUniquePtr<FImpl> Impl;
+
+	/** Current FGCObject* that references are being added from  */
+	FGCObject* CurrentlySerializingObject = nullptr;
+
 	/** True if the list of referenced objects has changed since we last ran VerifyGCObjectNames */
 	bool bReferencedObjectsChangedSinceLastNameVerify = false;
-	/** Currently serializing FGCObject*, only valid if bIsAddingReferencedObjects */
-	FGCObject* CurrentlySerializingObject = nullptr;
 	friend struct FReplaceReferenceHelper;
 
 public:
 	DECLARE_CASTED_CLASS_INTRINSIC_WITH_API_NO_CTOR(UGCObjectReferencer, UObject, CLASS_Transient, TEXT("/Script/CoreUObject"), CASTCLASS_None, NO_API);
 
-	UGCObjectReferencer(const FObjectInitializer& ObjectInitializer = FObjectInitializer::Get());
-	UGCObjectReferencer(FVTableHelper& Helper) : Super(Helper) {};
+	COREUOBJECT_API UGCObjectReferencer(const FObjectInitializer& ObjectInitializer = FObjectInitializer::Get());
+	COREUOBJECT_API UGCObjectReferencer(FVTableHelper& Helper);
 
 	/**
 	 * Adds an object to the referencer list
 	 *
 	 * @param Object The object to add to the list
 	 */
-	void AddObject(FGCObject* Object);
+	COREUOBJECT_API void AddObject(FGCObject* Object);
 
 	/**
 	 * Removes an object from the referencer list
 	 *
 	 * @param Object The object to remove from the list
 	 */
-	void RemoveObject(FGCObject* Object);
+	COREUOBJECT_API void RemoveObject(FGCObject* Object);
 
 	/**
 	 * Get the name of the first FGCObject that owns this object.
@@ -65,7 +59,7 @@ public:
 	 * @param bOnlyIfAddingReferenced Only try to find the name if we are currently inside AddReferencedObjects
 	 * @return true if the object was found.
 	 */
-	bool GetReferencerName(UObject* Object, FString& OutName, bool bOnlyIfAddingReferenced = false) const;
+	COREUOBJECT_API bool GetReferencerName(UObject* Object, FString& OutName, bool bOnlyIfAddingReferenced = false) const;
 
 	/**
 	 * Forwards this call to all registered objects so they can reference
@@ -74,8 +68,9 @@ public:
 	 * @param InThis This UGCObjectReferencer object.
 	 * @param Collector The collector of referenced objects.
 	 */
-	static void AddReferencedObjects(UObject* InThis, FReferenceCollector& Collector);
-	
+	COREUOBJECT_API static void AddReferencedObjects(UObject* InThis, FReferenceCollector& Collector);
+
+	void AddInitialReferences(TArray<UObject**>& Out);
 	/**
 	 * Destroy function that gets called before the object is freed. This might
 	 * be as late as from the destructor.
@@ -85,7 +80,7 @@ public:
 	/**
 	 * Verifies that all FGCObjects have specified names
 	 */
-	void VerifyGCObjectNames();
+	COREUOBJECT_API void VerifyGCObjectNames();
 
 	/**
 	 * Returns the currently serializing object
@@ -122,21 +117,6 @@ public:
  */
 class COREUOBJECT_API FGCObject
 {
-	bool bReferenceAdded = false;
-
-	void Init()
-	{
-		// Some objects can get created after the engine started shutting down (lazy init of singletons etc).
-		if (!IsEngineExitRequested())
-		{
-			StaticInit();
-			check(GGCObjectReferencer);
-			// Add this instance to the referencer's list
-			GGCObjectReferencer->AddObject(this);
-			bReferenceAdded = true;
-		}
-	}
-
 public:
 	/**
 	 * The static object referencer object that is shared across all
@@ -144,59 +124,59 @@ public:
 	 */
 	static UGCObjectReferencer* GGCObjectReferencer;
 
-	/** Default name for unnamed FGCObjects */
-	static const TCHAR* UnknownGCObjectName;
-
-	/**
-	 * Initializes the global object referencer and adds it to the root set.
-	 */
-	static void StaticInit(void)
-	{
-		if (GGCObjectReferencer == NULL)
-		{
-			GGCObjectReferencer = NewObject<UGCObjectReferencer>();
-			GGCObjectReferencer->AddToRoot();
-		}
-	}
+	/** Initializes the global object referencer and adds it to the root set. */
+	static void StaticInit();
 
 	/**
 	 * Tells the global object that forwards AddReferencedObjects calls on to objects
 	 * that a new object is requiring AddReferencedObjects call.
 	 */
-	FGCObject(void)
+	FGCObject()
 	{
-		Init();
+		RegisterGCObject();
 	}
 
-	/** Copy constructor */
-	FGCObject(FGCObject const&)
+	FGCObject(const FGCObject& Other)
 	{
-		Init();
+		RegisterGCObject();
 	}
 
-	/** Move constructor */
-	FGCObject(FGCObject&&)
+	FGCObject(FGCObject&& Other)
 	{
-		Init();
+		RegisterGCObject();
 	}
 
-	/**
-	 * Removes this instance from the global referencer's list
-	 */
-	virtual ~FGCObject(void)
+	enum class EFlags : uint32
+	{
+		None = 0,
+
+		/** Manually call RegisterGCObject() later to avoid collecting references from empty / late-initialized FGCObjects */
+		RegisterLater = 1 << 0,
+
+		/**
+		 * Declare that AddReferencedObjects *only* calls FReferenceCollector::AddStableReference*() functions
+		 * instead of the older FReferenceCollector::AddReferenceObject*() and only adds native references.
+		 * 
+		 * Allows gathering initial references before reachability analysis starts.
+		 */
+		AddStableNativeReferencesOnly = 1 << 1,
+	};
+	explicit FGCObject(EFlags Flags);
+
+	virtual ~FGCObject()
 	{
-		// GObjectSerializer will be NULL if this object gets destroyed after the exit purge.
-		// We want to make sure we remove any objects that were added to the GGCObjectReferencer during Init when exiting
-		if (GGCObjectReferencer && bReferenceAdded)
-		{
-			// Remove this instance from the referencer's list
-			GGCObjectReferencer->RemoveObject(this);
-		}
+		UnregisterGCObject();
 	}
 
 	FGCObject& operator=(const FGCObject&) {return *this;}
 	FGCObject& operator=(FGCObject&&) {return *this;}
 
+	/** Register with GC, only needed if constructed with EFlags::RegisterLater or after unregistering */
+	void RegisterGCObject();
+
+	/** Unregister ahead of destruction. Safe to call multiple times. */
+	void UnregisterGCObject();
+
 	/**
 	 * Pure virtual that must be overloaded by the inheriting class. Use this
 	 * method to serialize any UObjects contained that you wish to keep around.
@@ -205,20 +185,33 @@ public:
 	 */
 	virtual void AddReferencedObjects( FReferenceCollector& Collector ) = 0;
 
-	/**
-	 * Use this method to report a name for your referencer.
-	 */
+	/** Overload this method to report a name for your referencer */
 	virtual FString GetReferencerName() const
 	{
 		return UnknownGCObjectName;
 	}
 
-	/**
-	 * Use this method to report how the specified object is referenced, if necessary
-	 */
+	/** Overload this method to report how the specified object is referenced, if necessary */
 	virtual bool GetReferencerPropertyName(UObject* Object, FString& OutPropertyName) const
 	{
 		return false;
 	}
+
+private:
+	friend UGCObjectReferencer;
+	static const TCHAR* UnknownGCObjectName;
+
+	const bool bCanMakeInitialReferences = false;
+	bool bReferenceAdded = false;	
 };
 
+ENUM_CLASS_FLAGS(FGCObject::EFlags);
+
+FORCEINLINE_DEBUGGABLE FGCObject::FGCObject(EFlags Flags)
+: bCanMakeInitialReferences(EnumHasAllFlags(Flags, EFlags::AddStableNativeReferencesOnly))
+{
+	if (!EnumHasAnyFlags(Flags, EFlags::RegisterLater))
+	{
+		RegisterGCObject();
+	}
+}
\ No newline at end of file
diff --git a/Engine/Source/Runtime/CoreUObject/Public/UObject/GCObjectInfo.h b/Engine/Source/Runtime/CoreUObject/Public/UObject/GCObjectInfo.h
--- a/Engine/Source/Runtime/CoreUObject/Public/UObject/GCObjectInfo.h
+++ b/Engine/Source/Runtime/CoreUObject/Public/UObject/GCObjectInfo.h
@@ -146,5 +146,5 @@ public:
 	}
 
 	/** Helper function for adding info about an UObject into UObject to FGCObjectInfo map */
-	static FGCObjectInfo* FindOrAddInfoHelper(UObject* InObject, TMap<UObject*, FGCObjectInfo*>& InOutObjectToInfoMap);
+	static FGCObjectInfo* FindOrAddInfoHelper(const UObject* InObject, TMap<const UObject*, FGCObjectInfo*>& InOutObjectToInfoMap);
 };
diff --git a/Engine/Source/Runtime/CoreUObject/Public/UObject/GarbageCollection.h b/Engine/Source/Runtime/CoreUObject/Public/UObject/GarbageCollection.h
--- a/Engine/Source/Runtime/CoreUObject/Public/UObject/GarbageCollection.h
+++ b/Engine/Source/Runtime/CoreUObject/Public/UObject/GarbageCollection.h
@@ -18,7 +18,6 @@
 #include "Serialization/ArchiveUObject.h"
 #include "Stats/Stats.h"
 #include "Stats/Stats2.h"
-#include "UObject/FastReferenceCollectorOptions.h"
 #include "UObject/NameTypes.h"
 #include "UObject/UObjectGlobals.h"
 #include "UObject/UnrealNames.h"
@@ -55,6 +54,14 @@ DECLARE_STATS_GROUP(TEXT("Garbage Collection"), STATGROUP_GC, STATCAT_Advanced);
  */
 #define DO_POINTER_CHECKS_ON_GC 0 
 
+namespace UE::GC
+{
+struct FWorkerContext;
+struct FTokenId;
+}
+
+using FGCArrayStruct = UE::GC::FWorkerContext;
+
 /*-----------------------------------------------------------------------------
 	Realtime garbage collection helper classes.
 -----------------------------------------------------------------------------*/
@@ -64,38 +71,35 @@ DECLARE_STATS_GROUP(TEXT("Garbage Collection"), STATGROUP_GC, STATCAT_Advanced);
  */
 enum EGCReferenceType
 {
-	GCRT_None			= 0,
-	GCRT_Object,
-	GCRT_Class,
-	GCRT_PersistentObject,
-	GCRT_ExternalPackage,				// Specific reference type token for UObject external package
-	GCRT_ArrayObject,
-	GCRT_ArrayStruct,
-	GCRT_FixedArray,
-	GCRT_AddStructReferencedObjects,
-	GCRT_AddReferencedObjects,
-	GCRT_AddTMapReferencedObjects,
-	GCRT_AddTSetReferencedObjects,
-	GCRT_AddFieldPathReferencedObject,
-	GCRT_ArrayAddFieldPathReferencedObject,
-	GCRT_EndOfPointer,
-	GCRT_EndOfStream,
-	GCRT_NoopPersistentObject,
-	GCRT_NoopClass,
-	GCRT_ArrayObjectFreezable,
-	GCRT_ArrayStructFreezable,
-	GCRT_Optional,
-	GCRT_WeakObject,
-	GCRT_ArrayWeakObject,
-	GCRT_LazyObject,
-	GCRT_ArrayLazyObject,
-	GCRT_SoftObject,
-	GCRT_ArraySoftObject,
-	GCRT_Delegate,
-	GCRT_ArrayDelegate,
-	GCRT_MulticastDelegate,
-	GCRT_ArrayMulticastDelegate,
-	GCRT_DynamicallyTypedValue,
+	GCRT_None = 0,
+	GCRT_Object,							// Scalar reference, i.e. MyObject* and TObjectPtr<MyObject>
+	GCRT_ExternalPackage,					// UObject external package
+	GCRT_ArrayObject,						// Array of references, e.g. TArray<MyObject*> and TArray<TObjectPtr<MyObject>>
+	GCRT_ArrayStruct,						// Array of structs
+	GCRT_FixedArray,						// C-style fixed arrays, e.g. SomeType Props[4]
+	GCRT_AddStructReferencedObjects,		// AddStructReferencedObjects() function pointer
+	GCRT_AddReferencedObjects,				// AddReferencedObjects() function pointer
+	GCRT_AddTMapReferencedObjects,			// Map whose value or key contain references
+	GCRT_AddTSetReferencedObjects,			// Map whose elements contain references
+	GCRT_AddFieldPathReferencedObject,		// Field path strong reference to owner
+	GCRT_ArrayAddFieldPathReferencedObject, // Array of field paths
+	GCRT_EndOfPointer,						// Comes after pointer
+	GCRT_EndOfStream,						// Terminator
+	GCRT_ArrayObjectFreezable,				// Freezable array of references, e.g. TArray<MyObject*, FMemoryImageAllocator>
+	GCRT_ArrayStructFreezable,				// Freezable array of structs
+	GCRT_Optional,							// TOptional
+	GCRT_WeakObject,						// TWeakObjectPtr
+	GCRT_ArrayWeakObject,					// Array of weak pointers
+	GCRT_LazyObject,						// TLazyObjectPtr
+	GCRT_ArrayLazyObject,					// Array of lazy pointers
+	GCRT_SoftObject,						// TSoftObjectPtr
+	GCRT_ArraySoftObject,					// Array of soft object pointers
+	GCRT_Delegate,							// UobjectDelegate weak reference to target object
+	GCRT_ArrayDelegate,						// Array of delegates
+	GCRT_MulticastDelegate,					// Delegate weak reference to target object
+	GCRT_ArrayMulticastDelegate,			// Array of delegates
+	GCRT_DynamicallyTypedValue,				// FDynamicallyTypedValue
+	GCRT_SlowAddReferencedObjects,			// AddReferencedObjects() function using UE::GC::RegisterSlowImplementation
 };
 
 enum class EGCTokenType : uint32
@@ -104,47 +108,25 @@ enum class EGCTokenType : uint32
 	NonNative = 1	
 };
 
-/** 
- * Convenience struct containing all necessary information for a reference.
- */
+/** Transcodes uint32 token <-> reflection info for one reference property */
 struct FGCReferenceInfo
 {
-	/**
-	 * Constructor
-	 *
-	 * @param InType	type of reference
-	 * @param InOffset	offset into object/ struct
-	 */
 	FORCEINLINE FGCReferenceInfo(EGCReferenceType InReferenceType, uint32 InOffset)
 		: ReturnCount(0)
 		, Type(InReferenceType)
 		, Offset(InOffset)
 	{
+		static_assert(sizeof(FGCReferenceInfo) == sizeof(uint32));
 		checkf(InReferenceType != GCRT_None && InReferenceType <= 0x1F, TEXT("Invalid GC Token Reference Type (%d)"), (uint32)InReferenceType);
 		checkf((InOffset & ~0x7FFFF) == 0, TEXT("Invalid GC Token Offset (%d), max is %d"), InOffset, 0x7FFFF);
 	}
-	/**
-	 * Constructor
-	 *
-	 * @param InValue	value to set union mapping to a uint32 to
-	 */
-	FORCEINLINE FGCReferenceInfo( uint32 InValue )
-	:	Value( InValue )
-	{}
-	/**
-	 * uint32 conversion operator
-	 *
-	 * @return uint32 value of struct
-	 */
-	FORCEINLINE operator uint32() const 
-	{ 
-		return Value; 
-	}
 
-	/** Mapping to exactly one uint32 */
+	FORCEINLINE explicit FGCReferenceInfo(uint32 InValue) : Value(InValue) {}
+
+	FORCEINLINE operator uint32() const { return Value; }
+
 	union
 	{
-		/** Mapping to exactly one uint32 */
 		struct
 		{
 			/** Return depth, e.g. 1 for last entry in an array, 2 for last entry in an array of structs of arrays, ... */
@@ -157,41 +139,16 @@ struct FGCReferenceInfo
 		/** uint32 value of reference info, used for easy conversion to/ from uint32 for token array */
 		uint32 Value;
 	};
-
-	/** End of token stream token */
-	static const FGCReferenceInfo EndOfStreamToken;
 };
 
-/** 
- * Convenience structure containing all necessary information for skipping a dynamic array
- */
+/** Transcodes uint32 token <-> info needed to skip a dynamic array */
 struct FGCSkipInfo
 {
-	/**
-	 * Default constructor
-	 */
-	FORCEINLINE FGCSkipInfo()
-	{}
+	FORCEINLINE FGCSkipInfo() {}
+	FORCEINLINE explicit FGCSkipInfo(uint32 InValue) : Value(InValue) {}
 
-	/**
-	 * Constructor
-	 *
-	 * @param InValue	value to set union mapping to a uint32 to
-	 */
-	FORCEINLINE FGCSkipInfo( uint32 InValue )
-	:	Value( InValue )
-	{}
-	/**
-	 * uint32 conversion operator
-	 *
-	 * @return uint32 value of struct
-	 */
-	FORCEINLINE operator uint32() const 
-	{ 
-		return Value; 
-	}
+	FORCEINLINE operator uint32() const { return Value; }
 
-	/** Mapping to exactly one uint32 */
 	union
 	{
 		/** Mapping to exactly one uint32 */
@@ -207,212 +164,216 @@ struct FGCSkipInfo
 	};
 };
 
-/**
- * Stores debug info about the token.
- */
+/** Token debug name and offset */
 struct FTokenInfo
 {
-	/* Token offset. */
 	int32 Offset;
-	/* Token debug name. */
 	FName Name;
 };
 
-/**
- * Reference token stream class. Used for creating and parsing stream of object references.
- */
-struct COREUOBJECT_API FGCReferenceTokenStream
+namespace UE::GC {
+
+enum class EAROFlags
 {
-	/** Initialization value to ensure that we have the right skip index index */
-	enum EGCArrayInfoPlaceholder { E_GCSkipIndexPlaceholder = 0xDEDEDEDE }; 
+	None			= 0,
+	Unbalanced		= 1 << 0,		// Some instances are very slow but most are fast. GC can flush these more frequently.
+	ExtraSlow		= 2 << 0,		// All instances are slow. GC can work-steal these at finer batch granularity.
+};
+ENUM_CLASS_FLAGS(EAROFlags);
 
-	/** Constructor */
-	FGCReferenceTokenStream()
-	{
-		check( sizeof(FGCReferenceInfo) == sizeof(uint32) );
-	}
+// Reference collection batches up slow AddReferencedObjects calls
+COREUOBJECT_API void RegisterSlowImplementation(void (*AddReferencedObjects)(UObject*, FReferenceCollector&), EAROFlags Flags = EAROFlags::None);
 
-	/**
-	 * Shrinks the token stream, removing array slack.
-	 */
-	void Shrink()
-	{
-		Tokens.Shrink();
-#if ENABLE_GC_TOKEN_DEBUG_INFO
-		TokenDebugInfo.Shrink();
-#endif // ENABLE_GC_TOKEN_DEBUG_INFO
-	}
+struct FTokenStreamOwner;
+class FTokenStreamBuilderIterator;
+
+class /* ensure same cache line */ alignas(16) FTokenStreamView
+{
+public:
+	FTokenStreamView() : StackSize(0), DroppedNum(0), IsNonNative(0) {}
 
-	/** Empties the token stream entirely */
-	void Empty()
+	FORCEINLINE EGCTokenType GetTokenType() const
 	{
-		Tokens.Empty();
-#if ENABLE_GC_TOKEN_DEBUG_INFO
-		TokenDebugInfo.Empty();
-#endif // ENABLE_GC_TOKEN_DEBUG_INFO
+		return IsNonNative ? EGCTokenType::NonNative : EGCTokenType::Native;
 	}
 
-	/**
-	 * Returns the size ofthe reference token stream.
-	 *
-	 * @returns Size of the stream.
-	 */
-	int32 Size() const
+	FORCEINLINE FGCReferenceInfo AccessReferenceInfo(uint32 CurrentIndex) const
 	{
-		return Tokens.Num();
+		return FGCReferenceInfo(Tokens[CurrentIndex]);
 	}
 
-	/**
-	 * return true if this is empty
-	 */
-	bool IsEmpty() const
+	FORCEINLINE uint32 ReadStride(uint32& CurrentIndex) const
 	{
-		return Tokens.Num() == 0;
+		return Tokens[CurrentIndex++];
 	}
 
-	/**
-	 * Increments the stack size requirement for this stream
-	 */
-	void SetStackSize(int32 InStackSize)
+	FORCEINLINE FGCSkipInfo ReadSkipInfo(uint32& CurrentIndex) const
 	{
-		StackSize = InStackSize;
+		FGCSkipInfo SkipInfo(Tokens[CurrentIndex]);
+		SkipInfo.SkipIndex += CurrentIndex;
+		CurrentIndex++;
+		return SkipInfo;
 	}
 
 	/**
-	 * Returns the stack size required by this stream
+	 * Read return count stored at the index before the skip index. This is required 
+	 * to correctly return the right amount of levels when skipping over an empty array.
+	 *
+	 * @param SkipIndex index of first token after array
 	 */
-	int32 GetStackSize() const
+	FORCEINLINE uint8 GetSkipReturnCount( FGCSkipInfo SkipInfo ) const
 	{
-		return StackSize;
+		check( SkipInfo.SkipIndex > 0 && SkipInfo.SkipIndex <= Num );		
+		FGCReferenceInfo ReferenceInfo(Tokens[SkipInfo.SkipIndex-1]);
+		check( ReferenceInfo.Type != GCRT_None );
+		return static_cast<uint8>(ReferenceInfo.ReturnCount - SkipInfo.InnerReturnCount);
 	}
 
-	/**
-	 * Sets token type this stream contains
-	 */
-	void SetTokenType(EGCTokenType InType)
+	FORCEINLINE uint32 ReadCount(uint32& CurrentIndex) const
 	{
-		TokenType = InType;
+		return Tokens[CurrentIndex++];
 	}
 
-	/**
-	 * Returns token type this stream contains
-	 */
-	FORCEINLINE EGCTokenType GetTokenType() const
+	FORCEINLINE void* ReadPointer(uint32& CurrentIndex)
 	{
-		return TokenType;
+		UPTRINT Result = UPTRINT(Tokens[CurrentIndex]) | (UPTRINT(Tokens[CurrentIndex+1]) << 32);
+		CurrentIndex += 2;
+		return (void*)Result;
 	}
 
-	/**
-	 * Prepends passed in stream to existing one.
-	 *
-	 * @param Other	stream to concatenate
-	 */
-	void PrependStream( const FGCReferenceTokenStream& Other );
+	bool IsEmpty() const { return Num == 0; }
+	uint32 NumTokens() const { return Num; }
+	void* GetTokenData() const { return Tokens; }
+private:
+	friend FTokenStreamOwner;
+	friend class FTokenStreamBuilder;
+	static constexpr uint32 NumStackSizeBits = 28;
+	static constexpr uint32 NumDroppedNumBits = 3;
 
-	/**
-	 * Emit reference info.
-	 *
-	 * @param ReferenceInfo	reference info to emit.
-	 *
-	 * @return Index of the reference info in the token stream.
-	 */
-	int32 EmitReferenceInfo(FGCReferenceInfo ReferenceInfo, const FName& DebugName);
+	uint32* RESTRICT Tokens = nullptr; // Debug names stored after tokens
+	uint32 Num = 0;
+	uint32 StackSize : NumStackSizeBits;
+	uint32 DroppedNum : NumDroppedNumBits;
+	uint32 IsNonNative : 1;
+
+#if ENABLE_GC_TOKEN_DEBUG_INFO
+	const FName* GetDebugNames() const { return reinterpret_cast<const FName*>(Tokens + Num + DroppedNum); }
+#endif
+};
+
+
+struct FTokenStreamOwner
+{
+	UE_NONCOPYABLE(FTokenStreamOwner);
+	FTokenStreamOwner() = default;
+	~FTokenStreamOwner() { Reset(); }
+
+	FTokenStreamView Strong; // Strong references only
+	FTokenStreamView Mixed; // Strong and weak references
+	
+	COREUOBJECT_API void Reset();
+
+	/** Gets the number of bytes allocated for tokens */
+	int64 GetTokenAllocatedSize() const;
+
+	/** Gets the number of bytes allocated for token debug info */
+	int64 GetDebugInfoAllocatedSize() const;
+
+	/** Returns extended information about a token, including its name (if available) */
+	FTokenInfo GetTokenInfo(UE::GC::FTokenId Id) const;
+};
+
+/** Reference token stream class. Used for creating and parsing stream of object references. */
+class FTokenStreamBuilder
+{
+public:
+	explicit FTokenStreamBuilder(UClass& Owner) : DebugOwner(Owner) {}
+
+	uint32 Num() const { return static_cast<uint32>(Tokens.Num()); }
+	bool IsEmpty() const { return Tokens.Num() == 0; }
+	void SetStackSize(int32 InStackSize) { StackSize = InStackSize; }
+	uint32 GetStackSize() const { return StackSize; }
+	void SetTokenType(EGCTokenType InType) { TokenType = InType; }
+	EGCTokenType GetTokenType() const { return TokenType; }
+
+	/** Emit reference info and return index of the reference info */
+	COREUOBJECT_API int32 EmitReferenceInfo(FGCReferenceInfo ReferenceInfo, FName DebugName);
 
 	/**
-	 * Emit placeholder for aray skip index, updated in UpdateSkipIndexPlaceholder
+	 * Emit placeholder for array skip index, updated in UpdateSkipIndexPlaceholder
 	 *
 	 * @return the index of the skip index, used later in UpdateSkipIndexPlaceholder
 	 */
-	uint32 EmitSkipIndexPlaceholder();
+	COREUOBJECT_API uint32 EmitSkipIndexPlaceholder();
 
 	/**
 	 * Updates skip index place holder stored and passed in skip index index with passed
-	 * in skip index. The skip index is used to skip over tokens in the case of an emtpy 
+	 * in skip index. The skip index is used to skip over tokens in the case of an empty 
 	 * dynamic array.
 	 * 
 	 * @param SkipIndexIndex index where skip index is stored at.
 	 * @param SkipIndex index to store at skip index index
 	 */
-	void UpdateSkipIndexPlaceholder( uint32 SkipIndexIndex, uint32 SkipIndex );
-
-	/**
-	 * Emit count
-	 *
-	 * @param Count count to emit
-	 */
-	int32 EmitCount( uint32 Count );
+	COREUOBJECT_API void UpdateSkipIndexPlaceholder(uint32 SkipIndexIndex, uint32 SkipIndex);
 
-	/**
-	 * Emit a pointer
-	 *
-	 * @param Ptr pointer to emit
-	 */
-	int32 EmitPointer( void const* Ptr );
-
-	/**
-	 * Emit stride
-	 *
-	 * @param Stride stride to emit
-	 */
-	int32 EmitStride( uint32 Stride );
+	int32 EmitCount(uint32 Count);
+	int32 EmitPointer(const void* Ptr);
+	COREUOBJECT_API int32 EmitStride(uint32 Stride);
 
 	/**
 	 * Increase return count on last token.
 	 *
 	 * @return index of next token
 	 */
-	uint32 EmitReturn();
+	COREUOBJECT_API uint32 EmitReturn();
 
-	/**
-	 * Helper function to perform post parent token stream prepend fixup
-	 */
-	void Fixup(void (*AddReferencedObjectsPtr)(UObject*, class FReferenceCollector&), bool bKeepOuterToken, bool bKeepClassToken);
+	COREUOBJECT_API void EmitObjectReference(int32 Offset, FName DebugName, EGCReferenceType Kind = GCRT_Object);
+	COREUOBJECT_API void EmitObjectArrayReference(int32 Offset, FName DebugName);
+	COREUOBJECT_API uint32 EmitStructArrayBegin(int32 Offset, FName DebugName, int32 Stride);
 
-	/**
-	 * Reads count and advances stream.
-	 *
-	 * @return read in count
-	 */
-	FORCEINLINE uint32 ReadCount( uint32& CurrentIndex )
+	/** The index following the current one will be written to the passed in SkipIndexIndex to skip tokens for empty dynamic arrays. */
+	COREUOBJECT_API void EmitStructArrayEnd(uint32 SkipIndexIndex);
+
+	/** All tokens issues between Begin and End will be replayed Count times. */
+	COREUOBJECT_API void EmitFixedArrayBegin(int32 Offset, FName DebugName, int32 Stride, int32 Count);
+	COREUOBJECT_API void EmitFixedArrayEnd();
+	COREUOBJECT_API void EmitExternalPackageReference();
+
+	void EmitFinalTokens(void(*AddReferencedObjects)(UObject*, FReferenceCollector&));	
+	
+	static FTokenStreamView DropFinalTokens(FTokenStreamView Mixed, void(*DropARO)(UObject*, FReferenceCollector&));
+
+	/** Allocate merged stream with tokens from super class and current class */
+	static void Merge(FTokenStreamView& Out, const FTokenStreamBuilder& Class, FTokenStreamView Super);
+	
+	class FConstIterator;
+	static void CopyStrongTokens(FConstIterator MixedIt, FTokenStreamBuilder& Out);
+	static bool CopyNextStrongToken(FConstIterator& /* in-out */ MixedIt, FTokenStreamBuilder& Out, uint32& OutReturnCount);
+
+	/** Reads count and advances stream */
+	FORCEINLINE uint32 ReadCount(uint32& CurrentIndex)
 	{
 		return Tokens[CurrentIndex++];
 	}
 
-	/**
-	 * Reads stride and advances stream.
-	 *
-	 * @return read in stride
-	 */
-	FORCEINLINE uint32 ReadStride( uint32& CurrentIndex )
+	/** Reads stride and advances stream */
+	FORCEINLINE uint32 ReadStride(uint32& CurrentIndex)
 	{
 		return Tokens[CurrentIndex++];
 	}
 
-	/**
-	 * Reads pointer and advance stream
-	 *
-	 * @return read in pointer
-	 */
-	FORCEINLINE void* ReadPointer( uint32& CurrentIndex )
+	/** Reads pointer and advance stream */
+	FORCEINLINE void* ReadPointer(uint32& CurrentIndex) const
 	{
 		UPTRINT Result = UPTRINT(Tokens[CurrentIndex++]);
-#if PLATFORM_64BITS // warning C4293: '<<' : shift count negative or too big, undefined behavior, so we needed the ifdef
-		static_assert(sizeof(void*) == 8, "Pointer size mismatch.");
 		Result |= UPTRINT(Tokens[CurrentIndex++]) << 32;
-#else
-		static_assert(sizeof(void*) == 4, "Pointer size mismatch.");
-#endif
 		return (void*)Result;
 	}
 
-	/**
-	 * Reads in reference info and advances stream.
-	 *
-	 * @return read in reference info
-	 */
+	/** Reads in reference info and advances stream. */
 	FORCEINLINE FGCReferenceInfo ReadReferenceInfo( uint32& CurrentIndex )
 	{
-		return Tokens[CurrentIndex++];
+		return FGCReferenceInfo(Tokens[CurrentIndex++]);
 	}
 
 	/**
@@ -420,9 +381,9 @@ struct COREUOBJECT_API FGCReferenceTokenStream
 	 *
 	 * @return Reference info at passed in index
 	 */
-	FORCEINLINE FGCReferenceInfo AccessReferenceInfo( uint32 CurrentIndex ) const
+	FORCEINLINE FGCReferenceInfo AccessReferenceInfo(uint32 CurrentIndex) const
 	{
-		return Tokens[CurrentIndex];
+		return FGCReferenceInfo(Tokens[CurrentIndex]);
 	}
 
 	/**
@@ -432,7 +393,7 @@ struct COREUOBJECT_API FGCReferenceTokenStream
 	 */
 	FORCEINLINE FGCSkipInfo ReadSkipInfo( uint32& CurrentIndex )
 	{
-		FGCSkipInfo SkipInfo = Tokens[CurrentIndex];
+		FGCSkipInfo SkipInfo(Tokens[CurrentIndex]);
 		SkipInfo.SkipIndex += CurrentIndex;
 		CurrentIndex++;
 		return SkipInfo;
@@ -444,12 +405,12 @@ struct COREUOBJECT_API FGCReferenceTokenStream
 	 *
 	 * @param SkipIndex index of first token after array
 	 */
-	FORCEINLINE uint32 GetSkipReturnCount( FGCSkipInfo SkipInfo )
+	FORCEINLINE uint8 GetSkipReturnCount( FGCSkipInfo SkipInfo )
 	{
 		check( SkipInfo.SkipIndex > 0 && SkipInfo.SkipIndex <= (uint32)Tokens.Num() );		
-		FGCReferenceInfo ReferenceInfo = Tokens[SkipInfo.SkipIndex-1];
+		FGCReferenceInfo ReferenceInfo(Tokens[SkipInfo.SkipIndex-1]);
 		check( ReferenceInfo.Type != GCRT_None );
-		return ReferenceInfo.ReturnCount - SkipInfo.InnerReturnCount;		
+		return static_cast<uint8>(ReferenceInfo.ReturnCount - SkipInfo.InnerReturnCount);
 	}
 
 	/**
@@ -457,46 +418,24 @@ struct COREUOBJECT_API FGCReferenceTokenStream
 	 *
 	 * @return true if the end of the stream has been reached, false otherwise
 	 */
-	FORCEINLINE bool EndOfStream( uint32 CurrentIndex )
+	FORCEINLINE bool EndOfStream(uint32 CurrentIndex)
 	{
 		return CurrentIndex >= (uint32)Tokens.Num();
 	}
 
-	/**
-	 * Returns extended information about a token, including its name (if available)
-	 *
-	 * @param TokenIndex Index of a token to return extended information for
-	 */
-	FTokenInfo GetTokenInfo(int32 TokenIndex) const;
-
-	/** Sets the maximum stack size required by all token streams */
-	static void SetMaxStackSize(int32 InNewSize)
-	{
-		MaxStackSize = InNewSize;
-	}
-
 	/** Gets the maximum stack size required by all token streams */
-	FORCEINLINE static int32 GetMaxStackSize()
+	FORCEINLINE static uint32 GetMaxStackSize()
 	{
 		return MaxStackSize;
 	}
 
-	/** Gets the number of bytes allocated for tokens */
-	int64 GetTokenAllocatedSize() const
+	friend bool operator==(const FTokenStreamBuilder& A, const FTokenStreamBuilder& B)
 	{
-		return Tokens.GetAllocatedSize();
-	}
-	/** Gets the number of bytes allocated for token debug info */
-	int64 GetDebugInfoAllocatedSize() const
-	{
-#if ENABLE_GC_TOKEN_DEBUG_INFO
-		return TokenDebugInfo.GetAllocatedSize();
-#else
-		return 0;
-#endif
+		return A.Tokens == B.Tokens && A.StackSize == B.StackSize && A.TokenType == B.TokenType;
 	}
 
 private:
+	static constexpr uint32 SkipIndexPlaceholderMagic =  0xDEDEDEDEu; 
 
 	/**
 	 * Helper function to store a pointer into a preallocated token stream.
@@ -504,35 +443,33 @@ private:
 	 * @param Stream Preallocated token stream.
 	 * @param Ptr pointer to store
 	 */
-	FORCEINLINE void StorePointer( uint32* Stream, void const* Ptr )
+	FORCEINLINE void StorePointer(uint32* Stream, const void* Ptr)
 	{
-	#if PLATFORM_64BITS // warning C4293: '<<' : shift count negative or too big, undefined behavior, so we needed the ifdef
-		static_assert(sizeof(void*) == 8, "Pointer size mismatch.");
-		Stream[0] = UPTRINT(Ptr) & 0xffffffff;
+		Stream[0] = UPTRINT(Ptr) & 0xffffffffu;
 		Stream[1] = UPTRINT(Ptr) >> 32;
-	#else
-		static_assert(sizeof(void*) == 4, "Pointer size mismatch.");
-		Stream[0] = PTRINT(Ptr);
-	#endif
 	}
 
-	/** Token array */
-	TArray<uint32> Tokens;
 	/** Stack size required by this token stream */
-	int32 StackSize = 0;
-	/** Type of tokens in this stream (native / non-native) */
+	uint32 StackSize = 0;
 	EGCTokenType TokenType = EGCTokenType::Native;
-	/** Maximum stack size for TFastReferenceCollector */
-	static int32 MaxStackSize;
+	TArray<uint32, TInlineAllocator<128>> Tokens;
 #if ENABLE_GC_TOKEN_DEBUG_INFO
-	/** 
-	 * Name of the proprty that emitted the associated token or token type (pointer etc).
-	 * We want to keep it in a separate array for performance reasons
-	 */
-	TArray<FName> TokenDebugInfo;
-#endif // ENABLE_GC_TOKEN_DEBUG_INFO
+	TArray<FName, TInlineAllocator<128>> DebugNames;
+#endif
+	UClass& DebugOwner;
+
+	/** Maximum stack size for TFastReferenceCollector */
+	COREUOBJECT_API static uint32 MaxStackSize;
+};
+
+struct FIntrinsicClassTokens
+{
+	COREUOBJECT_API static FTokenStreamBuilder& AllocateBuilder(UClass* Class);
+	static TUniquePtr<FTokenStreamBuilder> ConsumeBuilder(UClass* Class);
 };
 
+} // namespace UE::GC
+
 /** Prevent GC from running in the current scope */
 class COREUOBJECT_API FGCScopeGuard
 {
@@ -542,20 +479,17 @@ public:
 };
 
 class FGCObject;
-template <EFastReferenceCollectorOptions Options> class FGCReferenceProcessor;
 
 /** Information about references to objects marked as Garbage that's gather by the Garbage Collector */
 struct FGarbageReferenceInfo
 {
 	/** Object marked as garbage */
 	UObject* GarbageObject;
-	/** Reference to the object marked as garbage */
-	UObject*& GarbageObjectRef;
 	/** Referencing object info */
 	union FReferencerUnion
 	{
 		/** Referencing UObject */
-		UObject* Object;
+		const UObject* Object;
 		/** Referencing FGCObject */
 		FGCObject* GCObject;
 	} Referencer;
@@ -564,19 +498,16 @@ struct FGarbageReferenceInfo
 	/** Referencing property name */
 	FName PropertyName;
 
-	FGarbageReferenceInfo(UObject* InReferencingObject, UObject* InGarbageObject, UObject*& InGarbageObjectRef, FName InPropertyName)
+	FGarbageReferenceInfo(const UObject* InReferencingObject, UObject* InGarbageObject, FName InPropertyName)
 		: GarbageObject(InGarbageObject)
-		, GarbageObjectRef(InGarbageObjectRef)
 		, bReferencerUObject(true)
 		, PropertyName(InPropertyName)
 	{
 		Referencer.Object = InReferencingObject;
 	}
-	FGarbageReferenceInfo(FGCObject* InReferencingObject, UObject* InGarbageObject, UObject*& InGarbageObjectRef)
+	FGarbageReferenceInfo(FGCObject* InReferencingObject, UObject* InGarbageObject)
 		: GarbageObject(InGarbageObject)
-		, GarbageObjectRef(InGarbageObjectRef)
 		, bReferencerUObject(false)
-		, PropertyName(NAME_None)
 	{
 		Referencer.GCObject = InReferencingObject;
 	}
@@ -588,118 +519,12 @@ struct FGarbageReferenceInfo
 struct FGCDirectReference
 {
 	FGCDirectReference() = default;
-	explicit FGCDirectReference(UObject* Obj)
-		: ReferencedObject(Obj)
-	{}
+	explicit FGCDirectReference(UObject* Obj) : ReferencedObject(Obj) {}
 	/** Property or FGCObject name referencing this object */
 	FName ReferencerName;
 	UObject* ReferencedObject = nullptr;
 };
 
-/** Struct to hold the objects to serialize array and the list of weak references. This is allocated by ArrayPool */
-struct FGCArrayStruct
-{
-	template <typename ReferenceProcessorType, typename CollectorType, typename ArrayPoolType, EFastReferenceCollectorOptions Options>
-	friend class TFastReferenceCollector;
-	
-	// Arrays filled during Garbage Collectionmn
-	TArray<UObject*> ObjectsToSerialize;
-	TArray<UObject**> WeakReferences;
-	TArray<FGarbageReferenceInfo> GarbageReferences;
-#if ENABLE_GC_HISTORY
-	TMap<UObject*, TArray<FGCDirectReference>*> History;
-#endif
-
-	FORCEINLINE UObject* GetReferencingObject()
-	{
-		return ReferencingObject;
-	}
-
-	/** Returns the size of memory allocated by internal arrays */
-	int64 GetAllocatedSize() const
-	{
-		return ObjectsToSerialize.GetAllocatedSize() +
-			WeakReferences.GetAllocatedSize() +
-			GarbageReferences.GetAllocatedSize();
-	}
-
-private:
-	// This is set by GC when processing references from the current referencing object
-	UObject* ReferencingObject = nullptr;
-};
-
-/**
-* Specialized FReferenceCollector that uses FGCReferenceProcessor to mark objects as reachable.
-*/
-template <EFastReferenceCollectorOptions Options>
-class FGCCollector : public FReferenceCollector
-{
-	FGCReferenceProcessor<Options>& ReferenceProcessor;
-	FGCArrayStruct& ObjectArrayStruct;
-	bool bAllowEliminatingReferences;
-	bool bIsProcessingNativeReferences;
-
-	constexpr FORCEINLINE bool IsParallel() const
-	{
-		return !!(Options & EFastReferenceCollectorOptions::Parallel);
-	}
-	constexpr FORCEINLINE bool IsWithClusters() const
-	{
-		return !!(Options & EFastReferenceCollectorOptions::WithClusters);
-	}
-
-public:
-
-	FGCCollector(FGCReferenceProcessor<Options>& InProcessor, FGCArrayStruct& InObjectArrayStruct);
-
-	virtual void HandleObjectReference(UObject*& InObject, const UObject* InReferencingObject, const FProperty* InReferencingProperty) override;
-
-	virtual void HandleObjectReferences(UObject** InObjects, const int32 ObjectNum, const UObject* InReferencingObject, const FProperty* InReferencingProperty) override;
-
-	virtual bool IsIgnoringArchetypeRef() const override
-	{
-		return false;
-	}
-	virtual bool IsIgnoringTransient() const override
-	{
-		return false;
-	}
-	virtual void AllowEliminatingReferences(bool bAllow) override
-	{
-		bAllowEliminatingReferences = bAllow;
-	}
-	virtual bool MarkWeakObjectReferenceForClearing(UObject** WeakReference) override
-	{
-		// Track this references for later destruction if necessary. These should be relatively rare
-		ObjectArrayStruct.WeakReferences.Add(WeakReference);
-		return true;
-	}
-	virtual void SetIsProcessingNativeReferences(bool bIsNative) override 
-	{
-		bIsProcessingNativeReferences = bIsNative;
-	}
-
-private:
-
-	FORCEINLINE void InternalHandleObjectReference(UObject*& Object, const UObject* ReferencingObject, const FProperty* ReferencingProperty, const EGCTokenType InTokenType);
-};
-
-/**
- * FGarbageCollectionTracer
- * Interface to allow external systems to trace additional object references, used for bridging GCs
- */
-class FGarbageCollectionTracer
-{
-public:
-	virtual ~FGarbageCollectionTracer() {}
-
-	UE_DEPRECATED(5.0, "Use version of PerformReachabilityAnalysisOnObjects that takes EFastReferenceCollectorOptions parameter.")
-	COREUOBJECT_API virtual void PerformReachabilityAnalysisOnObjects(FGCArrayStruct* ArrayStruct, bool bForceSingleThreaded, bool bWithClusters);
-
-	virtual void PerformReachabilityAnalysisOnObjects(FGCArrayStruct* ArrayStruct, const EFastReferenceCollectorOptions InOptions) = 0;
-
-};
-
 /** True if Garbage Collection is running. Use IsGarbageCollecting() functio n instead of using this variable directly */
 extern COREUOBJECT_API bool GIsGarbageCollecting;
 
diff --git a/Engine/Source/Runtime/CoreUObject/Public/UObject/GarbageCollectionHistory.h b/Engine/Source/Runtime/CoreUObject/Public/UObject/GarbageCollectionHistory.h
--- a/Engine/Source/Runtime/CoreUObject/Public/UObject/GarbageCollectionHistory.h
+++ b/Engine/Source/Runtime/CoreUObject/Public/UObject/GarbageCollectionHistory.h
@@ -40,7 +40,7 @@ struct FGCDirectReferenceInfo
 struct FGCSnapshot
 {
 	/** FGCObjectInfo structs generated for each of the objects encountered during GC */
-	TMap<UObject*, FGCObjectInfo*> ObjectToInfoMap;
+	TMap<const UObject*, FGCObjectInfo*> ObjectToInfoMap;
 	/** Lis of direct references for all objects */
 	TMap<FGCObjectInfo*, TArray<FGCDirectReferenceInfo>*> DirectReferences;
 
@@ -61,7 +61,7 @@ class COREUOBJECT_API FGCHistory
 	~FGCHistory();
 
 	void Cleanup(FGCSnapshot& InShapshot);	
-	void MergeArrayStructHistory(TMap<UObject*, TArray<FGCDirectReference>*>& History, FGCSnapshot& Snapshot);
+	void MergeArrayStructHistory(TMap<const UObject*, TArray<FGCDirectReference>*>& History, FGCSnapshot& Snapshot);
 
 public:
 
@@ -69,7 +69,7 @@ public:
 	static FGCHistory& Get();
 
 	/** Copies information generated by Garbage Collector into the next snapshot struct. */
-	void Update(TArray<FGCArrayStruct*>& AllArrays);
+	void Update(TConstArrayView<TUniquePtr<UE::GC::FWorkerContext>> Contexts);
 
 	/** Frees memory allocated by GC history. */
 	void Cleanup();
diff --git a/Engine/Source/Runtime/CoreUObject/Public/UObject/ObjectMacros.h b/Engine/Source/Runtime/CoreUObject/Public/UObject/ObjectMacros.h
--- a/Engine/Source/Runtime/CoreUObject/Public/UObject/ObjectMacros.h
+++ b/Engine/Source/Runtime/CoreUObject/Public/UObject/ObjectMacros.h
@@ -586,7 +586,6 @@ enum class EInternalObjectFlags : int32
 
 	LoaderImport = 1 << 20, ///< Object is ready to be imported by another package during loading
 	Garbage = 1 << 21, ///< Garbage from logical point of view and should not be referenced. This flag is mirrored in EObjectFlags as RF_Garbage for performance
-	PersistentGarbage = 1 << 22, ///< Same as above but referenced through a persistent reference so it can't be GC'd
 	ReachableInCluster = 1 << 23, ///< External reference to object in cluster exists
 	ClusterRoot = 1 << 24, ///< Root of a cluster
 	Native = 1 << 25, ///< Native (UClass only). 
@@ -602,7 +601,7 @@ enum class EInternalObjectFlags : int32
 	MirroredFlags = Garbage | PendingKill, /// Flags mirrored in EObjectFlags
 
 	//~ Make sure this is up to date!
-	AllFlags = LoaderImport | Garbage | PersistentGarbage | ReachableInCluster | ClusterRoot | Native | Async | AsyncLoading | Unreachable | PendingKill | RootSet | PendingConstruction
+	AllFlags = LoaderImport | Garbage | ReachableInCluster | ClusterRoot | Native | Async | AsyncLoading | Unreachable | PendingKill | RootSet | PendingConstruction
 	PRAGMA_ENABLE_DEPRECATION_WARNINGS
 };
 ENUM_CLASS_FLAGS(EInternalObjectFlags);
diff --git a/Engine/Source/Runtime/CoreUObject/Public/UObject/PersistentObjectPtr.h b/Engine/Source/Runtime/CoreUObject/Public/UObject/PersistentObjectPtr.h
--- a/Engine/Source/Runtime/CoreUObject/Public/UObject/PersistentObjectPtr.h
+++ b/Engine/Source/Runtime/CoreUObject/Public/UObject/PersistentObjectPtr.h
@@ -8,7 +8,6 @@
 
 #include "CoreMinimal.h"
 #include "UObject/WeakObjectPtr.h"
-#include "UObject/FastReferenceCollectorOptions.h"
 
 /**
  * TPersistentObjectPtr is a template base class for FLazyObjectPtr and FSoftObjectPtr
@@ -16,11 +15,8 @@
 template<class TObjectID>
 struct TPersistentObjectPtr
 {
+	friend struct FGCInternals;
 public:	
-
-	template <typename ReferenceProcessorType, typename CollectorType, typename ArrayPoolType, EFastReferenceCollectorOptions Options>
-	friend class TFastReferenceCollector;
-
 	/** Default constructor, will be null */
 	FORCEINLINE TPersistentObjectPtr()
 	{
diff --git a/Engine/Source/Runtime/CoreUObject/Public/UObject/ReferenceChainSearch.h b/Engine/Source/Runtime/CoreUObject/Public/UObject/ReferenceChainSearch.h
--- a/Engine/Source/Runtime/CoreUObject/Public/UObject/ReferenceChainSearch.h
+++ b/Engine/Source/Runtime/CoreUObject/Public/UObject/ReferenceChainSearch.h
@@ -341,7 +341,7 @@ private:
 	/** All nodes created during the search */
 	TMap<FGCObjectInfo*, FGraphNode*> AllNodes;
 	/** Maps UObject pointers to object info structs */
-	TMap<UObject*, FGCObjectInfo*> ObjectToInfoMap;
+	TMap<const UObject*, FGCObjectInfo*> ObjectToInfoMap;
 
 	/** Performs the search */
 	void PerformSearch();
@@ -356,7 +356,7 @@ private:
 	void Cleanup();
 
 	/** Tries to find a node for an object and if it doesn't exists creates a new one and returns it */
-	FGraphNode* FindOrAddNode(UObject* InObjectToFindNodeFor);
+	FGraphNode* FindOrAddNode(const UObject* InObjectToFindNodeFor);
 	FGraphNode* FindOrAddNode(FGCObjectInfo* InObjectInfo);
 
 	/** Link two nodes together with the given reference info */
diff --git a/Engine/Source/Runtime/CoreUObject/Public/UObject/UObjectAllocator.h b/Engine/Source/Runtime/CoreUObject/Public/UObject/UObjectAllocator.h
--- a/Engine/Source/Runtime/CoreUObject/Public/UObject/UObjectAllocator.h
+++ b/Engine/Source/Runtime/CoreUObject/Public/UObject/UObjectAllocator.h
@@ -38,12 +38,7 @@ public:
 	 */
 	void BootMessage();
 
-	/**
-	 * Checks whether object is part of permanent object pool.
-	 *
-	 * @param Object object to test as a member of permanent object pool
-	 * @return true if object is part of permanent object pool, false otherwise
-	 */
+	UE_DEPRECATED(5.1, "Use the more efficient FPermanentObjectPoolExtents instead")
 	FORCEINLINE bool ResidesInPermanentPool(const UObjectBase *Object) const
 	{
 		return ((const uint8*)Object >= PermanentObjectPool) && ((const uint8*)Object < PermanentObjectPoolTail);
@@ -67,6 +62,7 @@ public:
 	void FreeUObject(UObjectBase *Object) const;
 
 private:
+	friend class FPermanentObjectPoolExtents;
 
 	/** Size in bytes of pool for objects disregarded for GC.								*/
 	int32							PermanentObjectPoolSize;
@@ -81,3 +77,21 @@ private:
 /** Global UObjectBase allocator							*/
 extern COREUOBJECT_API FUObjectAllocator GUObjectAllocator;
 
+/** Helps check if an object is part of permanent object pool */
+class FPermanentObjectPoolExtents
+{
+public:
+	FORCEINLINE FPermanentObjectPoolExtents(const FUObjectAllocator& ObjectAllocator = GUObjectAllocator)
+		: Address(reinterpret_cast<uint64>(ObjectAllocator.PermanentObjectPool))
+		, Size(static_cast<uint64>(ObjectAllocator.PermanentObjectPoolSize))
+	{}
+	
+	FORCEINLINE bool Contains(const UObjectBase* Object) const
+	{
+		return reinterpret_cast<uint64>(Object) - Address < Size;
+	}
+
+private:
+	const uint64 Address;
+	const uint64 Size;
+};
diff --git a/Engine/Source/Runtime/CoreUObject/Public/UObject/UObjectArray.h b/Engine/Source/Runtime/CoreUObject/Public/UObject/UObjectArray.h
--- a/Engine/Source/Runtime/CoreUObject/Public/UObject/UObjectArray.h
+++ b/Engine/Source/Runtime/CoreUObject/Public/UObject/UObjectArray.h
@@ -503,10 +503,10 @@ public:
 	**/
 	FORCEINLINE_DEBUGGABLE FUObjectItem const* GetObjectPtr(int32 Index) const TSAN_SAFE
 	{
-		const int32 ChunkIndex = Index / NumElementsPerChunk;
-		const int32 WithinChunkIndex = Index % NumElementsPerChunk;
+		const uint32 ChunkIndex = (uint32)Index / NumElementsPerChunk;
+		const uint32 WithinChunkIndex = (uint32)Index % NumElementsPerChunk;
 		checkf(IsValidIndex(Index), TEXT("IsValidIndex(%d)"), Index);
-		checkf(ChunkIndex < NumChunks, TEXT("ChunkIndex (%d) < NumChunks (%d)"), ChunkIndex, NumChunks);
+		checkf(ChunkIndex < (uint32)NumChunks, TEXT("ChunkIndex (%d) < NumChunks (%d)"), ChunkIndex, NumChunks);
 		checkf(Index < MaxElements, TEXT("Index (%d) < MaxElements (%d)"), Index, MaxElements);
 		FUObjectItem* Chunk = Objects[ChunkIndex];
 		check(Chunk);
@@ -514,16 +514,24 @@ public:
 	}
 	FORCEINLINE_DEBUGGABLE FUObjectItem* GetObjectPtr(int32 Index) TSAN_SAFE
 	{
-		const int32 ChunkIndex = Index / NumElementsPerChunk;
-		const int32 WithinChunkIndex = Index % NumElementsPerChunk;
+		const uint32 ChunkIndex = (uint32)Index / NumElementsPerChunk;
+		const uint32 WithinChunkIndex = (uint32)Index % NumElementsPerChunk;
 		checkf(IsValidIndex(Index), TEXT("IsValidIndex(%d)"), Index);
-		checkf(ChunkIndex < NumChunks, TEXT("ChunkIndex (%d) < NumChunks (%d)"), ChunkIndex, NumChunks);
+		checkf(ChunkIndex < (uint32)NumChunks, TEXT("ChunkIndex (%d) < NumChunks (%d)"), ChunkIndex, NumChunks);
 		checkf(Index < MaxElements, TEXT("Index (%d) < MaxElements (%d)"), Index, MaxElements);
 		FUObjectItem* Chunk = Objects[ChunkIndex];
 		check(Chunk);
 		return Chunk + WithinChunkIndex;
 	}
 
+	FORCEINLINE_DEBUGGABLE void PrefetchObjectPtr(int32 Index) const TSAN_SAFE
+	{
+		const uint32 ChunkIndex = (uint32)Index / NumElementsPerChunk;
+		const uint32 WithinChunkIndex = (uint32)Index % NumElementsPerChunk;
+		const FUObjectItem* Chunk = Objects[ChunkIndex];
+		FPlatformMisc::Prefetch(Chunk + WithinChunkIndex);
+	}
+
 	/**
 	* Return a reference to an element
 	* @param	Index	Index to return
diff --git a/Engine/Source/Runtime/CoreUObject/Public/UObject/UObjectBase.h b/Engine/Source/Runtime/CoreUObject/Public/UObject/UObjectBase.h
--- a/Engine/Source/Runtime/CoreUObject/Public/UObject/UObjectBase.h
+++ b/Engine/Source/Runtime/CoreUObject/Public/UObject/UObjectBase.h
@@ -28,6 +28,8 @@ class UScriptStruct;
 
 DECLARE_DWORD_COUNTER_STAT_EXTERN(TEXT("STAT_UObjectsStatGroupTester"), STAT_UObjectsStatGroupTester, STATGROUP_UObjects, COREUOBJECT_API);
 
+namespace UE::GC { class FTokenStreamBuilder; }
+
 /** 
  * Low level implementation of UObject, should not be used directly in game code 
  */
@@ -78,7 +80,7 @@ public:
 	/**
 	 * Emit GC tokens for UObjectBase, this might be UObject::StaticClass or Default__Class
 	 **/
-	static void EmitBaseReferences(UClass *RootClass);
+	static void EmitBaseReferences(UE::GC::FTokenStreamBuilder& TokenStream);
 
 protected:
 	/**
@@ -227,6 +229,9 @@ public:
 		while( FPlatformAtomics::InterlockedCompareExchange( (int32*)&ObjectFlags, NewFlags, OldFlags) != OldFlags );
 	}
 
+	static void PrefetchClass(UObject* Object) { FPlatformMisc::Prefetch(Object, offsetof(UObjectBase, ClassPrivate)); }
+	static void PrefetchOuter(UObject* Object) { FPlatformMisc::Prefetch(Object, offsetof(UObjectBase, OuterPrivate)); }
+
 private:
 
 	/** Flags used to track and report various object states. This needs to be 8 byte aligned on 32-bit
diff --git a/Engine/Source/Runtime/CoreUObject/Public/UObject/UObjectGlobals.h b/Engine/Source/Runtime/CoreUObject/Public/UObject/UObjectGlobals.h
--- a/Engine/Source/Runtime/CoreUObject/Public/UObject/UObjectGlobals.h
+++ b/Engine/Source/Runtime/CoreUObject/Public/UObject/UObjectGlobals.h
@@ -2098,33 +2098,24 @@ private:
 };
 
 /** Base class for reference serialization archives */
-class COREUOBJECT_API FReferenceCollectorArchive : public FArchiveUObject
+class FReferenceCollectorArchive : public FArchiveUObject
 {
 	/** Object which is performing the serialization. */
-	const UObject* SerializingObject;
+	const UObject* SerializingObject = nullptr;
 	/** Object that owns the serialized data. */
-	const UObject* SerializedDataContainer;
-	/** Pointer to serialized data (read-only). */
-	const void* SerializedDataPtr;
+	const UObject* SerializedDataContainer  = nullptr;
 	/** Stored pointer to reference collector. */
 	class FReferenceCollector& Collector;
 
 protected:
-
-	class FReferenceCollector& GetCollector()
+	FORCEINLINE class FReferenceCollector& GetCollector()
 	{
 		return Collector;
 	}
 
 public:
+	FReferenceCollectorArchive(const UObject* InSerializingObject, FReferenceCollector& InCollector);
 
-	FReferenceCollectorArchive(const UObject* InSerializingObject, FReferenceCollector& InCollector)
-		: SerializingObject(InSerializingObject)
-		, SerializedDataContainer(nullptr)
-		, SerializedDataPtr(nullptr)
-		, Collector(InCollector)
-	{
-	}
 	void SetSerializingObject(const UObject* InSerializingObject)
 	{
 		SerializingObject = InSerializingObject;
@@ -2141,14 +2132,9 @@ public:
 	{
 		return SerializedDataContainer;
 	}
-	void SetSerializedDataPtr(const void* InSerializedDataPtr)
-	{
-		SerializedDataPtr = InSerializedDataPtr;
-	}
-	const void* GetSerializedDataPtr() const
-	{
-		return SerializedDataPtr;
-	}
+
+	COREUOBJECT_API virtual FArchive& operator<<(UObject*& Object) override;
+	COREUOBJECT_API virtual FArchive& operator<<(FObjectPtr& Object) override;
 };
 
 /** Helper class for setting and resetting attributes on the FReferenceCollectorArchive */
@@ -2158,27 +2144,23 @@ class COREUOBJECT_API FVerySlowReferenceCollectorArchiveScope
 	const UObject* OldSerializingObject;
 	FProperty* OldSerializedProperty;
 	const UObject* OldSerializedDataContainer;
-	const void* OldSerializedDataPtr;
 
 public:
-	FVerySlowReferenceCollectorArchiveScope(FReferenceCollectorArchive& InArchive, const UObject* InSerializingObject, FProperty* InSerializedProperty = nullptr, const UObject* InSerializedDataContainer = nullptr, const void* InSerializedDataPtr = nullptr)
+	FVerySlowReferenceCollectorArchiveScope(FReferenceCollectorArchive& InArchive, const UObject* InSerializingObject, FProperty* InSerializedProperty = nullptr, const UObject* InSerializedDataContainer = nullptr)
 		: Archive(InArchive)
 		, OldSerializingObject(InArchive.GetSerializingObject())
 		, OldSerializedProperty(InArchive.GetSerializedProperty())
 		, OldSerializedDataContainer(InArchive.GetSerializedDataContainer())
-		, OldSerializedDataPtr(InArchive.GetSerializedDataPtr())
 	{
 		Archive.SetSerializingObject(InSerializingObject);
 		Archive.SetSerializedProperty(InSerializedProperty);
 		Archive.SetSerializedDataContainer(InSerializedDataContainer);
-		Archive.SetSerializedDataPtr(InSerializedDataPtr);
 	}
 	~FVerySlowReferenceCollectorArchiveScope()
 	{
 		Archive.SetSerializingObject(OldSerializingObject);
 		Archive.SetSerializedProperty(OldSerializedProperty);
 		Archive.SetSerializedDataContainer(OldSerializedDataContainer);
-		Archive.SetSerializedDataPtr(OldSerializedDataPtr);
 	}
 	FReferenceCollectorArchive& GetArchive()
 	{
@@ -2186,16 +2168,82 @@ public:
 	}
 };
 
-/**
- * FReferenceCollector.
- * Helper class used by the garbage collector to collect object references.
- */
+/** Used by garbage collector to collect references via virtual AddReferencedObjects calls */
 class COREUOBJECT_API FReferenceCollector
 {
 public:
+	virtual ~FReferenceCollector() {}
+
+	/** Preferred way to add a reference that allows batching. Object must outlive GC tracing, can't be used for temporary/stack references. */
+	virtual void AddStableReference(UObject** Object);
+	
+	/** Preferred way to add a reference array that allows batching. Can't be used for temporary/stack array. */
+	virtual void AddStableReferenceArray(TArray<UObject*>* Objects);
+
+	/** Preferred way to add a reference set that allows batching. Can't be used for temporary/stack set. */
+	virtual void AddStableReferenceSet(TSet<UObject*>* Objects);
+
+	template<class UObjectType>
+	FORCEINLINE void AddStableReference(UObjectType** Object)
+	{
+		static_assert(sizeof(UObjectType) > 0, "Element must be a pointer to a fully-defined type");
+		static_assert(std::is_convertible_v<UObjectType*, const UObjectBase*>, "Element must be a pointer to a type derived from UObject");
+		AddStableReference(reinterpret_cast<UObject**>(Object));
+	}
+
+	template<class UObjectType>
+	FORCEINLINE void AddStableReferenceArray(TArray<UObjectType*>* Objects)
+	{
+		static_assert(sizeof(UObjectType) > 0, "Element must be a pointer to a fully-defined type");
+		static_assert(std::is_convertible_v<UObjectType*, const UObjectBase*>, "Element must be a pointer to a type derived from UObject");
+		AddStableReferenceArray(reinterpret_cast<TArray<UObject*>*>(Objects)); 
+	}
+
+	template<class UObjectType>
+	FORCEINLINE void AddStableReferenceSet(TSet<UObjectType*>* Objects)
+	{
+		static_assert(sizeof(UObjectType) > 0, "Element must be a pointer to a fully-defined type");
+		static_assert(std::is_convertible_v<UObjectType*, const UObjectBase*>, "Element must be a pointer to a type derived from UObject");
+		AddStableReferenceSet(reinterpret_cast<TSet<UObject*>*>(Objects)); 
+	}
+
+	template<class UObjectType>
+	FORCEINLINE void AddStableReference(TObjectPtr<UObjectType>* Object)
+	{
+		AddStableReference(reinterpret_cast<UObjectType**>(Object));
+	}
+
+	template<class UObjectType>
+	FORCEINLINE void AddStableReferenceArray(TArray<TObjectPtr<UObjectType>>* Objects)
+	{
+		AddStableReferenceArray(reinterpret_cast<TArray<UObjectType*>*>(Objects)); 
+	}
+
+	template<class UObjectType>
+	FORCEINLINE void AddStableReferenceSet(TSet<TObjectPtr<UObjectType>>* Objects)
+	{
+		AddStableReferenceSet(reinterpret_cast<TSet<UObjectType*>*>(Objects)); 
+	}
+
+	template <typename KeyType, typename ValueType, typename Allocator, typename KeyFuncs>
+	FORCEINLINE_DEBUGGABLE void AddStableReferenceMap(TMapBase<KeyType, ValueType, Allocator, KeyFuncs>& Map)
+	{
+		static constexpr bool bKeyReference =	std::is_convertible_v<KeyType, const UObjectBase*>;
+		static constexpr bool bValueReference =	std::is_convertible_v<ValueType, const UObjectBase*>;
+		static_assert(bKeyReference || bValueReference, "Key or value must be pointer to fully-defined UObject type");
 
-	FReferenceCollector();
-	virtual ~FReferenceCollector();
+		for (TPair<KeyType, ValueType>& Pair : Map)
+		{
+			if constexpr (bKeyReference)
+			{
+				AddStableReference(&Pair.Key);
+			}
+			if constexpr (bValueReference)
+			{
+				AddStableReference(&Pair.Value);
+			}
+		}
+	}
 
 	/**
 	 * Adds object reference.
@@ -2318,18 +2366,6 @@ public:
 		}
 	}
 
-
-
-	/**
-	 * Adds any applicable references from a UScriptStruct. Only necessary to handle cases of an unreflected/non-UPROPERTY struct that wants to have references emitted.
-	 *
-	 * @param ScriptStruct Script struct to add referenced objects from
-	 * @param StructMemory Memory of struct instance
-	 * @param ReferencingObject Referencing object (if available).
-	 * @param ReferencingProperty Referencing property (if available).
-	 */
-	void AddReferencedObjects(const class UScriptStruct*& ScriptStruct, void* StructMemory, const UObject* ReferencingObject = nullptr, const FProperty* ReferencingProperty = nullptr);
-
 	/**
 	 * Adds object reference.
 	 *
@@ -2472,6 +2508,36 @@ public:
 		}
 	}
 
+	/**
+	 * Adds all strong property references from a UScriptStruct instance including the struct itself
+	 * 
+	 * Only necessary to handle cases of an unreflected/non-UPROPERTY struct that wants to have references emitted.
+	 *
+	 * Calls AddStructReferencedObjects() but not recursively on nested structs. 
+	 * 
+	 * This and other AddPropertyReferences functions will hopefully merge into a single function in the future.
+	 * They're kept separate initially to maintain exact semantics while replacing the much slower
+	 * SerializeBin/TPropertyValueIterator/GetVerySlowReferenceCollectorArchive paths.
+	 */
+	void AddReferencedObjects(const UScriptStruct*& ScriptStruct, void* Instance, const UObject* ReferencingObject = nullptr, const FProperty* ReferencingProperty = nullptr);
+
+	/** Adds all strong property references from a struct instance, but not the struct itself. Skips AddStructReferencedObjects. */
+	void AddPropertyReferences(const UStruct* Struct, void* Instance, const UObject* ReferencingObject = nullptr);
+	
+	/** Same as AddPropertyReferences but also calls AddStructReferencedObjects on Struct and all nested structs */
+	void AddPropertyReferencesWithStructARO(const UScriptStruct* Struct, void* Instance, const UObject* ReferencingObject = nullptr);
+
+	/** Internal use only. Same as AddPropertyReferences but skips field path and interface properties. Might get removed. */
+	void AddPropertyReferencesLimitedToObjectProperties(const UStruct* Struct, void* Instance, const UObject* ReferencingObject = nullptr);
+
+	/**
+	 * Make Add[OnlyObject]PropertyReference/AddReferencedObjects(UScriptStruct) use AddReferencedObjects(UObject*&) callbacks
+	 * with ReferencingObject and ReferencingProperty context supplied and check for null references before making a callback.
+	 * 
+	 * Return false to use context free AddStableReference callbacks without null checks that avoid sync cache misses when batch processing references.
+	 */
+	virtual bool NeedsPropertyReferencer() const { return true; }
+
 	/**
 	 * If true archetype references should not be added to this collector.
 	 */
@@ -2484,6 +2550,8 @@ public:
 	 * Allows reference elimination by this collector.
 	 */
 	virtual void AllowEliminatingReferences(bool bAllow) {}
+
+
 	/**
 	 * Sets the property that is currently being serialized
 	 */
@@ -2506,9 +2574,13 @@ public:
 	 */
 	virtual bool IsProcessingNativeReferences() const { return true; }
 
+	/** Used by parallel reachability analysis to pre-collect and then exclude some initial FGCObject references */
+	virtual bool NeedsInitialReferences() const { return true; }
+
 	/**
-	* Returns the collector archive associated with this collector.
-	* NOTE THAT COLLECTING REFERENCES THROUGH SERIALIZATION IS VERY SLOW.
+	* Get archive to collect references via SerializeBin / Serialize.
+	*
+	* NOTE: Prefer using AddPropertyReferences or AddReferencedObjects(const UScriptStruct&) instead, they're much faster.
 	*/
 	FReferenceCollectorArchive& GetVerySlowReferenceCollectorArchive()
 	{
@@ -2519,19 +2591,6 @@ public:
 		return *DefaultReferenceCollectorArchive;
 	}
 
-	/**
-	* INTERNAL USE ONLY: returns the persistent frame collector archive associated with this collector.
-	* NOTE THAT COLLECTING REFERENCES THROUGH SERIALIZATION IS VERY SLOW.
-	*/
-	FReferenceCollectorArchive& GetInternalPersistentFrameReferenceCollectorArchive()
-	{
-		if (!PersistentFrameReferenceCollectorArchive)
-		{
-			CreatePersistentFrameReferenceCollectorArchive();
-		}
-		return *PersistentFrameReferenceCollectorArchive;
-	}
-
 protected:
 	/**
 	 * Handle object reference. Called by AddReferencedObject.
@@ -2570,16 +2629,11 @@ protected:
 	virtual void HandleObjectReferences(FObjectPtr* InObjects, const int32 ObjectNum, const UObject* InReferencingObject, const FProperty* InReferencingProperty);
 
 private:
-
-	/** Creates the roxy archive that uses serialization to add objects to this collector */
+	/** Creates the proxy archive that uses serialization to add objects to this collector */
 	void CreateVerySlowReferenceCollectorArchive();
-	/** Creates persistent frame proxy archive that uses serialization to add objects to this collector */
-	void CreatePersistentFrameReferenceCollectorArchive();
 
 	/** Default proxy archive that uses serialization to add objects to this collector */
-	FReferenceCollectorArchive* DefaultReferenceCollectorArchive;
-	/** Persistent frame proxy archive that uses serialization to add objects to this collector */
-	FReferenceCollectorArchive* PersistentFrameReferenceCollectorArchive;
+	TUniquePtr<FReferenceCollectorArchive> DefaultReferenceCollectorArchive;
 };
 
 /**
diff --git a/Engine/Source/Runtime/CoreUObject/Public/UObject/UnrealType.h b/Engine/Source/Runtime/CoreUObject/Public/UObject/UnrealType.h
--- a/Engine/Source/Runtime/CoreUObject/Public/UObject/UnrealType.h
+++ b/Engine/Source/Runtime/CoreUObject/Public/UObject/UnrealType.h
@@ -1080,7 +1080,7 @@ public:
 	 * Emits tokens used by realtime garbage collection code to passed in ReferenceTokenStream. The offset emitted is relative
 	 * to the passed in BaseOffset which is used by e.g. arrays of structs.
 	 */
-	virtual void EmitReferenceInfo(UClass& OwnerClass, int32 BaseOffset, TArray<const FStructProperty*>& EncounteredStructProps, FGCStackSizeHelper& StackSizeHelper);
+	virtual void EmitReferenceInfo(UE::GC::FTokenStreamBuilder& TokenStream, int32 BaseOffset, TArray<const FStructProperty*>& EncounteredStructProps, FGCStackSizeHelper& StackSizeHelper);
 
 	FORCEINLINE int32 GetSize() const
 	{
@@ -2931,7 +2931,7 @@ PRAGMA_ENABLE_DEPRECATION_WARNINGS
 
 	// FProperty interface
 	virtual void SerializeItem(FStructuredArchive::FSlot Slot, void* Value, void const* Defaults) const override;
-	virtual void EmitReferenceInfo(UClass& OwnerClass, int32 BaseOffset, TArray<const FStructProperty*>& EncounteredStructProps, FGCStackSizeHelper& StackSizeHelper) override;
+	virtual void EmitReferenceInfo(UE::GC::FTokenStreamBuilder& TokenStream, int32 BaseOffset, TArray<const FStructProperty*>& EncounteredStructProps, FGCStackSizeHelper& StackSizeHelper) override;
 	virtual const TCHAR* ImportText_Internal(const TCHAR* Buffer, void* ContainerOrPropertyPtr, EPropertyPointerType PropertyPointerType, UObject* OwnerObject, int32 PortFlags, FOutputDevice* ErrorText) const override;
 	virtual EConvertFromTypeResult ConvertFromType(const FPropertyTag& Tag, FStructuredArchive::FSlot Slot, uint8* Data, UStruct* DefaultsStruct) override;
 
@@ -3052,7 +3052,7 @@ PRAGMA_ENABLE_DEPRECATION_WARNINGS
 
 	// FProperty interface
 	virtual void SerializeItem(FStructuredArchive::FSlot Slot, void* Value, void const* Defaults) const override;
-	virtual void EmitReferenceInfo(UClass& OwnerClass, int32 BaseOffset, TArray<const FStructProperty*>& EncounteredStructProps, FGCStackSizeHelper& StackSizeHelper) override;
+	virtual void EmitReferenceInfo(UE::GC::FTokenStreamBuilder& TokenStream, int32 BaseOffset, TArray<const FStructProperty*>& EncounteredStructProps, FGCStackSizeHelper& StackSizeHelper) override;
 protected:
 	virtual void LinkInternal(FArchive& Ar) override;
 private:
@@ -3115,7 +3115,7 @@ PRAGMA_ENABLE_DEPRECATION_WARNINGS
 	virtual FName GetID() const override;
 	virtual bool Identical( const void* A, const void* B, uint32 PortFlags ) const override;
 	virtual void SerializeItem( FStructuredArchive::FSlot Slot, void* Value, void const* Defaults ) const override;
-	virtual void EmitReferenceInfo(UClass& OwnerClass, int32 BaseOffset, TArray<const FStructProperty*>& EncounteredStructProps, FGCStackSizeHelper& StackSizeHelper) override;
+	virtual void EmitReferenceInfo(UE::GC::FTokenStreamBuilder& TokenStream, int32 BaseOffset, TArray<const FStructProperty*>& EncounteredStructProps, FGCStackSizeHelper& StackSizeHelper) override;
 	// End of FProperty interface
 
 	// FObjectProperty interface
@@ -3189,7 +3189,7 @@ protected:
 	virtual const TCHAR* ImportText_Internal(const TCHAR* Buffer, void* ContainerOrPropertyPtr, EPropertyPointerType PropertyPointerType, UObject* OwnerObject, int32 PortFlags, FOutputDevice* ErrorText) const override;
 public:
 	virtual EConvertFromTypeResult ConvertFromType(const FPropertyTag& Tag, FStructuredArchive::FSlot Slot, uint8* Data, UStruct* DefaultsStruct) override;
-	virtual void EmitReferenceInfo(UClass& OwnerClass, int32 BaseOffset, TArray<const FStructProperty*>& EncounteredStructProps, FGCStackSizeHelper& StackSizeHelper) override;
+	virtual void EmitReferenceInfo(UE::GC::FTokenStreamBuilder& TokenStream, int32 BaseOffset, TArray<const FStructProperty*>& EncounteredStructProps, FGCStackSizeHelper& StackSizeHelper) override;
 	// End of FProperty interface
 
 	// FObjectProperty interface
@@ -3500,7 +3500,7 @@ public:
 
 	// UObject interface
 	virtual void Serialize( FArchive& Ar ) override;
-	virtual void EmitReferenceInfo(UClass& OwnerClass, int32 BaseOffset, TArray<const FStructProperty*>& EncounteredStructProps, FGCStackSizeHelper& StackSizeHelper) override;
+	virtual void EmitReferenceInfo(UE::GC::FTokenStreamBuilder& TokenStream, int32 BaseOffset, TArray<const FStructProperty*>& EncounteredStructProps, FGCStackSizeHelper& StackSizeHelper) override;
 	virtual void BeginDestroy() override;
 	virtual void AddReferencedObjects(FReferenceCollector& Collector) override;
 	// End of UObject interface
@@ -3757,7 +3757,7 @@ public:
 	virtual bool PassCPPArgsByRef() const override;
 	virtual void InstanceSubobjects( void* Data, void const* DefaultData, UObject* Owner, struct FObjectInstancingGraph* InstanceGraph ) override;
 	virtual bool ContainsObjectReference(TArray<const FStructProperty*>& EncounteredStructProps, EPropertyObjectReferenceType InReferenceType = EPropertyObjectReferenceType::Strong) const override;
-	virtual void EmitReferenceInfo(UClass& OwnerClass, int32 BaseOffset, TArray<const FStructProperty*>& EncounteredStructProps, FGCStackSizeHelper& StackSizeHelper) override;
+	virtual void EmitReferenceInfo(UE::GC::FTokenStreamBuilder& TokenStream, int32 BaseOffset, TArray<const FStructProperty*>& EncounteredStructProps, FGCStackSizeHelper& StackSizeHelper) override;
 	virtual bool SameType(const FProperty* Other) const override;
 	virtual EConvertFromTypeResult ConvertFromType(const FPropertyTag& Tag, FStructuredArchive::FSlot Slot, uint8* Data, UStruct* DefaultsStruct) override;
 
@@ -3893,7 +3893,7 @@ public:
 	virtual bool PassCPPArgsByRef() const override;
 	virtual void InstanceSubobjects(void* Data, void const* DefaultData, UObject* Owner, struct FObjectInstancingGraph* InstanceGraph) override;
 	virtual bool ContainsObjectReference(TArray<const FStructProperty*>& EncounteredStructProps, EPropertyObjectReferenceType InReferenceType = EPropertyObjectReferenceType::Strong) const override;
-	virtual void EmitReferenceInfo(UClass& OwnerClass, int32 BaseOffset, TArray<const FStructProperty*>& EncounteredStructProps, FGCStackSizeHelper& StackSizeHelper) override;
+	virtual void EmitReferenceInfo(UE::GC::FTokenStreamBuilder& TokenStream, int32 BaseOffset, TArray<const FStructProperty*>& EncounteredStructProps, FGCStackSizeHelper& StackSizeHelper) override;
 	virtual bool SameType(const FProperty* Other) const override;
 	virtual EConvertFromTypeResult ConvertFromType(const FPropertyTag& Tag, FStructuredArchive::FSlot Slot, uint8* Data, UStruct* DefaultsStruct) override;
 	virtual void* GetValueAddressAtIndex_Direct(const FProperty* Inner, void* InValueAddress, int32 Index) const override;
@@ -4017,7 +4017,7 @@ public:
 	virtual bool PassCPPArgsByRef() const override;
 	virtual void InstanceSubobjects(void* Data, void const* DefaultData, UObject* Owner, struct FObjectInstancingGraph* InstanceGraph) override;
 	virtual bool ContainsObjectReference(TArray<const FStructProperty*>& EncounteredStructProps, EPropertyObjectReferenceType InReferenceType = EPropertyObjectReferenceType::Strong) const override;
-	virtual void EmitReferenceInfo(UClass& OwnerClass, int32 BaseOffset, TArray<const FStructProperty*>& EncounteredStructProps, FGCStackSizeHelper& StackSizeHelper) override;
+	virtual void EmitReferenceInfo(UE::GC::FTokenStreamBuilder& TokenStream, int32 BaseOffset, TArray<const FStructProperty*>& EncounteredStructProps, FGCStackSizeHelper& StackSizeHelper) override;
 	virtual bool SameType(const FProperty* Other) const override;
 	virtual EConvertFromTypeResult ConvertFromType(const FPropertyTag& Tag, FStructuredArchive::FSlot Slot, uint8* Data, UStruct* DefaultsStruct) override;
 	virtual void* GetValueAddressAtIndex_Direct(const FProperty* Inner, void* InValueAddress, int32 Index) const override;
@@ -5852,7 +5852,7 @@ public:
 	virtual void InstanceSubobjects( void* Data, void const* DefaultData, UObject* Owner, struct FObjectInstancingGraph* InstanceGraph ) override;
 	virtual int32 GetMinAlignment() const override;
 	virtual bool ContainsObjectReference(TArray<const FStructProperty*>& EncounteredStructProps, EPropertyObjectReferenceType InReferenceType = EPropertyObjectReferenceType::Strong) const override;
-	virtual void EmitReferenceInfo(UClass& OwnerClass, int32 BaseOffset, TArray<const FStructProperty*>& EncounteredStructProps, FGCStackSizeHelper& StackSizeHelper) override;
+	virtual void EmitReferenceInfo(UE::GC::FTokenStreamBuilder& TokenStream, int32 BaseOffset, TArray<const FStructProperty*>& EncounteredStructProps, FGCStackSizeHelper& StackSizeHelper) override;
 	virtual bool SameType(const FProperty* Other) const override;
 	virtual EConvertFromTypeResult ConvertFromType(const FPropertyTag& Tag, FStructuredArchive::FSlot Slot, uint8* Data, UStruct* DefaultsStruct) override;
 #if WITH_EDITORONLY_DATA
@@ -5945,7 +5945,7 @@ protected:
 	virtual const TCHAR* ImportText_Internal(const TCHAR* Buffer, void* ContainerOrPropertyPtr, EPropertyPointerType PropertyPointerType, UObject* OwnerObject, int32 PortFlags, FOutputDevice* ErrorText) const override;
 public:
 	virtual bool ContainsObjectReference(TArray<const FStructProperty*>& EncounteredStructProps, EPropertyObjectReferenceType InReferenceType = EPropertyObjectReferenceType::Strong) const override;
-	virtual void EmitReferenceInfo(UClass& OwnerClass, int32 BaseOffset, TArray<const FStructProperty*>& EncounteredStructProps, FGCStackSizeHelper& StackSizeHelper) override;
+	virtual void EmitReferenceInfo(UE::GC::FTokenStreamBuilder& TokenStream, int32 BaseOffset, TArray<const FStructProperty*>& EncounteredStructProps, FGCStackSizeHelper& StackSizeHelper) override;
 	virtual void InstanceSubobjects( void* Data, void const* DefaultData, UObject* Owner, struct FObjectInstancingGraph* InstanceGraph ) override;
 	virtual bool SameType(const FProperty* Other) const override;
 #if WITH_EDITORONLY_DATA
@@ -6015,7 +6015,7 @@ protected:
 	virtual void ExportText_Internal( FString& ValueStr, const void* PropertyValueOrContainer, EPropertyPointerType PropertyPointerType, const void* DefaultValue, UObject* Parent, int32 PortFlags, UObject* ExportRootScope ) const override;
 public:
 	virtual bool ContainsObjectReference(TArray<const FStructProperty*>& EncounteredStructProps, EPropertyObjectReferenceType InReferenceType = EPropertyObjectReferenceType::Strong) const override;
-	virtual void EmitReferenceInfo(UClass& OwnerClass, int32 BaseOffset, TArray<const FStructProperty*>& EncounteredStructProps, FGCStackSizeHelper& StackSizeHelper) override;
+	virtual void EmitReferenceInfo(UE::GC::FTokenStreamBuilder& TokenStream, int32 BaseOffset, TArray<const FStructProperty*>& EncounteredStructProps, FGCStackSizeHelper& StackSizeHelper) override;
 	virtual void InstanceSubobjects( void* Data, void const* DefaultData, UObject* Owner, struct FObjectInstancingGraph* InstanceGraph ) override;
 	virtual bool SameType(const FProperty* Other) const override;
 	virtual EConvertFromTypeResult ConvertFromType(const FPropertyTag& Tag, FStructuredArchive::FSlot Slot, uint8* Data, UStruct* DefaultsStruct) override;
diff --git a/Engine/Source/Runtime/CoreUObject/Public/UObject/WeakObjectPtr.h b/Engine/Source/Runtime/CoreUObject/Public/UObject/WeakObjectPtr.h
--- a/Engine/Source/Runtime/CoreUObject/Public/UObject/WeakObjectPtr.h
+++ b/Engine/Source/Runtime/CoreUObject/Public/UObject/WeakObjectPtr.h
@@ -10,7 +10,6 @@
 #include "CoreTypes.h"
 #include "Misc/AssertionMacros.h"
 #include "Templates/UnrealTemplate.h"
-#include "UObject/FastReferenceCollectorOptions.h"
 #include "UObject/ScriptDelegates.h"
 #include "UObject/UObjectArray.h"
 
diff --git a/Engine/Source/Runtime/Engine/Classes/Components/ActorComponent.h b/Engine/Source/Runtime/Engine/Classes/Components/ActorComponent.h
--- a/Engine/Source/Runtime/Engine/Classes/Components/ActorComponent.h
+++ b/Engine/Source/Runtime/Engine/Classes/Components/ActorComponent.h
@@ -153,6 +153,8 @@ private:
 	int32 UCSSerializationIndex;
 
 protected:
+	/** If this actor currently has any UCS modified properties */
+	uint8 bHasUCSModifiedProperties:1;
 	/** 
 	 *  Indicates if this ActorComponent is currently registered with a scene. 
 	 */
@@ -421,9 +423,6 @@ private:
 	TArray<FSimpleMemberReference> UCSModifiedProperties_DEPRECATED;
 #endif
 
-	static FRWLock AllUCSModifiedPropertiesLock;
-	static TMap<UActorComponent*, TArray<FSimpleMemberReference>> AllUCSModifiedProperties;
-
 public:
 	/**
 	 * Activates the SceneComponent, should be overridden by native child classes.
diff --git a/Engine/Source/Runtime/Engine/Private/Actor.cpp b/Engine/Source/Runtime/Engine/Private/Actor.cpp
--- a/Engine/Source/Runtime/Engine/Private/Actor.cpp
+++ b/Engine/Source/Runtime/Engine/Private/Actor.cpp
@@ -396,7 +396,7 @@ bool AActor::CanBeInCluster() const
 void AActor::AddReferencedObjects(UObject* InThis, FReferenceCollector& Collector)
 {
 	AActor* This = CastChecked<AActor>(InThis);
-	Collector.AddReferencedObjects(This->OwnedComponents);
+	Collector.AddStableReferenceSet(&This->OwnedComponents);
 #if WITH_EDITOR
 	if (This->CurrentTransactionAnnotation.IsValid())
 	{
diff --git a/Engine/Source/Runtime/Engine/Private/BlueprintGeneratedClass.cpp b/Engine/Source/Runtime/Engine/Private/BlueprintGeneratedClass.cpp
--- a/Engine/Source/Runtime/Engine/Private/BlueprintGeneratedClass.cpp
+++ b/Engine/Source/Runtime/Engine/Private/BlueprintGeneratedClass.cpp
@@ -66,7 +66,6 @@ struct BPGCBreadcrumbsParams
 {
 	UObject& Object;
 	UBlueprintGeneratedClass& BPGC;
-	FArchive& Ar;
 	uint32 ThreadId;
 };
 
@@ -92,15 +91,6 @@ void WriteBPGCBreadcrumbs(FCrashContextExtendedWriter& Writer, const BPGCBreadcr
 		Params.BPGC.GetPathName(nullptr, Builder);
 		Builder.AppendChar(TEXT('\n'));
 
-		if (Params.Ar.GetSerializedProperty() != nullptr)
-		{
-			Params.Ar.GetSerializedProperty()->GetPathName(nullptr, Builder);
-		}
-		else
-		{
-			Builder.Append(TEXT("Serialized property was null!"));
-		}
-
 		TStringBuilder<MAX_FILENAME_SIZE> Filename;
 
 		Filename.Appendf(TEXT("BPGCBreadcrumb_%u"), Params.ThreadId);
@@ -1846,68 +1836,183 @@ void UBlueprintGeneratedClass::Bind()
 	}
 }
 
+#define UE_CHECK_BLUEPRINT_REFERENCES !(UE_BUILD_TEST || UE_BUILD_SHIPPING)
+
+class FPersistentFrameCollector final : public FReferenceCollector
+{
+	FReferenceCollector& InnerCollector;
+#if UE_CHECK_BLUEPRINT_REFERENCES
+	const uint8* PersistentFrameDataAddr;
+	const UObject* Blueprint;
+#endif
+
+public:
+	FPersistentFrameCollector(FReferenceCollector& Collector, uint8* Instance, UObject* InBlueprint)
+	: InnerCollector(Collector)
+#if UE_CHECK_BLUEPRINT_REFERENCES
+	, PersistentFrameDataAddr(Instance)
+	, Blueprint(InBlueprint)
+#endif
+	{}
+
+
+	virtual void HandleObjectReference(UObject*& Object, const UObject* Referencer, const FProperty* ReferencingProperty) override
+	{
+		check(Referencer);
+
+#if UE_CHECK_BLUEPRINT_REFERENCES
+		const bool bIsValidObjectReference = (Object == nullptr || Object->IsValidLowLevelFast());
+		if (!bIsValidObjectReference)
+		{
+			if (const UFunction* UberGraphFunction = Cast<UFunction>(Referencer))
+			{
+				const int32 PersistentFrameDataSize = UberGraphFunction->GetStructureSize();
+
+				FString PersistentFrameDataText;
+				const int32 MaxBytesToDisplayPerLine = 32;
+				PersistentFrameDataText.Reserve(PersistentFrameDataSize * 2 + PersistentFrameDataSize / MaxBytesToDisplayPerLine);
+				for (int32 PersistentFrameDataIdx = 0; PersistentFrameDataIdx < PersistentFrameDataSize; ++PersistentFrameDataIdx)
+				{
+					if (PersistentFrameDataIdx % MaxBytesToDisplayPerLine == 0)
+					{
+						PersistentFrameDataText += TEXT("\n");
+					}
+
+					PersistentFrameDataText += FString::Printf(TEXT("%02x "), PersistentFrameDataAddr[PersistentFrameDataIdx]);
+				}
+
+				UE_LOG(LogUObjectGlobals, Log, TEXT("PersistentFrame: Addr=0x%016llx, Size=%d%s"),
+					(int64)(PTRINT)PersistentFrameDataAddr,
+					PersistentFrameDataSize,
+					*PersistentFrameDataText);
+			}
+		}
+
+		auto GetBlueprintObjectNameLambda = [](const UObject* Referencer) -> FString
+		{
+			if (const UClass* BPGC = Referencer->GetTypedOuter<UClass>())
+			{
+#if WITH_EDITORONLY_DATA
+				if (BPGC->ClassGeneratedBy)
+				{
+					return BPGC->ClassGeneratedBy->GetFullName();
+				}
+#endif
+				return BPGC->GetFullName();
+			}
+
+			return TEXT("NULL");
+		};
+
+		if (!ensureMsgf(bIsValidObjectReference
+			, TEXT("Invalid object referenced by the PersistentFrame: 0x%016llx (Blueprint object: %s, ReferencingProperty: %s, Instance: %s, Address: 0x%016llx) - If you have a reliable repro for this, please contact the development team with it.")
+			, (int64)(PTRINT)Object
+			, *GetBlueprintObjectNameLambda(Referencer)
+			, *ReferencingProperty->GetFullName()
+			, *Blueprint->GetFullName()
+			, (int64)(PTRINT)&Object))
+		{
+			// clear the property value (it's garbage)... the ubergraph-frame
+			// has just lost a reference to whatever it was attempting to hold onto
+			Object = nullptr;
+			return;
+		}
+#endif
+
+		if (Object)
+		{
+			// If the property that serialized us is not an object property we are in some native serializer, we have to treat these as strong
+			if (!Object->HasAnyFlags(RF_StrongRefOnFrame))
+			{
+				if (const FObjectProperty* ObjectProperty = CastField<const FObjectProperty>(ReferencingProperty))
+				{
+					// This was a raw UObject* serialized by FObjectProperty, so just save the address
+					if (InnerCollector.MarkWeakObjectReferenceForClearing(&Object))
+					{
+						return;
+					}
+				}
+			}
+
+			// This is a hard reference or we don't know what's serializing it, so serialize it normally
+			InnerCollector.AddReferencedObject(Object, Referencer, ReferencingProperty);
+		}
+	}
+
+	virtual void HandleObjectReferences(UObject** Objects, int32 Num, const UObject* ReferencingObject, const FProperty* ReferencingProperty)
+	{
+		for (int32 Idx = 0; Idx < Num; ++Idx)
+		{
+			HandleObjectReference(Objects[Idx], ReferencingObject, ReferencingProperty);
+		}
+	}
+
+	virtual bool IsIgnoringArchetypeRef() const override { return InnerCollector.IsIgnoringArchetypeRef(); }
+	virtual bool IsIgnoringTransient() const override { return InnerCollector.IsIgnoringArchetypeRef(); }
+};
+
 void UBlueprintGeneratedClass::AddReferencedObjectsInUbergraphFrame(UObject* InThis, FReferenceCollector& Collector)
 {
 	checkSlow(InThis);
-	for (UClass* CurrentClass = InThis->GetClass(); CurrentClass; CurrentClass = CurrentClass->GetSuperClass())
+	checkSlow(InThis->GetClass());
+	checkSlow(InThis->GetClass() == Cast<UBlueprintGeneratedClass>(InThis->GetClass()));
+
+	UBlueprintGeneratedClass* BPGC = static_cast<UBlueprintGeneratedClass*>(InThis->GetClass());
+	UClass* SuperClass = BPGC->GetSuperClass();
+	
+	while (true)
 	{
-		if (UBlueprintGeneratedClass* BPGC = Cast<UBlueprintGeneratedClass>(CurrentClass))
-		{
 #if USE_UBER_GRAPH_PERSISTENT_FRAME
-			if (BPGC->UberGraphFramePointerProperty)
+		UFunction* UberGraphFunction = BPGC->UberGraphFunction;
+		if (BPGC->UberGraphFramePointerProperty)
+		{
+			FPointerToUberGraphFrame* PointerToUberGraphFrame = BPGC->UberGraphFramePointerProperty->ContainerPtrToValuePtr<FPointerToUberGraphFrame>(InThis);
+			checkSlow(PointerToUberGraphFrame)
+
+				
+			FPersistentFrameCollector ProxyCollector(Collector, PointerToUberGraphFrame->RawPointer, InThis);
+			if (uint8* Instance = PointerToUberGraphFrame->RawPointer)
 			{
-				FPointerToUberGraphFrame* PointerToUberGraphFrame = BPGC->UberGraphFramePointerProperty->ContainerPtrToValuePtr<FPointerToUberGraphFrame>(InThis);
-				checkSlow(PointerToUberGraphFrame)
-				if (PointerToUberGraphFrame->RawPointer)
-				{
 #if VALIDATE_UBER_GRAPH_PERSISTENT_FRAME
-					ensureMsgf(
-						PointerToUberGraphFrame->UberGraphFunctionKey == BPGC->UberGraphFunctionKey,
-						TEXT("Detected key mismatch in uber graph frame for instance %s of type %s, iteration will be unsafe"),
-						*InThis->GetPathName(),
-						*BPGC->GetPathName()
-					);
+				ensureMsgf(
+					PointerToUberGraphFrame->UberGraphFunctionKey == BPGC->UberGraphFunctionKey,
+					TEXT("Detected key mismatch in uber graph frame for instance %s of type %s, iteration will be unsafe"),
+					*InThis->GetPathName(),
+					*BPGC->GetPathName()
+				);
 #endif//VALIDATE_UBER_GRAPH_PERSISTENT_FRAME
 
-					{
-						checkSlow(BPGC->UberGraphFunction);
-						FVerySlowReferenceCollectorArchiveScope CollectorScope(
-							Collector.GetInternalPersistentFrameReferenceCollectorArchive(),
-							BPGC->UberGraphFunction,
-							BPGC->UberGraphFramePointerProperty,
-							InThis,
-							PointerToUberGraphFrame->RawPointer);
-
 #if WITH_ADDITIONAL_CRASH_CONTEXTS
-						BPGCBreadcrumbsParams Params =
-						{
-							*InThis,
-							*BPGC,
-							CollectorScope.GetArchive(),
-							FPlatformTLS::GetCurrentThreadId()
-						};
+				BPGCBreadcrumbsParams Params =
+				{
+					*InThis,
+					*BPGC,
+					FPlatformTLS::GetCurrentThreadId()
+				};
 
-						UE_ADD_CRASH_CONTEXT_SCOPE([&](FCrashContextExtendedWriter& Writer) { WriteBPGCBreadcrumbs(Writer, Params); });
+				UE_ADD_CRASH_CONTEXT_SCOPE([&](FCrashContextExtendedWriter& Writer) { WriteBPGCBreadcrumbs(Writer, Params); });
 #endif // WITH_ADDITIONAL_CRASH_CONTEXTS
-						// All encountered references should be treated as non-native by GC
-						Collector.SetIsProcessingNativeReferences(false);
-						BPGC->UberGraphFunction->SerializeBin(CollectorScope.GetArchive(), PointerToUberGraphFrame->RawPointer);
-						Collector.SetIsProcessingNativeReferences(true);
-					}
+				
+				// All encountered references should be treated as non-native by GC
+				Collector.SetIsProcessingNativeReferences(false);
+				//Collector.AddPersistentFrameReferences(*UberGraphFunction, Instance, InThis);
+				ProxyCollector.AddPropertyReferences(UberGraphFunction, Instance, UberGraphFunction);
+				Collector.SetIsProcessingNativeReferences(true);
 
-				}
 			}
-#endif // USE_UBER_GRAPH_PERSISTENT_FRAME
 		}
-		else if (CurrentClass->HasAllClassFlags(CLASS_Native))
+#endif // USE_UBER_GRAPH_PERSISTENT_FRAME
+
+		if (SuperClass->HasAllClassFlags(CLASS_Native))
 		{
-			CurrentClass->CallAddReferencedObjects(InThis, Collector);
+			checkSlow(Cast<UBlueprintGeneratedClass>(SuperClass) == nullptr);
+			SuperClass->CallAddReferencedObjects(InThis, Collector);
 			break;
 		}
-		else
-		{
-			checkSlow(false);
-		}
+
+		checkSlow(SuperClass == Cast<UBlueprintGeneratedClass>(SuperClass));
+		BPGC = static_cast<UBlueprintGeneratedClass*>(SuperClass);
+		SuperClass = SuperClass->GetSuperClass();
 	}
 }
 
diff --git a/Engine/Source/Runtime/Engine/Private/Components/ActorComponent.cpp b/Engine/Source/Runtime/Engine/Private/Components/ActorComponent.cpp
--- a/Engine/Source/Runtime/Engine/Private/Components/ActorComponent.cpp
+++ b/Engine/Source/Runtime/Engine/Private/Components/ActorComponent.cpp
@@ -314,8 +314,9 @@ FActorComponentGlobalDestroyPhysicsSignature UActorComponent::GlobalDestroyPhysi
 UActorComponent::FOnMarkRenderStateDirty UActorComponent::MarkRenderStateDirtyEvent;
 
 const FString UActorComponent::ComponentTemplateNameSuffix(TEXT("_GEN_VARIABLE"));
-TMap<UActorComponent*, TArray<FSimpleMemberReference>> UActorComponent::AllUCSModifiedProperties;
-FRWLock UActorComponent::AllUCSModifiedPropertiesLock;
+
+static TMap<UActorComponent*, TArray<FSimpleMemberReference, TInlineAllocator<1>>> GAllUCSModifiedProperties;
+static FRWLock GAllUCSModifiedPropertiesLock;
 
 UActorComponent::UActorComponent(const FObjectInitializer& ObjectInitializer /*= FObjectInitializer::Get()*/)
 	: Super(ObjectInitializer)
@@ -447,8 +448,9 @@ void UActorComponent::PostLoad()
 		{
 			if (UCSModifiedProperties_DEPRECATED.Num())
 			{
-				FRWScopeLock Lock(AllUCSModifiedPropertiesLock, SLT_Write);
-				AllUCSModifiedProperties.Add(this, MoveTemp(UCSModifiedProperties_DEPRECATED));
+				FWriteScopeLock WriteLock(GAllUCSModifiedPropertiesLock);
+				GAllUCSModifiedProperties.Add(this) = MoveTemp(UCSModifiedProperties_DEPRECATED);
+				bHasUCSModifiedProperties = true;
 			}
 		}
 #endif
@@ -2319,7 +2321,7 @@ void UActorComponent::DetermineUCSModifiedProperties()
 {
 	if (CreationMethod == EComponentCreationMethod::SimpleConstructionScript)
 	{
-		TArray<FSimpleMemberReference> UCSModifiedProperties;
+		TArray<FSimpleMemberReference, TInlineAllocator<1>> UCSModifiedProperties;
 
 		class FComponentPropertySkipper : public FArchive
 		{
@@ -2368,58 +2370,79 @@ void UActorComponent::DetermineUCSModifiedProperties()
 			}
 		}
 
-		FRWScopeLock Lock(AllUCSModifiedPropertiesLock, SLT_Write);
-		if (UCSModifiedProperties.Num() > 0)
+		bHasUCSModifiedProperties = UCSModifiedProperties.Num() > 0;
+		FWriteScopeLock WriteLock(GAllUCSModifiedPropertiesLock);
+		if (bHasUCSModifiedProperties)
 		{
-			AllUCSModifiedProperties.Add(this, MoveTemp(UCSModifiedProperties));
+			GAllUCSModifiedProperties.Add(this, MoveTemp(UCSModifiedProperties));
 		}
 		else
 		{
-			AllUCSModifiedProperties.Remove(this);
+			GAllUCSModifiedProperties.Remove(this);
 		}
 	}
 }
 
 void UActorComponent::GetUCSModifiedProperties(TSet<const FProperty*>& ModifiedProperties) const
 {
-	FRWScopeLock Lock(AllUCSModifiedPropertiesLock, SLT_ReadOnly);
-	if (TArray<FSimpleMemberReference>* UCSModifiedProperties = AllUCSModifiedProperties.Find(this))
+	if (bHasUCSModifiedProperties)
 	{
-		for (const FSimpleMemberReference& MemberReference : *UCSModifiedProperties)
+		FReadScopeLock ReadLock(GAllUCSModifiedPropertiesLock);
+		if (TArray<FSimpleMemberReference, TInlineAllocator<1>>* UCSModifiedProperties = GAllUCSModifiedProperties.Find(this))
 		{
-			ModifiedProperties.Add(FMemberReference::ResolveSimpleMemberReference<FProperty>(MemberReference));
+			for (const FSimpleMemberReference& MemberReference : *UCSModifiedProperties)
+			{
+				ModifiedProperties.Add(FMemberReference::ResolveSimpleMemberReference<FProperty>(MemberReference));
+			}
 		}
 	}
 }
 
 void UActorComponent::RemoveUCSModifiedProperties(const TArray<FProperty*>& Properties)
 {
-	FRWScopeLock Lock(AllUCSModifiedPropertiesLock, SLT_Write);
-	if (TArray<FSimpleMemberReference>* UCSModifiedProperties = AllUCSModifiedProperties.Find(this))
+	if (bHasUCSModifiedProperties && Properties.Num())
 	{
-		for (FProperty* Property : Properties)
+		FWriteScopeLock WriteLock(GAllUCSModifiedPropertiesLock);
+		if (TArray<FSimpleMemberReference, TInlineAllocator<1>>* UCSModifiedProperties = GAllUCSModifiedProperties.Find(this))
 		{
-			FSimpleMemberReference MemberReference;
-			FMemberReference::FillSimpleMemberReference<FProperty>(Property, MemberReference);
-			UCSModifiedProperties->RemoveSwap(MemberReference);
+			for (FProperty* Property : Properties)
+			{
+				FSimpleMemberReference MemberReference;
+				FMemberReference::FillSimpleMemberReference<FProperty>(Property, MemberReference);
+				UCSModifiedProperties->RemoveSwap(MemberReference);
+			}
+
+			if (UCSModifiedProperties->IsEmpty())
+			{
+				bHasUCSModifiedProperties = false;
+				GAllUCSModifiedProperties.Remove(this);
+			}
 		}
 	}
 }
 
 void UActorComponent::ClearUCSModifiedProperties()
 {
-	FRWScopeLock Lock(AllUCSModifiedPropertiesLock, SLT_Write);
-	AllUCSModifiedProperties.Remove(this);
+	if (bHasUCSModifiedProperties)
+	{
+		bHasUCSModifiedProperties = false;
+		FWriteScopeLock WriteLock(GAllUCSModifiedPropertiesLock);
+		GAllUCSModifiedProperties.Remove(this);
+	}
 }
 
 void UActorComponent::AddReferencedObjects(UObject* InThis, FReferenceCollector& Collector)
 {
-	FRWScopeLock Lock(AllUCSModifiedPropertiesLock, SLT_ReadOnly);
-	if (TArray<FSimpleMemberReference>* UCSModifiedProperties = AllUCSModifiedProperties.Find(CastChecked<UActorComponent>(InThis)))
+	UActorComponent* ThisActor = CastChecked<UActorComponent>(InThis);
+	if (ThisActor->bHasUCSModifiedProperties)
 	{
-		for (FSimpleMemberReference& MemberReference : *UCSModifiedProperties)
+		FReadScopeLock ReadLock(GAllUCSModifiedPropertiesLock);
+		if (TArray<FSimpleMemberReference, TInlineAllocator<1>>* UCSModifiedProperties = GAllUCSModifiedProperties.Find(ThisActor))
 		{
-			Collector.AddReferencedObject(MemberReference.MemberParent);
+			for (FSimpleMemberReference& MemberReference : *UCSModifiedProperties)
+			{
+				Collector.AddStableReference(&MemberReference.MemberParent);
+			}
 		}
 	}
 }
@@ -2467,29 +2490,29 @@ void UActorComponent::Serialize(FArchive& Ar)
 	{
 		if (Ar.IsLoading())
 		{
-			TArray<FSimpleMemberReference> UCSModifiedProperties;
+			TArray<FSimpleMemberReference, TInlineAllocator<1>> UCSModifiedProperties;
 			Ar << UCSModifiedProperties;
 
-			FRWScopeLock Lock(AllUCSModifiedPropertiesLock, SLT_Write);
+			FWriteScopeLock WriteLock(GAllUCSModifiedPropertiesLock);
 			if (UCSModifiedProperties.Num() > 0)
 			{
-				AllUCSModifiedProperties.Add(this, MoveTemp(UCSModifiedProperties));
+				GAllUCSModifiedProperties.Add(this, MoveTemp(UCSModifiedProperties));
 			}
 			else
 			{
-				AllUCSModifiedProperties.Remove(this);
+				GAllUCSModifiedProperties.Remove(this);
 			}
 		}
 		else
 		{
-			FRWScopeLock Lock(AllUCSModifiedPropertiesLock, SLT_ReadOnly);
-			if (TArray<FSimpleMemberReference>* UCSModifiedProperties = AllUCSModifiedProperties.Find(this))
+			FReadScopeLock ReadLock(GAllUCSModifiedPropertiesLock);
+			if (TArray<FSimpleMemberReference, TInlineAllocator<1>>* UCSModifiedProperties = GAllUCSModifiedProperties.Find(this))
 			{
 				Ar << *UCSModifiedProperties;
 			}
 			else
 			{
-				TArray<FSimpleMemberReference> EmptyUCSModifiedProperties;
+				TArray<FSimpleMemberReference, TInlineAllocator<1>> EmptyUCSModifiedProperties;
 				Ar << EmptyUCSModifiedProperties;
 			}
 		}
diff --git a/Engine/Source/Runtime/Engine/Private/Components/StaticMeshComponent.cpp b/Engine/Source/Runtime/Engine/Private/Components/StaticMeshComponent.cpp
--- a/Engine/Source/Runtime/Engine/Private/Components/StaticMeshComponent.cpp
+++ b/Engine/Source/Runtime/Engine/Private/Components/StaticMeshComponent.cpp
@@ -302,13 +302,14 @@ FString UStaticMeshComponent::GetDetailedInfoInternal() const
 void UStaticMeshComponent::AddReferencedObjects(UObject* InThis, FReferenceCollector& Collector)
 {	
 	UStaticMeshComponent* This = CastChecked<UStaticMeshComponent>(InThis);
+	FPlatformMisc::Prefetch(This, offsetof(UStaticMeshComponent, LODData));
 	Super::AddReferencedObjects(This, Collector);
 
-	for (int32 LODIndex = 0; LODIndex < This->LODData.Num(); LODIndex++)
+	for (FStaticMeshComponentLODInfo& LodInfo : This->LODData)
 	{
-		if (This->LODData[LODIndex].OverrideMapBuildData)
+		if (LodInfo.OverrideMapBuildData)
 		{
-			This->LODData[LODIndex].OverrideMapBuildData->AddReferencedObjects(Collector);
+			LodInfo.OverrideMapBuildData->AddReferencedObjects(Collector);
 		}
 	}
 }
diff --git a/Engine/Source/Runtime/Engine/Private/DataTable.cpp b/Engine/Source/Runtime/Engine/Private/DataTable.cpp
--- a/Engine/Source/Runtime/Engine/Private/DataTable.cpp
+++ b/Engine/Source/Runtime/Engine/Private/DataTable.cpp
@@ -34,6 +34,8 @@ namespace
 		}
 	}
 #endif // WITH_EDITORONLY_DATA
+
+	UE_CALL_ONCE(UE::GC::RegisterSlowImplementation, &UDataTable::AddReferencedObjects, UE::GC::EAROFlags::ExtraSlow);
 }
 
 UDataTable::FScopedDataTableChange::FScopedDataTableChange(UDataTable* InTable)
@@ -268,15 +270,11 @@ void UDataTable::AddReferencedObjects(UObject* InThis, FReferenceCollector& Coll
 	if(This->RowStruct != nullptr && This->RowStruct->RefLink != nullptr)
 	{
 		// Now iterate over rows in the map
-		for ( auto RowIt = This->RowMap.CreateIterator(); RowIt; ++RowIt )
+		for (const TPair<FName, uint8*>& Pair : This->RowMap)
 		{
-			uint8* RowData = RowIt.Value();
-
-			if (RowData)
+			if (uint8* RowData = Pair.Value)
 			{
-				FVerySlowReferenceCollectorArchiveScope CollectorScope(Collector.GetVerySlowReferenceCollectorArchive(), This);
-				// Serialize all of the properties to make sure they get in the collector
-				This->RowStruct->SerializeBin(CollectorScope.GetArchive(), RowData);
+				Collector.AddPropertyReferences(This->RowStruct, RowData, This);
 			}
 		}
 	}
diff --git a/Engine/Source/Runtime/Engine/Private/Level.cpp b/Engine/Source/Runtime/Engine/Private/Level.cpp
--- a/Engine/Source/Runtime/Engine/Private/Level.cpp
+++ b/Engine/Source/Runtime/Engine/Private/Level.cpp
@@ -496,11 +496,11 @@ void ULevel::AddReferencedObjects(UObject* InThis, FReferenceCollector& Collecto
 	// Let GC know that we're referencing some AActor objects
 	if (FPlatformProperties::RequiresCookedData() && GActorClusteringEnabled && This->bActorClusterCreated)
 	{
-		Collector.AddReferencedObjects(This->ActorsForGC, This);
+		Collector.AddStableReferenceArray(&This->ActorsForGC);
 	}
 	else
 	{
-		Collector.AddReferencedObjects(This->Actors, This);
+		Collector.AddStableReferenceArray(&This->Actors);
 	}
 
 	Super::AddReferencedObjects( This, Collector );
diff --git a/Engine/Source/Runtime/Engine/Private/LevelActorContainer.cpp b/Engine/Source/Runtime/Engine/Private/LevelActorContainer.cpp
--- a/Engine/Source/Runtime/Engine/Private/LevelActorContainer.cpp
+++ b/Engine/Source/Runtime/Engine/Private/LevelActorContainer.cpp
@@ -73,7 +73,7 @@ public:
 		check(Obj->CanBeInCluster());
 		if (ObjectIndex != ClusterRootIndex && ObjectItem->GetOwnerIndex() == 0 && !GUObjectArray.IsDisregardForGC(Obj) && !Obj->IsRooted())
 		{
-			ObjectsToSerializeStruct.ObjectsToSerialize.Add(Obj);
+			ObjectsToSerializeStruct.ObjectsToSerialize.Add<Options>(Obj);
 			check(!ObjectItem->HasAnyFlags(EInternalObjectFlags::ClusterRoot));
 			ObjectItem->SetOwnerIndex(ClusterRootIndex);
 			Cluster.Objects.Add(ObjectIndex);
@@ -83,7 +83,7 @@ public:
 				UObject* ObjOuter = Obj->GetOuter();
 				if (CanAddToCluster(ObjOuter))
 				{
-					HandleTokenStreamObjectReference(ObjectsToSerializeStruct, Obj, ObjOuter, INDEX_NONE, EGCTokenType::Native, true);
+					HandleTokenStreamObjectReference(ObjectsToSerializeStruct, Obj, ObjOuter, UE::GC::ETokenlessId::Outer, EGCTokenType::Native, true);
 				}
 				else
 				{
@@ -92,9 +92,9 @@ public:
 				if (!Obj->GetClass()->HasAllClassFlags(CLASS_Native))
 				{
 					UObject* ObjectClass = Obj->GetClass();
-					HandleTokenStreamObjectReference(ObjectsToSerializeStruct, Obj, ObjectClass, INDEX_NONE, EGCTokenType::Native, true);
+					HandleTokenStreamObjectReference(ObjectsToSerializeStruct, Obj, ObjectClass, UE::GC::ETokenlessId::Class, EGCTokenType::Native, true);
 					UObject* ObjectClassOuter = Obj->GetClass()->GetOuter();
-					HandleTokenStreamObjectReference(ObjectsToSerializeStruct, Obj, ObjectClassOuter, INDEX_NONE, EGCTokenType::Native, true);
+					HandleTokenStreamObjectReference(ObjectsToSerializeStruct, Obj, ObjectClassOuter, UE::GC::ETokenlessId::ClassOuter, EGCTokenType::Native, true);
 				}
 			}
 		}
@@ -108,7 +108,7 @@ public:
 	* @param TokenIndex Index to the token stream where the reference was found.
 	* @param bAllowReferenceElimination True if reference elimination is allowed (ignored when constructing clusters).
 	*/
-	FORCEINLINE void HandleTokenStreamObjectReference(FGCArrayStruct& ObjectsToSerializeStruct, UObject* ReferencingObject, UObject*& Object, const int32 TokenIndex, const EGCTokenType TokenType, bool bAllowReferenceElimination)
+	FORCEINLINE void HandleTokenStreamObjectReference(FGCArrayStruct& ObjectsToSerializeStruct, UObject* ReferencingObject, UObject*& Object, UE::GC::FTokenId TokenIndex, EGCTokenType TokenType, bool bAllowReferenceElimination)
 	{
 		if (Object)
 		{
@@ -182,19 +182,10 @@ void ULevelActorContainer::CreateCluster()
 
 	// Collect all objects referenced by cluster root and by all objects it's referencing
 	FActorClusterReferenceProcessor Processor(ContainerInternalIndex, Cluster, CastChecked<ULevel>(GetOuter()));
-	TFastReferenceCollector<
-		FActorClusterReferenceProcessor, 
-		TDefaultReferenceCollector<FActorClusterReferenceProcessor>, 
-		FGCArrayPool, 
-		EFastReferenceCollectorOptions::AutogenerateTokenStream | EFastReferenceCollectorOptions::ProcessNoOpTokens
-	> ReferenceCollector(Processor, FGCArrayPool::Get());
+	TArray<UObject*> ObjectsToProcess = {static_cast<UObject*>(this)};
 	FGCArrayStruct ArrayStruct;
-	TArray<UObject*>& ObjectsToProcess = ArrayStruct.ObjectsToSerialize;
-	ObjectsToProcess.Add(static_cast<UObject*>(this));
-	ReferenceCollector.CollectReferences(ArrayStruct);
-#if UE_BUILD_DEBUG
-	FGCArrayPool::Get().CheckLeaks();
-#endif
+	ArrayStruct.SetInitialObjectsUnpadded(ObjectsToProcess);
+	CollectReferences(Processor, ArrayStruct);
 
 	check(RootItem->GetOwnerIndex() == 0);
 	RootItem->SetClusterIndex(ClusterIndex);
diff --git a/Engine/Source/Runtime/Engine/Private/Materials/MaterialCachedData.cpp b/Engine/Source/Runtime/Engine/Private/Materials/MaterialCachedData.cpp
--- a/Engine/Source/Runtime/Engine/Private/Materials/MaterialCachedData.cpp
+++ b/Engine/Source/Runtime/Engine/Private/Materials/MaterialCachedData.cpp
@@ -68,17 +68,17 @@ FMaterialCachedExpressionData::FMaterialCachedExpressionData()
 
 void FMaterialCachedExpressionData::AddReferencedObjects(FReferenceCollector& Collector)
 {
-	Collector.AddReferencedObjects(ReferencedTextures);
-	Collector.AddReferencedObjects(GrassTypes);
-	Collector.AddReferencedObjects(MaterialLayers.Layers);
-	Collector.AddReferencedObjects(MaterialLayers.Blends);
+	Collector.AddStableReferenceArray(&ReferencedTextures);
+	Collector.AddStableReferenceArray(&GrassTypes);
+	Collector.AddStableReferenceArray(&MaterialLayers.Layers);
+	Collector.AddStableReferenceArray(&MaterialLayers.Blends);
 	for (FMaterialFunctionInfo& FunctionInfo : FunctionInfos)
 	{
-		Collector.AddReferencedObject(FunctionInfo.Function);
+		Collector.AddStableReference(&FunctionInfo.Function);
 	}
 	for (FMaterialParameterCollectionInfo& ParameterCollectionInfo : ParameterCollectionInfos)
 	{
-		Collector.AddReferencedObject(ParameterCollectionInfo.ParameterCollection);
+		Collector.AddStableReference(&ParameterCollectionInfo.ParameterCollection);
 	}
 }
 
diff --git a/Engine/Source/Runtime/Engine/Private/Materials/MaterialShared.cpp b/Engine/Source/Runtime/Engine/Private/Materials/MaterialShared.cpp
--- a/Engine/Source/Runtime/Engine/Private/Materials/MaterialShared.cpp
+++ b/Engine/Source/Runtime/Engine/Private/Materials/MaterialShared.cpp
@@ -1262,7 +1262,7 @@ void FMaterial::SetRenderingThreadShaderMap(TRefCountPtr<FMaterialShaderMap>& In
 void FMaterial::AddReferencedObjects(FReferenceCollector& Collector)
 {
 #if WITH_EDITOR
-	Collector.AddReferencedObjects(ErrorExpressions);
+	Collector.AddStableReferenceArray(&ErrorExpressions);
 #endif
 }
 
@@ -1452,8 +1452,8 @@ void FMaterialResource::AddReferencedObjects(FReferenceCollector& Collector)
 {
 	FMaterial::AddReferencedObjects(Collector);
 
-	Collector.AddReferencedObject(Material);
-	Collector.AddReferencedObject(MaterialInstance);
+	Collector.AddStableReference(&Material);
+	Collector.AddStableReference(&MaterialInstance);
 }
 
 bool FMaterialResource::GetAllowDevelopmentShaderCompile()const
