Bulk Data - bulk data resource table in package header

#rnx

diff --git a/Engine/Source/Runtime/Core/Public/UObject/ObjectVersion.h b/Engine/Source/Runtime/Core/Public/UObject/ObjectVersion.h
--- a/Engine/Source/Runtime/Core/Public/UObject/ObjectVersion.h
+++ b/Engine/Source/Runtime/Core/Public/UObject/ObjectVersion.h
@@ -70,6 +70,9 @@ enum class EUnrealEngineObjectUE5Version : uint32
 	// Add a soft object path list to the package summary for fast remap
 	ADD_SOFTOBJECTPATH_LIST,
 
+	// Added bulk/data resource table
+	DATA_RESOURCES,
+
 	// -----<new versions can be added before this line>-------------------------------------------------
 	// - this needs to be the last line (see note below)
 	AUTOMATIC_VERSION_PLUS_ONE,
diff --git a/Engine/Source/Runtime/CoreUObject/Private/UObject/LinkerLoad.cpp b/Engine/Source/Runtime/CoreUObject/Private/UObject/LinkerLoad.cpp
--- a/Engine/Source/Runtime/CoreUObject/Private/UObject/LinkerLoad.cpp
+++ b/Engine/Source/Runtime/CoreUObject/Private/UObject/LinkerLoad.cpp
@@ -816,6 +816,12 @@ FLinkerLoad::ELinkerStatus FLinkerLoad::ProcessPackageSummary(TMap<TPair<FName,
 		SCOPED_LOADTIMER(LinkerLoad_SerializePreloadDependencies);
 		Status = SerializePreloadDependencies();
 	}
+	
+	if (Status == LINKER_Loaded)
+	{
+		SCOPED_LOADTIMER(LinkerLoad_SerializeDataResources);
+		Status = SerializeDataResourceMap();
+	}
 
 	TRACE_LOADTIME_END_PROCESS_SUMMARY;
 
@@ -2339,6 +2345,30 @@ FLinkerLoad::ELinkerStatus FLinkerLoad::SerializePreloadDependencies()
 	return !IsTimeLimitExceeded(TEXT("serialize preload dependencies")) ? LINKER_Loaded : LINKER_TimedOut;
 }
 
+FLinkerLoad::ELinkerStatus FLinkerLoad::SerializeDataResourceMap()
+{
+	DECLARE_SCOPE_CYCLE_COUNTER(TEXT("FLinkerLoad::SerializeDataResourceMap"), STAT_LinkerLoad_SerializeDataResourceMap, STATGROUP_LinkerLoad);
+	
+	TOptional<FStructuredArchive::FSlot> DataResourcesSlot;
+	
+	if (IsTextFormat())
+	{
+		DataResourcesSlot = StructuredArchiveRootRecord->TryEnterField(TEXT("DataResources"), false);
+	}
+	else if (Summary.DataResourceOffset > 0)
+	{
+		Seek(Summary.DataResourceOffset);
+		DataResourcesSlot = StructuredArchiveRootRecord->EnterField(TEXT("DataResources"));
+	}
+
+	if (DataResourcesSlot.IsSet())
+	{
+		FObjectDataResource::Serialize(*DataResourcesSlot, DataResourceMap);
+	}
+
+	return LINKER_Loaded;
+}
+
 /**
  * Serializes thumbnails
  */
diff --git a/Engine/Source/Runtime/CoreUObject/Private/UObject/ObjectResource.cpp b/Engine/Source/Runtime/CoreUObject/Private/UObject/ObjectResource.cpp
--- a/Engine/Source/Runtime/CoreUObject/Private/UObject/ObjectResource.cpp
+++ b/Engine/Source/Runtime/CoreUObject/Private/UObject/ObjectResource.cpp
@@ -362,3 +362,47 @@ void operator<<(FStructuredArchive::FSlot Slot, FObjectImport& I)
 		I.XObject = NULL;
 	}
 }
+
+FArchive& FObjectDataResource::Serialize(FArchive& Ar, TArray<FObjectDataResource>& DataResources)
+{
+	Serialize(FStructuredArchiveFromArchive(Ar).GetSlot(), DataResources);
+	return Ar;
+}
+
+void FObjectDataResource::Serialize(FStructuredArchive::FSlot Slot, TArray<FObjectDataResource>& DataResources)
+{
+	auto SerializeDataResource = [](FStructuredArchive::FSlot Slot, uint32 Version, FObjectDataResource& D)
+	{
+		FStructuredArchive::FRecord Record = Slot.EnterRecord();
+		
+		Record << SA_VALUE(TEXT("Flags"), D.Flags);
+		Record << SA_VALUE(TEXT("SerialOffset"), D.SerialOffset);
+		Record << SA_VALUE(TEXT("DuplicateSerialOffset"), D.DuplicateSerialOffset);
+		Record << SA_VALUE(TEXT("SerialSize"), D.SerialSize);
+		Record << SA_VALUE(TEXT("RawSize"), D.RawSize);
+		Record << SA_VALUE(TEXT("OuterIndex"), D.OuterIndex);
+		Record << SA_VALUE(TEXT("LegacyBulkDataFlags"), D.LegacyBulkDataFlags);
+	};
+
+	FStructuredArchive::FRecord Record = Slot.EnterRecord();
+	uint32 Version = static_cast<uint32>(FObjectDataResource::EVersion::Latest);
+	Record << SA_VALUE(TEXT("Version"), Version);
+	int32 Count = DataResources.Num();
+	Record << SA_VALUE(TEXT("Count"), Count);
+
+	if (Count == 0)
+	{
+		return;
+	}
+
+	if (Slot.GetUnderlyingArchive().IsLoading())
+	{
+		DataResources.SetNum(Count);
+	}
+
+	FStructuredArchive::FStream Stream = Record.EnterStream(TEXT("Resources"));
+	for (FObjectDataResource& DataResource : DataResources)
+	{
+		SerializeDataResource(Stream.EnterElement(), Version, DataResource);
+	}
+}
diff --git a/Engine/Source/Runtime/CoreUObject/Private/UObject/PackageFileSummary.cpp b/Engine/Source/Runtime/CoreUObject/Private/UObject/PackageFileSummary.cpp
--- a/Engine/Source/Runtime/CoreUObject/Private/UObject/PackageFileSummary.cpp
+++ b/Engine/Source/Runtime/CoreUObject/Private/UObject/PackageFileSummary.cpp
@@ -431,6 +431,15 @@ void operator<<(FStructuredArchive::FSlot Slot, FPackageFileSummary& Sum)
 		{
 			Sum.PayloadTocOffset = INDEX_NONE;
 		}
+		
+		if (BaseArchive.IsSaving() || Sum.GetFileVersionUE() >= EUnrealEngineObjectUE5Version::DATA_RESOURCES)
+		{
+			Record << SA_VALUE(TEXT("DataResourceOffset"), Sum.DataResourceOffset);
+		}
+		else
+		{
+			Sum.DataResourceOffset = -1;
+		}
 	}
 }
 
diff --git a/Engine/Source/Runtime/CoreUObject/Public/UObject/Linker.h b/Engine/Source/Runtime/CoreUObject/Public/UObject/Linker.h
--- a/Engine/Source/Runtime/CoreUObject/Public/UObject/Linker.h
+++ b/Engine/Source/Runtime/CoreUObject/Public/UObject/Linker.h
@@ -445,7 +445,10 @@ public:
 
 	/** Gatherable text data contained within this package */
 	TArray<FGatherableTextData> GatherableTextDataMap;
-
+	
+	/** Raw/bulk data references */
+	TArray<FObjectDataResource> DataResourceMap;
+	
 	/** The name of the file for this package */
 	UE_DEPRECATED(5.0, "Use GetDebugName for logging identifiers. For other purposes, use GetPackagePath on LinkerLoad and GetFilename on LinkerSave.")
 	FString					Filename;
diff --git a/Engine/Source/Runtime/CoreUObject/Public/UObject/LinkerLoad.h b/Engine/Source/Runtime/CoreUObject/Public/UObject/LinkerLoad.h
--- a/Engine/Source/Runtime/CoreUObject/Public/UObject/LinkerLoad.h
+++ b/Engine/Source/Runtime/CoreUObject/Public/UObject/LinkerLoad.h
@@ -1134,6 +1134,11 @@ private:
 	 * Serializes the preload dependencies.
 	 */
 	ELinkerStatus SerializePreloadDependencies();
+	
+	/**
+	 * Serializes the data resource map.
+	 */
+	ELinkerStatus SerializeDataResourceMap();
 
 	/** Sets the basic linker archive info */
 	void ResetStatusInfo();
diff --git a/Engine/Source/Runtime/CoreUObject/Public/UObject/ObjectResource.h b/Engine/Source/Runtime/CoreUObject/Public/UObject/ObjectResource.h
--- a/Engine/Source/Runtime/CoreUObject/Public/UObject/ObjectResource.h
+++ b/Engine/Source/Runtime/CoreUObject/Public/UObject/ObjectResource.h
@@ -7,6 +7,7 @@
 #include "CoreTypes.h"
 #include "Misc/AssertionMacros.h"
 #include "Misc/CString.h"
+#include "Misc/EnumClassFlags.h"
 #include "Misc/Guid.h"
 #include "Serialization/StructuredArchive.h"
 #include "Serialization/StructuredArchiveAdapters.h"
@@ -527,3 +528,48 @@ struct FObjectImport : public FObjectResource
 	friend COREUOBJECT_API FArchive& operator<<(FArchive& Ar, FObjectImport& I);
 	friend COREUOBJECT_API void operator<<( FStructuredArchive::FSlot Slot, FObjectImport& I );
 };
+
+/** Data resource flags. */
+enum class EObjectDataResourceFlags : uint32
+{
+	None					= 0,
+	Inline					= (1 << 0),
+	Streaming				= (1 << 1),
+	Optional				= (1 << 2),
+	Duplicate				= (1 << 3),
+	MemoryMapped			= (1 << 4),
+	DerivedDataReference	= (1 << 5),
+};
+ENUM_CLASS_FLAGS(EObjectDataResourceFlags);
+
+/**
+ * UObject binary/bulk data resource type.
+ */
+struct FObjectDataResource
+{
+	enum class EVersion : uint32
+	{
+		Invalid,
+		Initial,
+		LatestPlusOne,
+		Latest = LatestPlusOne - 1
+	};
+
+	/** Data resource flags. */
+	EObjectDataResourceFlags Flags = EObjectDataResourceFlags::None;
+	/** Location of the data in the underlying storage type. */
+	int64 SerialOffset = -1;
+	/** Location of the data in the underlying storage type if this data is duplicated. */
+	int64 DuplicateSerialOffset = -1;
+	/** Number of bytes to serialize when loading/saving the data. */
+	int64 SerialSize = -1;
+	/** Uncompressed size of the data. */
+	int64 RawSize = -1;
+	/** Location of this resource outer/owning object in the export table. */
+	FPackageIndex OuterIndex;
+	/** Bulk data flags. */
+	uint32 LegacyBulkDataFlags = 0;
+	/** I/O functions. */
+	COREUOBJECT_API static FArchive& Serialize(FArchive& Ar, TArray<FObjectDataResource>& DataResources);
+	COREUOBJECT_API static void Serialize(FStructuredArchive::FSlot Slot, TArray<FObjectDataResource>& DataResources);
+};
diff --git a/Engine/Source/Runtime/CoreUObject/Public/UObject/PackageFileSummary.h b/Engine/Source/Runtime/CoreUObject/Public/UObject/PackageFileSummary.h
--- a/Engine/Source/Runtime/CoreUObject/Public/UObject/PackageFileSummary.h
+++ b/Engine/Source/Runtime/CoreUObject/Public/UObject/PackageFileSummary.h
@@ -240,6 +240,11 @@ public:
 	* Location into the file on disk for the payload table of contents data
 	*/
 	int64		PayloadTocOffset;
+	
+	/**
+	* Location into the file of the data resource(s)
+	*/
+	int32		DataResourceOffset;
 
 	/** Constructor */
 	COREUOBJECT_API FPackageFileSummary();
