Bulk Data - data resource table

This change aims to unify how BulkData and DerivedData references are stored and loaded from disk.
The structured part of Bulk Data i.e. Flags,Offset and Size has been moved out from the export data
and into a lookup table in the package header. This allows for bulk data to be rearranged w/o needing to recook but it will
also be used for derived data references to allow loading DDC data both from the runtime.
Currenlty this change ONLY effects cooked output.
Inline data can now be reloaded when loading from I/O store, all consistency checks
when discarding/reloading inline data has been removed.
When cooking, exports are serialized into it's own archive enabling the package header
to be modified after cooking the exports. This also makes it easer to rearrange the exports
in the export bundles for optimal loading performance, i.e. there is no longer a dependency
of the cooked serial offset from the begnning of the .uasset file.

#rnx

diff --git a/Engine/Source/Developer/IoStoreUtilities/Internal/PackageStoreOptimizer.h b/Engine/Source/Developer/IoStoreUtilities/Internal/PackageStoreOptimizer.h
--- a/Engine/Source/Developer/IoStoreUtilities/Internal/PackageStoreOptimizer.h
+++ b/Engine/Source/Developer/IoStoreUtilities/Internal/PackageStoreOptimizer.h
@@ -276,6 +276,7 @@ private:
 	TArray<FPackageId> ImportedPackageIds;
 	TArray<uint64> ImportedPublicExportHashes;
 	TSet<FSHAHash> ShaderMapHashes;
+	TArray<FBulkDataMapEntry> BulkDataEntries;
 
 	FIoBuffer HeaderBuffer;
 
@@ -290,6 +291,7 @@ private:
 	uint64 GraphDataSize = 0;
 	uint64 NameMapSize = 0;
 	uint64 VersioningInfoSize = 0;
+	uint64 BulkDataMapSize = 0;
 	uint32 LoadOrder = 0;
 	
 	TArray<FExportBundleGraphNode> ExportBundleGraphNodes;
@@ -369,6 +371,7 @@ private:
 		TArray<FObjectImport> ObjectImports;
 		TArray<FObjectExport> ObjectExports;
 		TArray<FPackageIndex> PreloadDependencies;
+		TArray<FObjectDataResource> DataResources;
 	};
 
 	struct FPackageStoreHeaderData
@@ -378,6 +381,7 @@ private:
 		TArray<FPackageId> ImportedPackageIds;
 		TArray<uint64> ImportedPublicExportHashes;
 		TArray<FDisplayNameEntryId> NameMap;
+		TArray<FBulkDataMapEntry> BulkDataEntries;
 		TArray<FPackageObjectIndex> Imports; // FH: Imports might need to have more info to be able to resolve export hash with import as outer
 		TArray<FExportMapEntry> Exports;
 		TArray<FExportBundleHeader> ExportBundleHeaders;
@@ -401,6 +405,7 @@ private:
 	void ProcessExports(const FPackageStoreHeaderData& PackageStoreHeaderData, FPackageStorePackage* Package) const;
 	void ProcessPreloadDependencies(const FCookedHeaderData& CookedHeaderData, FPackageStorePackage* Package) const;
 	void ProcessPreloadDependencies(const FPackageStoreHeaderData& PackageStoreHeaderData, FPackageStorePackage* Package) const;
+	void ProcessDataResources(const FCookedHeaderData& CookedHeaderData, FPackageStorePackage* Package) const;
 	TArray<FPackageStorePackage*> SortPackagesInLoadOrder(const TMap<FPackageId, FPackageStorePackage*>& PackagesMap) const;
 	void SerializeGraphData(const TArray<FPackageId>& ImportedPackageIds, FPackageStorePackage::FGraphData& GraphData, FBufferWriter& GraphArchive) const;
 	TArray<FPackageStorePackage::FExportGraphNode*> SortExportGraphNodesInLoadOrder(FPackageStorePackage* Package, FExportGraphEdges& Edges) const;
diff --git a/Engine/Source/Developer/IoStoreUtilities/Private/PackageStoreOptimizer.cpp b/Engine/Source/Developer/IoStoreUtilities/Private/PackageStoreOptimizer.cpp
--- a/Engine/Source/Developer/IoStoreUtilities/Private/PackageStoreOptimizer.cpp
+++ b/Engine/Source/Developer/IoStoreUtilities/Private/PackageStoreOptimizer.cpp
@@ -144,7 +144,8 @@ FPackageStorePackage* FPackageStoreOptimizer::CreatePackageFromCookedHeader(cons
 	ProcessImports(CookedHeaderData, Package, Imports);
 	ProcessExports(CookedHeaderData, Package, Imports.GetData());
 	ProcessPreloadDependencies(CookedHeaderData, Package);
-	
+	ProcessDataResources(CookedHeaderData, Package);
+
 	CreateExportBundles(Package);
 
 	return Package;
@@ -178,6 +179,7 @@ FPackageStorePackage* FPackageStoreOptimizer::CreatePackageFromPackageStoreHeade
 	{
 		Package->NameMapBuilder.AddName(DisplayId);
 	}
+	Package->BulkDataEntries = MoveTemp(PackageStoreHeaderData.BulkDataEntries);
 	ProcessImports(PackageStoreHeaderData, Package);
 	ProcessExports(PackageStoreHeaderData, Package);
 	ProcessPreloadDependencies(PackageStoreHeaderData, Package);
@@ -282,6 +284,12 @@ FPackageStoreOptimizer::FCookedHeaderData FPackageStoreOptimizer::LoadCookedHead
 			ProxyAr << ObjectExport;
 		}
 	}
+	
+	if (Summary.DataResourceOffset > 0)
+	{
+		ProxyAr.Seek(Summary.DataResourceOffset);
+		FObjectDataResource::Serialize(ProxyAr, CookedHeaderData.DataResources);
+	}
 
 	return CookedHeaderData;
 }
@@ -309,6 +317,15 @@ FPackageStoreOptimizer::FPackageStoreHeaderData FPackageStoreOptimizer::LoadPack
 	TArray<FDisplayNameEntryId>& NameMap = PackageStoreHeaderData.NameMap;
 	NameMap = LoadNameBatch(HeaderDataReader);
 
+	const FZenPackageVersioningInfo* VersioningInfo = PackageStoreHeaderData.VersioningInfo.GetPtrOrNull();
+	if (VersioningInfo == nullptr || VersioningInfo->PackageVersion >= EUnrealEngineObjectUE5Version::DATA_RESOURCES)
+	{
+		int64 BulkDataMapSize = 0;
+		HeaderDataReader << BulkDataMapSize;
+		const uint8* BulkDataMapData = HeaderData + sizeof(FZenPackageSummary) + HeaderDataReader.Tell();
+		PackageStoreHeaderData.BulkDataEntries = MakeArrayView(reinterpret_cast<const FBulkDataMapEntry*>(BulkDataMapData), BulkDataMapSize / sizeof(FBulkDataMapEntry));
+	}
+
 	for (FPackageId PackageId : PackageStoreEntry.ImportedPackageIds)
 	{
 		PackageStoreHeaderData.ImportedPackageIds.Add(PackageId);
@@ -1005,6 +1022,20 @@ void FPackageStoreOptimizer::ProcessPreloadDependencies(const FCookedHeaderData&
 	}
 }
 
+void FPackageStoreOptimizer::ProcessDataResources(const FCookedHeaderData& CookedHeaderData, FPackageStorePackage* Package) const
+{
+	for (const FObjectDataResource& DataResource : CookedHeaderData.DataResources)
+	{
+		FBulkDataMapEntry& Entry = Package->BulkDataEntries.AddDefaulted_GetRef();
+		checkf(DataResource.SerialSize == DataResource.RawSize, TEXT("Compressed bulk data is not supported in cooked builds"));
+
+		Entry.SerialOffset = DataResource.SerialOffset;
+		Entry.DuplicateSerialOffset = DataResource.DuplicateSerialOffset;
+		Entry.SerialSize = DataResource.SerialSize;
+		Entry.Flags = DataResource.LegacyBulkDataFlags;
+	}
+}
+
 void FPackageStoreOptimizer::ProcessPreloadDependencies(const FPackageStoreHeaderData& PackageStoreHeaderData, FPackageStorePackage* Package) const
 {
 	TRACE_CPUPROFILER_EVENT_SCOPE(ProcessPreloadDependencies);
@@ -1489,6 +1520,13 @@ void FPackageStoreOptimizer::FinalizePackageHeader(FPackageStorePackage* Package
 		VersioningInfoArchive << Package->VersioningInfo.GetValue();
 		Package->VersioningInfoSize = VersioningInfoArchive.Tell();
 	}
+	
+	FBufferWriter BulkDataMapAr(nullptr, 0, EBufferWriterFlags::AllowResize | EBufferWriterFlags::TakeOwnership);
+	for (FBulkDataMapEntry& Entry : Package->BulkDataEntries)
+	{
+		BulkDataMapAr << Entry;
+	}
+	Package->BulkDataMapSize = BulkDataMapAr.Tell();
 
 	Package->HeaderSize =
 		sizeof(FZenPackageSummary)
@@ -1498,7 +1536,8 @@ void FPackageStoreOptimizer::FinalizePackageHeader(FPackageStorePackage* Package
 		+ Package->ImportMapSize
 		+ Package->ExportMapSize
 		+ Package->ExportBundleEntriesSize
-		+ Package->GraphDataSize;
+		+ Package->GraphDataSize
+		+ Package->BulkDataMapSize + sizeof(int64);
 
 	Package->HeaderBuffer = FIoBuffer(Package->HeaderSize);
 	uint8* HeaderData = Package->HeaderBuffer.Data();
@@ -1522,6 +1561,10 @@ void FPackageStoreOptimizer::FinalizePackageHeader(FPackageStorePackage* Package
 	}
 
 	HeaderArchive.Serialize(NameMapArchive.GetWriterData(), NameMapArchive.Tell());
+
+	HeaderArchive << Package->BulkDataMapSize;
+	HeaderArchive.Serialize(BulkDataMapAr.GetWriterData(), Package->BulkDataMapSize);
+
 	PackageSummary->ImportedPublicExportHashesOffset = HeaderArchive.Tell();
 	HeaderArchive.Serialize(ImportedPublicExportHashesArchive.GetWriterData(), ImportedPublicExportHashesArchive.Tell());
 	PackageSummary->ImportMapOffset = HeaderArchive.Tell();
@@ -2019,4 +2062,4 @@ FIoContainerHeader FPackageStoreOptimizer::CreateContainerHeader(const FIoContai
 FIoContainerHeader FPackageStoreOptimizer::CreateOptionalContainerHeader(const FIoContainerId& ContainerId, TArrayView<const FPackageStoreEntryResource> PackageStoreEntries) const
 {
 	return CreateContainerHeaderInternal(ContainerId, PackageStoreEntries, true);
-}
\ No newline at end of file
+}
diff --git a/Engine/Source/Runtime/Core/Public/Serialization/Archive.h b/Engine/Source/Runtime/Core/Public/Serialization/Archive.h
--- a/Engine/Source/Runtime/Core/Public/Serialization/Archive.h
+++ b/Engine/Source/Runtime/Core/Public/Serialization/Archive.h
@@ -1634,6 +1634,19 @@ public:
 	virtual void DetachBulkData(FBulkData* BulkData, bool bEnsureBulkDataIsLoaded) { }
 	virtual void DetachBulkData(UE::Serialization::FEditorBulkData* BulkData, bool bEnsureBulkDataIsLoaded) {}
 
+	/**
+	 * Serialize bulk data.
+	 *
+	 * @param	BulkData		Bulk data object to serialize
+	 * @param	Params			Serialization parameters
+	 *
+	 * @return true if the bulk data was serialized, false to fallback to default serialization (inline)
+	 */
+	virtual bool SerializeBulkData(FBulkData& BulkData, const struct FBulkDataSerializationParams& Params)
+	{ 
+		return false;
+	}
+
 	/**
 	* Determine if the given archive is a valid "child" of this archive. In general, this means "is exactly the same" but
 	* this function allows a derived archive to support "child" or "internal" archives which are different objects that proxy the
@@ -2132,6 +2145,30 @@ public:
 	};
 #endif
 
+	/** Seeks to and restores the position of an archive. */
+	class FScopeSeekTo
+	{
+	public:
+		FScopeSeekTo(FArchive& InAr, int64 InPos)
+			: Ar(InAr)
+			, SavedPos(InAr.Tell())
+		{
+			Ar.Seek(InPos);
+		}
+		
+		~FScopeSeekTo()
+		{
+			Ar.Seek(SavedPos);
+		}
+
+		FScopeSeekTo(const FScopeSeekTo&) = delete;
+		FScopeSeekTo& operator=(const FScopeSeekTo&) = delete;
+
+	private:
+		FArchive& Ar;
+		int64 SavedPos;
+	};
+
 	/** Called whilst cooking to provide file region hints to the cooker. */
 	virtual void PushFileRegionType(EFileRegionType Type) { }
 	virtual void PopFileRegionType() { }
diff --git a/Engine/Source/Runtime/Core/Public/Serialization/ArchiveProxy.h b/Engine/Source/Runtime/Core/Public/Serialization/ArchiveProxy.h
--- a/Engine/Source/Runtime/Core/Public/Serialization/ArchiveProxy.h
+++ b/Engine/Source/Runtime/Core/Public/Serialization/ArchiveProxy.h
@@ -154,6 +154,11 @@ public:
 		InnerArchive.DetachBulkData(BulkData, bEnsureBulkDataIsLoaded);
 	}
 
+	virtual bool SerializeBulkData(class FBulkData& BulkData, const struct FBulkDataSerializationParams& Params) override
+	{
+		return InnerArchive.SerializeBulkData(BulkData, Params);
+	}
+
 	virtual bool Precache(int64 PrecacheOffset, int64 PrecacheSize) override
 	{
 		return InnerArchive.Precache(PrecacheOffset, PrecacheSize);
@@ -236,6 +241,7 @@ public:
 
 	virtual void SetFilterEditorOnly(bool InFilterEditorOnly) override
 	{
+		FArchive::SetFilterEditorOnly(InFilterEditorOnly);
 		InnerArchive.SetFilterEditorOnly(InFilterEditorOnly);
 	}
 
diff --git a/Engine/Source/Runtime/CoreUObject/Internal/Serialization/FileRegionArchive.h b/Engine/Source/Runtime/CoreUObject/Internal/Serialization/FileRegionArchive.h
new file mode 100644
--- /dev/null
+++ b/Engine/Source/Runtime/CoreUObject/Internal/Serialization/FileRegionArchive.h
@@ -0,0 +1,42 @@
+// Copyright Epic Games, Inc. All Rights Reserved.
+
+#pragma once
+
+#include "Containers/Array.h"
+#include "Serialization/FileRegions.h"
+#include "Serialization/LargeMemoryWriter.h"
+
+template<typename ArchiveType>
+class TFileRegionArchive
+	: public ArchiveType
+{
+public:
+	using ArchiveType::ArchiveType;
+	using ArchiveType::Tell;
+
+	const TArray<FFileRegion>& GetFileRegions() const
+	{
+		return FileRegions;
+	}
+
+	virtual void PushFileRegionType(EFileRegionType Type) override final
+	{
+		check(CurrentRegion.Offset == 0);
+		CurrentRegion.Offset = Tell();
+		CurrentRegion.Type = Type;
+	}
+
+	virtual void PopFileRegionType() override final
+	{
+		check(CurrentRegion.Type != EFileRegionType::None);
+		CurrentRegion.Length = Tell() - CurrentRegion.Offset;
+		FileRegions.Add(CurrentRegion);
+		CurrentRegion = FFileRegion();
+	}
+
+private:
+	TArray<FFileRegion> FileRegions;
+	FFileRegion CurrentRegion;
+};
+
+using FFileRegionMemoryWriter = TFileRegionArchive<FLargeMemoryWriter>;
diff --git a/Engine/Source/Runtime/CoreUObject/Private/Serialization/AsyncLoading.cpp b/Engine/Source/Runtime/CoreUObject/Private/Serialization/AsyncLoading.cpp
--- a/Engine/Source/Runtime/CoreUObject/Private/Serialization/AsyncLoading.cpp
+++ b/Engine/Source/Runtime/CoreUObject/Private/Serialization/AsyncLoading.cpp
@@ -3515,15 +3515,18 @@ void FAsyncPackage::EventDrivenSerializeExport(int32 LocalExportIndex)
 			ACCUM_LOADTIMECOUNT_STAT(StaticGetNativeClassName(Object->GetClass()).ToString());
 			SCOPED_ACCUM_LOADTIME_STAT(StaticGetNativeClassName(Object->GetClass()).ToString());
 			SCOPED_ACCUM_LOADTIME(Serialize, StaticGetNativeClassName(Object->GetClass()));
+			
+			FLinkerExportArchive ExportArchive(*Linker, Export.SerialOffset, Export.SerialSize);
+			ExportArchive.SetFilterEditorOnly(Linker->IsFilterEditorOnly());
 		
 			if (Object->HasAnyFlags(RF_ClassDefaultObject))
 			{
-				Object->GetClass()->SerializeDefaultObject(Object, *Linker);
+				Object->GetClass()->SerializeDefaultObject(Object, ExportArchive);
 			}
 			else
 			{
 				UE_SCOPED_IO_ACTIVITY(*Object->GetName());
-				Object->Serialize(*Linker);
+				Object->Serialize(ExportArchive);
 			}
 		}
 		check(Linker->TemplateForGetArchetypeFromLoader == Template);
@@ -8553,6 +8556,38 @@ void FAsyncArchive::DiscardInlineBufferAndUpdateCurrentPos()
 }
 #endif
 
+FLinkerExportArchive::FLinkerExportArchive(FLinkerLoad& InLinker, int64 InExportSerialOffset, int64 InExportSerialSize)
+	: FArchiveProxy(InLinker)
+	, ExportSerialOffset(InExportSerialOffset)
+	, ExportSerialSize(InExportSerialSize)
+{
+	bExportsCookedToSeparateArchive = InLinker.IsLoadingFromCookedPackage() && InLinker.Summary.GetFileVersionUE() >= EUnrealEngineObjectUE5Version::DATA_RESOURCES;
+}
+
+int64 FLinkerExportArchive::Tell() 
+{
+	const int64 Position = FArchiveProxy::Tell();
+
+	if (bExportsCookedToSeparateArchive)
+	{
+		return Position - ExportSerialOffset;
+	}
+
+	return Position; 
+}
+
+void FLinkerExportArchive::Seek(int64 Position)
+{
+	if (bExportsCookedToSeparateArchive && Position <= ExportSerialSize)
+	{
+		FArchiveProxy::Seek(ExportSerialOffset + Position);
+	}
+	else
+	{
+		FArchiveProxy::Seek(Position);
+	}
+}
+
 #if TRACK_SERIALIZE
 
 #include "Containers/StackTracker.h"
diff --git a/Engine/Source/Runtime/CoreUObject/Private/Serialization/AsyncLoading2.cpp b/Engine/Source/Runtime/CoreUObject/Private/Serialization/AsyncLoading2.cpp
--- a/Engine/Source/Runtime/CoreUObject/Private/Serialization/AsyncLoading2.cpp
+++ b/Engine/Source/Runtime/CoreUObject/Private/Serialization/AsyncLoading2.cpp
@@ -142,6 +142,17 @@ FArchive& operator<<(FArchive& Ar, FExportMapEntry& ExportMapEntry)
 
 	return Ar;
 }
+	
+FArchive& operator<<(FArchive& Ar, FBulkDataMapEntry& BulkDataEntry)
+{
+	Ar << BulkDataEntry.SerialOffset;
+	Ar << BulkDataEntry.DuplicateSerialOffset;
+	Ar << BulkDataEntry.SerialSize;
+	Ar << BulkDataEntry.Flags;
+	Ar << BulkDataEntry.Pad;
+
+	return Ar;
+}
 
 uint64 FPackageObjectIndex::GenerateImportHashFromObjectPath(const FStringView& ObjectPath)
 {
@@ -1344,6 +1355,7 @@ struct FAsyncPackageHeaderData
 	TArrayView<const FPackageObjectIndex> ImportMap;
 	TArrayView<const FExportMapEntry> ExportMap;
 	TArrayView<const uint8> ArcsData;
+	TArrayView<const FBulkDataMapEntry> BulkDataMap;
 	// Backed by allocation in FAsyncPackageData
 	TArrayView<FPackageId> ImportedPackageIds;
 	TArrayView<FName> ImportedPackageNames;
@@ -1357,6 +1369,7 @@ struct FAsyncPackageHeaderData
 		ImportMap = TArrayView<const FPackageObjectIndex>();
 		ExportMap = TArrayView<const FExportMapEntry>();
 		ArcsData = TArrayView<const uint8>();
+		BulkDataMap = TArrayView<const FBulkDataMapEntry>();
 	}
 };
 
@@ -1506,6 +1519,7 @@ class FExportArchive final : public FArchive
 {
 public:
 	FExportArchive(const uint8* AllExportDataPtr, const uint8* CurrentExportPtr, uint64 AllExportDataSize)
+		: IoDispatcher(FIoDispatcher::Get())
 	{
 #if (!DEVIRTUALIZE_FLinkerLoad_Serialize)
 		ActiveFPLB = &InlineFPLB;
@@ -1564,18 +1578,32 @@ public:
 
 	virtual int64 Tell() override
 	{
-		int64 CookedFilePosition = (ActiveFPLB->StartFastPathLoadBuffer - ActiveFPLB->OriginalFastPathLoadBuffer);
-		CookedFilePosition -= BufferSerialOffset;
-		CookedFilePosition += CookedSerialOffset;
-		return CookedFilePosition;
+		if (bExportsCookedToSeparateArchive)
+		{
+			return (ActiveFPLB->StartFastPathLoadBuffer - ActiveFPLB->OriginalFastPathLoadBuffer) - BufferSerialOffset;
+		}
+		else
+		{
+			int64 CookedFilePosition = (ActiveFPLB->StartFastPathLoadBuffer - ActiveFPLB->OriginalFastPathLoadBuffer);
+			CookedFilePosition -= BufferSerialOffset;
+			CookedFilePosition += CookedSerialOffset;
+			return CookedFilePosition;
+		}
 	}
 
 	virtual void Seek(int64 Position) override
 	{
-		uint64 BufferPosition = (uint64)Position;
-		BufferPosition -= CookedSerialOffset;
-		BufferPosition += BufferSerialOffset;
-		ActiveFPLB->StartFastPathLoadBuffer = ActiveFPLB->OriginalFastPathLoadBuffer + BufferPosition;
+		if (bExportsCookedToSeparateArchive)
+		{
+			ActiveFPLB->StartFastPathLoadBuffer = ActiveFPLB->OriginalFastPathLoadBuffer + BufferSerialOffset + Position;
+		}
+		else
+		{
+			uint64 BufferPosition = (uint64)Position;
+			BufferPosition -= CookedSerialOffset;
+			BufferPosition += BufferSerialOffset;
+			ActiveFPLB->StartFastPathLoadBuffer = ActiveFPLB->OriginalFastPathLoadBuffer + BufferPosition;
+		}
 		CheckBufferPosition(TEXT("InvalidSeek"));
 	}
 
@@ -1728,10 +1756,107 @@ public:
 		}
 		return *this;
 	}
+
+	inline virtual bool SerializeBulkData(FBulkData& BulkData, const FBulkDataSerializationParams& Params) override
+	{
+		const FPackageId& PackageId = PackageDesc->PackageIdToLoad;
+		const uint16 ChunkIndex = bIsOptionalSegment ? 1 : 0; 
+
+		UE::BulkData::Private::FBulkMetaData& Meta = BulkData.BulkMeta;
+		int64 DuplicateSerialOffset = -1;
+		SerializeBulkMeta(Meta, DuplicateSerialOffset, Params.ElementSize);
+
+		const bool bIsInline = Meta.HasAnyFlags(BULKDATA_PayloadAtEndOfFile) == false;
+		if (bIsInline)
+		{
+			if (int64 PayloadSize = Meta.GetSize(); PayloadSize > 0 && Meta.HasAnyFlags(BULKDATA_Unused) == false)
+			{
+				// Make inline bulk data reloadable
+				const int64 ExportBundleOffset = (ActiveFPLB->StartFastPathLoadBuffer - ActiveFPLB->OriginalFastPathLoadBuffer);
+				Meta.SetOffset(ExportBundleOffset);
+				BulkData.BulkChunkId = CreateIoChunkId(PackageId.Value(), ChunkIndex, EIoChunkType::ExportBundleData);
+				Serialize(BulkData.ReallocateData(PayloadSize), PayloadSize);
+			}
+		}
+		else
+		{
+			const EIoChunkType ChunkType = Meta.HasAnyFlags(BULKDATA_OptionalPayload) ? EIoChunkType::OptionalBulkData : EIoChunkType::BulkData;
+			FIoChunkId ChunkId = CreateIoChunkId(PackageId.Value(), ChunkIndex, ChunkType);
+
+			if (Meta.HasAnyFlags(BULKDATA_DuplicateNonOptionalPayload))
+			{
+				const FIoChunkId OptionalChunkId = CreateIoChunkId(PackageId.Value(), ChunkIndex, EIoChunkType::OptionalBulkData);
+				if (IoDispatcher.DoesChunkExist(OptionalChunkId))
+				{
+					ChunkId = OptionalChunkId;
+					Meta.ClearFlags(BULKDATA_DuplicateNonOptionalPayload);
+					Meta.AddFlags(BULKDATA_OptionalPayload);
+					Meta.SetOffset(DuplicateSerialOffset);
+				}
+			}
+			else if (Meta.HasAnyFlags(BULKDATA_MemoryMappedPayload) && Params.bAttemptMemoryMapping)
+			{
+				const FIoChunkId MappedChunkId = CreateIoChunkId(PackageId.Value(), ChunkIndex, EIoChunkType::MemoryMappedBulkData);
+				TIoStatusOr<FIoMappedRegion> Status = IoDispatcher.OpenMapped(MappedChunkId, FIoReadOptions(Meta.GetOffset(), Meta.GetSize()));
+
+				if (Status.IsOk())
+				{
+					ChunkId = MappedChunkId;
+					FIoMappedRegion Mapping = Status.ConsumeValueOrDie(); 
+					BulkData.DataAllocation.SetMemoryMappedData(&BulkData, Mapping.MappedFileHandle, Mapping.MappedFileRegion);
+				}
+				else
+				{
+					UE_LOG(LogSerialization, Warning, TEXT("Memory map bulk data from chunk '%s', offset '%lld', size '%lld' FAILED"),
+						*LexToString(MappedChunkId), Meta.GetOffset(), Meta.GetSize());
+					
+					BulkData.BulkChunkId = ChunkId;
+					BulkData.ForceBulkDataResident();
+				}
+			}
+
+			BulkData.BulkChunkId = ChunkId;
+		}
+
+		return true;
+	}
 	//~ End FArchive::FLinkerLoad Interface
 
 private:
+	inline void SerializeBulkMeta(UE::BulkData::Private::FBulkMetaData& Meta, int64& DuplicateSerialOffset, int32 ElementSize)
+	{
+		using namespace UE::BulkData::Private;
+		FArchive& Ar = *this;
+
+		if (UNLIKELY(HeaderData->BulkDataMap.IsEmpty()))
+		{
+			FBulkMetaResource SerializedMeta;
+			Ar << SerializedMeta;
+			Meta = FBulkMetaData::FromSerialized(SerializedMeta, ElementSize);
+			DuplicateSerialOffset = SerializedMeta.DuplicateOffset;
+		}
+		else
+		{
+			int32 EntryIndex = INDEX_NONE;
+			Ar << EntryIndex;
+			const FBulkDataMapEntry& Entry = HeaderData->BulkDataMap[EntryIndex];
+			Meta.SetFlags(static_cast<EBulkDataFlags>(Entry.Flags));
+			Meta.SetOffset(Entry.SerialOffset);
+			Meta.SetSize(Entry.SerialSize);
+			DuplicateSerialOffset = Entry.DuplicateSerialOffset;
+		}
+
+		Meta.AddFlags(static_cast<EBulkDataFlags>(BULKDATA_UsesIoDispatcher | BULKDATA_LazyLoadable));
+#if WITH_EDITOR
+		if (GIsEditor)
+		{
+			Meta.ClearFlags(BULKDATA_SingleUse);
+		}
+#endif
+	}
+
 	friend FAsyncPackage2;
+	FIoDispatcher& IoDispatcher;
 #if (!DEVIRTUALIZE_FLinkerLoad_Serialize)
 	FArchive::FFastPathLoadBuffer InlineFPLB;
 	FArchive::FFastPathLoadBuffer* ActiveFPLB;
@@ -1749,6 +1874,8 @@ private:
 	uint64 CookedSerialOffset = 0;
 	uint64 CookedSerialSize = 0;
 	uint64 BufferSerialOffset = 0;
+	bool bIsOptionalSegment = false;
+	bool bExportsCookedToSeparateArchive = false;
 
 	void FixupSoftObjectPathForInstancedPackage(FSoftObjectPath& InOutSoftObjectPath)
 	{
@@ -4721,6 +4848,15 @@ static void ReadAsyncPackageHeader(FAsyncPackageSerializationState& Serializatio
 		HeaderData.NameMap.Load(PackageHeaderDataReader, FMappedName::EType::Package);
 	}
 	HeaderData.PackageName = HeaderData.NameMap.GetName(PackageSummary->Name);
+	
+	const FZenPackageVersioningInfo* VersioningInfo = HeaderData.VersioningInfo.GetPtrOrNull();
+	if (VersioningInfo == nullptr || VersioningInfo->PackageVersion >= EUnrealEngineObjectUE5Version::DATA_RESOURCES)
+	{
+		int64 BulkDataMapSize = 0;
+		PackageHeaderDataReader << BulkDataMapSize;
+		const uint8* BulkDataMapData = PackageHeaderDataPtr + sizeof(FZenPackageSummary) + PackageHeaderDataReader.Tell();
+		HeaderData.BulkDataMap = MakeArrayView(reinterpret_cast<const FBulkDataMapEntry*>(BulkDataMapData), BulkDataMapSize / sizeof(FBulkDataMapEntry));
+	}
 
 	HeaderData.CookedHeaderSize = PackageSummary->CookedHeaderSize;
 	HeaderData.ImportedPublicExportHashes = TArrayView<const uint64>(
@@ -5400,12 +5536,15 @@ EEventLoadNodeExecutionResult FAsyncPackage2::Event_ProcessExportBundle(FAsyncLo
 	
 	if (!Package->bLoadHasFailed)
 	{
+		bool bIsOptionalSegment = false;
 #if WITH_EDITOR
 		const FAsyncPackageHeaderData* HeaderData;
 		const FExportBundleHeader* ExportBundle;
 		FAsyncPackageSerializationState* SerializationState;
 		TArrayView<FExportObject> Exports = Package->Data.Exports;
-		if (InExportBundleIndex >= Package->HeaderData.ExportBundleHeaders.Num())
+		bIsOptionalSegment = InExportBundleIndex >= Package->HeaderData.ExportBundleHeaders.Num();
+
+		if (bIsOptionalSegment)
 		{
 			HeaderData = Package->OptionalSegmentHeaderData.GetPtrOrNull();
 			check(HeaderData);
@@ -5459,6 +5598,8 @@ EEventLoadNodeExecutionResult FAsyncPackage2::Event_ProcessExportBundle(FAsyncLo
 			Ar.Exports = Exports;
 			Ar.ExternalReadDependencies = &Package->ExternalReadDependencies;
 			Ar.InstanceContext = Package->InstanceContext.Get();
+			Ar.bIsOptionalSegment = bIsOptionalSegment;
+			Ar.bExportsCookedToSeparateArchive = Ar.UEVer() >= EUnrealEngineObjectUE5Version::DATA_RESOURCES;			
 		}
 
 		while (Package->ExportBundleEntryIndex < int32(ExportBundle->EntryCount))
diff --git a/Engine/Source/Runtime/CoreUObject/Private/Serialization/AsyncLoadingPrivate.h b/Engine/Source/Runtime/CoreUObject/Private/Serialization/AsyncLoadingPrivate.h
--- a/Engine/Source/Runtime/CoreUObject/Private/Serialization/AsyncLoadingPrivate.h
+++ b/Engine/Source/Runtime/CoreUObject/Private/Serialization/AsyncLoadingPrivate.h
@@ -3,6 +3,7 @@
 #pragma once
 
 #include "CoreMinimal.h"
+#include "Serialization/ArchiveProxy.h"
 #include "UObject/Linker.h"
 #include "Async/AsyncFileHandle.h"
 
@@ -149,3 +150,24 @@ private:
 
 	FLinkerLoad* OwnerLinker;
 };
+
+/**
+ * The Linker export archive converts Export relative offsets
+ * to file relative offsets when exports are cooked to separate
+ * archives.
+ */
+class COREUOBJECT_API FLinkerExportArchive final
+	: public FArchiveProxy
+{
+public:
+	FLinkerExportArchive(class FLinkerLoad& InLinker, int64 InExportSerialOffset, int64 InExportSerialSize);
+
+	virtual int64 Tell() override;
+	virtual void Seek(int64 Position) override;
+
+private:
+	int64 ExportSerialOffset;
+	int64 ExportSerialSize;
+	bool bExportsCookedToSeparateArchive;
+};
+
diff --git a/Engine/Source/Runtime/CoreUObject/Private/Serialization/BulkData.cpp b/Engine/Source/Runtime/CoreUObject/Private/Serialization/BulkData.cpp
--- a/Engine/Source/Runtime/CoreUObject/Private/Serialization/BulkData.cpp
+++ b/Engine/Source/Runtime/CoreUObject/Private/Serialization/BulkData.cpp
@@ -136,6 +136,11 @@ static FArchive& SerializeAsInt32(FArchive& Ar, int64& Value)
 
 FArchive& operator<<(FArchive& Ar, FBulkMetaResource& BulkMeta)
 {
+	if (Ar.IsSaving() && BulkMeta.SizeOnDisk >= (1LL << 31))
+	{
+		FBulkData::SetBulkDataFlagsOn(BulkMeta.Flags, BULKDATA_Size64Bit);
+	}
+
 	Ar << BulkMeta.Flags;
 
 	if (UNLIKELY(BulkMeta.Flags & BULKDATA_Size64Bit))
@@ -188,7 +193,7 @@ FArchive& operator<<(FArchive& Ar, FBulkMetaResource& BulkMeta)
 	return Ar;
 }
 
-} // namespace UE::BulkData
+} // namespace UE::BulkData::Private
 
 /*-----------------------------------------------------------------------------
 	Memory managament
@@ -607,12 +612,6 @@ bool FBulkData::CanLoadFromDisk() const
 #if WITH_EDITOR
 	return AttachedAr != nullptr || BulkChunkId.IsValid();
 #else
-#if UE_KEEP_INLINE_RELOADING_CONSISTENT
-	if (IsInlined())
-	{
-		return false;
-	}
-#endif //UE_KEEP_INLINE_RELOADING_CONSISTENT
 	return BulkChunkId.IsValid();
 #endif // WITH_EDITOR
 }
@@ -1107,6 +1106,11 @@ void FBulkData::ClearBulkDataFlagsOn(EBulkDataFlags& InOutAccumulator, EBulkData
 	InOutAccumulator = static_cast<EBulkDataFlags>(InOutAccumulator & ~FlagsToClear);
 }
 
+bool FBulkData::HasFlags(EBulkDataFlags Flags, EBulkDataFlags Contains)
+{
+	return (uint32(Flags) & uint32(Contains)) == Contains; 
+}
+
 void FBulkData::Serialize(FArchive& Ar, UObject* Owner, bool bAttemptFileMapping, int32 ElementSize, EFileRegionType FileRegionType)
 {
 	DECLARE_SCOPE_CYCLE_COUNTER(TEXT("FBulkData::Serialize"), STAT_UBD_Serialize, STATGROUP_Memory);
@@ -1117,6 +1121,12 @@ void FBulkData::Serialize(FArchive& Ar, UObject* Owner, bool bAttemptFileMapping
 	
 	check(!bAttemptFileMapping || Ar.IsLoading()); // makes no sense to map unless we are loading
 
+	if (Ar.SerializeBulkData(*this, FBulkDataSerializationParams {Owner, ElementSize, FileRegionType, bAttemptFileMapping}))
+	{
+		// Just early out when the archive overrides the serialization of bulk data
+		return;
+	}
+
 #if !USE_RUNTIME_BULKDATA
 	if (Ar.IsTransacting())
 	{
@@ -1178,356 +1188,49 @@ void FBulkData::Serialize(FArchive& Ar, UObject* Owner, bool bAttemptFileMapping
 #endif // !USE_RUNTIME_BULKDATA
 	if (Ar.IsPersistent() && !Ar.IsObjectReferenceCollector() && !Ar.ShouldSkipBulkData())
 	{
+			using namespace UE::BulkData::Private;
 #if TRACK_BULKDATA_USE
 		FThreadSafeBulkDataToObjectMap::Get().Add( this, Owner );
 #endif
 		if (Ar.IsLoading())
 		{
-			using namespace UE::BulkData::Private;
 			checkf(IsUnlocked(), TEXT("Serialize bulk data FAILED, bulk data is locked"));
 
-			UPackage* Package					= Owner ? Owner->GetPackage() : nullptr;
-			FArchive* CacheableAr				= Ar.GetCacheableArchive();
-			const bool bCanLazyLoad				= CacheableAr != nullptr && Ar.IsAllowingLazyLoading();
-			FLinkerLoad* LinkerLoad				= nullptr;
-
 			FBulkMetaResource MetaResource;
 			Ar << MetaResource;
-#if WITH_EDITOR
-			if (Owner == nullptr)
-			{
-				// Temp fix for accidentially uploading bulk data with wrong flags to DDC
-				ClearBulkDataFlagsOn(MetaResource.Flags, static_cast<EBulkDataFlags>(BULKDATA_PayloadAtEndOfFile | BULKDATA_PayloadInSeperateFile | BULKDATA_WorkspaceDomainPayload));
-			}
-#endif
 			BulkMeta = FBulkMetaData::FromSerialized(MetaResource, ElementSize);
+			BulkChunkId = FIoChunkId::InvalidChunkId;
 
-			bool bIsUsingIoDispatcher = false;
-			if (Ar.IsLoadingFromCookedPackage())
-			{
-				BulkMeta.SetMetaFlags(FBulkMetaData::EMetaFlags::CookedPackage);
-				bIsUsingIoDispatcher = Package && Package->GetPackageId().IsValid();
-			}
-
-			check(MetaResource.ElementCount <= 0 || BulkMeta.GetSize() == MetaResource.ElementCount * ElementSize);
-#if !USE_RUNTIME_BULKDATA
-			check(MetaResource.ElementCount <= 0 || BulkMeta.GetSizeOnDisk() == MetaResource.SizeOnDisk);
-#endif
-			check(BulkMeta.GetOffset() == MetaResource.Offset);
-			check(BulkMeta.GetFlags() == MetaResource.Flags);
-
+			const int64 BulkDataSize = GetBulkDataSize();
+			SerializeBulkData(Ar, ReallocateData(BulkDataSize), BulkDataSize, BulkMeta.GetFlags());
+#if WITH_EDITOR
 			if (GIsEditor)
 			{
 				ClearBulkDataFlags(BULKDATA_SingleUse);
 			}
-
-			if (bIsUsingIoDispatcher)
-			{
-				if (IsInlined() == false)
-				{
-					check(IsInSeparateFile());
-					check(Package != nullptr);
-					BulkMeta.SetFlags(EBulkDataFlags(BulkMeta.GetFlags() | BULKDATA_UsesIoDispatcher));
-					BulkChunkId = CreateBulkDataIoChunkId(BulkMeta, Owner);
-				}
-			}
-			else
-			{
-				if (Owner)
-				{
-					if (LinkerLoad = Owner->GetLinker(); LinkerLoad == nullptr)
-					{
-						LinkerLoad = FLinkerLoad::FindExistingLinkerForPackage(Package);
-					}
-
-					if (LinkerLoad != nullptr)
-					{
-						BulkChunkId = CreateBulkDataIoChunkId(BulkMeta, Owner);
-					}
-#if WITH_EDITOR
-					Linker = LinkerLoad;
-					if (bCanLazyLoad)
-					{
-						check(CacheableAr->IsTextFormat() == false);
-						AttachedAr = CacheableAr;
-						AttachedAr->AttachBulkData(Owner, this);
-					}
-#endif // WITH_EDITOR
-				}
-			}
-
-			if (IsInlined())
-			{
-				checkf(bAttemptFileMapping == false || bIsUsingIoDispatcher == false, TEXT("Trying to memory map inline bulk data '%s' which is not supported when loading from I/O store"), *Owner->GetPackage()->GetFName().ToString());
-
-				const int64 BulkDataSize = GetBulkDataSize();
-				void* DataBuffer = ReallocateData(BulkDataSize);
-				SerializeBulkData(Ar, DataBuffer, BulkDataSize, BulkMeta.GetFlags());
-
-				ConditionalSetInlineAlwaysAllowDiscard(bIsUsingIoDispatcher);
-			}
-			else
-			{
-				if (NeedsOffsetFixup())
-				{
-					check(LinkerLoad);
-					check(bIsUsingIoDispatcher == false);
-					BulkMeta.SetOffset(BulkMeta.GetOffset() + LinkerLoad->Summary.BulkDataStartOffset);
-				}
-
-				if (IsInSeparateFile())
-				{
-					SetBulkDataFlags(BULKDATA_LazyLoadable);
-
-					if (IsDuplicateNonOptional())
-					{
-						const EBulkDataFlags OptionalFlags = static_cast<EBulkDataFlags>(BULKDATA_PayloadAtEndOfFile | BULKDATA_PayloadInSeperateFile | BULKDATA_OptionalPayload);
-						const FIoChunkId OptionalChunkId = CreateBulkDataIoChunkId(FBulkMetaData(OptionalFlags, BulkMeta.GetMetaFlags()), Owner);
-
-						if (DoesBulkDataExist(OptionalChunkId))
-						{
-							BulkMeta.SetFlags(static_cast<EBulkDataFlags>((MetaResource.DuplicateFlags & ~BULKDATA_DuplicateNonOptionalPayload) | BULKDATA_OptionalPayload | BULKDATA_PayloadInSeperateFile | BULKDATA_PayloadAtEndOfFile));
-							BulkMeta.SetOffset(MetaResource.DuplicateOffset);
-							BulkMeta.SetSizeOnDisk(MetaResource.DuplicateSizeOnDisk);
-							BulkChunkId = OptionalChunkId;
-
-							if (bIsUsingIoDispatcher)
-							{
-								BulkMeta.SetFlags(EBulkDataFlags(BulkMeta.GetFlags() | BULKDATA_UsesIoDispatcher));
-							}
-
-							if (NeedsOffsetFixup())
-							{
-								checkf(bIsUsingIoDispatcher == false, TEXT("Loading bulk data from I/O store doesn't require offset fixup, use flag BULKDATA_NoOffsetFixUp when cooking"));
-								checkf(LinkerLoad != nullptr, TEXT("BulkData needs its offset fixed on load but no linker found"));
-								BulkMeta.SetOffset(BulkMeta.GetOffset() + LinkerLoad->Summary.BulkDataStartOffset);
-							}
-						}
-					}
-
-					bool bFileMappingFailed = false;
-					if (bAttemptFileMapping && (bCanLazyLoad || bIsUsingIoDispatcher))
-					{
-						const EBulkDataFlags MemoryMappedFlags = static_cast<EBulkDataFlags>(BULKDATA_PayloadAtEndOfFile | BULKDATA_PayloadInSeperateFile | BULKDATA_MemoryMappedPayload);
-						const FIoChunkId MmapChunkId = CreateBulkDataIoChunkId(FBulkMetaData(MemoryMappedFlags, BulkMeta.GetMetaFlags()), Owner);
-
-						FIoMappedRegion MappedRegion;
-						if (TryMemoryMapBulkData(BulkMeta, MmapChunkId, BulkMeta.GetOffset(), BulkMeta.GetSize(), MappedRegion))
-						{
-							DataAllocation.SetMemoryMappedData(this, MappedRegion.MappedFileHandle, MappedRegion.MappedFileRegion);
-							BulkChunkId = MmapChunkId;
-						}
-						else
-						{
-							bFileMappingFailed = true;
-						}
-					}
-
-					if (bFileMappingFailed || (bCanLazyLoad == false && bIsUsingIoDispatcher == false))
-					{
-						UE_CLOG(bFileMappingFailed, LogSerialization, Warning, TEXT("Memory map bulk data '%s' FAILED"), *LexToString(BulkChunkId));
-						ForceBulkDataResident();
-					}
-				}
-				else
-				{
-					check(bIsUsingIoDispatcher == false);
-					check(Ar.IsLoadingFromCookedPackage() == false);
-
-					if (bCanLazyLoad)
-					{
-						SetBulkDataFlags(BULKDATA_LazyLoadable);
-					}
-					else
-					{
-						ClearBulkDataFlags(BULKDATA_LazyLoadable);
-
-						const int64 BulkDataSize = GetBulkDataSize();
-						const int64 BulkDataOffset = BulkMeta.GetOffset();
-						const int64 SavedPos = Ar.Tell();
-
-						Ar.Seek(BulkDataOffset);
-
-						void* DataBuffer = ReallocateData(BulkDataSize);
-						SerializeBulkData(Ar, DataBuffer, BulkDataSize, BulkMeta.GetFlags());
-
-						Ar.Seek(SavedPos);
-					}
-				}
-			}
+#endif
 		}
 #if !USE_RUNTIME_BULKDATA
 		else if (Ar.IsSaving())
 		{
-			// Make sure bulk data is loaded.
 			MakeSureBulkDataIsLoaded();
 
-			// Make mutable copies of the bulkdata location variables
-			EBulkDataFlags LocalBulkDataFlags = BulkMeta.GetFlags();
-			int64 LocalBulkDataSize = GetBulkDataSize(); 
-			int64 LocalBulkDataSizeOnDisk = GetBulkDataSizeOnDisk(); 
-			int64 LocalBulkDataOffsetInFile = BulkMeta.GetOffset();
-
-			// If the bulk data size is greater than can be held in an int32, then potentially the ElementCount
-			// and BulkDataSizeOnDisk need to be held as int64s, so set a flag indicating the new format.
-			if (GetBulkDataSize() >= (1LL << 31))
-			{
-				SetBulkDataFlagsOn(LocalBulkDataFlags, BULKDATA_Size64Bit);
-			}
-			// Remove single element serialization requirement before saving out bulk data flags.
-			ClearBulkDataFlagsOn(LocalBulkDataFlags, BULKDATA_ForceSingleElementSerialization);
-
-			// Save offset where we are serializing BulkDataFlags and store a placeholder
-			int64 SavedBulkDataFlagsPos = Ar.Tell();
-			{
-				Ar << LocalBulkDataFlags;
-			}
-
-			// Number of elements in array.
-			int64 ElementCount = LocalBulkDataSize / ElementSize;
-			SerializeBulkDataSizeInt(Ar, ElementCount, LocalBulkDataFlags);
-
-			// Only serialize status information if wanted.
-			int64 SavedBulkDataSizeOnDiskPos	= INDEX_NONE;
-			int64 SavedBulkDataOffsetInFilePos	= INDEX_NONE;
-			
-			{
-				// Save offset where we are serializing BulkDataSizeOnDisk and store a placeholder
-				SavedBulkDataSizeOnDiskPos = Ar.Tell();
-				LocalBulkDataSizeOnDisk = INDEX_NONE;
-				SerializeBulkDataSizeInt(Ar, LocalBulkDataSizeOnDisk, LocalBulkDataFlags);
-
-				// Save offset where we are serializing BulkDataOffsetInFile and store a placeholder
-				SavedBulkDataOffsetInFilePos = Ar.Tell();
-				LocalBulkDataOffsetInFile = INDEX_NONE;
-				Ar << LocalBulkDataOffsetInFile;
-			}
-
-			// try to get the linkersave object
-			FLinkerSave* LinkerSave = Cast<FLinkerSave>(Ar.GetLinker());
-
-			// determine whether we are going to store the payload inline or not.
-			bool bStoreInline = !!(LocalBulkDataFlags & BULKDATA_ForceInlinePayload) || !LinkerSave || Ar.IsTextFormat();
-			if (Ar.IsCooking() && !(LocalBulkDataFlags & BULKDATA_Force_NOT_InlinePayload))
-			{
-				bStoreInline = true;
-			}
-
-			if (!bStoreInline)
-			{
-				// set the flag indicating where the payload is stored
-				SetBulkDataFlagsOn(LocalBulkDataFlags, BULKDATA_PayloadAtEndOfFile);
-				ClearBulkDataFlagsOn(LocalBulkDataFlags,
-					static_cast<EBulkDataFlags>(BULKDATA_PayloadInSeperateFile | BULKDATA_WorkspaceDomainPayload)); // SavePackageUtilities::SaveBulkData will add these back if required
-
-				// with no LinkerSave we have to store the data inline
-				check(LinkerSave != NULL);
-
-				// add the bulkdata storage info object to the linkersave
-				FLinkerSave::FBulkDataStorageInfo& BulkStore = LinkerSave->BulkDataToAppend.AddZeroed_GetRef();
-
-				BulkStore.BulkDataOffsetInFilePos = SavedBulkDataOffsetInFilePos;
-				BulkStore.BulkDataSizeOnDiskPos = SavedBulkDataSizeOnDiskPos;
-				BulkStore.BulkDataFlagsPos = SavedBulkDataFlagsPos;
-				BulkStore.BulkDataFlags = LocalBulkDataFlags;
-				BulkStore.BulkDataFileRegionType = FileRegionType;
-				BulkStore.BulkData = this;
-
-				// If having flag BULKDATA_DuplicateNonOptionalPayload, duplicate bulk data in optional storage (.uptnl)
-				if (LocalBulkDataFlags & BULKDATA_DuplicateNonOptionalPayload)
-				{
-					int64 SavedDupeBulkDataFlagsPos = INDEX_NONE;
-					int64 SavedDupeBulkDataSizeOnDiskPos = INDEX_NONE;
-					int64 SavedDupeBulkDataOffsetInFilePos = INDEX_NONE;
+			FBulkMetaResource SerializedMeta;
+			SerializedMeta.Flags = BulkMeta.GetFlags();
+			SerializedMeta.ElementCount = GetBulkDataSize() / ElementSize;
+			SerializedMeta.SizeOnDisk = GetBulkDataSize();
 
-					EBulkDataFlags SavedDupeBulkDataFlags = static_cast<EBulkDataFlags>(
-						(LocalBulkDataFlags & ~BULKDATA_DuplicateNonOptionalPayload) | BULKDATA_OptionalPayload);
-					{
-						// Save offset where we are serializing SavedDupeBulkDataFlags and store a placeholder
-						SavedDupeBulkDataFlagsPos = Ar.Tell();
-						Ar << SavedDupeBulkDataFlags;
-
-						// Save offset where we are serializing SavedDupeBulkDataSizeOnDisk and store a placeholder
-						SavedDupeBulkDataSizeOnDiskPos = Ar.Tell();
-						int64 DupeBulkDataSizeOnDisk = INDEX_NONE;
-						SerializeBulkDataSizeInt(Ar, DupeBulkDataSizeOnDisk, SavedDupeBulkDataFlags);
-
-						// Save offset where we are serializing SavedDupeBulkDataOffsetInFile and store a placeholder
-						SavedDupeBulkDataOffsetInFilePos = Ar.Tell();
-						int64 DupeBulkDataOffsetInFile = INDEX_NONE;
-						Ar << DupeBulkDataOffsetInFile;
-					}
+			const EBulkDataFlags FlagsToClear = static_cast<EBulkDataFlags>(BULKDATA_PayloadAtEndOfFile | BULKDATA_PayloadInSeperateFile | BULKDATA_WorkspaceDomainPayload | BULKDATA_ForceSingleElementSerialization);
+			FBulkData::ClearBulkDataFlagsOn(SerializedMeta.Flags, FlagsToClear);
 
-					// add duplicate bulkdata with different flag
-					FLinkerSave::FBulkDataStorageInfo& DupeBulkStore = LinkerSave->BulkDataToAppend.AddZeroed_GetRef();
+			const int64 MetaOffset = Ar.Tell();
+			Ar << SerializedMeta;
 
-					DupeBulkStore.BulkDataOffsetInFilePos = SavedDupeBulkDataOffsetInFilePos;
-					DupeBulkStore.BulkDataSizeOnDiskPos = SavedDupeBulkDataSizeOnDiskPos;
-					DupeBulkStore.BulkDataFlagsPos = SavedDupeBulkDataFlagsPos;
-					DupeBulkStore.BulkDataFlags = SavedDupeBulkDataFlags;
-					DupeBulkStore.BulkDataFileRegionType = FileRegionType;
-					DupeBulkStore.BulkData = this;
-				}
-			}
-			else
+			SerializedMeta.Offset = Ar.Tell();
+			SerializedMeta.SizeOnDisk = SerializePayload(Ar, SerializedMeta.Flags, FileRegionType);
 			{
-				// set the flag indicating where the payload is stored
-				ClearBulkDataFlagsOn(LocalBulkDataFlags,
-					static_cast<EBulkDataFlags>(BULKDATA_PayloadAtEndOfFile | BULKDATA_PayloadInSeperateFile | BULKDATA_WorkspaceDomainPayload));
-
-				int64 SavedBulkDataStartPos = Ar.Tell();
-
-				// Serialize bulk data.
-				if (FileRegionType != EFileRegionType::None)
-				{
-					Ar.PushFileRegionType(FileRegionType);
-				}
-				SerializeBulkData(Ar, GetDataBufferForWrite(), LocalBulkDataFlags);
-				if (FileRegionType != EFileRegionType::None)
-				{
-					Ar.PopFileRegionType();
-				}
-
-				// store the payload endpos
-				int64 SavedBulkDataEndPos = Ar.Tell();
-
-				checkf(SavedBulkDataStartPos >= 0 && SavedBulkDataEndPos >= 0,
-					TEXT("Bad archive positions for bulkdata. StartPos=%d EndPos=%d"),
-					SavedBulkDataStartPos, SavedBulkDataEndPos);
-
-				LocalBulkDataSizeOnDisk = SavedBulkDataEndPos - SavedBulkDataStartPos;
-				LocalBulkDataOffsetInFile = SavedBulkDataStartPos;
-
-				// Since we are storing inline we are not relying on SavePackageUtilities::SaveBulkData to update the placeholder
-				// location data, so we need to do it here.
-
-				// store current file offset before seeking back
-				int64 CurrentFileOffset = Ar.Tell();
-				{
-					// Seek back and overwrite the flags 
-					Ar.Seek(SavedBulkDataFlagsPos);
-					Ar << LocalBulkDataFlags;
-
-					// Seek back and overwrite placeholder for BulkDataSizeOnDisk
-					Ar.Seek(SavedBulkDataSizeOnDiskPos);
-					SerializeBulkDataSizeInt(Ar, LocalBulkDataSizeOnDisk, LocalBulkDataFlags);
-
-					// Seek back and overwrite placeholder for BulkDataOffsetInFile
-					Ar.Seek(SavedBulkDataOffsetInFilePos);
-					Ar << LocalBulkDataOffsetInFile;
-				}
-				// Seek to the end of written data so we don't clobber any data in subsequent writes
-				Ar.Seek(CurrentFileOffset);
-
-#if WITH_EDITOR
-				// If we are overwriting the LoadedPath for the current package, set the location variables on *this equal to the new values
-				// that we are writing into the package on disk
-				if (LinkerSave && LinkerSave->bUpdatingLoadedPath)
-				{
-					SetFlagsFromDiskWrittenValues(LocalBulkDataFlags, LocalBulkDataOffsetInFile, LocalBulkDataSizeOnDisk,
-						INDEX_NONE /* LinkerSummaryBulkDataStartOffset, not applicable */);
-				}
-#endif
+				FArchive::FScopeSeekTo _(Ar, MetaOffset);
+				Ar << SerializedMeta;
 			}
 		}
 #endif // !USE_RUNTIME_BULKDATA
@@ -1700,6 +1403,32 @@ void FBulkData::SerializeBulkData(FArchive& Ar, void* Data, int64 DataSize, EBul
 	}
 }
 
+int64 FBulkData::SerializePayload(FArchive& Ar, EBulkDataFlags SerializationFlags, EFileRegionType RegionType)
+{
+	check(Ar.IsSaving());
+
+	MakeSureBulkDataIsLoaded();
+
+	const int64 PayloadStart = Ar.Tell();
+
+	if (int64 PayloadSize = GetBulkDataSize(); PayloadSize > 0)
+	{
+		if (RegionType != EFileRegionType::None)
+		{
+			Ar.PushFileRegionType(RegionType);
+		}
+
+		SerializeBulkData(Ar, GetDataBufferForWrite(), PayloadSize, SerializationFlags);
+
+		if (RegionType != EFileRegionType::None)
+		{
+			Ar.PopFileRegionType();
+		}
+	}
+
+	return Ar.Tell() - PayloadStart;
+}
+
 IAsyncReadFileHandle* FBulkData::OpenAsyncReadHandle() const
 {
 	return UE::BulkData::Private::OpenAsyncReadBulkData(BulkMeta, BulkChunkId).Release();
@@ -1816,11 +1545,10 @@ bool FBulkData::TryLoadDataIntoMemory(FIoBuffer Dest)
 	{
 		if (FArchive* Ar = AttachedAr)
 		{
-			const int64 SavedPos = Ar->Tell();
-			Ar->Seek(BulkDataOffset);
-			SerializeBulkData(*Ar, Dest.GetData(), Dest.GetSize(), BulkMeta.GetFlags());
-
-			Ar->Seek(SavedPos);
+			{
+				FArchive::FScopeSeekTo _(*Ar, BulkDataOffset);
+				SerializeBulkData(*Ar, Dest.GetData(), Dest.GetSize(), BulkMeta.GetFlags());
+			}
 			Ar->FlushCache();
 		
 			return true;
@@ -1850,60 +1578,7 @@ bool FBulkData::TryLoadDataIntoMemory(FIoBuffer Dest)
 
 bool FBulkData::CanDiscardInternalData() const
 {	
-	// Data marked as single use should always be discarded
-	if (IsSingleUse())
-	{
-		return true;
-	}
-
-	// If we can load from disk then we can discard it as it can be reloaded later
-	if (CanLoadFromDisk())
-	{
-		return true;
-	}
-
-	// If BULKDATA_AlwaysAllowDiscard has been set then we should always allow the data to 
-	// be discarded even if it cannot be reloaded again.
-	if (BulkMeta.HasAnyFlags(BULKDATA_AlwaysAllowDiscard))
-	{
-		return true;
-	}
-
-	return false;
-}
-
-void FBulkData::ConditionalSetInlineAlwaysAllowDiscard(bool bPackageUsesIoStore)
-{
-	// If PackagePath is null and we do not have BULKDATA_UsesIoDispatcher, we will not be able to reload the bulkdata.
-	// We will not have a PackagePath if the engine is using IoStore.
-	// So in IoStore with BulkData stored inline, we can not reload the bulkdata.
-	// We do not need to consider the end-of-file bulkdata section because IoStore  guarantees during cooking
-	// (SaveBulkData) that only inlined data is in the package file; there is no end-of-file bulkdata section.
-
-	// In the Inlined IoStore case therefore we need to not discard the bulk data if !BULKDATA_SingleUse;
-	// we have to keep it in case it gets requested twice.
-	// However, some systems (Audio,Animation) have large inline data they for legacy reasons have not marked as
-	// BULKDATA_SingleUse.
-	// In the old loader this data was discarded for them since CanLoadFromDisk was true.
-	// Since CanLoadFromDisk is now false, we are keeping that data around, and this causes memory bloat.
-	// Licensees have asked us to fix this memory bloat.
-	//
-	// To hack-fix the memory bloat in these systems when using IoStore we will mark all inline BulkData
-	// as discardable when using IoStore.
-	// This will cause a bug when using IoStore in any systems that do actually need to reload inline data.
-	// Each project that uses IoStore needs to guarantee that they do not have any systems that need to
-	// reload inline data.
-	// Note that when the define UE_KEEP_INLINE_RELOADING_CONSISTENT is enabled the old loading path should also
-	// not allow the reloading of inline data so that it behaves the same way as the new loading path. So when it
-	// is enabled we do not need to check if the loader is or isn't enabled, we can just set the flag and assume
-	// all inline data should be allowed to be discarded.
-
-#if !UE_KEEP_INLINE_RELOADING_CONSISTENT
-	if (bPackageUsesIoStore)
-#endif // !UE_KEEP_INLINE_RELOADING_CONSISTENT
-	{
-		SetBulkDataFlags(BULKDATA_AlwaysAllowDiscard);
-	}
+	return BulkMeta.HasAnyFlags(static_cast<EBulkDataFlags>(BULKDATA_AlwaysAllowDiscard | BULKDATA_SingleUse)) || CanLoadFromDisk();
 }
 
 void FUntypedBulkData::SerializeElements(FArchive& Ar, void* Data)
diff --git a/Engine/Source/Runtime/CoreUObject/Private/Serialization/BulkDataStreaming.cpp b/Engine/Source/Runtime/CoreUObject/Private/Serialization/BulkDataStreaming.cpp
--- a/Engine/Source/Runtime/CoreUObject/Private/Serialization/BulkDataStreaming.cpp
+++ b/Engine/Source/Runtime/CoreUObject/Private/Serialization/BulkDataStreaming.cpp
@@ -93,128 +93,10 @@ void FBulkDataIORequest::Cancel()
 	ReadRequest->Cancel();
 }
 
-namespace UE::BulkData::Private
-{
-
-//////////////////////////////////////////////////////////////////////////////
-
-/**
- * Triggers an ensure when trying to stream inline bulk data. This is to prevent
- * inconsistencies between the Zen loader loading from I/O store (utoc/ucas) and the legacy EDL loader using
- * the file system (.pak). The Zen loader currenlty does not allow reloading of inline bulk data.
- * The trigger can be ignored by setting [Core.System]IgnoreInlineBulkDataReloadEnsures to true in the config file.
- * Ignoring this ensure will most likely break loading of bulk data when packaging with Zen/IO store (.utoc/.ucas).
- */
-void EnsureCanStreamBulkData(const FBulkMetaData& BulkMeta)
-{
-#if !WITH_EDITOR && !UE_KEEP_INLINE_RELOADING_CONSISTENT	
-	static struct FIgnoreInlineDataReloadEnsures
-	{
-		bool bEnabled = false;
-
-		FIgnoreInlineDataReloadEnsures()
-		{
-			FConfigFile PlatformEngineIni;
-			FConfigCacheIni::LoadLocalIniFile(PlatformEngineIni, TEXT("Engine"), true, ANSI_TO_TCHAR(FPlatformProperties::IniPlatformName()));
-
-			PlatformEngineIni.GetBool(TEXT("Core.System"), TEXT("IgnoreInlineBulkDataReloadEnsures"), bEnabled);
-
-			UE_LOG(LogSerialization, Display, TEXT("IgnoreInlineDataReloadEnsures: '%s'"), bEnabled ? TEXT("true") : TEXT("false"));
-		}
-	} IgnoreInlineDataReloadEnsures;
-
-	const bool bIsInlined = BulkMeta.HasAnyFlags(BULKDATA_PayloadAtEndOfFile) == false;
-	// Note that we only want the ensure to trigger if we have a valid offset (the bulkdata references data from disk)
-	ensureMsgf(!bIsInlined || BulkMeta.GetOffset() == INDEX_NONE || IgnoreInlineDataReloadEnsures.bEnabled,
-				TEXT("Attempting to stream inline BulkData! This operation is not supported by the IoDispatcher and so will eventually stop working."
-				" The calling code should be fixed to retain the inline data in memory and re-use it rather than discard it and then try to reload from disk!"));
-#endif 
-}
-
-//////////////////////////////////////////////////////////////////////////////
-
-EPackageSegment GetPackageSegmentFromFlags(const FBulkMetaData& BulkMeta)
-{
-	const EBulkDataFlags BulkDataFlags = BulkMeta.GetFlags();
-
-	if ((BulkDataFlags & BULKDATA_PayloadInSeperateFile) == 0)
-	{
-		const bool bLoadingFromCookedPackage = EnumHasAnyFlags(BulkMeta.GetMetaFlags(), FBulkMetaData::EMetaFlags::CookedPackage);
-		if (bLoadingFromCookedPackage)
-		{
-			// Cooked packages are split into EPackageSegment::Header (summary and linker tables) and
-			// EPackageSegment::Exports (serialized UObject bytes and the bulk data section)
-			// Inline and end-of-file bulk data is in the Exports section
-			return EPackageSegment::Exports;
-		}
-		else
-		{
-			return EPackageSegment::Header;
-		}
-	}
-	else if (BulkDataFlags & BULKDATA_OptionalPayload )
-	{
-		return EPackageSegment::BulkDataOptional;
-	}
-	else if (BulkDataFlags & BULKDATA_MemoryMappedPayload)
-	{
-		return EPackageSegment::BulkDataMemoryMapped;
-	}
-	else
-	{
-		return EPackageSegment::BulkDataDefault;
-	}
-}
-
 //////////////////////////////////////////////////////////////////////////////
 
-FIoChunkId CreateBulkDataIoChunkId(const FBulkMetaData& BulkMeta, UObject* Owner)
+namespace UE::BulkData::Private
 {
-	UPackage* Pkg = Owner ? Owner->GetPackage() : nullptr;
-
-	if (Pkg == nullptr)
-	{
-		return FIoChunkId::InvalidChunkId;
-	}
-
-	const bool bCooked = EnumHasAnyFlags(BulkMeta.GetMetaFlags(), FBulkMetaData::EMetaFlags::CookedPackage);
-	
-	if (bCooked && Pkg->GetPackageId().IsValid())
-	{
-		const EBulkDataFlags BulkDataFlags = BulkMeta.GetFlags();
-
-		const EIoChunkType ChunkType = BulkDataFlags & BULKDATA_OptionalPayload
-			? EIoChunkType::OptionalBulkData
-			: BulkDataFlags & BULKDATA_MemoryMappedPayload
-				? EIoChunkType::MemoryMappedBulkData
-				: EIoChunkType::BulkData;
-
-		const uint16 ChunkIndex = EnumHasAnyFlags(BulkMeta.GetMetaFlags(), FBulkMetaData::EMetaFlags::OptionalPackage) ? 1 : 0;
-		
-		return CreateIoChunkId(Pkg->GetPackageIdToLoad().Value(), ChunkIndex, ChunkType);
-	}
-	else
-	{
-		FLinkerLoad* Linker = Owner->GetLinker();
-		if (Linker == nullptr)
-		{
-			Linker = FLinkerLoad::FindExistingLinkerForPackage(Pkg);
-		}
-	
-		if (Linker != nullptr)
-		{
-			const FPackagePath& PkgPath = Linker->GetPackagePath();
-			const EPackageSegment Segment = GetPackageSegmentFromFlags(BulkMeta);
-			const bool bExternalResource = BulkMeta.HasAnyFlags(BULKDATA_WorkspaceDomainPayload);
-
-			return CreatePackageResourceChunkId(PkgPath.GetPackageFName(), Segment, bExternalResource);
-		}
-	}
-
-	return FIoChunkId::InvalidChunkId;
-}
-
-//////////////////////////////////////////////////////////////////////////////
 
 enum class EChunkRequestStatus : uint32
 {
@@ -555,8 +437,6 @@ bool OpenReadBulkData(
 		return false;
 	}
 
-	EnsureCanStreamBulkData(BulkMeta);
-	
 	FIoBatch Batch = FIoDispatcher::Get().NewBatch();
 	FIoRequest Request = Batch.Read(BulkChunkId, FIoReadOptions(Offset, Size), ConvertToIoDispatcherPriority(Priority));
 	FEventRef Event;
@@ -583,7 +463,6 @@ TUniquePtr<IAsyncReadFileHandle> OpenAsyncReadBulkData(const FBulkMetaData& Bulk
 		return TUniquePtr<IAsyncReadFileHandle>();
 	}
 
-	EnsureCanStreamBulkData(BulkMeta);
 	return MakeUnique<FChunkReadFileHandle>(BulkChunkId);
 }
 
@@ -603,8 +482,6 @@ TUniquePtr<IBulkDataIORequest> CreateStreamingRequest(
 		return TUniquePtr<IBulkDataIORequest>();
 	}
 
-	EnsureCanStreamBulkData(BulkMeta);
-
 	FIoBuffer Buffer = UserSuppliedMemory ? FIoBuffer(FIoBuffer::Wrap, UserSuppliedMemory, Size) : FIoBuffer(Size);
 	FChunkBulkDataRequest* Request = new FChunkBulkDataRequest(Callback, MoveTemp(Buffer));
 	Request->Issue(BulkChunkId, FIoReadOptions(Offset, Size), ConvertToIoDispatcherPriority(Priority));
diff --git a/Engine/Source/Runtime/CoreUObject/Private/Serialization/DerivedData.cpp b/Engine/Source/Runtime/CoreUObject/Private/Serialization/DerivedData.cpp
--- a/Engine/Source/Runtime/CoreUObject/Private/Serialization/DerivedData.cpp
+++ b/Engine/Source/Runtime/CoreUObject/Private/Serialization/DerivedData.cpp
@@ -353,17 +353,10 @@ void FDerivedData::Serialize(FArchive& Ar, UObject* Owner)
 			BulkData.SetBulkDataFlags(BulkDataFlags | BULKDATA_Force_NOT_InlinePayload);
 		}
 
-		const int64 Offset = Ar.Tell();
-
-		FLinkerSave::FBulkDataStorageInfo& BulkDataInfo = Linker->BulkDataToAppend.AddZeroed_GetRef();
-		BulkDataInfo.BulkDataOffsetInFilePos = Offset;
-		// Size will be written as 64 bits because of the flag BULKDATA_Size64Bit.
-		BulkDataInfo.BulkDataSizeOnDiskPos = Offset + sizeof(uint64);
-		// Use the offset of the offset because SaveBulkData writes flags before offset and size anyway.
-		BulkDataInfo.BulkDataFlagsPos = Offset;
-		BulkDataInfo.BulkDataFlags = BulkDataFlags;
-		BulkDataInfo.BulkDataFileRegionType = EFileRegionType::None;
-		BulkDataInfo.BulkData = &EditorData->BulkData;
+		const int32 ElementSize = 1;
+		const bool bAttemptFileMapping = false;
+		const EFileRegionType FileRegionType = EFileRegionType::None;
+		Ar.SerializeBulkData(EditorData->BulkData, FBulkDataSerializationParams {Owner, ElementSize, FileRegionType, bAttemptFileMapping});
 
 		return LocalCookedData.Serialize(Ar);
 
diff --git a/Engine/Source/Runtime/CoreUObject/Private/UObject/LinkerLoad.cpp b/Engine/Source/Runtime/CoreUObject/Private/UObject/LinkerLoad.cpp
--- a/Engine/Source/Runtime/CoreUObject/Private/UObject/LinkerLoad.cpp
+++ b/Engine/Source/Runtime/CoreUObject/Private/UObject/LinkerLoad.cpp
@@ -16,6 +16,7 @@
 #include "UObject/ObjectRedirector.h"
 #include "UObject/Package.h"
 #include "UObject/PackageResourceManager.h"
+#include "UObject/PackageResourceIoDispatcherBackend.h"
 #include "UObject/PackageTrailer.h"
 #include "UObject/UObjectHash.h"
 #include "Misc/PackageName.h"
@@ -55,7 +56,7 @@
 #include "Misc/StringBuilder.h"
 #include "Misc/EngineBuildSettings.h"
 #include "Internationalization/GatherableTextData.h"
-
+#include "Async/MappedFileHandle.h"
 class FTexture2DResourceMem;
 
 #define LOCTEXT_NAMESPACE "LinkerLoad"
@@ -113,9 +114,25 @@ void TrackPackageAssetClass(UPackage* Package, FLinkerLoad& LinkerLoad, const TA
 #endif
 }
 
-/*----------------------------------------------------------------------------
-Helpers
-----------------------------------------------------------------------------*/
+EPackageSegment GetBulkDataPackageSegmentFromFlags(const EBulkDataFlags BulkDataFlags, bool bLoadingFromCookedPackage)
+{
+	if (FBulkData::HasFlags(BulkDataFlags, BULKDATA_PayloadInSeperateFile) == false)
+	{
+		return bLoadingFromCookedPackage ? EPackageSegment::Exports : EPackageSegment::Header;
+	}
+	else if (BulkDataFlags & BULKDATA_OptionalPayload )
+	{
+		return EPackageSegment::BulkDataOptional;
+	}
+	else if (BulkDataFlags & BULKDATA_MemoryMappedPayload)
+	{
+		return EPackageSegment::BulkDataMemoryMapped;
+	}
+	else
+	{
+		return EPackageSegment::BulkDataDefault;
+	}
+}
 
 #if WITH_EDITOR
 bool FLinkerLoad::ShouldCreateThrottledSlowTask() const
@@ -6682,4 +6699,153 @@ bool FLinkerLoad::TryGetPreloadedLoader(const FPackagePath& InPackagePath, FOpen
 
 #endif
 
+bool FLinkerLoad::SerializeBulkData(FBulkData& BulkData, const FBulkDataSerializationParams& Params)
+{
+	using namespace UE::BulkData::Private;
+
+	if (ShouldSkipBulkData())
+	{
+		return false;
+	}
+	
+	checkf(BulkData.IsUnlocked(), TEXT("Serialize bulk data FAILED, bulk data is locked"));
+
+	FBulkMetaData& Meta = BulkData.BulkMeta;
+	int64 DuplicateSerialOffset = -1;
+	SerializeBulkMeta(Meta, DuplicateSerialOffset, Params.ElementSize);
+
+	if (IsAllowingLazyLoading())
+	{
+		Meta.AddFlags(BULKDATA_LazyLoadable);
+#if WITH_EDITOR
+		check(IsTextFormat() == false);
+		BulkData.AttachedAr = this;
+		AttachBulkData(Params.Owner, &BulkData);
+#endif // WITH_EDITOR
+	}
+	else
+	{
+		Meta.ClearFlags(BULKDATA_LazyLoadable);
+	}
+
+	const bool bExternalResource = Meta.HasAnyFlags(BULKDATA_WorkspaceDomainPayload);
+	EPackageSegment Segment = GetBulkDataPackageSegmentFromFlags(Meta.GetFlags(), IsLoadingFromCookedPackage());  
+	BulkData.BulkChunkId = UE::CreatePackageResourceChunkId(PackagePath.GetPackageFName(), Segment, bExternalResource);
+
+	const bool bIsInline = Meta.HasAnyFlags(BULKDATA_PayloadAtEndOfFile) == false;
+	if (bIsInline)
+	{
+		if (IsLoadingFromCookedPackage())
+		{
+			// Cooked packages are split into .uasset/.exp files and the offset needs to be adjusted accordingly.
+			const int64 PkgHeaderSize = IPackageResourceManager::Get().FileSize(PackagePath,  EPackageSegment::Header);
+			Meta.SetOffset(Meta.GetOffset() - PkgHeaderSize);
+		}
+		void* Payload = BulkData.ReallocateData(Meta.GetSize());
+		BulkData.SerializeBulkData(*this, Payload, Meta.GetSize(), Meta.GetFlags());
+	}
+#if USE_RUNTIME_BULKDATA
+	else
+	{
+		// Streaming cooked bulk data. Since I/O store is the default staged output format
+		// this path is only used when running from loose cooked files.
+		check(Meta.HasAnyFlags(BULKDATA_PayloadInSeperateFile));
+		check(IsLoadingFromCookedPackage());
+		Meta.AddFlags(BULKDATA_LazyLoadable);
+
+		IPackageResourceManager& ResourceMgr = IPackageResourceManager::Get();
+	
+		if (Meta.HasAnyFlags(BULKDATA_DuplicateNonOptionalPayload) &&
+			ResourceMgr.DoesPackageExist(PackagePath, EPackageSegment::BulkDataOptional))
+		{
+			Segment = EPackageSegment::BulkDataOptional;
+			Meta.ClearFlags(BULKDATA_DuplicateNonOptionalPayload);
+			Meta.AddFlags(BULKDATA_OptionalPayload);
+			Meta.SetOffset(DuplicateSerialOffset);
+		}
+		else if (Meta.HasAnyFlags(BULKDATA_MemoryMappedPayload) && Params.bAttemptMemoryMapping)
+		{
+			bool bMemoryMapped = false;
+			if (IsAllowingLazyLoading())
+			{
+				TUniquePtr<IMappedFileHandle> MappedFile;
+				MappedFile.Reset(ResourceMgr.OpenMappedHandleToPackage(PackagePath, EPackageSegment::BulkDataMemoryMapped));
+				IMappedFileRegion* MappedRegion = MappedFile.IsValid() ? MappedFile->MapRegion(Meta.GetOffset(), Meta.GetSize(), true) : nullptr;
+				if (MappedRegion)
+				{
+					Segment = EPackageSegment::BulkDataMemoryMapped;
+					BulkData.DataAllocation.SetMemoryMappedData(&BulkData, MappedFile.Release(), MappedRegion);
+					bMemoryMapped  = true;
+				}
+			}
+
+			if (bMemoryMapped == false)
+			{
+				UE_LOG(LogSerialization, Warning, TEXT("Memory map bulk data '%s' FAILED"), *PackagePath.GetDebugName());
+				
+				BulkData.BulkChunkId = UE::CreatePackageResourceChunkId(PackagePath.GetPackageFName(), EPackageSegment::BulkDataDefault, false);
+				BulkData.ForceBulkDataResident();
+			}
+		}
+	
+		BulkData.BulkChunkId = UE::CreatePackageResourceChunkId(PackagePath.GetPackageFName(), Segment, bExternalResource);
+	}
+#else
+	else
+	{
+		// Streaming uncooked bulk data (editor only)
+		check(IsLoadingFromCookedPackage() == false);
+
+		// Unless this package is loaded from the EditorDomain, the offset needs
+		// to be adjusted to the start of non-inline bulk data in the .uasset file. 
+		if (Meta.HasAnyFlags(BULKDATA_WorkspaceDomainPayload) == false)
+		{
+			Meta.SetOffset(Meta.GetOffset() + Summary.BulkDataStartOffset);
+		}
+
+		if (IsAllowingLazyLoading() == false)
+		{
+			FArchive& Ar = *this;
+			FArchive::FScopeSeekTo _(Ar, Meta.GetOffset());
+			void* Payload = BulkData.ReallocateData(Meta.GetSize());
+			BulkData.SerializeBulkData(Ar, Payload, Meta.GetSize(), Meta.GetFlags());
+		}
+	}
+#endif // USE_RUNTIME_BULKDATA
+
+	return true;
+}
+
+void FLinkerLoad::SerializeBulkMeta(UE::BulkData::Private::FBulkMetaData& Meta, int64& DuplicateSerialOffset, int32 ElementSize)
+{
+	using namespace UE::BulkData::Private;
+	FArchive& Ar = *this;
+
+	if (DataResourceMap.IsEmpty())
+	{
+		FBulkMetaResource SerializedMeta;
+		Ar << SerializedMeta;
+		Meta = FBulkMetaData::FromSerialized(SerializedMeta, ElementSize);
+		DuplicateSerialOffset = SerializedMeta.DuplicateOffset;
+	}
+	else
+	{
+		int32 DataResourceIndex = INDEX_NONE;
+		Ar << DataResourceIndex;
+		const FObjectDataResource& DataResource = DataResourceMap[DataResourceIndex];
+		Meta.SetFlags(static_cast<EBulkDataFlags>(DataResource.LegacyBulkDataFlags));
+		Meta.SetOffset(DataResource.SerialOffset);
+		Meta.SetSize(DataResource.RawSize);
+		Meta.SetSizeOnDisk(DataResource.SerialSize);
+		DuplicateSerialOffset = DataResource.DuplicateSerialOffset;
+	}
+
+#if WITH_EDITOR
+	if (GIsEditor)
+	{
+		Meta.ClearFlags(BULKDATA_SingleUse);
+	}
+#endif // WITH_EDITOR
+}
+
 #undef LOCTEXT_NAMESPACE
diff --git a/Engine/Source/Runtime/CoreUObject/Private/UObject/LinkerSave.cpp b/Engine/Source/Runtime/CoreUObject/Private/UObject/LinkerSave.cpp
--- a/Engine/Source/Runtime/CoreUObject/Private/UObject/LinkerSave.cpp
+++ b/Engine/Source/Runtime/CoreUObject/Private/UObject/LinkerSave.cpp
@@ -2,7 +2,9 @@
 
 #include "UObject/LinkerSave.h"
 #include "HAL/FileManager.h"
+#include "Misc/CommandLine.h"
 #include "Misc/ConfigCacheIni.h"
+#include "Serialization/BulkData.h"
 #include "Serialization/LargeMemoryWriter.h"
 #include "UObject/Package.h"
 #include "UObject/Class.h"
@@ -14,7 +16,6 @@
 #include "UObject/UObjectThreadContext.h"
 #include "UObject/UnrealType.h"
 #include "HAL/PlatformStackWalk.h"
-
 #if WITH_EDITORONLY_DATA
 #include "IO/IoDispatcher.h"
 #include "Serialization/DerivedData.h"
@@ -200,8 +201,6 @@ bool FLinkerSave::CloseAndDestroySaver()
 	if (Saver)
 	{
 		// first, do an explicit close to check for archive errors
-		bSuccess = Saver->Close();
-		// then, destroy it
 		delete Saver;
 	}
 	Saver = nullptr;
@@ -555,3 +554,175 @@ UE::FDerivedData FLinkerSave::AddDerivedData(const UE::FDerivedData& Data)
 	return UE::FDerivedData(CookedData);
 }
 #endif // WITH_EDITORONLY_DATA
+
+bool FLinkerSave::SerializeBulkData(FBulkData& BulkData, const FBulkDataSerializationParams& Params) 
+{
+	using namespace UE::BulkData::Private;
+
+	auto CanSaveBulkDataByReference = [](FBulkData& BulkData) -> bool
+	{
+		return BulkData.GetBulkDataOffsetInFile() != INDEX_NONE &&
+			// We don't support yet loading from a separate file
+			!BulkData.IsInSeparateFile() &&
+			// It is possible to have a BulkData marked as optional without putting it into a separate file, and we
+			// assume that if BulkData is optional and in a separate file, then it is in the BulkDataOptional
+			// segment. Rather than changing that assumption to support optional ExternalResource bulkdata, we
+			// instead require that optional inlined/endofpackagedata BulkDatas can not be read from an
+			// ExternalResource and must remain inline.
+			!BulkData.IsOptional() &&
+			// Inline or end-of-package-file data can only be loaded from the workspace domain package file if the
+			// archive used by the bulk data was actually from the package file; BULKDATA_LazyLoadable is set by
+			// Serialize iff that is the case										
+			(BulkData.GetBulkDataFlags() & BULKDATA_LazyLoadable);
+	};
+	
+	if (ShouldSkipBulkData())
+	{
+		return false;
+	}
+
+	const EBulkDataFlags BulkDataFlags	= static_cast<EBulkDataFlags>(BulkData.GetBulkDataFlags());
+	int32 ResourceIndex					= DataResourceMap.Num();
+	const int64 PayloadSize				= BulkData.GetBulkDataSize();
+	const EFileRegionType RegionToUse	= bFileRegionsEnabled && IsCooking() ? Params.RegionType : EFileRegionType::None;
+	const bool bSupportsMemoryMapping	= IsCooking() && MemoryMappingAlignment >= 0;
+	const bool bSaveAsResourceIndex		= IsCooking();
+	
+	FBulkMetaResource SerializedMeta;
+	SerializedMeta.Flags = BulkDataFlags;
+	SerializedMeta.ElementCount = PayloadSize / Params.ElementSize;
+	SerializedMeta.SizeOnDisk = PayloadSize; 
+
+	EBulkDataFlags FlagsToClear = static_cast<EBulkDataFlags>(BULKDATA_PayloadAtEndOfFile | BULKDATA_PayloadInSeperateFile | BULKDATA_WorkspaceDomainPayload | BULKDATA_ForceSingleElementSerialization);
+	if (IsCooking())
+	{
+		FBulkData::SetBulkDataFlagsOn(FlagsToClear, static_cast<EBulkDataFlags>(BULKDATA_SerializeCompressed));
+	}
+
+	FBulkData::ClearBulkDataFlagsOn(SerializedMeta.Flags, FlagsToClear);
+
+	const bool bSerializeInline =
+		FBulkData::HasFlags(BulkDataFlags, BULKDATA_ForceInlinePayload) ||
+		(IsCooking() && (FBulkData::HasFlags(BulkDataFlags, BULKDATA_Force_NOT_InlinePayload) == false)) ||
+		IsTextFormat();
+
+	if (bSerializeInline)
+	{
+		FArchive& Ar = *this;
+
+		const int64 MetaOffset = Tell();
+		if (bSaveAsResourceIndex)
+		{
+			Ar << ResourceIndex;
+		}
+		else
+		{
+			Ar << SerializedMeta;
+		}
+
+		SerializedMeta.Offset = Tell();
+		SerializedMeta.SizeOnDisk = BulkData.SerializePayload(Ar, SerializedMeta.Flags, RegionToUse);
+
+		if (bSaveAsResourceIndex == false)
+		{
+			FArchive::FScopeSeekTo _(Ar, MetaOffset);
+			Ar << SerializedMeta;
+		}
+	}
+	else
+	{
+		FBulkData::SetBulkDataFlagsOn(SerializedMeta.Flags, static_cast<EBulkDataFlags>(BULKDATA_PayloadAtEndOfFile));
+
+		if (bSaveBulkDataToSeparateFiles)
+		{
+			check(bSaveBulkDataByReference == false);
+			FBulkData::SetBulkDataFlagsOn(SerializedMeta.Flags, static_cast<EBulkDataFlags>(BULKDATA_PayloadInSeperateFile | BULKDATA_NoOffsetFixUp));
+		}
+		
+		const bool bSaveByReference = bSaveBulkDataByReference && CanSaveBulkDataByReference(BulkData);
+		if (bSaveByReference)
+		{
+			check(IsCooking() == false);
+			FBulkData::SetBulkDataFlagsOn(SerializedMeta.Flags, static_cast<EBulkDataFlags>(BULKDATA_NoOffsetFixUp | BULKDATA_WorkspaceDomainPayload));
+		}
+
+		if (bSaveBulkDataToSeparateFiles && FBulkData::HasFlags(SerializedMeta.Flags, BULKDATA_OptionalPayload))
+		{
+			SerializedMeta.Offset = OptionalBulkDataAr.Tell();
+			SerializedMeta.SizeOnDisk = BulkData.SerializePayload(OptionalBulkDataAr, SerializedMeta.Flags, RegionToUse);
+		}
+		else if (bSaveBulkDataToSeparateFiles && FBulkData::HasFlags(SerializedMeta.Flags, BULKDATA_MemoryMappedPayload) && bSupportsMemoryMapping)
+		{
+			if (int64 Padding = Align(MemoryMappedBulkDataAr.Tell(), MemoryMappingAlignment) - MemoryMappedBulkDataAr.Tell(); Padding > 0)
+			{
+				TArray<uint8> Zeros;
+				Zeros.SetNumZeroed(int32(Padding));
+				MemoryMappedBulkDataAr.Serialize(Zeros.GetData(), Padding);
+			}
+			SerializedMeta.Offset = MemoryMappedBulkDataAr.Tell();
+			SerializedMeta.SizeOnDisk = BulkData.SerializePayload(MemoryMappedBulkDataAr, SerializedMeta.Flags, RegionToUse);
+		}
+		else
+		{
+			if (bSaveBulkDataToSeparateFiles && FBulkData::HasFlags(SerializedMeta.Flags, BULKDATA_DuplicateNonOptionalPayload))
+			{
+				SerializedMeta.DuplicateFlags = SerializedMeta.Flags;
+				SerializedMeta.DuplicateOffset = OptionalBulkDataAr.Tell();
+				SerializedMeta.DuplicateSizeOnDisk = BulkData.SerializePayload(OptionalBulkDataAr, SerializedMeta.Flags, RegionToUse);
+
+				FBulkData::ClearBulkDataFlagsOn(SerializedMeta.DuplicateFlags, BULKDATA_DuplicateNonOptionalPayload);
+				FBulkData::SetBulkDataFlagsOn(SerializedMeta.DuplicateFlags, BULKDATA_OptionalPayload);
+			}
+
+			if (bSaveByReference)
+			{
+				SerializedMeta.Offset = BulkData.GetBulkDataOffsetInFile();
+				SerializedMeta.SizeOnDisk = BulkData.GetBulkDataSizeOnDisk();
+			}
+			else
+			{
+				SerializedMeta.Offset = BulkDataAr.Tell();
+				SerializedMeta.SizeOnDisk = BulkData.SerializePayload(BulkDataAr, SerializedMeta.Flags, RegionToUse);
+			}
+		}
+		
+		FArchive& Ar = *this;
+		if (bSaveAsResourceIndex)
+		{
+			Ar << ResourceIndex;
+		}
+		else
+		{
+			Ar << SerializedMeta;
+		}
+	}
+
+	FObjectDataResource& DataResource = DataResourceMap.AddDefaulted_GetRef();
+	DataResource.RawSize				= PayloadSize;
+	DataResource.SerialSize				= SerializedMeta.SizeOnDisk;
+	DataResource.SerialOffset			= SerializedMeta.Offset;
+	DataResource.DuplicateSerialOffset	= SerializedMeta.DuplicateOffset;
+	DataResource.LegacyBulkDataFlags	= SerializedMeta.Flags;
+	DataResource.OuterIndex				= ObjectIndicesMap.FindRef(Params.Owner);
+
+	SerializedBulkData.Add(&BulkData, ResourceIndex);
+
+	return true;
+}
+
+void FLinkerSave::OnPostSaveBulkData()
+{
+#if WITH_EDITOR
+	if (bUpdatingLoadedPath)
+	{
+		for (TPair<FBulkData*, int32>& Kv : SerializedBulkData)
+		{
+			FBulkData& BulkData = *Kv.Key;
+			const FObjectDataResource& DataResource = DataResourceMap[Kv.Value];
+			BulkData.SetFlagsFromDiskWrittenValues(static_cast<EBulkDataFlags>(DataResource.Flags), DataResource.SerialOffset, DataResource.SerialSize, Summary.BulkDataStartOffset);
+		}
+	}
+#endif
+
+	SerializedBulkData.Empty();
+}
diff --git a/Engine/Source/Runtime/CoreUObject/Private/UObject/SavePackage/SavePackageUtilities.cpp b/Engine/Source/Runtime/CoreUObject/Private/UObject/SavePackage/SavePackageUtilities.cpp
--- a/Engine/Source/Runtime/CoreUObject/Private/UObject/SavePackage/SavePackageUtilities.cpp
+++ b/Engine/Source/Runtime/CoreUObject/Private/UObject/SavePackage/SavePackageUtilities.cpp
@@ -20,6 +20,7 @@
 #include "Serialization/EditorBulkData.h"
 #include "Serialization/CompactBinarySerialization.h"
 #include "Serialization/CompactBinaryWriter.h"
+#include "Serialization/FileRegionArchive.h"
 #include "Serialization/LargeMemoryWriter.h"
 #include "Serialization/PackageWriter.h"
 #include "Tasks/Task.h"
@@ -1812,16 +1813,6 @@ void SaveThumbnails(UPackage* InOuter, FLinkerSave* Linker, FStructuredArchive::
 #endif
 }
 
-class FLargeMemoryWriterWithRegions : public FLargeMemoryWriter
-{
-public:
-	FLargeMemoryWriterWithRegions()
-		: FLargeMemoryWriter(0, /* IsPersistent */ true)
-	{}
-
-	TArray<FFileRegion> FileRegions;
-};
-
 ESavePackageResult AppendAdditionalData(FLinkerSave& Linker, int64& InOutDataStartOffset, FSavePackageContext* SavePackageContext)
 {
 	if (Linker.AdditionalDataToAppend.Num() == 0)
@@ -1833,21 +1824,23 @@ ESavePackageResult AppendAdditionalData(FLinkerSave& Linker, int64& InOutDataSta
 	if (PackageWriter)
 	{
 		bool bDeclareRegionForEachAdditionalFile = SavePackageContext->PackageWriterCapabilities.bDeclareRegionForEachAdditionalFile;
-		FLargeMemoryWriterWithRegions DataArchive;
+		FFileRegionMemoryWriter DataArchive;
 		for (FLinkerSave::AdditionalDataCallback& Callback : Linker.AdditionalDataToAppend)
 		{
-			int64 RegionStart = DataArchive.Tell();
-			Callback(Linker, DataArchive, InOutDataStartOffset + RegionStart);
-			int64 RegionEnd = DataArchive.Tell();
-			if (RegionEnd != RegionStart && bDeclareRegionForEachAdditionalFile)
+			if (bDeclareRegionForEachAdditionalFile)
 			{
-				DataArchive.FileRegions.Add(FFileRegion(RegionStart, RegionEnd - RegionStart, EFileRegionType::None));
+				DataArchive.PushFileRegionType(EFileRegionType::None);
+			}
+			Callback(Linker, DataArchive, InOutDataStartOffset + DataArchive.Tell());
+			if (bDeclareRegionForEachAdditionalFile)
+			{
+				DataArchive.PopFileRegionType();
 			}
 		}
 		IPackageWriter::FLinkerAdditionalDataInfo DataInfo{ Linker.LinkerRoot->GetFName() };
 		int64 DataSize = DataArchive.TotalSize();
 		FIoBuffer DataBuffer = FIoBuffer(FIoBuffer::AssumeOwnership, DataArchive.ReleaseOwnership(), DataSize);
-		PackageWriter->WriteLinkerAdditionalData(DataInfo, DataBuffer, DataArchive.FileRegions);
+		PackageWriter->WriteLinkerAdditionalData(DataInfo, DataBuffer, DataArchive.GetFileRegions());
 		InOutDataStartOffset += DataSize;
 	}
 	else
@@ -1958,365 +1951,6 @@ ESavePackageResult CreatePayloadSidecarFile(FLinkerSave& Linker, const FPackageP
 	return ESavePackageResult::Success;
 }
 
-ESavePackageResult SaveBulkData(FLinkerSave* Linker, int64& InOutStartOffset, const UPackage* InOuter, const TCHAR* Filename, const ITargetPlatform* TargetPlatform,
-				  FSavePackageContext* SavePackageContext, uint32 SaveFlags, const bool bTextFormat, 
-				  int64& TotalPackageSizeUncompressed, bool bIsOptionalRealm)
-{
-	// Now we write all the bulkdata that is supposed to be at the end of the package
-	// and fix up the offset
-	IPackageWriter* PackageWriter = SavePackageContext != nullptr ? SavePackageContext->PackageWriter : nullptr;
-	Linker->Summary.BulkDataStartOffset = InOutStartOffset;
-
-	if (Linker->BulkDataToAppend.Num() == 0)
-	{
-		return ESavePackageResult::Success;;
-	}
-	check(!bTextFormat);
-
-	COOK_STAT(FScopedDurationTimer SaveTimer(FSavePackageStats::SerializeBulkDataTimeSec));
-
-	FScopedSlowTask BulkDataFeedback((float)Linker->BulkDataToAppend.Num());
-
-	TUniquePtr<FLargeMemoryWriterWithRegions> BulkArchive;
-	TUniquePtr<FLargeMemoryWriterWithRegions> OptionalBulkArchive;
-	TUniquePtr<FLargeMemoryWriterWithRegions> MappedBulkArchive;
-
-	const bool bSeparateSegmentsEnabled = Linker->IsCooking();
-
-	int64 LinkerStart = 0;
-	if (PackageWriter || bSeparateSegmentsEnabled)
-	{
-		BulkArchive.Reset(new FLargeMemoryWriterWithRegions);
-		if (bSeparateSegmentsEnabled)
-		{
-			OptionalBulkArchive.Reset(new FLargeMemoryWriterWithRegions);
-			MappedBulkArchive.Reset(new FLargeMemoryWriterWithRegions);
-		}
-	}
-	else
-	{
-		LinkerStart = Linker->Tell();
-	}
-	bool bRequestSaveByReference = SaveFlags & SAVE_BulkDataByReference;
-
-	bool bAlignBulkData = false;
-	bool bDeclareSpecialRegions = false;
-	bool bDeclareRegionForEachAdditionalFile = false;
-	int64 BulkDataAlignment = 0;
-
-	if (TargetPlatform)
-	{
-		bAlignBulkData = TargetPlatform->SupportsFeature(ETargetPlatformFeatures::MemoryMappedFiles);
-		bDeclareSpecialRegions = TargetPlatform->SupportsFeature(ETargetPlatformFeatures::CookFileRegionMetadata);
-		BulkDataAlignment = TargetPlatform->GetMemoryMappingAlignment();
-	}
-	if (PackageWriter)
-	{
-		bDeclareRegionForEachAdditionalFile = SavePackageContext->PackageWriterCapabilities.bDeclareRegionForEachAdditionalFile;
-	}
-
-	if (bRequestSaveByReference)
-	{
-		const TCHAR* FailureReason = nullptr;
-		if (Linker->bUpdatingLoadedPath)
-		{
-			FailureReason = TEXT("SAVE_BulkDataByReference is incompatible with bUpdatingLoadedPath");
-		}
-		if (FailureReason)
-		{
-			UE_LOG(LogSavePackage, Error, TEXT("SaveBulkData failed for %s: %s."), Filename, FailureReason);
-			return ESavePackageResult::Error;
-		}
-	}
-
-	for (FLinkerSave::FBulkDataStorageInfo& BulkDataStorageInfo : Linker->BulkDataToAppend)
-	{
-		BulkDataFeedback.EnterProgressFrame();
-
-		// Set bulk data flags to what they were during initial serialization (they might have changed after that)
-		uint32 BulkDataFlags = BulkDataStorageInfo.BulkDataFlags;
-		FBulkData* BulkData = BulkDataStorageInfo.BulkData;
-		checkf(BulkDataFlags& BULKDATA_PayloadAtEndOfFile, TEXT("Inlined BulkData data should not have been added to BulkDataToAppend"));
-
-		const bool bBulkItemIsOptional = (BulkDataFlags & BULKDATA_OptionalPayload) != 0;
-		bool bBulkItemIsMapped = bAlignBulkData && ((BulkDataFlags & BULKDATA_MemoryMappedPayload) != 0);
-
-		if (bBulkItemIsMapped && bBulkItemIsOptional)
-		{
-			UE_LOG(LogSavePackage, Warning, TEXT("%s has bulk data that is both mapped and optional. This is not currently supported. Will not be mapped."), Filename);
-			BulkDataFlags &= ~BULKDATA_MemoryMappedPayload;
-			bBulkItemIsMapped = false;
-		}
-		if (bBulkItemIsMapped && bRequestSaveByReference)
-		{
-			UE_LOG(LogSavePackage, Warning, TEXT("%s has bulk data that is mapped, but the save method is SAVE_BulkDataByReference."
-				"This is not currently supported. Will not be mapped."), Filename);
-			BulkDataFlags &= ~BULKDATA_MemoryMappedPayload;
-			bBulkItemIsMapped = false;
-		}
-
-		enum ESaveLocation
-		{
-			Invalid,
-			EndOfPackage,
-			SeparateSegment,
-			Reference,
-			SeparateArchiveAtEndOfPackage,
-		};
-		ESaveLocation SaveLocation = ESaveLocation::Invalid;
-		auto CanSaveByReference = [](FBulkData* BulkData)
-		{
-			return BulkData->GetBulkDataOffsetInFile() != INDEX_NONE &&
-				// We don't support yet loading from a separate file
-				!BulkData->IsInSeparateFile() &&
-				// It is possible to have a BulkData marked as optional without putting it into a separate file, and we
-				// assume that if BulkData is optional and in a separate file, then it is in the BulkDataOptional
-				// segment. Rather than changing that assumption to support optional ExternalResource bulkdata, we
-				// instead require that optional inlined/endofpackagedata BulkDatas can not be read from an
-				// ExternalResource and must remain inline.
-				!BulkData->IsOptional() &&
-				// Inline or end-of-package-file data can only be loaded from the workspace domain package file if the
-				// archive used by the bulk data was actually from the package file; BULKDATA_LazyLoadable is set by
-				// Serialize iff that is the case										
-				(BulkData->GetBulkDataFlags() & BULKDATA_LazyLoadable);
-		};
-		if (bRequestSaveByReference && CanSaveByReference(BulkData))
-		{
-			SaveLocation = ESaveLocation::Reference;
-		}
-		else if (bSeparateSegmentsEnabled)
-		{
-			SaveLocation = ESaveLocation::SeparateSegment;
-		}
-		else if (PackageWriter)
-		{
-			SaveLocation = ESaveLocation::SeparateArchiveAtEndOfPackage;
-		}
-		else
-		{
-			SaveLocation = ESaveLocation::EndOfPackage;
-		}
-
-		int64 BulkDataOffsetInFile;
-		int64 BulkDataSizeOnDisk;
-		if (SaveLocation == ESaveLocation::Reference)
-		{
-			BulkDataFlags |= BULKDATA_PayloadInSeperateFile | BULKDATA_WorkspaceDomainPayload;
-			BulkDataFlags |= BULKDATA_NoOffsetFixUp; // We don't use legacy offset fixups when referencing
-			BulkDataFlags &= ~BULKDATA_MemoryMappedPayload; // We don't support memory mapping from referenced
-
-			BulkDataOffsetInFile = BulkData->GetBulkDataOffsetInFile();
-			BulkDataSizeOnDisk = BulkData->GetBulkDataSizeOnDisk();
-		}
-		else
-		{
-			TArray<FFileRegion>* TargetRegions = &Linker->FileRegions;
-			FArchive* TargetArchive = Linker;
-			if (SaveLocation == ESaveLocation::SeparateSegment)
-			{
-				BulkDataFlags |= BULKDATA_PayloadInSeperateFile;
-				// OffsetFixup is not useful when bulkdata is in separate segments.
-				// It is also not supported if we are using IoStore, which always uses SeparateSegments, because 
-				// the Linker's Summary.BulkDataStartOffset information is not available when loading with AsyncLoader2.
-				BulkDataFlags |= BULKDATA_NoOffsetFixUp;
-
-				if (bBulkItemIsOptional)
-				{
-					TargetArchive = OptionalBulkArchive.Get();
-					TargetRegions = &OptionalBulkArchive->FileRegions;
-				}
-				else if (bBulkItemIsMapped)
-				{
-					TargetArchive = MappedBulkArchive.Get();
-					TargetRegions = &MappedBulkArchive->FileRegions;
-				}
-				else
-				{
-					TargetArchive = BulkArchive.Get();
-					TargetRegions = &BulkArchive->FileRegions;
-				}
-			}
-			else if (SaveLocation == ESaveLocation::SeparateArchiveAtEndOfPackage)
-			{
-				TargetArchive = BulkArchive.Get();
-				TargetRegions = &BulkArchive->FileRegions;
-			}
-
-			check(TargetArchive && TargetRegions);
-
-			// Pad archive for proper alignment for memory mapping
-			if (bBulkItemIsMapped && BulkDataAlignment > 0)
-			{
-				const int64 BulkStartOffset = TargetArchive->Tell();
-
-				if (!IsAligned(BulkStartOffset, BulkDataAlignment))
-				{
-					const int64 AlignedOffset = Align(BulkStartOffset, BulkDataAlignment);
-
-					int64 Padding = AlignedOffset - BulkStartOffset;
-					check(Padding > 0);
-
-					uint64 Zero64 = 0;
-					while (Padding >= 8)
-					{
-						*TargetArchive << Zero64;
-						Padding -= 8;
-					}
-
-					uint8 Zero8 = 0;
-					while (Padding > 0)
-					{
-						*TargetArchive << Zero8;
-						Padding--;
-					}
-
-					check(TargetArchive->Tell() == AlignedOffset);
-				}
-			}
-
-			const int64 BulkStartOffset = TargetArchive->Tell();
-			BulkData->SerializeBulkData(*TargetArchive, BulkData->Lock(LOCK_READ_ONLY), static_cast<EBulkDataFlags>(BulkDataFlags));
-			BulkData->Unlock();
-			const int64 BulkEndOffset = TargetArchive->Tell();
-
-			BulkDataOffsetInFile = BulkStartOffset;
-			if (SaveLocation == ESaveLocation::SeparateArchiveAtEndOfPackage)
-			{
-				// BulkDatas are written in a separate archive that is appended onto the exports archive;
-				// the input BulkStartOffset is relative to the beginning of this separate archive.
-				// But records of the BulkData's offset needs to instead be relative to the beginning of the exports archive.
-				// We have to do this to distinguish them from inline bulkdatas in the exports segment
-				BulkDataOffsetInFile += Linker->Summary.BulkDataStartOffset;
-			}
-			if ((BulkDataFlags & BULKDATA_NoOffsetFixUp) == 0)
-			{
-				// The runtime will add in the Summary.BulkDataStartOffset, so we subtract it here.
-				// This allows the decoupling of the values written into the package from the size of the exports section.
-				BulkDataOffsetInFile -= Linker->Summary.BulkDataStartOffset;
-			}
-			BulkDataSizeOnDisk = BulkEndOffset - BulkStartOffset;
-
-			bool bDeclareBecauseSpecial = bDeclareSpecialRegions && BulkDataStorageInfo.BulkDataFileRegionType != EFileRegionType::None;
-			if ((bDeclareRegionForEachAdditionalFile || bDeclareBecauseSpecial) && BulkDataSizeOnDisk > 0)
-			{
-				TargetRegions->Add(FFileRegion(BulkStartOffset, BulkDataSizeOnDisk, BulkDataStorageInfo.BulkDataFileRegionType));
-			}
-		}
-
-		const int64 SavedLinkerOffset = Linker->Tell();
-		Linker->Seek(BulkDataStorageInfo.BulkDataFlagsPos);
-		*Linker << BulkDataFlags;
-		Linker->Seek(BulkDataStorageInfo.BulkDataSizeOnDiskPos);
-		SerializeBulkDataSizeInt(*Linker, BulkDataSizeOnDisk, static_cast<EBulkDataFlags>(BulkDataFlags));
-		Linker->Seek(BulkDataStorageInfo.BulkDataOffsetInFilePos);
-		*Linker << BulkDataOffsetInFile;
-		Linker->Seek(SavedLinkerOffset);
-
-#if WITH_EDITOR
-		// If we are overwriting the LoadedPath for the current package, the bulk data flags and location need to be updated to match
-		// the values set in the package on disk.
-		if (Linker->bUpdatingLoadedPath)
-		{
-			BulkData->SetFlagsFromDiskWrittenValues(static_cast<EBulkDataFlags>(BulkDataFlags), BulkDataOffsetInFile,
-				BulkDataSizeOnDisk, Linker->Summary.BulkDataStartOffset);
-		}
-#endif
-	}
-
-	if (BulkArchive)
-	{
-		if (PackageWriter)
-		{
-			auto AddSizeAndConvertToIoBuffer = [&TotalPackageSizeUncompressed](FLargeMemoryWriter* Writer)
-			{
-				const int64 TotalSize = Writer->TotalSize();
-				TotalPackageSizeUncompressed += TotalSize;
-				return FIoBuffer(FIoBuffer::AssumeOwnership, Writer->ReleaseOwnership(), TotalSize);
-			};
-
-			IPackageWriter::FBulkDataInfo BulkInfo;
-			BulkInfo.PackageName = InOuter->GetFName();
-			// Adjust LooseFilePath if needed
-			if (bIsOptionalRealm)
-			{
-				// Optional output have the form PackagePath.o.ext
-				BulkInfo.LooseFilePath = FPathViews::ChangeExtension(Filename, TEXT("o.") + FPaths::GetExtension(Filename));
-				BulkInfo.MultiOutputIndex = 1;
-			}
-			else
-			{
-				BulkInfo.LooseFilePath = Filename;
-			}
-			FPackageId PackageId = FPackageId::FromName(BulkInfo.PackageName);
-				
-			if (BulkArchive->TotalSize())
-			{
-				BulkInfo.ChunkId = CreateIoChunkId(PackageId.Value(), BulkInfo.MultiOutputIndex, EIoChunkType::BulkData);
-				BulkInfo.BulkDataType = bSeparateSegmentsEnabled ?
-					IPackageWriter::FBulkDataInfo::BulkSegment : IPackageWriter::FBulkDataInfo::AppendToExports;
-				BulkInfo.LooseFilePath = FPathViews::ChangeExtension(BulkInfo.LooseFilePath, LexToString(EPackageExtension::BulkDataDefault));
-					
-				PackageWriter->WriteBulkData(BulkInfo, AddSizeAndConvertToIoBuffer(BulkArchive.Get()), BulkArchive->FileRegions);
-			}
-			// @note FH: temporarily do not handle optional bulk data into editor optional packages, proper support will be added soon
-			if (OptionalBulkArchive && OptionalBulkArchive->TotalSize() && !bIsOptionalRealm)
-			{
-				BulkInfo.ChunkId = CreateIoChunkId(PackageId.Value(), BulkInfo.MultiOutputIndex, EIoChunkType::OptionalBulkData);
-				BulkInfo.BulkDataType = IPackageWriter::FBulkDataInfo::Optional;
-				BulkInfo.LooseFilePath = FPathViews::ChangeExtension(BulkInfo.LooseFilePath, LexToString(EPackageExtension::BulkDataOptional));
-				PackageWriter->WriteBulkData(BulkInfo, AddSizeAndConvertToIoBuffer(OptionalBulkArchive.Get()), OptionalBulkArchive->FileRegions);
-			}
-			if (MappedBulkArchive && MappedBulkArchive->TotalSize())
-			{
-				checkf(!bIsOptionalRealm, TEXT("MemoryMappedBulkData is currently unsupported with optional package multi output for %s"), *InOuter->GetName());
-				BulkInfo.ChunkId = CreateIoChunkId(PackageId.Value(), 0, EIoChunkType::MemoryMappedBulkData);
-				BulkInfo.BulkDataType = IPackageWriter::FBulkDataInfo::Mmap;
-				BulkInfo.LooseFilePath = FPathViews::ChangeExtension(BulkInfo.LooseFilePath, LexToString(EPackageExtension::BulkDataMemoryMapped));
-				PackageWriter->WriteBulkData(BulkInfo, AddSizeAndConvertToIoBuffer(MappedBulkArchive.Get()), MappedBulkArchive->FileRegions);
-			}
-		}
-		else
-		{
-			checkf(!bIsOptionalRealm, TEXT("Package optional package multi output is unsupported without a PackageWriter for %s"), *InOuter->GetName());
-			auto WriteBulkData = [&](FLargeMemoryWriterWithRegions* Archive, const TCHAR* BulkFileExtension)
-			{
-				if (const int64 DataSize = Archive ? Archive->TotalSize() : 0)
-				{
-					TotalPackageSizeUncompressed += DataSize;
-
-					FLargeMemoryPtr DataPtr(Archive->ReleaseOwnership());
-
-					const FString ArchiveFilename = FPaths::ChangeExtension(Filename, BulkFileExtension);
-
-					EAsyncWriteOptions WriteOptions(EAsyncWriteOptions::None);
-					SavePackageUtilities::AsyncWriteFile(MoveTemp(DataPtr), DataSize, *ArchiveFilename, WriteOptions, Archive->FileRegions);
-				}
-			};
-
-			WriteBulkData(BulkArchive.Get(), LexToString(EPackageExtension::BulkDataDefault));
-			WriteBulkData(OptionalBulkArchive.Get(), LexToString(EPackageExtension::BulkDataOptional));
-			WriteBulkData(MappedBulkArchive.Get(), LexToString(EPackageExtension::BulkDataMemoryMapped));
-		}
-	}
-
-	if (!bSeparateSegmentsEnabled)
-	{
-		if (PackageWriter)
-		{
-			check(BulkArchive.IsValid() && LinkerStart == 0);
-			InOutStartOffset += BulkArchive->TotalSize();
-		}
-		else
-		{
-			check(!BulkArchive.IsValid() && LinkerStart > 0);
-			InOutStartOffset += Linker->Tell() - LinkerStart;
-		}
-	}
-
-	Linker->BulkDataToAppend.Empty();
-	return ESavePackageResult::Success;
-}
-
 void SaveWorldLevelInfo(UPackage* InOuter, FLinkerSave* Linker, FStructuredArchive::FRecord Record)
 {
 	Linker->Summary.WorldTileInfoDataOffset = 0;
diff --git a/Engine/Source/Runtime/CoreUObject/Private/UObject/SavePackage/SavePackageUtilities.h b/Engine/Source/Runtime/CoreUObject/Private/UObject/SavePackage/SavePackageUtilities.h
--- a/Engine/Source/Runtime/CoreUObject/Private/UObject/SavePackage/SavePackageUtilities.h
+++ b/Engine/Source/Runtime/CoreUObject/Private/UObject/SavePackage/SavePackageUtilities.h
@@ -352,29 +352,6 @@ namespace SavePackageUtilities
 
 	void SaveThumbnails(UPackage* InOuter, FLinkerSave* Linker, FStructuredArchive::FSlot Slot);
 
-	/**
-	 * Save all of the BulkDatas that have registered as not being inline.
-	 * They may be saved to the end of the file, or to an archive for a separate file.
-	 * Kicks off async writes, use UPackage::WaitForAsyncFileWrites(); to ensure they are complete.
-	 * 
-	 * @param Linker The linker containing the exports. Provides metadata for BulkData, and BulkData may write to it as their target archive.
-	 * @param InOutStartOffset In value is the offset in the Linker's archive where the BulkDatas will be put. If SavePackageContext settings direct
-	 *        the bulkdatas to write in a separate archive that will be combined after the linker, the value is the Linker archive's totalsize.
-	 *        Output value is incremented by the number of bytes written the Linker or the separate archive at the end of the linker.
-	 *        Bytes written to separate files do not contribute to this offset.
-	 * @param InOuter The package being saved
-	 * @param Filename The filename the package is being saved to. Used to decide the filename of separate archives, if any are used.
-	 * @param SavePackageContext Provides The PackageWriter and other parameters of how the bulkdata should be saved.
-	 * @param SaveFlags The flags passed into SavePackage; affects how the bulkdata should be saved.
-	 * @param bTextFormat True if package is saving to text. Bulkdata has special handling for text output.
-	 * @param bDiffing True if the package is only diffing, so bulk data should be recorded but not saved to disk.
-	 * @param TotalPackageSizeUncompressed Output: Bulkdata sizes are added to it.
-	 */
-	ESavePackageResult SaveBulkData(FLinkerSave* Linker, int64& InOutStartOffset, const UPackage* InOuter,
-		const TCHAR* Filename, const ITargetPlatform* TargetPlatform, FSavePackageContext* SavePackageContext,
-		uint32 SaveFlags, const bool bTextFormat,
-		int64& TotalPackageSizeUncompressed, bool bIsOptionalRealm = false);
-	
 	/**
 	 * Used to append additional data to the end of the package file by invoking callbacks stored in the linker.
 	 * They may be saved to the end of the file, or to a separate archive passed into the PackageWriter.
diff --git a/Engine/Source/Runtime/CoreUObject/Private/UObject/SavePackage2.cpp b/Engine/Source/Runtime/CoreUObject/Private/UObject/SavePackage2.cpp
--- a/Engine/Source/Runtime/CoreUObject/Private/UObject/SavePackage2.cpp
+++ b/Engine/Source/Runtime/CoreUObject/Private/UObject/SavePackage2.cpp
@@ -11,6 +11,7 @@
 #include "HAL/FileManager.h"
 #include "Internationalization/TextPackageNamespaceUtil.h"
 #include "Internationalization/PackageLocalizationManager.h"
+#include "Interfaces/ITargetPlatform.h"
 #include "IO/IoDispatcher.h"
 #include "Misc/AssetRegistryInterface.h"
 #include "Misc/ConfigCacheIni.h"
@@ -1000,8 +1001,24 @@ ESavePackageResult BuildLinker(FSaveContext& SaveContext)
 		if (SaveContext.IsCooking())
 		{
 			SaveContext.GetLinker()->SetDebugSerializationFlags(DSF_EnableCookerWarnings | SaveContext.GetLinker()->GetDebugSerializationFlags());
+			SaveContext.GetLinker()->SetSaveBulkDataToSeparateFiles(true);
+
+			if (const ITargetPlatform* TargetPlatform = SaveContext.GetTargetPlatform())
+			{
+				if (TargetPlatform->SupportsFeature(ETargetPlatformFeatures::MemoryMappedFiles))
+				{
+					SaveContext.GetLinker()->SetMemoryMapAlignment(TargetPlatform->GetMemoryMappingAlignment());
+				}
+				SaveContext.GetLinker()->SetFileRegionsEnabled(TargetPlatform->SupportsFeature(ETargetPlatformFeatures::CookFileRegionMetadata));
+			}
 		}
 #endif
+		if (SaveContext.IsCooking() == false)
+		{
+			const bool bSaveBulkDataByReference = (SaveContext.GetSaveArgs().SaveFlags & SAVE_BulkDataByReference) != 0;
+			SaveContext.GetLinker()->SetSaveBulkDataByReference(bSaveBulkDataByReference);
+		}
+
 		// Make sure the package has the same version as the linker
 		SaveContext.UpdatePackageLinkerVersions();
 	}
@@ -1615,6 +1632,24 @@ void WriteGatherableText(FStructuredArchive::FRecord& StructuredArchiveRoot, FSa
 	}
 }
 
+int64 WriteObjectDataResources(TArray<FObjectDataResource>& DataResources, FStructuredArchive::FRecord& StructuredArchiveRoot, FSaveContext& SaveContext)
+{
+	check(SaveContext.GetLinker());
+	FLinkerSave& Linker = *SaveContext.GetLinker();
+
+	// The data resource table is only saved for cooked output
+	if (Linker.IsCooking() == false || DataResources.IsEmpty())
+	{
+		Linker.Summary.DataResourceOffset = -1; 
+		return 0;
+	}
+	
+	Linker.Summary.DataResourceOffset = (int32)Linker.Tell();
+	FObjectDataResource::Serialize(StructuredArchiveRoot.EnterField(TEXT("DataResources")), DataResources);
+
+	return Linker.Tell() - Linker.Summary.DataResourceOffset;
+}
+
 ESavePackageResult WritePackageHeader(FStructuredArchive::FRecord& StructuredArchiveRoot, FSaveContext& SaveContext)
 {
 	FLinkerSave* Linker = SaveContext.GetLinker();
@@ -1815,6 +1850,67 @@ ESavePackageResult WritePackageTextHeader(FStructuredArchive::FRecord& Structure
 	return ReturnSuccessOrCancel();
 }
 
+[[nodiscard]] ESavePackageResult WriteCookedExports(TArray64<uint8>& ExportsBuffer, FSaveContext& SaveContext)
+{
+	SCOPED_SAVETIMER(UPackage_Save_SaveExports);
+
+	check(SaveContext.GetLinker() && SaveContext.GetLinker()->IsCooking());
+	FLinkerSave& Linker = *SaveContext.GetLinker();
+	FScopedSlowTask SlowTask((float)Linker.ExportMap.Num(), FText(), SaveContext.IsUsingSlowTask());
+	TArray64<uint8> ExportBuffer;
+
+	for (int32 i = 0; i < Linker.ExportMap.Num(); i++)
+	{
+		if (GWarn->ReceivedUserCancel())
+		{
+			return ESavePackageResult::Canceled;
+		}
+		SlowTask.EnterProgressFrame();
+
+		FObjectExport& Export = Linker.ExportMap[i];
+		if (Export.Object == nullptr)
+		{
+			continue;
+		}
+			
+		SCOPED_SAVETIMER(UPackage_Save_SaveExport);
+		SCOPED_SAVETIMER_TEXT(*WriteToString<128>(GetClassTraceScope(Export.Object), TEXT("_SaveSerialize")));
+
+		Export.SerialOffset = ExportsBuffer.Num();
+		Linker.CurrentlySavingExport = FPackageIndex::FromExport(i);
+
+		// Use a separate archive for each export to enable rearranging the exports according to the load/dependency order. Seeks/Tell
+		// during export serialize will be relative to the beginning of the export buffer.
+		ExportBuffer.Reset();
+		FMemoryWriter64 ExportArchive(ExportBuffer);
+		TGuardValue<FArchive*> GuardSaver(Linker.Saver, &ExportArchive);
+
+		if (Export.Object->HasAnyFlags(RF_ClassDefaultObject))
+		{
+			Export.Object->GetClass()->SerializeDefaultObject(Export.Object, Linker);
+		}
+		else
+		{
+			TGuardValue<UObject*> GuardSerializedObject(SaveContext.GetSerializeContext()->SerializedObject, Export.Object);
+			Export.Object->Serialize(Linker);
+#if WITH_EDITOR
+			Export.Object->CookAdditionalFiles(SaveContext.GetFilename(), SaveContext.GetTargetPlatform(),
+				[&SaveContext](const TCHAR* Filename, void* Data, int64 Size)
+				{
+					FLargeMemoryWriter& Writer = SaveContext.AdditionalFilesFromExports.Emplace_GetRef(0, true, Filename);
+					Writer.Serialize(Data, Size);
+				});
+#endif
+		}
+
+		Linker.CurrentlySavingExport = FPackageIndex();
+		Export.SerialSize = ExportBuffer.Num(); 
+		ExportsBuffer.Append(ExportBuffer);
+	}
+
+	return Linker.IsError() ? ESavePackageResult::Error : ReturnSuccessOrCancel();
+}
+
 ESavePackageResult WriteExports(FStructuredArchive::FRecord& StructuredArchiveRoot, FSaveContext& SaveContext)
 {
 	//COOK_STAT(FScopedDurationTimer SaveTimer(FSavePackageStats::SerializeExportsTimeSec));
@@ -1902,6 +1998,117 @@ ESavePackageResult WriteExports(FStructuredArchive::FRecord& StructuredArchiveRo
 	return Linker->IsError() ? ESavePackageResult::Error : ReturnSuccessOrCancel();
 }
 
+[[nodiscard]] ESavePackageResult WriteBulkData(FSaveContext& SaveContext)
+{
+	COOK_STAT(FScopedDurationTimer SaveTimer(FSavePackageStats::SerializeBulkDataTimeSec));
+
+	FLinkerSave& Linker = *SaveContext.GetLinker();
+
+	Linker.Summary.BulkDataStartOffset = Linker.Tell();
+
+	const bool bIsOptionalRealm = SaveContext.GetCurrentHarvestingRealm() == ESaveRealm::Optional;
+
+	IPackageWriter* PackageWriter = SaveContext.GetPackageWriter();
+
+	if (PackageWriter == nullptr)
+	{
+		// Saving non-inline bulk data to the end of the package (Editor)
+
+		check(Linker.IsCooking() == false);
+		check(bIsOptionalRealm == false);
+
+		FFileRegionMemoryWriter& Ar = Linker.GetBulkDataArchive();
+		if (const int64 TotalSize = Ar.TotalSize(); TotalSize > 0)
+		{
+			FIoBuffer Buffer(FIoBuffer::AssumeOwnership, Ar.ReleaseOwnership(), TotalSize);
+			Linker.Serialize(Buffer.GetData(), Buffer.GetSize());
+			SaveContext.TotalPackageSizeUncompressed += TotalSize;
+		}
+
+		return ESavePackageResult::Success;
+	}
+
+	// Saving non-inline bulk data to seperate file(s) (Cooking/EditorDomain)
+
+	if (Linker.IsCooking() == false)
+	{
+		const bool bSaveBulkDataByReference = SaveContext.GetSaveArgs().SaveFlags & SAVE_BulkDataByReference;
+		if (bSaveBulkDataByReference)
+		{
+			if (Linker.bUpdatingLoadedPath)
+			{
+				UE_LOG(LogSavePackage, Error, TEXT("Save bulk data '%s' FAILED, reason '%s'"),
+					SaveContext.GetFilename(), TEXT("SAVE_BulkDataByReference is incompatible with bUpdatingLoadedPath"));
+				return ESavePackageResult::Error;
+			}
+
+			// Nothing to do, load bulk data from the .uasset source file
+			return ESavePackageResult::Success;
+		}
+	}
+
+	const FName PackageName = SaveContext.GetPackage()->GetFName();
+	const FPackageId PackageId = FPackageId::FromName(PackageName);
+	const uint16 MultiOutputIndex = bIsOptionalRealm ? 1 : 0;
+
+	auto GetFilePath = [&SaveContext, bIsOptionalRealm](EPackageExtension Ext) -> FString
+	{
+		FString FileExt = bIsOptionalRealm ? FString(TEXT(".o")) + LexToString(Ext) : LexToString(Ext);
+		return FPathViews::ChangeExtension(SaveContext.GetFilename(), FileExt);
+	};
+
+	auto WriteToPackageWriter = [&SaveContext, PackageWriter, bIsOptionalRealm](FFileRegionMemoryWriter& Ar, IPackageWriter::FBulkDataInfo Info) -> int64
+	{
+		if (int64 TotalSize = Ar.TotalSize(); TotalSize > 0)
+		{
+			checkf(bIsOptionalRealm == false || Info.BulkDataType != IPackageWriter::FBulkDataInfo::Mmap,
+				TEXT("Memory mapped bulk data is currently not supported for optional package '%s'"), *SaveContext.GetPackage()->GetName());
+
+			FIoBuffer Buffer(FIoBuffer::AssumeOwnership, Ar.ReleaseOwnership(), TotalSize);
+			PackageWriter->WriteBulkData(Info, Buffer, Ar.GetFileRegions());
+			return TotalSize;
+		}
+		return 0;
+	};
+
+	FScopedSlowTask Feedback(3.0f);
+	SaveContext.TotalPackageSizeUncompressed += WriteToPackageWriter(Linker.GetBulkDataArchive(), IPackageWriter::FBulkDataInfo
+	{
+		PackageName,
+		IPackageWriter::FBulkDataInfo::BulkSegment,
+		GetFilePath(EPackageExtension::BulkDataDefault),
+		CreateIoChunkId(PackageId.Value(), MultiOutputIndex, EIoChunkType::BulkData),
+		MultiOutputIndex
+	});
+	Feedback.EnterProgressFrame();
+
+	// @note FH: temporarily do not handle optional bulk data into editor optional packages, proper support will be added soon
+	if (bIsOptionalRealm == false)
+	{
+		SaveContext.TotalPackageSizeUncompressed += WriteToPackageWriter(Linker.GetOptionalBulkDataArchive(), IPackageWriter::FBulkDataInfo
+		{
+			PackageName,
+			IPackageWriter::FBulkDataInfo::Optional,
+			GetFilePath(EPackageExtension::BulkDataOptional),
+			CreateIoChunkId(PackageId.Value(), MultiOutputIndex, EIoChunkType::OptionalBulkData),
+			MultiOutputIndex
+		});
+	}
+	Feedback.EnterProgressFrame();
+
+	SaveContext.TotalPackageSizeUncompressed += WriteToPackageWriter(Linker.GetMemoryMappedBulkDataArchive(), IPackageWriter::FBulkDataInfo
+	{
+		PackageName,
+		IPackageWriter::FBulkDataInfo::Mmap,
+		GetFilePath(EPackageExtension::BulkDataMemoryMapped),
+		CreateIoChunkId(PackageId.Value(), MultiOutputIndex, EIoChunkType::MemoryMappedBulkData),
+		MultiOutputIndex
+	});
+	Feedback.EnterProgressFrame();
+
+	return ESavePackageResult::Success;
+}
+
 [[nodiscard]] ESavePackageResult BuildAndWriteTrailer(IPackageWriter* PackageWriter, FStructuredArchive::FRecord& StructuredArchiveRoot, FSaveContext& SaveContext, int64& InOutCurrentOffset)
 {
 	SaveContext.GetLinker()->Summary.PayloadTocOffset = INDEX_NONE;
@@ -2349,7 +2556,7 @@ ESavePackageResult SaveHarvestedRealms(FSaveContext& SaveContext, ESaveRealm Har
 	FSaveContext::FSetSaveRealmToSaveScope Scope(SaveContext, HarvestingContextToSave);
 
 	// Create slow task dialog if needed
-	const int32 TotalSaveSteps = 11;
+	const int32 TotalSaveSteps = 12;
 	FScopedSlowTask SlowTask(TotalSaveSteps, FText(), SaveContext.IsUsingSlowTask());
 
 	// Validate Exports
@@ -2409,13 +2616,62 @@ ESavePackageResult SaveHarvestedRealms(FSaveContext& SaveContext, ESaveRealm Har
 		}
 	}
 
-	// Write Exports
-	SlowTask.EnterProgressFrame();
-	SaveContext.Result = WriteExports(StructuredArchiveRoot, SaveContext);
-	if (SaveContext.Result != ESavePackageResult::Success)
+	// Write Exports 
 	{
-		return SaveContext.Result;
+		SlowTask.EnterProgressFrame();
+		FLinkerSave& Linker = *SaveContext.GetLinker();
+
+		if (Linker.IsCooking())
+		{
+			// Write the exports into a seperate archive
+			TArray64<uint8> ExportsBuffer;
+			ExportsBuffer.Reserve(4 << 10);
+			SaveContext.Result = WriteCookedExports(ExportsBuffer, SaveContext);
+
+			if (SaveContext.Result == ESavePackageResult::Success)
+			{
+				// Write the data resource table to the header section before appending the export(s)
+				const int64 DataResourceSize = WriteObjectDataResources(Linker.DataResourceMap, StructuredArchiveRoot, SaveContext);
+				check(DataResourceSize >= 0);
+
+				Linker.Summary.TotalHeaderSize += DataResourceSize;
+				Linker.Serialize(ExportsBuffer.GetData(), ExportsBuffer.Num()); 
+
+				// Adjust the export offsets with the total header size
+				for (FObjectExport& Export : Linker.ExportMap)
+				{
+					if (Export.Object)
+					{
+						Export.SerialOffset += Linker.Summary.TotalHeaderSize;
+					}
+				}
+			}
+		}
+		else
+		{
+			SaveContext.Result = WriteExports(StructuredArchiveRoot, SaveContext);
+		}
+
+		if (SaveContext.Result != ESavePackageResult::Success)
+		{
+			return SaveContext.Result;
+		}
+	}
+
+	const int64 EndOfExportsOffset = SaveContext.GetLinker()->Tell();
+
+	// Write bulk data
+	{
+		SlowTask.EnterProgressFrame();
+
+		SaveContext.Result = WriteBulkData(SaveContext); 
+		if (SaveContext.Result != ESavePackageResult::Success)
+		{
+			return SaveContext.Result;
+		}
+		SaveContext.GetLinker()->OnPostSaveBulkData();
 	}
+
 	// Get SHA Key
 	{
 		// if we want to generate the SHA key, get it out now that the package has finished saving
@@ -2431,7 +2687,6 @@ ESavePackageResult SaveHarvestedRealms(FSaveContext& SaveContext, ESaveRealm Har
 	}
 
 	IPackageWriter* PackageWriter = SaveContext.GetPackageWriter();
-	const int64 EndOfExportsOffset = SaveContext.GetLinker()->Tell();
 	// When not using a PackageWriter, VirtualExportsFileOffset is identical to the offset in the
 	// Exports archive: SaveContext.GetLinker()->Tell. When using a PackageWriter however, additional
 	// blobs such as bulkdata are not written into the exports archive, they are stored as separate archives
@@ -2459,7 +2714,6 @@ ESavePackageResult SaveHarvestedRealms(FSaveContext& SaveContext, ESaveRealm Har
 		}
 	}
 
-
 	// Now that the package is written out we can write the package trailer that is appended
 	// to the file. This should be the last thing written to the file!
 	SlowTask.EnterProgressFrame();
@@ -2606,13 +2860,7 @@ ESavePackageResult InnerSave(FSaveContext& SaveContext)
  */
 ESavePackageResult WriteAdditionalFiles(FSaveContext& SaveContext, FScopedSlowTask& SlowTask, int64& VirtualExportsFileOffset)
 {
-	// Save Bulk Data
 	SlowTask.EnterProgressFrame();
-	SaveContext.Result = SavePackageUtilities::SaveBulkData(SaveContext.GetLinker(), VirtualExportsFileOffset,
-		SaveContext.GetPackage(), SaveContext.GetFilename(), SaveContext.GetTargetPlatform(),
-		SaveContext.GetSavePackageContext(), SaveContext.GetSaveArgs().SaveFlags, SaveContext.IsTextFormat(),
-		SaveContext.TotalPackageSizeUncompressed,
-		SaveContext.GetCurrentHarvestingRealm() == ESaveRealm::Optional);
 
 	// Add any pending data blobs to the end of the file by invoking the callbacks
 	ESavePackageResult Result = SavePackageUtilities::AppendAdditionalData(*SaveContext.GetLinker(), VirtualExportsFileOffset, SaveContext.GetSavePackageContext());
diff --git a/Engine/Source/Runtime/CoreUObject/Public/Serialization/AsyncLoading2.h b/Engine/Source/Runtime/CoreUObject/Public/Serialization/AsyncLoading2.h
--- a/Engine/Source/Runtime/CoreUObject/Public/Serialization/AsyncLoading2.h
+++ b/Engine/Source/Runtime/CoreUObject/Public/Serialization/AsyncLoading2.h
@@ -245,6 +245,7 @@ enum class EExportFilterFlags : uint8
 enum class EZenPackageVersion : uint32
 {
 	Initial,
+	DataResourceTable,
 
 	LatestPlusOne,
 	Latest = LatestPlusOne - 1
@@ -344,6 +345,17 @@ struct FExportMapEntry
 	COREUOBJECT_API friend FArchive& operator<<(FArchive& Ar, FExportMapEntry& ExportMapEntry);
 };
 
+struct FBulkDataMapEntry
+{
+	int64 SerialOffset = 0;
+	int64 DuplicateSerialOffset = 0;
+	int64 SerialSize = 0;
+	uint32 Flags = 0;
+	uint32 Pad = 0;
+	
+	COREUOBJECT_API friend FArchive& operator<<(FArchive& Ar, FBulkDataMapEntry& BulkDataEntry);
+};
+
 COREUOBJECT_API void FindAllRuntimeScriptPackages(TArray<UPackage*>& OutPackages);
 
 /**
diff --git a/Engine/Source/Runtime/CoreUObject/Public/Serialization/BulkData.h b/Engine/Source/Runtime/CoreUObject/Public/Serialization/BulkData.h
--- a/Engine/Source/Runtime/CoreUObject/Public/Serialization/BulkData.h
+++ b/Engine/Source/Runtime/CoreUObject/Public/Serialization/BulkData.h
@@ -46,6 +46,7 @@ class IAsyncReadRequest;
 class UObject;
 struct FTimespan;
 namespace UE { namespace Serialization { class FEditorBulkData; } }
+enum class EFileRegionType : uint8;
 
 #if WITH_EDITOR == 0 && WITH_EDITORONLY_DATA == 0
 	#define USE_NEW_BULKDATA UE_DEPRECATED_MACRO(5.1, "The USE_NEW_BULKDATA macro has been deprecated in favor of USE_RUNTIME_BULKDATA.") 1
@@ -183,22 +184,6 @@ inline FArchive& operator<<(FArchive& Ar, EBulkDataFlags& Flags)
 	return Ar;
 }
 
-/** Serialize the given Value as an int32 or int64 depending on InBulkDataFlags&BULKDATA_Size64Bit. */
-inline void SerializeBulkDataSizeInt(FArchive& Ar, int64& Value, EBulkDataFlags InBulkDataFlags)
-{
-	if (InBulkDataFlags & BULKDATA_Size64Bit)
-	{
-		Ar << Value;
-	}
-	else
-	{
-		check(!Ar.IsSaving() || (MIN_int32 <= Value && Value <= MAX_int32));
-		int32 ValueAsInt32 = static_cast<int32>(Value);
-		Ar << ValueAsInt32;
-		Value = ValueAsInt32;
-	}
-}
-
 /**
  * Enumeration for bulk data lock status.
  */
@@ -351,22 +336,11 @@ struct FBulkMetaResource
  *
  * Uses 5 bytes for size and offset.
  * [0 - 4][5 - 9][  10	][    11	][12    -   15]([16   -  23])
- * [Size][Offset][MetaFlags][LockFlags][BulkDataFlags]([SizeOnDisk])
+ * [Size][Offset][Unused][LockFlags][BulkDataFlags]([SizeOnDisk])
  */
 class FBulkMetaData
 {
 public:
-	enum class EMetaFlags : uint8
-	{
-		/** No additional bulk data flags. */
-		None				= 0,
-		/** Loading from a cooked package. */
-		CookedPackage		= (1 << 0),
-		/** Loading optional package data. */
-		OptionalPackage		= (1 << 1)
-	};
-	FRIEND_ENUM_CLASS_FLAGS(EMetaFlags);
-
 	/** 40 bits for max bulk data size. */
 	static constexpr int64 MaxSize = 0xFFffFFffFF;
 	/** 39 bits for max bulk data offset and 1 bit to indicate INDEX_NONE. */
@@ -380,11 +354,10 @@ public:
 #endif
 	}
 
-	explicit FBulkMetaData(EBulkDataFlags Flags, EMetaFlags MetaFlags)
+	explicit FBulkMetaData(EBulkDataFlags Flags)
 		: FBulkMetaData()
 	{
 		SetFlags(Flags);
-		SetMetaFlags(MetaFlags);
 	}
 
 	inline int64 GetSize() const
@@ -452,6 +425,16 @@ public:
 		*reinterpret_cast<uint32*>(&Data[12]) = static_cast<uint32>(Flags);
 	}
 
+	void AddFlags(EBulkDataFlags Flags)
+	{
+		*reinterpret_cast<uint32*>(&Data[12]) |= static_cast<uint32>(Flags);
+	}
+	
+	void ClearFlags(EBulkDataFlags Flags)
+	{
+		*reinterpret_cast<uint32*>(&Data[12]) &= ~static_cast<uint32>(Flags);
+	}
+
 	inline bool HasAnyFlags(EBulkDataFlags Flags) const
 	{
 		return (GetFlags() & Flags) != 0;
@@ -462,16 +445,6 @@ public:
 		return (GetFlags() & Flags) == Flags;
 	}
 
-	void SetMetaFlags(EMetaFlags MetaFlags)
-	{
-		Data[10] = static_cast<uint8>(MetaFlags);
-	}
-	
-	EMetaFlags GetMetaFlags() const
-	{
-		return static_cast<EMetaFlags>(Data[10]);
-	}
-	
 	static FBulkMetaData FromSerialized(const FBulkMetaResource& MetaResource, int64 ElementSize)
 	{
 		FBulkMetaData Meta;
@@ -485,6 +458,12 @@ public:
 		Meta.SetOffset(MetaResource.Offset);
 		Meta.SetFlags(MetaResource.Flags);
 
+		check(MetaResource.ElementCount <= 0 || Meta.GetSize() == MetaResource.ElementCount * ElementSize);
+		check(Meta.GetOffset() == MetaResource.Offset);
+		check(Meta.GetFlags() == MetaResource.Flags);
+#if !USE_RUNTIME_BULKDATA
+		check(MetaResource.ElementCount <= 0 || Meta.GetSizeOnDisk() == MetaResource.SizeOnDisk);
+#endif
 		return Meta;
 	}
 
@@ -505,14 +484,21 @@ private:
 #endif // USE_RUNTIME_BULKDATA
 };
 
-/** Returns an I/O chunk ID to be used when loading from I/O store. */
-FIoChunkId CreateBulkDataIoChunkId(const FBulkMetaData& BulkMeta, UObject* Owner);
-
-/** Returns the package segment when loading from the package resource manager. */
-COREUOBJECT_API EPackageSegment GetPackageSegmentFromFlags(const FBulkMetaData& BulkMeta);
-
 } // namespace UE::BulkData::Private
 
+/** Parameters when serializing bulk data. */
+struct FBulkDataSerializationParams
+{
+	/** The owner of the bulk data. */
+	UObject* Owner = nullptr;
+	/** Bulk data element size. */
+	int32 ElementSize = 0;
+	/** The region type. */
+	EFileRegionType RegionType;
+	/** Flag indicating whether to try to memory map the bulk data payload or not. */
+	bool bAttemptMemoryMapping = false;
+};
+
 /**
  * @documentation @todo documentation
  */
@@ -557,6 +543,7 @@ class COREUOBJECT_API FBulkData
 
 public:
 	friend class FLinkerLoad;
+	friend class FLinkerSave;
 	friend class FExportArchive;
 	friend class UE::Serialization::FEditorBulkData; // To allow access to AttachedAr
 	friend class FBulkDataBatchRequest;
@@ -1042,6 +1029,8 @@ public:
 	static void SetBulkDataFlagsOn(EBulkDataFlags& InOutAccumulator, EBulkDataFlags FlagsToSet);
 	/** Disable the given flags in the given accumulator variable. */
 	static void ClearBulkDataFlagsOn(EBulkDataFlags& InOutAccumulator, EBulkDataFlags FlagsToClear);
+	/** Returns whether all of the specified flags are set. */
+	static bool HasFlags(EBulkDataFlags Flags, EBulkDataFlags Contains);
 	/** Returns decompress method flags specified by the given bulk data flags. */
 	static FName GetDecompressionFormat(EBulkDataFlags InFlags);
 
@@ -1054,6 +1043,8 @@ protected:
 	void SerializeBulkData(FArchive& Ar, void* Data, int64 DataSize, EBulkDataFlags InBulkDataFlags);
 
 private:
+
+	int64 SerializePayload(FArchive& Ar, EBulkDataFlags SerializationFlags, EFileRegionType RegionType);
 #if WITH_EDITOR
 	/**
 	 * Detaches the bulk data from the passed in archive. Needs to match the archive we are currently
@@ -1092,9 +1083,6 @@ private:
 	
 	bool CanDiscardInternalData() const;
 	
-	/** Sets whether inline bulk data is allowed to be unloaded or not */
-	void ConditionalSetInlineAlwaysAllowDiscard(bool bPackageUsesIoStore);
-
 	/** Reallocate bulk data */
 	inline void* ReallocateData(SIZE_T SizeInBytes) { return DataAllocation.ReallocateData(this, SizeInBytes); }
 
diff --git a/Engine/Source/Runtime/CoreUObject/Public/UObject/LinkerLoad.h b/Engine/Source/Runtime/CoreUObject/Public/UObject/LinkerLoad.h
--- a/Engine/Source/Runtime/CoreUObject/Public/UObject/LinkerLoad.h
+++ b/Engine/Source/Runtime/CoreUObject/Public/UObject/LinkerLoad.h
@@ -44,6 +44,7 @@ class UObject;
 class UPackage;
 class UStruct;
 namespace UE::Serialization{ class FEditorBulkData; }
+namespace UE::BulkData::Private{ class FBulkMetaData; }
 namespace UE{ class FPackageTrailer; }
 struct FObjectPtr;
 struct FOpenPackageResult;
@@ -380,8 +381,12 @@ public:
 	 */
 	COREUOBJECT_API static void AddGameNameRedirect(const FName OldName, const FName NewName);
 
+	virtual bool SerializeBulkData(FBulkData& BulkData, const FBulkDataSerializationParams& Params) override;
+
 private:
 
+	void SerializeBulkMeta(UE::BulkData::Private::FBulkMetaData& Meta, int64& DuplicateSerialOffset, int32 ElementSize);
+
 	// Variables used during async linker creation.
 
 	/** Current index into soft object path list, used by async linker creation for spreading out serializing soft object path entries.	*/
diff --git a/Engine/Source/Runtime/CoreUObject/Public/UObject/LinkerSave.h b/Engine/Source/Runtime/CoreUObject/Public/UObject/LinkerSave.h
--- a/Engine/Source/Runtime/CoreUObject/Public/UObject/LinkerSave.h
+++ b/Engine/Source/Runtime/CoreUObject/Public/UObject/LinkerSave.h
@@ -12,7 +12,7 @@
 #include "IO/IoHash.h"
 #include "Serialization/Archive.h"
 #include "Serialization/ArchiveUObject.h"
-#include "Serialization/FileRegions.h"
+#include "Serialization/FileRegionArchive.h"
 #include "Templates/Function.h"
 #include "Templates/RefCounting.h"
 #include "Templates/UniquePtr.h"
@@ -77,27 +77,8 @@ public:
 	/** Save context associated with this linker */
 	TRefCountPtr<FUObjectSerializeContext> SaveContext;
 
-	/** List of bulkdata that needs to be stored at the end of the file */
-	struct FBulkDataStorageInfo
-	{
-		/** Offset to the location where the payload offset is stored */
-		int64 BulkDataOffsetInFilePos;
-		/** Offset to the location where the payload size is stored */
-		int64 BulkDataSizeOnDiskPos;
-		/** Offset to the location where the bulk data flags are stored */
-		int64 BulkDataFlagsPos;
-		/** Bulk data flags at the time of serialization */
-		uint32 BulkDataFlags;
-		/** The file region type to apply to this bulk data */
-		EFileRegionType BulkDataFileRegionType;
-		/** The bulkdata */
-		FBulkData* BulkData;
-	};
-	TArray<FBulkDataStorageInfo> BulkDataToAppend;
 	TArray<FFileRegion> FileRegions;
 
-	// TODO: Look into removing BulkDataToAppend and use AdditionalDataToAppend instead.
-
 	/**
 	 * Callback for arbitrary serializers to append data to the end of the ExportsArchive.
 	 * Some PackageWriters used by SavePackage will write this data to a separate archive.
@@ -192,6 +173,29 @@ public:
 	 */
 	virtual void SetFilterEditorOnly(bool bInFilterEditorOnly) override;
 
+	/** Set target platform memory map alignment. A negative value disables memory mapped bulk data. */
+	void SetMemoryMapAlignment(int64 InAlignment)
+	{
+		MemoryMappingAlignment = InAlignment;
+	}
+
+	/** Sets whether file regions will be written or not. */
+	void SetFileRegionsEnabled(bool bEnabled)
+	{
+		bFileRegionsEnabled = bEnabled;
+	}
+
+	/** Sets whether saving bulk data by reference, i.e. leaving the bulk data payload in the original .uasset file when using EditorDomain. */
+	void SetSaveBulkDataByReference(bool bValue)
+	{
+		bSaveBulkDataByReference = bValue;
+	}
+	
+	/** Sets whether bulk data will be stored in the Linker archive (.uasset) or in sepearate files (.ubulk, .m.ubulk, .opt.ubulk) .*/ 
+	void SetSaveBulkDataToSeparateFiles(bool bValue)
+	{
+		bSaveBulkDataToSeparateFiles = bValue;
+	}
 
 #if WITH_EDITOR
 	// proxy for debugdata
@@ -214,6 +218,9 @@ public:
 
 	/** Invoke all of the callbacks in PostSaveCallbacks and then empty it. */
 	void OnPostSave(const FPackagePath& PackagePath, FObjectPostSaveContext ObjectSaveContext);
+	
+	/** Triggered after bulk data payloads has been serialized to disk/package writer.*/
+	void OnPostSaveBulkData();
 
 	// FLinker interface
 	virtual FString GetDebugName() const override;
@@ -254,6 +261,23 @@ public:
 	UE::FDerivedData AddDerivedData(const UE::FDerivedData& Data);
 #endif // WITH_EDITORONLY_DATA
 
+	virtual bool SerializeBulkData(FBulkData& BulkData, const FBulkDataSerializationParams& Params) override;
+	
+	FFileRegionMemoryWriter& GetBulkDataArchive()
+	{
+		return BulkDataAr;
+	}
+
+	FFileRegionMemoryWriter& GetOptionalBulkDataArchive()
+	{
+		return OptionalBulkDataAr;
+	}
+
+	FFileRegionMemoryWriter& GetMemoryMappedBulkDataArchive()
+	{
+		return MemoryMappedBulkDataAr;
+	}
+
 protected:
 	/** Set the filename being saved to */
 	void SetFilename(FStringView InFilename);
@@ -266,4 +290,20 @@ private:
 	/** The index of the last derived data chunk added to the package. */
 	int32 LastDerivedDataIndex = -1;
 #endif
+	/** Map from bulk data object to resource index. */
+	TMap<FBulkData*, int32> SerializedBulkData;
+	/** Default bulk data archive. */
+	FFileRegionMemoryWriter BulkDataAr;
+	/** Optional bulk data archive. */
+	FFileRegionMemoryWriter OptionalBulkDataAr;
+	/** Memory mapped bulk data archive. */
+	FFileRegionMemoryWriter MemoryMappedBulkDataAr;
+	/** Alignment for memory mapped data .*/ 
+	int64 MemoryMappingAlignment = -1;
+	/** Whether file regions are enabled. */
+	bool bFileRegionsEnabled = false;
+	/** Whether saving bulk data by reference. */
+	bool bSaveBulkDataByReference = false;
+	/** Whether saving bulk data to separate files. */
+	bool bSaveBulkDataToSeparateFiles = false;
 };
