Add proper support for null dynamic delegates on struct pins and struct default values
This accidentally worked before (except on make struct nodes) but warned and tried to bind to a function named (null).None

diff --git a/Engine/Source/Editor/BlueprintGraph/Private/EdGraphSchema_K2.cpp b/Engine/Source/Editor/BlueprintGraph/Private/EdGraphSchema_K2.cpp
--- a/Engine/Source/Editor/BlueprintGraph/Private/EdGraphSchema_K2.cpp
+++ b/Engine/Source/Editor/BlueprintGraph/Private/EdGraphSchema_K2.cpp
@@ -4421,6 +4421,14 @@ bool UEdGraphSchema_K2::DefaultValueSimpleValidation(const FEdGraphPinType& PinT
 	{
 		// Anything is allowed
 	}
+	else if (PinCategory == TEXT("Delegate"))
+	{
+		// Only empty delegates are allowed, support both empty string and the format used in ExportText
+		if (!NewDefaultValue.IsEmpty() && NewDefaultValue != TEXT("(null).None"))
+		{
+			DVSV_RETURN_MSG(FString::Printf(TEXT("Unsupported value %s for delegate pin %s, only empty delegates are supported"), *NewDefaultValue, *PinName.ToString()));
+		}
+	}
 	else
 	{
 		//@TODO: MessageLog.Error(*FString::Printf(TEXT("Unsupported type %s on @@"), *UEdGraphSchema_K2::TypeToText(Type).ToString()), SourceObject);
diff --git a/Engine/Source/Editor/KismetCompiler/Private/KismetCompilerVMBackend.cpp b/Engine/Source/Editor/KismetCompiler/Private/KismetCompilerVMBackend.cpp
--- a/Engine/Source/Editor/KismetCompiler/Private/KismetCompilerVMBackend.cpp
+++ b/Engine/Source/Editor/KismetCompiler/Private/KismetCompilerVMBackend.cpp
@@ -1023,17 +1023,16 @@ public:
 			}
 			else if (FLiteralTypeHelper::IsDelegate(&Term->Type, CoerceProperty))
 			{
-				if (Term->Name == TEXT(""))
-				{
-					ensureMsgf(false, TEXT("Cannot use an empty literal expression for a delegate property"));
-				}
-				else
-				{
-					FName FunctionName(*(Term->Name)); //@TODO: K2 Delegate Support: Need to verify this function actually exists and has the right signature?
+				FName FunctionName;
 
-					Writer << EX_InstanceDelegate;
-					Writer << FunctionName;
+				// Deliberately null delegates are allowed, using empty string or the ExportText format
+				if (Term->Name != TEXT("") && Term->Name != TEXT("(null).None"))
+				{
+					FunctionName = *Term->Name; //@TODO: K2 Delegate Support: Need to verify this function actually exists and has the right signature?
 				}
+	
+				Writer << EX_InstanceDelegate;
+				Writer << FunctionName;
 			}
 			else if (FLiteralTypeHelper::IsSoftObject(&Term->Type, CoerceProperty))
 			{
diff --git a/Engine/Source/Runtime/CoreUObject/Private/UObject/PropertyHelper.h b/Engine/Source/Runtime/CoreUObject/Private/UObject/PropertyHelper.h
--- a/Engine/Source/Runtime/CoreUObject/Private/UObject/PropertyHelper.h
+++ b/Engine/Source/Runtime/CoreUObject/Private/UObject/PropertyHelper.h
@@ -185,6 +185,7 @@ namespace DelegatePropertyTools
 
 		TCHAR ObjName[NAME_SIZE];
 		TCHAR FuncName[NAME_SIZE];
+
 		// Get object name
 		int32 i;
 		while (*Buffer == TEXT('('))
@@ -198,6 +199,15 @@ namespace DelegatePropertyTools
 		ObjName[i] = TCHAR('\0');
 		UClass* Cls = nullptr;
 		UObject* Object = nullptr;
+		bool bIsNullClass = false;
+		
+		// Handle (null).None syntax used for empty delegates in ExportText
+		if (*Buffer == TCHAR(')') && FCString::Strcmp(ObjName, TEXT("null")) == 0)
+		{
+			++Buffer;
+			bIsNullClass = true;
+		}
+
 		// Get function name
 		if (*Buffer == TCHAR('.'))
 		{
@@ -235,6 +245,26 @@ namespace DelegatePropertyTools
 			}
 			FCString::Strncpy(FuncName, ObjName, i + 1);
 		}
+
+		if (bIsNullClass)
+		{
+			FName ParsedName(FuncName);
+			if (ParsedName == NAME_None)
+			{
+				// Deliberately null
+				Delegate.Unbind();
+				return Buffer;
+			}
+			else
+			{
+				FString DelegateString = Delegate.GetFunctionName().ToString();
+				FString SignatureFunctionString = GetNameSafe(SignatureFunction);
+				ErrorText->Logf(ELogVerbosity::Warning, TEXT("Cannot import delegate with function name %s but null class. Delegate=%s SignatureFunction=%s"),
+					*FuncName, *DelegateString, *SignatureFunctionString);
+				return nullptr;
+			}
+		}
+
 		if (Cls == nullptr)
 		{
 			Cls = UClass::TryFindTypeSlow<UClass>(ObjName);
