Silence bogus PVS warnings and fix checkSlow

#rnx

diff --git a/Engine/Source/Runtime/CoreUObject/Private/UObject/GarbageCollection.cpp b/Engine/Source/Runtime/CoreUObject/Private/UObject/GarbageCollection.cpp
--- a/Engine/Source/Runtime/CoreUObject/Private/UObject/GarbageCollection.cpp
+++ b/Engine/Source/Runtime/CoreUObject/Private/UObject/GarbageCollection.cpp
@@ -934,7 +934,7 @@ struct FResolvedMutableReference : FImmutableReference
 
 FORCEINLINE_DEBUGGABLE static bool ValidateReference(UObject* Object, FPermanentObjectPoolExtents PermanentPool, const UObject* ReferencingObject, FTokenId TokenId)
 {
-	bool bOk = (!PermanentPool.Contains(Object)) & (!!Object) & IsObjectHandleResolved(reinterpret_cast<FObjectHandle&>(Object));
+	bool bOk = (!PermanentPool.Contains(Object)) & (!!Object) & IsObjectHandleResolved(reinterpret_cast<FObjectHandle&>(Object)); //-V792
 
 #if ENABLE_GC_OBJECT_CHECKS
 	if (bOk)
@@ -1015,7 +1015,9 @@ protected:
 		{
 			uint32 Out = 0;
 			for (uint64 Word : Words)
+			{
 				Out += FPlatformMath::CountBits(Word);
+			}
 			return Out;
 		}
 
@@ -1023,7 +1025,9 @@ protected:
 		{
 			FValidatedBitmask Out;
 			for (uint32 WordIdx = 0; WordIdx < NumWords; ++WordIdx)
+			{
 				Out.Words[WordIdx] = A.Words[WordIdx] | B.Words[WordIdx];
+			}
 			return Out;
 		}
 
@@ -1031,12 +1035,14 @@ protected:
 		{
 			FValidatedBitmask Out;
 			for (uint32 WordIdx = 0; WordIdx < NumWords; ++WordIdx)
+			{
 				Out.Words[WordIdx] = A.Words[WordIdx] & B.Words[WordIdx];
+			}
 			return Out;
 		}
 
 
-		static const uint32 NumWords = 1 + (UnvalidatedBatchSize - 1) / 64;
+		static const uint32 NumWords = 1 + (UnvalidatedBatchSize - 1) / 64; //-V1064
 		uint64 Words[NumWords] = {};
 	};
 };
@@ -1103,7 +1109,7 @@ public:
 		if (Range.Num())
 		{
 			// Fill up UnvalidatedReferences slack and drain it to avoid inserting an element at a time and testing if full
-			for (uint32 Slack = UnvalidatedReferences.Slack(); (uint32)Range.Num() >= Slack; Slack = UnvalidatedBatchSize)
+			for (uint32 Slack = UnvalidatedReferences.Slack(); (uint32)Range.Num() >= Slack; Slack = UnvalidatedBatchSize) //-V1021
 			{
 				QueueUnvalidated(Range.Left(Slack));
 				Range.RightChopInline(Slack);
@@ -1207,7 +1213,7 @@ private:
 
 	FORCEINLINE_DEBUGGABLE void QueueValidReferences(uint32 NumToAppend, FValidatedBitmask Validations, uint32& InOutIdx)
 	{
-		checkSlow(NumToAppend == Validations.CountBits());
+		checkSlow(NumToAppend <= Validations.CountBits());
 		checkSlow(ValidatedReferences.Num + NumToAppend <= ValidatedReferences.Max());
 		uint32 NewQueueNum = ValidatedReferences.Num + NumToAppend;
 		for (uint32 QueueIdx = ValidatedReferences.Num; QueueIdx < NewQueueNum; ++InOutIdx)
