Silence more bogus PVS studio warnings

#rnx

diff --git a/Engine/Source/Runtime/CoreUObject/Private/UObject/GarbageCollection.cpp b/Engine/Source/Runtime/CoreUObject/Private/UObject/GarbageCollection.cpp
--- a/Engine/Source/Runtime/CoreUObject/Private/UObject/GarbageCollection.cpp
+++ b/Engine/Source/Runtime/CoreUObject/Private/UObject/GarbageCollection.cpp
@@ -1185,13 +1185,13 @@ private:
 		for (uint32 Idx = 0; Idx < Num; ++Idx)
 		{
 			UObject* Object = GetObject(UnvalidatedReferences[Idx].Reference);
-			ValidsB.Set(Idx, !!Object & IsObjectHandleResolved(reinterpret_cast<FObjectHandle&>(Object)));
+			ValidsB.Set(Idx, !!Object & IsObjectHandleResolved(reinterpret_cast<FObjectHandle&>(Object))); //-V792
 		}
 
 		FValidatedBitmask Validations = FValidatedBitmask::And(ValidsA, ValidsB);
 		uint32 NumValid = Validations.CountBits();
 		uint32 UnvalidatedIdx = 0;
-		for (uint32 Slack = ValidatedReferences.Slack(); NumValid >= Slack; Slack = ValidatedBatchSize)
+		for (uint32 Slack = ValidatedReferences.Slack(); NumValid >= Slack; Slack = ValidatedBatchSize) //-V1021
 		{
 			QueueValidReferences(Slack, Validations, /* in-out */  UnvalidatedIdx);
 			check(ValidatedReferences.IsFull());
@@ -2644,7 +2644,7 @@ public:
 		if (ValidateReference(Object, PermanentPool, ReferencingObject, TokenId))
 		{
 			FReferenceMetadata Metadata(GUObjectArray.ObjectToIndex(Object));
-			if (bAllowReferenceElimination & Metadata.Has(TReachabilityProcessor<Options>::KillFlag))
+			if (bAllowReferenceElimination & Metadata.Has(TReachabilityProcessor<Options>::KillFlag)) //-V792
 			{
 				if (MayKill<Options>(TokenType, bAllowReferenceElimination) == EKillable::Yes)
 				{
diff --git a/Engine/Source/Runtime/CoreUObject/Private/UObject/UObjectGlobals.cpp b/Engine/Source/Runtime/CoreUObject/Private/UObject/UObjectGlobals.cpp
--- a/Engine/Source/Runtime/CoreUObject/Private/UObject/UObjectGlobals.cpp
+++ b/Engine/Source/Runtime/CoreUObject/Private/UObject/UObjectGlobals.cpp
@@ -4737,7 +4737,7 @@ FORCEINLINE_DEBUGGABLE void CollectObjectReference(FReferenceCollector& Collecto
 	if constexpr (EnumHasAnyFlags(CollectFlags, EPropertyCollectFlags::NeedsReferencer))
 	{
 		// Sync reference processors will inspect Reference immediately so might as well avoid virtual call
-		if (!!Reference & IsObjectHandleResolved(*reinterpret_cast<FObjectHandle*>(Value)))
+		if (!!Reference & IsObjectHandleResolved(*reinterpret_cast<FObjectHandle*>(Value))) //-V792
 		{
 			Collector.AddReferencedObject(Reference, Referencer, &Property);
 		}
@@ -4791,8 +4791,8 @@ void CollectPropertyReferences(FReferenceCollector& Collector, FProperty& Proper
 	EPropertyFlags PropertyFlags = Property.GetPropertyFlags();
 	const EClassCastFlags CastFlags = static_cast<EClassCastFlags>(Class->GetCastFlags());
 	
-	if (EnumHasAnyFlags(CastFlags, GetCastFlagsToSkip(CollectFlags)) |
-		EnumHasAnyFlags(PropertyFlags, GetPropertyFlagsToSkip(CollectFlags)))
+	if (EnumHasAnyFlags(CastFlags, GetCastFlagsToSkip(CollectFlags)) |			//-V792
+		EnumHasAnyFlags(PropertyFlags, GetPropertyFlagsToSkip(CollectFlags)))	//-V792
 	{
 		return;
 	}
