[core] Added utils for creating soft object/class pointers (preferred over using implicit conversion)

diff --git a/Engine/Source/Runtime/CoreUObject/Public/UObject/SoftObjectPtr.h b/Engine/Source/Runtime/CoreUObject/Public/UObject/SoftObjectPtr.h
--- a/Engine/Source/Runtime/CoreUObject/Public/UObject/SoftObjectPtr.h
+++ b/Engine/Source/Runtime/CoreUObject/Public/UObject/SoftObjectPtr.h
@@ -415,6 +415,21 @@ FArchive& operator<<(FArchive& Ar, TSoftObjectPtr<T>& Ptr)
 template<class T> struct TIsPODType<TSoftObjectPtr<T> > { enum { Value = TIsPODType<FSoftObjectPtr>::Value }; };
 template<class T> struct TIsWeakPointerType<TSoftObjectPtr<T> > { enum { Value = TIsWeakPointerType<FSoftObjectPtr>::Value }; };
 
+/** Utility to create a TSoftObjectPtr without specifying the type */
+template <class T>
+TSoftObjectPtr<std::remove_cv_t<T>> MakeSoftObjectPtr(T* Object)
+{
+	static_assert(std::is_base_of_v<UObject, T>, "Type must derive from UObject");
+	return TSoftObjectPtr<std::remove_cv_t<T>>(Object);
+}
+
+template <class T>
+TSoftObjectPtr<std::remove_cv_t<T>> MakeSoftObjectPtr(TObjectPtr<T> Object)
+{
+	static_assert(std::is_base_of_v<UObject, T>, "Type must derive from UObject");
+	return TSoftObjectPtr<std::remove_cv_t<T>>(ToRawPtr(Object));
+}
+
 /**
  * TSoftClassPtr is a templatized wrapper around FSoftObjectPtr that works like a TSubclassOf, it can be used in UProperties for blueprint subclasses
  */
@@ -636,6 +651,21 @@ private:
 template <class T> struct TIsPODType<TSoftClassPtr<T> > { enum { Value = TIsPODType<FSoftObjectPtr>::Value }; };
 template <class T> struct TIsWeakPointerType<TSoftClassPtr<T> > { enum { Value = TIsWeakPointerType<FSoftObjectPtr>::Value }; };
 
+/** Utility to create a TSoftObjectPtr without specifying the type */
+template <class T>
+TSoftClassPtr<std::remove_cv_t<T>> MakeSoftClassPtr(T* Object)
+{
+	static_assert(std::is_base_of_v<UClass, T>, "Type must derive from UClass");
+	return TSoftClassPtr<std::remove_cv_t<T>>(Object);
+}
+
+template <class T>
+TSoftClassPtr<std::remove_cv_t<T>> MakeSoftClassPtr(TObjectPtr<T> Object)
+{
+	static_assert(std::is_base_of_v<UClass, T>, "Type must derive from UClass");
+	return TSoftClassPtr<std::remove_cv_t<T>>(ToRawPtr(Object));
+}
+
 /** Fast non-alphabetical order that is only stable during this process' lifetime. */
 struct FSoftObjectPtrFastLess : private FSoftObjectPathFastLess
 {
