Silence MSVC static analysis warnings

#rnx

diff --git a/Engine/Source/Runtime/CoreUObject/Private/UObject/GarbageCollection.cpp b/Engine/Source/Runtime/CoreUObject/Private/UObject/GarbageCollection.cpp
--- a/Engine/Source/Runtime/CoreUObject/Private/UObject/GarbageCollection.cpp
+++ b/Engine/Source/Runtime/CoreUObject/Private/UObject/GarbageCollection.cpp
@@ -1185,7 +1185,7 @@ private:
 		for (uint32 Idx = 0; Idx < Num; ++Idx)
 		{
 			UObject* Object = GetObject(UnvalidatedReferences[Idx].Reference);
-			ValidsB.Set(Idx, !!Object & IsObjectHandleResolved(reinterpret_cast<FObjectHandle&>(Object))); //-V792
+			ValidsB.Set(Idx, (!!Object) & IsObjectHandleResolved(reinterpret_cast<FObjectHandle&>(Object))); //-V792
 		}
 
 		FValidatedBitmask Validations = FValidatedBitmask::And(ValidsA, ValidsB);
@@ -1935,7 +1935,9 @@ FAROQueue::~FAROQueue()
 FORCEINLINE_DEBUGGABLE bool FAROQueue::TryPush(UObject* Object)
 {
 	uint32 HeadIdx = Head.load(std::memory_order_relaxed);
-	HeadBlock->Objects[HeadIdx % FAROBlock::NumWords] = Object;
+	uint32 ObjectsIdx = HeadIdx % FAROBlock::NumWords;
+	CA_ASSUME(ObjectsIdx < FAROBlock::Capacity);
+	HeadBlock->Objects[ObjectsIdx] = Object;
 	++HeadIdx;
 
 	// Must store OldBlock->FirstIndexInNextBlock before Head.store
@@ -2337,7 +2339,8 @@ public:
 			const int32 ObjectIndex = GUObjectArray.ObjectToIndex(Object);
 			FImmutableReference Reference = {Object};	
 			FReferenceMetadata Metadata(ObjectIndex);
-			if (Metadata.Has(KillFlag) & (Killable == EKillable::Yes))
+			bool bKillable = Killable == EKillable::Yes;
+			if (Metadata.Has(KillFlag) & bKillable) //-V792
 			{
 				check(ReferencingObject || IsWithPendingKill());
 				checkSlow(Metadata.ObjectItem->GetOwnerIndex() <= 0);
