Fixed editor nulling out Outer/Package causing "Reachable garbage object" errors in changelist/asset validation

#rnx

diff --git a/Engine/Source/Runtime/CoreUObject/Private/UObject/GarbageCollection.cpp b/Engine/Source/Runtime/CoreUObject/Private/UObject/GarbageCollection.cpp
--- a/Engine/Source/Runtime/CoreUObject/Private/UObject/GarbageCollection.cpp
+++ b/Engine/Source/Runtime/CoreUObject/Private/UObject/GarbageCollection.cpp
@@ -4401,8 +4401,10 @@ void UObject::AddReferencedObjects(UObject* This, FReferenceCollector& Collector
 		UObject* LoadOuter = This->GetOuter();
 		UClass* Class = This->GetClass();
 		UPackage* Package = This->GetExternalPackageInternal();
+		Collector.AllowEliminatingReferences(false);
 		Collector.AddReferencedObject(LoadOuter, This);
 		Collector.AddReferencedObject(Package, This);
+		Collector.AllowEliminatingReferences(true);
 		Collector.AddReferencedObject(Class, This);
 	}
 #endif
