Silence MSVC static analysis warnings

#rnx

diff --git a/Engine/Source/Runtime/CoreUObject/Private/UObject/UObjectGlobals.cpp b/Engine/Source/Runtime/CoreUObject/Private/UObject/UObjectGlobals.cpp
--- a/Engine/Source/Runtime/CoreUObject/Private/UObject/UObjectGlobals.cpp
+++ b/Engine/Source/Runtime/CoreUObject/Private/UObject/UObjectGlobals.cpp
@@ -4737,7 +4737,7 @@ FORCEINLINE_DEBUGGABLE void CollectObjectReference(FReferenceCollector& Collecto
 	if constexpr (EnumHasAnyFlags(CollectFlags, EPropertyCollectFlags::NeedsReferencer))
 	{
 		// Sync reference processors will inspect Reference immediately so might as well avoid virtual call
-		if (!!Reference & IsObjectHandleResolved(*reinterpret_cast<FObjectHandle*>(Value))) //-V792
+		if ((!!Reference) & IsObjectHandleResolved(*reinterpret_cast<FObjectHandle*>(Value))) //-V792
 		{
 			Collector.AddReferencedObject(Reference, Referencer, &Property);
 		}
