#UE Fix crash that can happen if loc data is gathered from a UDataTable using a RowStruct or a class containing a struct property whose struct couldn't be found or loaded (i.e. the struct was deleted, or otherwise missing)

diff --git a/Engine/Source/Runtime/CoreUObject/Private/Serialization/PropertyLocalizationDataGathering.cpp b/Engine/Source/Runtime/CoreUObject/Private/Serialization/PropertyLocalizationDataGathering.cpp
--- a/Engine/Source/Runtime/CoreUObject/Private/Serialization/PropertyLocalizationDataGathering.cpp
+++ b/Engine/Source/Runtime/CoreUObject/Private/Serialization/PropertyLocalizationDataGathering.cpp
@@ -84,6 +84,8 @@ const FPropertyLocalizationDataGatherer::FGatherableFieldsForType& FPropertyLoca
 
 const FPropertyLocalizationDataGatherer::FGatherableFieldsForType& FPropertyLocalizationDataGatherer::CacheGatherableFieldsForType(const UStruct* InType)
 {
+	check(InType);
+
 	TUniquePtr<FGatherableFieldsForType> GatherableFieldsForType = MakeUnique<FGatherableFieldsForType>();
 
 	// Include the parent fields (this will recursively cache any parent types)
@@ -171,8 +173,11 @@ bool FPropertyLocalizationDataGatherer::CanGatherFromInnerProperty(const FProper
 
 	if (const FStructProperty* StructInnerProp = CastField<const FStructProperty>(InInnerProperty))
 	{
-		// Call the "Get" version as we may have already cached a result for this type
-		return !GetGatherableFieldsForType(StructInnerProp->Struct).IsEmpty();
+		if (StructInnerProp->Struct)
+		{
+			// Call the "Get" version as we may have already cached a result for this type
+			return !GetGatherableFieldsForType(StructInnerProp->Struct).IsEmpty();
+		}
 	}
 
 	return false;
@@ -504,7 +509,10 @@ void FPropertyLocalizationDataGatherer::GatherLocalizationDataFromChildTextPrope
 		// Property is a struct property.
 		else if (StructProperty)
 		{
-			GatherLocalizationDataFromStructWithCallbacks(PathToElement, StructProperty->Struct, ElementValueAddress, DefaultElementValueAddress, ElementChildPropertyGatherTextFlags);
+			if (StructProperty->Struct)
+			{
+				GatherLocalizationDataFromStructWithCallbacks(PathToElement, StructProperty->Struct, ElementValueAddress, DefaultElementValueAddress, ElementChildPropertyGatherTextFlags);
+			}
 		}
 		// Property is an object property.
 		else if (ObjectProperty && !(GatherTextFlags & EPropertyLocalizationGathererTextFlags::SkipSubObjects))
diff --git a/Engine/Source/Runtime/Engine/Private/DataTable.cpp b/Engine/Source/Runtime/Engine/Private/DataTable.cpp
--- a/Engine/Source/Runtime/Engine/Private/DataTable.cpp
+++ b/Engine/Source/Runtime/Engine/Private/DataTable.cpp
@@ -29,11 +29,14 @@ namespace
 
 		PropertyLocalizationDataGatherer.GatherLocalizationDataFromObject(DataTable, GatherTextFlags);
 
-		const FString PathToObject = DataTable->GetPathName();
-		for (const auto& Pair : DataTable->GetRowMap())
+		if (DataTable->RowStruct)
 		{
-			const FString PathToRow = PathToObject + TEXT(".") + Pair.Key.ToString();
-			PropertyLocalizationDataGatherer.GatherLocalizationDataFromStructWithCallbacks(PathToRow, DataTable->RowStruct, Pair.Value, nullptr, GatherTextFlags);
+			const FString PathToObject = DataTable->GetPathName();
+			for (const auto& Pair : DataTable->GetRowMap())
+			{
+				const FString PathToRow = PathToObject + TEXT(".") + Pair.Key.ToString();
+				PropertyLocalizationDataGatherer.GatherLocalizationDataFromStructWithCallbacks(PathToRow, DataTable->RowStruct, Pair.Value, nullptr, GatherTextFlags);
+			}
 		}
 	}
 #endif // WITH_EDITORONLY_DATA
