Fix editor only soft object path not being loaded under the proper context and ultimately being cooked

diff --git a/Engine/Source/Runtime/CoreUObject/Private/UObject/LinkerLoad.cpp b/Engine/Source/Runtime/CoreUObject/Private/UObject/LinkerLoad.cpp
--- a/Engine/Source/Runtime/CoreUObject/Private/UObject/LinkerLoad.cpp
+++ b/Engine/Source/Runtime/CoreUObject/Private/UObject/LinkerLoad.cpp
@@ -1742,12 +1742,7 @@ FLinkerLoad::ELinkerStatus FLinkerLoad::SerializeSoftObjectPathList()
 	}
 
 #if WITH_EDITOR
-	FName PackageName, PropertyName;
-	ESoftObjectPathCollectType CollectType = ESoftObjectPathCollectType::AlwaysCollect;
-	ESoftObjectPathSerializeType SerializeType = ESoftObjectPathSerializeType::AlwaysSerialize;
-	FSoftObjectPathThreadContext& ThreadContext = FSoftObjectPathThreadContext::Get();
-	ThreadContext.GetSerializationOptions(PackageName, PropertyName, CollectType, SerializeType, this);
-	FSoftObjectPathSerializationScope SerializationScope(PackageName, PropertyName, ESoftObjectPathCollectType::NeverCollect, SerializeType);
+	FSoftObjectPathSerializationScope SerializationScope(NAME_None, NAME_None, ESoftObjectPathCollectType::NonPackage, ESoftObjectPathSerializeType::AlwaysSerialize);
 #endif // WITH_EDITOR
 
 	FStructuredArchive::FStream Stream = StructuredArchiveRootRecord->EnterStream(TEXT("SoftObjectPathList"));
