Improve some error messages by adding the package path

#rnx

diff --git a/Engine/Source/Runtime/CoreUObject/Private/Serialization/EditorBulkData.cpp b/Engine/Source/Runtime/CoreUObject/Private/Serialization/EditorBulkData.cpp
--- a/Engine/Source/Runtime/CoreUObject/Private/Serialization/EditorBulkData.cpp
+++ b/Engine/Source/Runtime/CoreUObject/Private/Serialization/EditorBulkData.cpp
@@ -1080,7 +1080,7 @@ void FEditorBulkData::Serialize(FArchive& Ar, UObject* Owner, bool bAllowRegiste
 				
 			if (IsStoredInPackageTrailer())
 			{
-				checkf(Trailer != nullptr, TEXT("Payload was stored in a package trailer, but there no trailer loaded"));
+				checkf(Trailer != nullptr, TEXT("Payload was stored in a package trailer, but there no trailer loaded. [%s]"), *GetDebugNameFromOwner(Owner));
 				// Cache the offset from the trailer (if we move the loading of the payload to the trailer 
 				// at a later point then we can skip this)
 				OffsetInFile = Trailer->FindPayloadOffsetInFile(PayloadContentId);
@@ -1090,7 +1090,7 @@ void FEditorBulkData::Serialize(FArchive& Ar, UObject* Owner, bool bAllowRegiste
 				// TODO: Consider removing in 5.2+
 				if (IsDataVirtualized() && Trailer->FindPayloadStatus(PayloadContentId) != EPayloadStatus::StoredVirtualized)
 				{
-					UE_LOG(LogSerialization, Warning, TEXT("Payload in '%s' is was saved with an invalid flag and required fixing. Please re-save the package!"), *GetDebugNameFromOwner(Owner));
+					UE_LOG(LogSerialization, Warning, TEXT("Payload was saved with an invalid flag and required fixing. Please re-save the package! [%s]"), *GetDebugNameFromOwner(Owner));
 					EnumRemoveFlags(Flags, EFlags::IsVirtualized);
 				}
 			}
@@ -1099,7 +1099,7 @@ void FEditorBulkData::Serialize(FArchive& Ar, UObject* Owner, bool bAllowRegiste
 				// This check is for older virtualized formats that might be seen in older test projects.
 				// But we only care if the archive has a linker! (loading from a package)
 				// TODO: Consider removing in 5.2+
-				UE_CLOG(IsDataVirtualized() && Ar.GetLinker() != nullptr, LogSerialization, Warning, TEXT("Payload in '%s' is virtualized in an older format and should be re-saved!"), *GetDebugNameFromOwner(Owner));
+				UE_CLOG(IsDataVirtualized() && Ar.GetLinker() != nullptr, LogSerialization, Warning, TEXT("Payload is virtualized in an older format and should be re-saved! [%s]"), *GetDebugNameFromOwner(Owner));
 				if (!IsDataVirtualized() && !EnumHasAnyFlags(Flags, EFlags::IsCooked))
 				{
 					Ar << OffsetInFile;
