Work around to load packages that were saved with incorrect FPackageFileSummary offsets

#rnx

- CL 23489396  introduced a problem for packages that contain both BulkData and EditorBulkData. When saved for the workspace domain FPackageFileSummary would be provided the wrong values for PayloadTocOffset (it would be a lower value than BulkDataStartOffset) causing the trailer to fail to load. When an editor bulkdata is later serialized it would be unable to find itself in the trailer (since the trailer doesn't exist) and we would assert as this should never be the case.
- Although the bug generating the bad data will be fixed, we can still load the trailer by jumping to the end of the package file and reading the trailer in reverse order, although this is a slower method which is why we prefer the forward reading method based on the stored offset.
- We will limit this fix to packages with the version number == EUnrealEngineObjectUE5Version::DATA_RESOURCES so that we only try this with packages that we know could have incorrect offsets.

diff --git a/Engine/Source/Runtime/CoreUObject/Private/UObject/LinkerLoad.cpp b/Engine/Source/Runtime/CoreUObject/Private/UObject/LinkerLoad.cpp
--- a/Engine/Source/Runtime/CoreUObject/Private/UObject/LinkerLoad.cpp
+++ b/Engine/Source/Runtime/CoreUObject/Private/UObject/LinkerLoad.cpp
@@ -1635,7 +1635,21 @@ FLinkerLoad::ELinkerStatus FLinkerLoad::SerializePackageTrailer()
 		Seek(Summary.PayloadTocOffset);
 		
 		PackageTrailer = MakeUnique<UE::FPackageTrailer>();
-		if (!PackageTrailer->TryLoad(*this))
+
+		bool bResult = PackageTrailer->TryLoad(*this);
+		if (!bResult && Summary.GetFileVersionUE().ToValue() == (int32)EUnrealEngineObjectUE5Version::DATA_RESOURCES)
+		{
+			// There was an issue that was causing incorrect values to be written to PayloadTocOffset for
+			// a limited time. In these cases we can try the slower ::TryLoadBackwards method of loading 
+			// the trailer. Note that we only do this if the FileVersion is 
+			// EUnrealEngineObjectUE5Version::DATA_RESOURCES as the bug was introduced while this was the
+			// current version, so any package with an older or newer version should be safe.
+
+			Seek(TotalSize());
+			bResult = PackageTrailer->TryLoadBackwards(*this);
+		}
+
+		if (!bResult)
 		{
 			// If the archive has an error then we found a package trailer but it failed to serialize
 			// correctly and we most likely have a problem with the file.
