SavePackage2 - fix bug with virtual export size offset

#rnx

diff --git a/Engine/Source/Runtime/CoreUObject/Private/UObject/SavePackage2.cpp b/Engine/Source/Runtime/CoreUObject/Private/UObject/SavePackage2.cpp
--- a/Engine/Source/Runtime/CoreUObject/Private/UObject/SavePackage2.cpp
+++ b/Engine/Source/Runtime/CoreUObject/Private/UObject/SavePackage2.cpp
@@ -1998,13 +1998,18 @@ ESavePackageResult WriteExports(FStructuredArchive::FRecord& StructuredArchiveRo
 	return Linker->IsError() ? ESavePackageResult::Error : ReturnSuccessOrCancel();
 }
 
-[[nodiscard]] ESavePackageResult WriteBulkData(FSaveContext& SaveContext)
+[[nodiscard]] ESavePackageResult WriteBulkData(FSaveContext& SaveContext, int64& VirtualExportsFileOffset)
 {
 	COOK_STAT(FScopedDurationTimer SaveTimer(FSavePackageStats::SerializeBulkDataTimeSec));
 
 	FLinkerSave& Linker = *SaveContext.GetLinker();
 
 	Linker.Summary.BulkDataStartOffset = Linker.Tell();
+	
+	if (Linker.IsCooking() == false)
+	{
+		VirtualExportsFileOffset += Linker.GetBulkDataArchive();
+	}
 
 	const bool bIsOptionalRealm = SaveContext.GetCurrentHarvestingRealm() == ESaveRealm::Optional;
 
@@ -2660,11 +2665,23 @@ ESavePackageResult SaveHarvestedRealms(FSaveContext& SaveContext, ESaveRealm Har
 
 	const int64 EndOfExportsOffset = SaveContext.GetLinker()->Tell();
 
+	// When not using a PackageWriter, VirtualExportsFileOffset is identical to the offset in the
+	// Exports archive: SaveContext.GetLinker()->Tell. When using a PackageWriter however, additional
+	// blobs such as bulkdata are not written into the exports archive, they are stored as separate archives
+	// in the PackageWriter. But various structs need to know the "offset" in the combined file that would
+	// be created by appending all of these blobs after the exports. VirtualExportsFileOffset holds that value.
+	int64 VirtualExportsFileOffset = EndOfExportsOffset;
+	if (SaveContext.GetLinker()->IsCooking() == false)
+	{
+		// When not cooking, the bulk data archive contains all streaming bulk data.
+		VirtualExportsFileOffset += SaveContext.GetLinker()->GetBulkDataArchive().TotalSize();
+	}
+
 	// Write bulk data
 	{
 		SlowTask.EnterProgressFrame();
 
-		SaveContext.Result = WriteBulkData(SaveContext); 
+		SaveContext.Result = WriteBulkData(SaveContext, VirtualExportsFileOffset); 
 		if (SaveContext.Result != ESavePackageResult::Success)
 		{
 			return SaveContext.Result;
@@ -2687,12 +2704,6 @@ ESavePackageResult SaveHarvestedRealms(FSaveContext& SaveContext, ESaveRealm Har
 	}
 
 	IPackageWriter* PackageWriter = SaveContext.GetPackageWriter();
-	// When not using a PackageWriter, VirtualExportsFileOffset is identical to the offset in the
-	// Exports archive: SaveContext.GetLinker()->Tell. When using a PackageWriter however, additional
-	// blobs such as bulkdata are not written into the exports archive, they are stored as separate archives
-	// in the PackageWriter. But various structs need to know the "offset" in the combined file that would
-	// be created by appending all of these blobs after the exports. VirtualExportsFileOffset holds that value.
-	int64 VirtualExportsFileOffset = EndOfExportsOffset;
 	SaveContext.Result = WriteAdditionalFiles(SaveContext, SlowTask, VirtualExportsFileOffset /* In/Out */);
 	if (SaveContext.Result != ESavePackageResult::Success)
 	{
