Savepackage2 - fix bad submit with virtual export offset

#rnx

diff --git a/Engine/Source/Runtime/CoreUObject/Private/UObject/SavePackage2.cpp b/Engine/Source/Runtime/CoreUObject/Private/UObject/SavePackage2.cpp
--- a/Engine/Source/Runtime/CoreUObject/Private/UObject/SavePackage2.cpp
+++ b/Engine/Source/Runtime/CoreUObject/Private/UObject/SavePackage2.cpp
@@ -2008,7 +2008,7 @@ ESavePackageResult WriteExports(FStructuredArchive::FRecord& StructuredArchiveRo
 	
 	if (Linker.IsCooking() == false)
 	{
-		VirtualExportsFileOffset += Linker.GetBulkDataArchive();
+		VirtualExportsFileOffset += Linker.GetBulkDataArchive().TotalSize();
 	}
 
 	const bool bIsOptionalRealm = SaveContext.GetCurrentHarvestingRealm() == ESaveRealm::Optional;
@@ -2671,11 +2671,6 @@ ESavePackageResult SaveHarvestedRealms(FSaveContext& SaveContext, ESaveRealm Har
 	// in the PackageWriter. But various structs need to know the "offset" in the combined file that would
 	// be created by appending all of these blobs after the exports. VirtualExportsFileOffset holds that value.
 	int64 VirtualExportsFileOffset = EndOfExportsOffset;
-	if (SaveContext.GetLinker()->IsCooking() == false)
-	{
-		// When not cooking, the bulk data archive contains all streaming bulk data.
-		VirtualExportsFileOffset += SaveContext.GetLinker()->GetBulkDataArchive().TotalSize();
-	}
 
 	// Write bulk data
 	{
