Editor Domain - fix fallouts after bulk data refactor

- allow file region of type None
- don't early out when saving bulk data by reference

#rnx

diff --git a/Engine/Source/Runtime/CoreUObject/Internal/Serialization/FileRegionArchive.h b/Engine/Source/Runtime/CoreUObject/Internal/Serialization/FileRegionArchive.h
--- a/Engine/Source/Runtime/CoreUObject/Internal/Serialization/FileRegionArchive.h
+++ b/Engine/Source/Runtime/CoreUObject/Internal/Serialization/FileRegionArchive.h
@@ -28,7 +28,6 @@ public:
 
 	virtual void PopFileRegionType() override final
 	{
-		check(CurrentRegion.Type != EFileRegionType::None);
 		CurrentRegion.Length = Tell() - CurrentRegion.Offset;
 		FileRegions.Add(CurrentRegion);
 		CurrentRegion = FFileRegion();
diff --git a/Engine/Source/Runtime/CoreUObject/Private/UObject/SavePackage2.cpp b/Engine/Source/Runtime/CoreUObject/Private/UObject/SavePackage2.cpp
--- a/Engine/Source/Runtime/CoreUObject/Private/UObject/SavePackage2.cpp
+++ b/Engine/Source/Runtime/CoreUObject/Private/UObject/SavePackage2.cpp
@@ -2046,9 +2046,6 @@ ESavePackageResult WriteExports(FStructuredArchive::FRecord& StructuredArchiveRo
 					SaveContext.GetFilename(), TEXT("SAVE_BulkDataByReference is incompatible with bUpdatingLoadedPath"));
 				return ESavePackageResult::Error;
 			}
-
-			// Nothing to do, load bulk data from the .uasset source file
-			return ESavePackageResult::Success;
 		}
 	}
 
