Add HasAsyncFileWrites() to see if there are pending writes without blocking

diff --git a/Engine/Source/Runtime/CoreUObject/Private/UObject/SavePackage/SavePackageUtilities.cpp b/Engine/Source/Runtime/CoreUObject/Private/UObject/SavePackage/SavePackageUtilities.cpp
--- a/Engine/Source/Runtime/CoreUObject/Private/UObject/SavePackage/SavePackageUtilities.cpp
+++ b/Engine/Source/Runtime/CoreUObject/Private/UObject/SavePackage/SavePackageUtilities.cpp
@@ -1979,6 +1979,11 @@ void UPackage::WaitForAsyncFileWrites()
 	}
 }
 
+bool UPackage::HasAsyncFileWrites()
+{
+	return OutstandingAsyncWrites.GetValue() > 0;
+}
+
 bool UPackage::IsEmptyPackage(UPackage* Package, const UObject* LastReferencer)
 {
 	// Don't count null or volatile packages as empty, just let them be NULL or get GCed
diff --git a/Engine/Source/Runtime/CoreUObject/Public/UObject/Package.h b/Engine/Source/Runtime/CoreUObject/Public/UObject/Package.h
--- a/Engine/Source/Runtime/CoreUObject/Public/UObject/Package.h
+++ b/Engine/Source/Runtime/CoreUObject/Public/UObject/Package.h
@@ -912,6 +912,9 @@ public:
 	/** Wait for any SAVE_Async file writes to complete **/
 	static void WaitForAsyncFileWrites();
 
+	/** Return true if SAVE_Async file writes are pending **/
+	static bool HasAsyncFileWrites();
+
 	/**
 	* Determines if a package contains no more assets.
 	*
