EditorDomain:  fix fallouts after bulk data refactor.
Allow EditorDomain to save BulkDataType FBulkDataInfo::BulkSegment.
Declare a Region for each BulkData passed into the EditorDomain PackageWriter.
Mark BulkDatas saved by reference as PayloadInSeparateFile.
Allow RegionTypeNone in SerializePayload.
#rnx

diff --git a/Engine/Source/Editor/UnrealEd/Private/EditorDomain/EditorDomainUtils.cpp b/Engine/Source/Editor/UnrealEd/Private/EditorDomain/EditorDomainUtils.cpp
--- a/Engine/Source/Editor/UnrealEd/Private/EditorDomain/EditorDomainUtils.cpp
+++ b/Engine/Source/Editor/UnrealEd/Private/EditorDomain/EditorDomainUtils.cpp
@@ -203,7 +203,7 @@ FBlake3Hash GGlobalConstructClassesHash;
 int64 GMaxBulkDataSize = -1;
 
 // Change to a new guid when EditorDomain needs to be invalidated
-const TCHAR* EditorDomainVersion = TEXT("382517D23ED840A48501AE70D95B3CE2");
+const TCHAR* EditorDomainVersion = TEXT("25CC763ED9BB4CEF9823E7B41FA7C771");
 
 // Identifier of the CacheBuckets for EditorDomain tables
 const TCHAR* EditorDomainPackageBucketName = TEXT("EditorDomainPackage");
@@ -1642,8 +1642,9 @@ protected:
 		BulkDataSize = 0;
 		for (const FBulkDataRecord& BulkRecord : Record.BulkDatas)
 		{
-			checkf(BulkRecord.Info.BulkDataType == IPackageWriter::FBulkDataInfo::AppendToExports,
-				TEXT("Does not support BulkData types other than AppendToExports."));
+			checkf(BulkRecord.Info.BulkDataType == IPackageWriter::FBulkDataInfo::AppendToExports ||
+				BulkRecord.Info.BulkDataType == IPackageWriter::FBulkDataInfo::BulkSegment,
+				TEXT("Does not support BulkData types other than AppendToExports and BulkSegment."));
 
 			const uint8* BufferStart = reinterpret_cast<const uint8*>(BulkRecord.Buffer.GetData());
 			uint64 SizeFromRegions = 0;
diff --git a/Engine/Source/Runtime/CoreUObject/Private/Serialization/BulkData.cpp b/Engine/Source/Runtime/CoreUObject/Private/Serialization/BulkData.cpp
--- a/Engine/Source/Runtime/CoreUObject/Private/Serialization/BulkData.cpp
+++ b/Engine/Source/Runtime/CoreUObject/Private/Serialization/BulkData.cpp
@@ -1,9 +1,11 @@
 // Copyright Epic Games, Inc. All Rights Reserved.
 #include "Serialization/BulkData.h"
+
 #include "Async/MappedFileHandle.h"
 #include "HAL/IConsoleManager.h"
 #include "IO/IoDispatcher.h"
 #include "Misc/ConfigCacheIni.h"
+#include "Misc/Optional.h"
 #include "ProfilingDebugging/LoadTimeTracker.h"
 #include "Serialization/LargeMemoryWriter.h"
 #include "Serialization/LargeMemoryReader.h"
@@ -1226,8 +1228,13 @@ void FBulkData::Serialize(FArchive& Ar, UObject* Owner, bool bAttemptFileMapping
 			const int64 MetaOffset = Ar.Tell();
 			Ar << SerializedMeta;
 
+			TOptional<EFileRegionType> FileRegionTypeOptional;
+			if (FileRegionType != EFileRegionType::None)
+			{
+				FileRegionTypeOptional = FileRegionType;
+			}
 			SerializedMeta.Offset = Ar.Tell();
-			SerializedMeta.SizeOnDisk = SerializePayload(Ar, SerializedMeta.Flags, FileRegionType);
+			SerializedMeta.SizeOnDisk = SerializePayload(Ar, SerializedMeta.Flags, FileRegionTypeOptional);
 			{
 				FArchive::FScopeSeekTo _(Ar, MetaOffset);
 				Ar << SerializedMeta;
@@ -1403,7 +1410,7 @@ void FBulkData::SerializeBulkData(FArchive& Ar, void* Data, int64 DataSize, EBul
 	}
 }
 
-int64 FBulkData::SerializePayload(FArchive& Ar, EBulkDataFlags SerializationFlags, EFileRegionType RegionType)
+int64 FBulkData::SerializePayload(FArchive& Ar, EBulkDataFlags SerializationFlags, const TOptional<EFileRegionType>& RegionType)
 {
 	check(Ar.IsSaving());
 
@@ -1413,14 +1420,14 @@ int64 FBulkData::SerializePayload(FArchive& Ar, EBulkDataFlags SerializationFlag
 
 	if (int64 PayloadSize = GetBulkDataSize(); PayloadSize > 0)
 	{
-		if (RegionType != EFileRegionType::None)
+		if (RegionType)
 		{
-			Ar.PushFileRegionType(RegionType);
+			Ar.PushFileRegionType(*RegionType);
 		}
 
 		SerializeBulkData(Ar, GetDataBufferForWrite(), PayloadSize, SerializationFlags);
 
-		if (RegionType != EFileRegionType::None)
+		if (RegionType)
 		{
 			Ar.PopFileRegionType();
 		}
diff --git a/Engine/Source/Runtime/CoreUObject/Private/UObject/LinkerSave.cpp b/Engine/Source/Runtime/CoreUObject/Private/UObject/LinkerSave.cpp
--- a/Engine/Source/Runtime/CoreUObject/Private/UObject/LinkerSave.cpp
+++ b/Engine/Source/Runtime/CoreUObject/Private/UObject/LinkerSave.cpp
@@ -4,6 +4,7 @@
 #include "HAL/FileManager.h"
 #include "Misc/CommandLine.h"
 #include "Misc/ConfigCacheIni.h"
+#include "Misc/Optional.h"
 #include "Serialization/BulkData.h"
 #include "Serialization/LargeMemoryWriter.h"
 #include "UObject/Package.h"
@@ -584,10 +585,21 @@ bool FLinkerSave::SerializeBulkData(FBulkData& BulkData, const FBulkDataSerializ
 	const EBulkDataFlags BulkDataFlags	= static_cast<EBulkDataFlags>(BulkData.GetBulkDataFlags());
 	int32 ResourceIndex					= DataResourceMap.Num();
 	const int64 PayloadSize				= BulkData.GetBulkDataSize();
-	const EFileRegionType RegionToUse	= bFileRegionsEnabled && IsCooking() ? Params.RegionType : EFileRegionType::None;
+	TOptional<EFileRegionType> RegionToUse;
 	const bool bSupportsMemoryMapping	= IsCooking() && MemoryMappingAlignment >= 0;
 	const bool bSaveAsResourceIndex		= IsCooking();
 	
+	if (bFileRegionsEnabled)
+	{
+		if (IsCooking())
+		{
+			RegionToUse = Params.RegionType;
+		}
+		else if (bDeclareRegionForEachAdditionalFile)
+		{
+			RegionToUse = EFileRegionType::None;
+		}
+	}
 	FBulkMetaResource SerializedMeta;
 	SerializedMeta.Flags = BulkDataFlags;
 	SerializedMeta.ElementCount = PayloadSize / Params.ElementSize;
@@ -643,7 +655,7 @@ bool FLinkerSave::SerializeBulkData(FBulkData& BulkData, const FBulkDataSerializ
 		if (bSaveByReference)
 		{
 			check(IsCooking() == false);
-			FBulkData::SetBulkDataFlagsOn(SerializedMeta.Flags, static_cast<EBulkDataFlags>(BULKDATA_NoOffsetFixUp | BULKDATA_WorkspaceDomainPayload));
+			FBulkData::SetBulkDataFlagsOn(SerializedMeta.Flags, static_cast<EBulkDataFlags>(BULKDATA_NoOffsetFixUp | BULKDATA_WorkspaceDomainPayload | BULKDATA_PayloadInSeperateFile));
 		}
 
 		if (bSaveBulkDataToSeparateFiles && FBulkData::HasFlags(SerializedMeta.Flags, BULKDATA_OptionalPayload))
diff --git a/Engine/Source/Runtime/CoreUObject/Private/UObject/SavePackage2.cpp b/Engine/Source/Runtime/CoreUObject/Private/UObject/SavePackage2.cpp
--- a/Engine/Source/Runtime/CoreUObject/Private/UObject/SavePackage2.cpp
+++ b/Engine/Source/Runtime/CoreUObject/Private/UObject/SavePackage2.cpp
@@ -1012,6 +1012,12 @@ ESavePackageResult BuildLinker(FSaveContext& SaveContext)
 				SaveContext.GetLinker()->SetFileRegionsEnabled(TargetPlatform->SupportsFeature(ETargetPlatformFeatures::CookFileRegionMetadata));
 			}
 		}
+		else if (SaveContext.GetPackageWriter() && SaveContext.GetSavePackageContext() &&
+			SaveContext.GetSavePackageContext()->PackageWriterCapabilities.bDeclareRegionForEachAdditionalFile)
+		{
+			SaveContext.GetLinker()->SetFileRegionsEnabled(true);
+			SaveContext.GetLinker()->SetDeclareRegionForEachAdditionalFile(true);
+		}
 #endif
 		if (SaveContext.IsCooking() == false)
 		{
diff --git a/Engine/Source/Runtime/CoreUObject/Public/Serialization/BulkData.h b/Engine/Source/Runtime/CoreUObject/Public/Serialization/BulkData.h
--- a/Engine/Source/Runtime/CoreUObject/Public/Serialization/BulkData.h
+++ b/Engine/Source/Runtime/CoreUObject/Public/Serialization/BulkData.h
@@ -21,6 +21,7 @@
 #include "Misc/AssertionMacros.h"
 #include "Misc/CompressionFlags.h"
 #include "Misc/EnumClassFlags.h"
+#include "Misc/OptionalFwd.h"
 #include "Misc/PackagePath.h"
 #include "Misc/PackageSegment.h"
 #include "Serialization/Archive.h"
@@ -1044,7 +1045,7 @@ protected:
 
 private:
 
-	int64 SerializePayload(FArchive& Ar, EBulkDataFlags SerializationFlags, EFileRegionType RegionType);
+	int64 SerializePayload(FArchive& Ar, EBulkDataFlags SerializationFlags, const TOptional<EFileRegionType>& RegionType);
 #if WITH_EDITOR
 	/**
 	 * Detaches the bulk data from the passed in archive. Needs to match the archive we are currently
diff --git a/Engine/Source/Runtime/CoreUObject/Public/UObject/LinkerSave.h b/Engine/Source/Runtime/CoreUObject/Public/UObject/LinkerSave.h
--- a/Engine/Source/Runtime/CoreUObject/Public/UObject/LinkerSave.h
+++ b/Engine/Source/Runtime/CoreUObject/Public/UObject/LinkerSave.h
@@ -185,6 +185,12 @@ public:
 		bFileRegionsEnabled = bEnabled;
 	}
 
+	/** Sets whether a separate file regions should be declared for each bulkdata */
+	void SetDeclareRegionForEachAdditionalFile(bool bValue)
+	{
+		bDeclareRegionForEachAdditionalFile = bValue;
+	}
+
 	/** Sets whether saving bulk data by reference, i.e. leaving the bulk data payload in the original .uasset file when using EditorDomain. */
 	void SetSaveBulkDataByReference(bool bValue)
 	{
@@ -302,6 +308,8 @@ private:
 	int64 MemoryMappingAlignment = -1;
 	/** Whether file regions are enabled. */
 	bool bFileRegionsEnabled = false;
+	/** Whether a separate file region should be declared for each bulkdata. */
+	bool bDeclareRegionForEachAdditionalFile = false;
 	/** Whether saving bulk data by reference. */
 	bool bSaveBulkDataByReference = false;
 	/** Whether saving bulk data to separate files. */
