[CoreUObject]
* ObjectPtr - Skipped static_assert when building for IWYU... (for compile time reasons this static assert should probably be moved to ObjectPtr dtor or something)

diff --git a/Engine/Source/Runtime/CoreUObject/Public/UObject/ObjectPtr.h b/Engine/Source/Runtime/CoreUObject/Public/UObject/ObjectPtr.h
--- a/Engine/Source/Runtime/CoreUObject/Public/UObject/ObjectPtr.h
+++ b/Engine/Source/Runtime/CoreUObject/Public/UObject/ObjectPtr.h
@@ -345,6 +345,7 @@ namespace ObjectPtr_Private
 template <typename T>
 struct TObjectPtr
 {
+#ifndef PLATFORM_COMPILER_IWYU
 	// TObjectPtr should only be used on types T that are EITHER:
 	// - incomplete (ie: forward declared and we have not seen their definition yet)
 	// - complete and derived from UObject
@@ -352,6 +353,8 @@ struct TObjectPtr
 	// - TObjectPtr<int>
 	// - TObjectPtr<IInterface>
 	static_assert(std::disjunction<std::negation<std::bool_constant<sizeof(ObjectPtr_Private::ResolveTypeIsComplete<T>(1)) == 2>>, std::is_base_of<UObject, T>>::value, "TObjectPtr<T> can only be used with types derived from UObject");
+#endif
+
 public:
 	using ElementType = T;
 
