Async Loading: downgrade a warning to a log when loading objects from PostLoad, which is valid for levels using external actors.

diff --git a/Engine/Source/Runtime/CoreUObject/Private/Serialization/AsyncLoading.cpp b/Engine/Source/Runtime/CoreUObject/Private/Serialization/AsyncLoading.cpp
--- a/Engine/Source/Runtime/CoreUObject/Private/Serialization/AsyncLoading.cpp
+++ b/Engine/Source/Runtime/CoreUObject/Private/Serialization/AsyncLoading.cpp
@@ -7003,7 +7003,7 @@ EAsyncPackageState::Type FAsyncPackage::PostLoadDeferredObjects(double InTickSta
 			{
 				// If there were any LoadObject calls inside of PostLoad, we need to pre-load those objects here. 
 				// There's no going back to the async tick loop from here.
-				UE_LOG(LogStreaming, Warning, TEXT("Detected %d objects loaded in PostLoad while streaming, this may cause hitches as we're blocking async loading to pre-load them."), ObjLoadedInPostLoad.Num());
+				UE_LOG(LogStreaming, Log, TEXT("Detected %d objects loaded in PostLoad while streaming, this may cause hitches as we're blocking async loading to pre-load them."), ObjLoadedInPostLoad.Num());
 
 				// Copy to local array because ObjLoadedInPostLoad can change while we're iterating over it
 				ObjLoadedInPostLoadLocal.Append(ObjLoadedInPostLoad);
