Fix issue where weak object pointers are invalidated when re-creating a UObject with the same name.
Make sure to also reuse the serial number when reusing the internal index.

#rnx

diff --git a/Engine/Source/Runtime/CoreUObject/Private/UObject/UObjectArray.cpp b/Engine/Source/Runtime/CoreUObject/Private/UObject/UObjectArray.cpp
--- a/Engine/Source/Runtime/CoreUObject/Private/UObject/UObjectArray.cpp
+++ b/Engine/Source/Runtime/CoreUObject/Private/UObject/UObjectArray.cpp
@@ -186,10 +186,10 @@ void FUObjectArray::DisableDisregardForGC()
 	}
 }
 
-void FUObjectArray::AllocateUObjectIndex(UObjectBase* Object, bool bMergingThreads /*= false*/, int32 AlreadyAllocatedIndex /*= -1*/)
+void FUObjectArray::AllocateUObjectIndex(UObjectBase* Object, int32 AlreadyAllocatedIndex, int32 SerialNumber)
 {
 	int32 Index = INDEX_NONE;
-	check(Object->InternalIndex == INDEX_NONE || bMergingThreads);
+	check(Object->InternalIndex == INDEX_NONE);
 
 	LockInternalArray();
 
@@ -233,10 +233,11 @@ void FUObjectArray::AllocateUObjectIndex(UObjectBase* Object, bool bMergingThrea
 	// Add to global table.
 	FUObjectItem* ObjectItem = IndexToObject(Index);
 	UE_CLOG(ObjectItem->Object != nullptr, LogUObjectArray, Fatal, TEXT("Attempting to add %s at index %d but another object (0x%016llx) exists at that index!"), *Object->GetFName().ToString(), Index, (int64)(PTRINT)ObjectItem->Object);
-	ObjectItem->ResetSerialNumberAndFlags();
+	ObjectItem->Object = Object;
 	// At this point all not-compiled-in objects are not fully constructed yet and this is the earliest we can mark them as such
-	ObjectItem->SetFlags(EInternalObjectFlags::PendingConstruction);
-	ObjectItem->Object = Object;		
+	ObjectItem->Flags = (int32)EInternalObjectFlags::PendingConstruction;
+	ObjectItem->ClusterRootIndex = 0;
+	ObjectItem->SerialNumber = SerialNumber;
 	Object->InternalIndex = Index;
 
 	UnlockInternalArray();
@@ -296,7 +297,9 @@ void FUObjectArray::FreeUObjectIndex(UObjectBase* Object)
 	FUObjectItem* ObjectItem = IndexToObject(Index);
 	UE_CLOG(ObjectItem->Object != Object, LogUObjectArray, Fatal, TEXT("Removing object (0x%016llx) at index %d but the index points to a different object (0x%016llx)!"), (int64)(PTRINT)Object, Index, (int64)(PTRINT)ObjectItem->Object);
 	ObjectItem->Object = nullptr;
-	ObjectItem->ResetSerialNumberAndFlags();
+	ObjectItem->Flags = 0;
+	ObjectItem->ClusterRootIndex = 0;
+	ObjectItem->SerialNumber = 0;
 
 	// You cannot safely recycle indicies in the non-GC range
 	// No point in filling this list when doing exit purge. Nothing should be allocated afterwards anyway.
diff --git a/Engine/Source/Runtime/CoreUObject/Private/UObject/UObjectBase.cpp b/Engine/Source/Runtime/CoreUObject/Private/UObject/UObjectBase.cpp
--- a/Engine/Source/Runtime/CoreUObject/Private/UObject/UObjectBase.cpp
+++ b/Engine/Source/Runtime/CoreUObject/Private/UObject/UObjectBase.cpp
@@ -118,7 +118,13 @@ UObjectBase::UObjectBase( EObjectFlags InFlags )
  * @param	InName				name of the new object
  * @param	InObjectArchetype	archetype to assign
  */
-UObjectBase::UObjectBase(UClass* InClass, EObjectFlags InFlags, EInternalObjectFlags InInternalFlags, UObject *InOuter, FName InName, int32 InInternalIndex)
+UObjectBase::UObjectBase(UClass* InClass,
+	EObjectFlags InFlags,
+	EInternalObjectFlags InInternalFlags,
+	UObject *InOuter,
+	FName InName,
+	int32 InInternalIndex,
+	int32 InSerialNumber)
 :	ObjectFlags			(InFlags)
 ,	InternalIndex		(INDEX_NONE)
 ,	ClassPrivate		(InClass)
@@ -126,7 +132,7 @@ UObjectBase::UObjectBase(UClass* InClass, EObjectFlags InFlags, EInternalObjectF
 {
 	check(ClassPrivate);
 	// Add to global table.
-	AddObject(InName, InInternalFlags, InInternalIndex);
+	AddObject(InName, InInternalFlags, InInternalIndex, InSerialNumber);
 	
 #if CSV_PROFILER && CSV_TRACK_UOBJECT_COUNT
 	UObjectStats::IncrementUObjectCount();
@@ -191,7 +197,7 @@ void UObjectBase::DeferredRegister(UClass *UClassStaticClass,const TCHAR* Packag
  *
  * @param Name name to assign to this uobject
  */
-void UObjectBase::AddObject(FName InName, EInternalObjectFlags InSetInternalFlags, int32 InInternalIndex)
+void UObjectBase::AddObject(FName InName, EInternalObjectFlags InSetInternalFlags, int32 InInternalIndex, int32 InSerialNumber)
 {
 	NamePrivate = InName;
 	EInternalObjectFlags InternalFlagsToSet = InSetInternalFlags;
@@ -209,7 +215,7 @@ void UObjectBase::AddObject(FName InName, EInternalObjectFlags InSetInternalFlag
 		InternalFlagsToSet |= EInternalObjectFlags::Native;
 		ObjectFlags &= ~RF_MarkAsNative;
 	}
-	GUObjectArray.AllocateUObjectIndex(this, /* bMergingThreads */ false, InInternalIndex);
+	GUObjectArray.AllocateUObjectIndex(this, InInternalIndex, InSerialNumber);
 	check(InName != NAME_None && InternalIndex >= 0);
 	if (InternalFlagsToSet != EInternalObjectFlags::None)
 	{
diff --git a/Engine/Source/Runtime/CoreUObject/Private/UObject/UObjectGlobals.cpp b/Engine/Source/Runtime/CoreUObject/Private/UObject/UObjectGlobals.cpp
--- a/Engine/Source/Runtime/CoreUObject/Private/UObject/UObjectGlobals.cpp
+++ b/Engine/Source/Runtime/CoreUObject/Private/UObject/UObjectGlobals.cpp
@@ -3285,6 +3285,7 @@ UObject* StaticAllocateObject
 	checkSlow(TotalSize);
 
 	int32 OldIndex = -1;
+	int32 OldSerialNumber = 0;
 
 	if( Obj == nullptr )
 	{	
@@ -3342,6 +3343,7 @@ UObject* StaticAllocateObject
 		if (!bCreatingCDO && (!bCanRecycleSubobjects || !Obj->IsDefaultSubobject()))
 		{
 			OldIndex = GUObjectArray.ObjectToIndex(Obj);
+			OldSerialNumber = GUObjectArray.GetSerialNumber(OldIndex);
 
 			// Destroy the object.
 			SCOPE_CYCLE_COUNTER(STAT_DestroyObject);
@@ -3406,7 +3408,7 @@ UObject* StaticAllocateObject
 	if (!bSubObject)
 	{
 		FMemory::Memzero((void *)Obj, TotalSize);
-		new ((void *)Obj) UObjectBase(const_cast<UClass*>(InClass), InFlags|RF_NeedInitialization, InternalSetFlags, InOuter, InName, OldIndex);
+		new ((void *)Obj) UObjectBase(const_cast<UClass*>(InClass), InFlags|RF_NeedInitialization, InternalSetFlags, InOuter, InName, OldIndex, OldSerialNumber);
 	}
 	else
 	{
diff --git a/Engine/Source/Runtime/CoreUObject/Public/UObject/UObjectArray.h b/Engine/Source/Runtime/CoreUObject/Public/UObject/UObjectArray.h
--- a/Engine/Source/Runtime/CoreUObject/Public/UObject/UObjectArray.h
+++ b/Engine/Source/Runtime/CoreUObject/Public/UObject/UObjectArray.h
@@ -213,13 +213,6 @@ struct FUObjectItem
 		return !!(GetFlagsInternal() & int32(EInternalObjectFlags::RootSet));
 	}
 
-	FORCEINLINE void ResetSerialNumberAndFlags()
-	{
-		Flags = 0;
-		ClusterRootIndex = 0;
-		SerialNumber = 0;
-	}
-
 #if STATS || ENABLE_STATNAMEDEVENTS_UOBJECT
 	COREUOBJECT_API void CreateStatID() const;
 #endif
@@ -701,8 +694,10 @@ public:
 	 * Adds a uobject to the global array which is used for uobject iteration
 	 *
 	 * @param	Object Object to allocate an index for
+	 * @param	AlreadyAllocatedIndex already allocated internal index to use, negative value means allocate a new index
+	 * @param	SerialNumber serial number to use
 	 */
-	void AllocateUObjectIndex(class UObjectBase* Object, bool bMergingThreads = false, int32 AlreadyAllocatedIndex = -1);
+	void AllocateUObjectIndex(class UObjectBase* Object, int32 AlreadyAllocatedIndex = -1, int32 SerialNumber = 0);
 
 	/**
 	 * Returns a UObject index top to the global uobject array
diff --git a/Engine/Source/Runtime/CoreUObject/Public/UObject/UObjectBase.h b/Engine/Source/Runtime/CoreUObject/Public/UObject/UObjectBase.h
--- a/Engine/Source/Runtime/CoreUObject/Public/UObject/UObjectBase.h
+++ b/Engine/Source/Runtime/CoreUObject/Public/UObject/UObjectBase.h
@@ -69,8 +69,15 @@ public:
 	 * @param	InOuter				outer for this object
 	 * @param	InName				name of the new object
 	 * @param	InInternalIndex		internal index to use (if already allocated), negative value means allocate a new index
+	 * @param	InSerialNumber		serial number to re-use (if already allocated)
 	 */
-	UObjectBase( UClass* InClass, EObjectFlags InFlags, EInternalObjectFlags InInternalFlags, UObject *InOuter, FName InName, int32 InInternalIndex = -1 );
+	UObjectBase(UClass* InClass,
+			EObjectFlags InFlags,
+			EInternalObjectFlags InInternalFlags,
+			UObject *InOuter,
+			FName InName,
+			int32 InInternalIndex = -1,
+			int32 InSerialNumber = 0);
 
 	/**
 	 * Final destructor, removes the object from the object array, and indirectly, from any annotations
@@ -111,8 +118,9 @@ private:
 	 * @param Name name to assign to this uobject
 	 * @param InSetInternalFlags Internal object flags to be set on the object once it's been added to the array
 	 * @param InInternalIndex already allocated internal index to use, negative value means allocate a new index
+	 * @param InSerialNumber already allocated serial number to re-use
 	 */
-	void AddObject(FName Name, EInternalObjectFlags InSetInternalFlags, int32 InInternalIndex = -1);
+	void AddObject(FName Name, EInternalObjectFlags InSetInternalFlags, int32 InInternalIndex = -1, int32 InSerialNumber = 0);
 
 public:
 	/**
