Fix problem with the FStructUtils::AttemptToFindUninitializedScriptStructMembers automation test sometimes failing without logging any errors.
Use the same log verbosity logic for all properties, no matter if they have a custom default constructor or not.
Let the test succeed or fail based on the absence or presence of errors instead of based on the uninitialized struct count.

#rnx

diff --git a/Engine/Source/Runtime/CoreUObject/Private/UObject/Class.cpp b/Engine/Source/Runtime/CoreUObject/Private/UObject/Class.cpp
--- a/Engine/Source/Runtime/CoreUObject/Private/UObject/Class.cpp
+++ b/Engine/Source/Runtime/CoreUObject/Private/UObject/Class.cpp
@@ -3615,45 +3615,62 @@ int32 FStructUtils::AttemptToFindUninitializedScriptStructMembers()
 	{
 		UScriptStruct* ScriptStruct = *ScriptIt;
 
-		if (FScriptStructTestWrapper::CanRunTests(ScriptStruct) && ScriptStruct != TestUninitializedScriptStructMembersTestStruct)
+		if (!FScriptStructTestWrapper::CanRunTests(ScriptStruct) || ScriptStruct == TestUninitializedScriptStructMembersTestStruct)
 		{
-			UninitializedPropertiesNoInit.Reset();
-			UninitializedPropertiesZeroed.Reset();
+			continue;
+		}
 
-			// Test the struct by constructing it with 'new FMyStruct();' syntax first. The compiler should zero all members in this case if the 
-			// struct doesn't have a custom default constructor defined
-			FindUninitializedScriptStructMembers(ScriptStruct, EScriptStructTestCtorSyntax::CompilerZeroed, UninitializedPropertiesZeroed);
-			// Test the struct by constructing it with 'new FStruct;' syntax in which case the compiler doesn't zero the properties automatically
-			FindUninitializedScriptStructMembers(ScriptStruct, EScriptStructTestCtorSyntax::NoInit, UninitializedPropertiesNoInit);			
+		UninitializedPropertiesNoInit.Reset();
+		UninitializedPropertiesZeroed.Reset();
 
-			for (const FProperty* Property : UninitializedPropertiesZeroed)
-			{
-				++UninitializedScriptStructMemberCount;
-				if (Property->IsA<FObjectPropertyBase>())
-				{
-					++UninitializedObjectPropertyCount;
-				}
-				UE_LOG(LogClass, Warning, TEXT("%s %s%s::%s is not initialized properly even though its struct probably has a custom default constructor.%s"), *Property->GetClass()->GetName(), ScriptStruct->GetPrefixCPP(), *ScriptStruct->GetName(), *Property->GetNameCPP(), *GetFieldLocation(ScriptStruct));
-			}
-			for (const FProperty* Property : UninitializedPropertiesNoInit)
-			{
-				if (!UninitializedPropertiesZeroed.Contains(Property))
-				{
-					++UninitializedScriptStructMemberCount;
-
-					ELogVerbosity::Type Verbosity = DetermineIfModuleIsEngine(ScriptStruct) ? Settings.EngineVerbosity : Settings.ProjectVerbosity;
+		// Test the struct by constructing it with 'new FMyStruct();' syntax first. The compiler should zero all members in this case if the 
+		// struct doesn't have a custom default constructor defined
+		FindUninitializedScriptStructMembers(ScriptStruct, EScriptStructTestCtorSyntax::CompilerZeroed, UninitializedPropertiesZeroed);
+		// Test the struct by constructing it with 'new FStruct;' syntax in which case the compiler doesn't zero the properties automatically
+		FindUninitializedScriptStructMembers(ScriptStruct, EScriptStructTestCtorSyntax::NoInit, UninitializedPropertiesNoInit);			
 
-					if (Property->IsA<FObjectPropertyBase>())
-					{
-						++UninitializedObjectPropertyCount;
-						Verbosity = FMath::Min(Verbosity, Settings.PointerVerbosity);
-					}
+		if (UninitializedPropertiesNoInit.Num() == 0 && UninitializedPropertiesZeroed.Num() == 0)
+		{
+			continue;
+		}
 
+		ELogVerbosity::Type StructVerbosity = DetermineIfModuleIsEngine(ScriptStruct) ? Settings.EngineVerbosity : Settings.ProjectVerbosity;
+		ELogVerbosity::Type PointerVerbosity = FMath::Min(StructVerbosity, Settings.PointerVerbosity);
+		auto LogUninitializedProperty =
+			[&UninitializedScriptStructMemberCount, &UninitializedObjectPropertyCount, ScriptStruct, StructVerbosity, PointerVerbosity]
+			(const FProperty* Property, const TCHAR* MessageText)
+		{
+			++UninitializedScriptStructMemberCount;
+			ELogVerbosity::Type Verbosity = StructVerbosity;
+			if (Property->IsA<FObjectPropertyBase>())
+			{
+				++UninitializedObjectPropertyCount;
+				Verbosity = PointerVerbosity;
+			}
 #if !NO_LOGGING
-					FMsg::Logf(__FILE__, __LINE__, LogClass.GetCategoryName(), Verbosity, TEXT("%s %s%s::%s is not initialized properly.%s"), *Property->GetClass()->GetName(), ScriptStruct->GetPrefixCPP(), *ScriptStruct->GetName(), *Property->GetNameCPP(), *GetFieldLocation(ScriptStruct));
+			// LogClass: Error: SetProperty FStructSerializerSetTestStruct::StructSet is not initialized properly even though its struct probably has a custom default constructor. Module:Serialization
+			FMsg::Logf(__FILE__, __LINE__, LogClass.GetCategoryName(), Verbosity, TEXT("%s %s%s::%s %s.%s"),
+				*Property->GetClass()->GetName(),
+				ScriptStruct->GetPrefixCPP(),
+				*ScriptStruct->GetName(),
+				*Property->GetNameCPP(),
+				MessageText,
+				*GetFieldLocation(ScriptStruct));
 #endif
-				}
+		};
+
+		for (const FProperty* Property : UninitializedPropertiesZeroed)
+		{
+			LogUninitializedProperty(Property, TEXT("is not initialized properly even though its struct probably has a custom default constructor"));
+		}
+
+		for (const FProperty* Property : UninitializedPropertiesNoInit)
+		{
+			if (UninitializedPropertiesZeroed.Contains(Property))
+			{
+				continue;
 			}
+			LogUninitializedProperty(Property, TEXT("is not initialized properly"));
 		}
 	}
 
@@ -3683,7 +3700,8 @@ bool FAutomationTestAttemptToFindUninitializedScriptStructMembers::RunTest(const
 #if !HACK_HEADER_GENERATOR
 	if (UObjectInitialized())
 	{
-		return FStructUtils::AttemptToFindUninitializedScriptStructMembers() == 0;
+		FStructUtils::AttemptToFindUninitializedScriptStructMembers();
+		return !HasAnyErrors();
 	}
 	else
 #endif
