Weak Object Pointers Serialization: if the objects share the same outer but aren't in the same package, we treat this weak object reference as a hard reference, to make sure they end up in the import table. This guarantee that they can be referenced by objects saved in other packages, such as different external actors from the same world.

diff --git a/Engine/Source/Runtime/CoreUObject/Private/UObject/SavePackage/PackageHarvester.cpp b/Engine/Source/Runtime/CoreUObject/Private/UObject/SavePackage/PackageHarvester.cpp
--- a/Engine/Source/Runtime/CoreUObject/Private/UObject/SavePackage/PackageHarvester.cpp
+++ b/Engine/Source/Runtime/CoreUObject/Private/UObject/SavePackage/PackageHarvester.cpp
@@ -506,19 +506,26 @@ FArchive& FPackageHarvester::operator<<(UObject*& Obj)
 
 FArchive& FPackageHarvester::operator<<(struct FWeakObjectPtr& Value)
 {
+	UObject* Object = Value.Get(true);
+
+	// If the referenced object share the same outer as the referencer but aren't in the same package, we treat this weak object
+	// reference as a hard reference, to make sure they end up in the import table. This guarantee that they can be referenced by objects
+	// saved in other packages,  such as different external actors from the same world.
+	const UObject* CurrentExport = CurrentExportDependencies.CurrentExport;
+
 	// @todo FH: Should we really force weak import in cooked builds?
-	if (IsCooking())
+	if (IsCooking() || (!Object->IsInPackage(SaveContext.GetPackage()) && (Object->GetOutermostObject() == CurrentExport->GetOutermostObject())))
 	{
-		// Always serialize weak pointers for the purposes of object tagging
-		UObject* Object = static_cast<UObject*>(Value.Get(true));
 		*this << Object;
 	}
 	else
 	{
 		FArchiveUObject::SerializeWeakObjectPtr(*this, Value);
 	}
+
 	return *this;
 }
+
 FArchive& FPackageHarvester::operator<<(FLazyObjectPtr& LazyObjectPtr)
 {
 	// @todo FH: Does this really do anything as far as tagging goes?
diff --git a/Engine/Source/Runtime/Engine/Private/ActorReferencesUtils.cpp b/Engine/Source/Runtime/Engine/Private/ActorReferencesUtils.cpp
--- a/Engine/Source/Runtime/Engine/Private/ActorReferencesUtils.cpp
+++ b/Engine/Source/Runtime/Engine/Private/ActorReferencesUtils.cpp
@@ -47,6 +47,15 @@ public:
 		return *this;
 	}
 
+	virtual FArchive& operator<<(FWeakObjectPtr& Value) override
+	{
+		if (UObject* Object = Value.Get(true))
+		{
+			return *this << Object;
+		}
+		return *this;
+	}
+
 private:
 	void HandleObjectReference(UObject* Obj)
 	{
