support for TObjectPtr handling renames and GetTypeHash(TObjectPtr) no long has to resolve
GetTypeHash will create a packaged object ref for hashing if the class is defined as lazy load otherwise raw pointer is hashed
a map was added for moved objects and ObjectPathId can be a WeakObjectPtr
the map allows UObjects to be mapped back to packed object refs

https://p4-swarm.epicgames.net/reviews/23356491

diff --git a/Engine/Source/Runtime/CoreUObject/Internal/UObject/LinkerLoadImportBehavior.h b/Engine/Source/Runtime/CoreUObject/Internal/UObject/LinkerLoadImportBehavior.h
--- a/Engine/Source/Runtime/CoreUObject/Internal/UObject/LinkerLoadImportBehavior.h
+++ b/Engine/Source/Runtime/CoreUObject/Internal/UObject/LinkerLoadImportBehavior.h
@@ -4,7 +4,6 @@
 #include "UObject/ObjectHandle.h"
 
 #if UE_WITH_OBJECT_HANDLE_LATE_RESOLVE
-#define UE_API COREUOBJECT_API
 
 class FLinkerLoad;
 struct FObjectImport;
@@ -32,7 +31,7 @@ EImportBehavior GetPropertyImportLoadBehavior(const FObjectImport& Import, const
 /// @return true if import was lazy loaded otherwise false
 bool TryLazyImport(const IAssetRegistryInterface& AssetRegistry, const FObjectImport& Import, const FLinkerLoad& LinkerLoad, FObjectPtr& ObjectPtr);
 
-/// @brief tries to resovle an object path to a lazy loaded TObjectPtr
+/// @brief tries to resolve an object path to a lazy loaded TObjectPtr
 ///
 /// @param Class type of the object that will be resolved
 /// @param ObjectPath object path of to resolve
@@ -40,7 +39,15 @@ bool TryLazyImport(const IAssetRegistryInterface& AssetRegistry, const FObjectIm
 /// @return true if was lazy loaded otherwise false
 bool TryLazyLoad(const UClass& Class, const FSoftObjectPath& ObjectPath, TObjectPtr<UObject>& OutObjectPtr);
 
+
+/// @brief Finds LoadBehavior meta data recursively
+/// @return Eager by default in not found
+EImportBehavior FindLoadBehavior(const UClass& Class);
+
+/// @brief returns true if lazy load imports is enabled via command line or config
+bool IsImportLazyLoadEnabled();
+
+
 }
 
-#undef UE_API
 #endif // UE_WITH_OBJECT_HANDLE_LATE_RESOLVE
diff --git a/Engine/Source/Runtime/CoreUObject/Private/UObject/LinkerLoad.cpp b/Engine/Source/Runtime/CoreUObject/Private/UObject/LinkerLoad.cpp
--- a/Engine/Source/Runtime/CoreUObject/Private/UObject/LinkerLoad.cpp
+++ b/Engine/Source/Runtime/CoreUObject/Private/UObject/LinkerLoad.cpp
@@ -5252,6 +5252,15 @@ UObject* FLinkerLoad::CreateExport( int32 Index )
 		{
 			TRACE_LOADTIME_CREATE_EXPORT_SCOPE(this, &Export.Object);
 			Export.Object = StaticConstructObject_Internal(Params);
+
+#if UE_WITH_OBJECT_HANDLE_LATE_RESOLVE
+			//if lazy load is enabled construct a packed ref if possible.
+			//this is to have a reverse map of UObject to FPackedObjectRef
+			if (UE::LinkerLoad::IsImportLazyLoadEnabled())
+			{
+				MakePackedObjectRef(Export.Object);
+			}
+#endif
 		}
 
 		if (FPlatformProperties::RequiresCookedData())
@@ -6281,22 +6290,7 @@ bool FLinkerLoad::RemoveKnownMissingPackage(FName PackageName)
 #if UE_WITH_OBJECT_HANDLE_LATE_RESOLVE
 bool FLinkerLoad::IsImportLazyLoadEnabled()
 {
-	auto ImportLazyLoadEnabled = []()
-	{
-		if (FParse::Param(FCommandLine::Get(), TEXT("LazyLoadImports")))
-		{
-			return true;
-		}
-		else if (GConfig)
-		{
-			bool bLazyLoadImportsConfig = false;
-			GConfig->GetBool(TEXT("Core.System.Experimental"), TEXT("LazyLoadImports"), bLazyLoadImportsConfig, GEngineIni);
-			return bLazyLoadImportsConfig;
-		}
-		return false;
-	};
-	static const bool bImportLazyLoadEnabled = ImportLazyLoadEnabled();
-	return bImportLazyLoadEnabled;
+	return UE::LinkerLoad::IsImportLazyLoadEnabled();
 }
 #endif
 
diff --git a/Engine/Source/Runtime/CoreUObject/Private/UObject/LinkerLoadImportBehavior.cpp b/Engine/Source/Runtime/CoreUObject/Private/UObject/LinkerLoadImportBehavior.cpp
--- a/Engine/Source/Runtime/CoreUObject/Private/UObject/LinkerLoadImportBehavior.cpp
+++ b/Engine/Source/Runtime/CoreUObject/Private/UObject/LinkerLoadImportBehavior.cpp
@@ -8,6 +8,7 @@
 #include "UObject/Package.h"
 #include "Misc/AssetRegistryInterface.h"
 #include "Misc/CommandLine.h"
+#include "Misc/ConfigCacheIni.h"
 
 #if UE_WITH_OBJECT_HANDLE_LATE_RESOLVE
 
@@ -16,7 +17,7 @@ namespace UE::LinkerLoad
 
 /// @brief Finds LoadBehavior meta data recursively
 /// @return Eager by default in not found
-static EImportBehavior FindLoadBehavior(const UClass& Class)
+EImportBehavior FindLoadBehavior(const UClass& Class)
 {
 	//Package class can't have meta data because of UHT
 	if (&Class == UPackage::StaticClass())
@@ -24,10 +25,10 @@ static EImportBehavior FindLoadBehavior(const UClass& Class)
 		return EImportBehavior::LazyOnDemand;
 	}
 
-	static const FName Name_LoadBehavior("LoadBehavior");
+	static const FName Name_LoadBehavior(TEXT("LoadBehavior"));
 	if (const FString* LoadBehaviorMeta = Class.FindMetaData(Name_LoadBehavior))
 	{
-		if (*LoadBehaviorMeta == "LazyOnDemand")
+		if (*LoadBehaviorMeta == TEXT("LazyOnDemand"))
 		{
 			return EImportBehavior::LazyOnDemand;
 		}
@@ -216,6 +217,30 @@ bool TryLazyLoad(const UClass& Class, const FSoftObjectPath& ObjectPath, TObject
 	return TryLazyLoad(*AssetRegistry, ObjectPath, ObjectPtr);
 }
 
+bool IsImportLazyLoadEnabled()
+{
+#if WITH_LOW_LEVEL_TESTS //not ideal but need a way to force lazyload on for tests
+	return true;
+#else
+	auto ImportLazyLoadEnabled = []()
+	{
+		if (FParse::Param(FCommandLine::Get(), TEXT("LazyLoadImports")))
+		{
+			return true;
+		}
+		else if (GConfig)
+		{
+			bool bLazyLoadImportsConfig = false;
+			GConfig->GetBool(TEXT("Core.System.Experimental"), TEXT("LazyLoadImports"), bLazyLoadImportsConfig, GEngineIni);
+			return bLazyLoadImportsConfig;
+		}
+		return false;
+	};
+	static const bool bImportLazyLoadEnabled = ImportLazyLoadEnabled();
+	return bImportLazyLoadEnabled;
+#endif
+}
+
 }
 
 #endif // UE_WITH_OBJECT_HANDLE_LATE_RESOLVE
diff --git a/Engine/Source/Runtime/CoreUObject/Private/UObject/Obj.cpp b/Engine/Source/Runtime/CoreUObject/Private/UObject/Obj.cpp
--- a/Engine/Source/Runtime/CoreUObject/Private/UObject/Obj.cpp
+++ b/Engine/Source/Runtime/CoreUObject/Private/UObject/Obj.cpp
@@ -312,7 +312,9 @@ bool UObject::Rename( const TCHAR* InName, UObject* NewOuter, ERenameFlags Flags
 				NewOuter->MarkPackageDirty();
 			}
 		}
-
+#if UE_WITH_OBJECT_HANDLE_LATE_RESOLVE
+		ObjectHandle_Private::UpdateRenamedObject(*this, NewName, NewOuter);
+#endif
 		LowLevelRename(NewName, NewOuter);
 	}
 
diff --git a/Engine/Source/Runtime/CoreUObject/Private/UObject/ObjectHandle.cpp b/Engine/Source/Runtime/CoreUObject/Private/UObject/ObjectHandle.cpp
--- a/Engine/Source/Runtime/CoreUObject/Private/UObject/ObjectHandle.cpp
+++ b/Engine/Source/Runtime/CoreUObject/Private/UObject/ObjectHandle.cpp
@@ -10,6 +10,7 @@
 #include "UObject/Class.h"
 #include "UObject/Linker.h"
 #include "UObject/LinkerLoad.h"
+#include "UObject/LinkerLoadImportBehavior.h"
 #include "UObject/ObjectRedirector.h"
 #include "UObject/Package.h"
 #include "UObject/UObjectHash.h"
@@ -22,6 +23,10 @@ FName FObjectRef::GetFName() const
 	{
 		return NAME_None;
 	}
+	else if (ObjectPath.IsNone())
+	{
+		return PackageName;
+	}
 
 	FObjectPathId::ResolvedNameContainerType ResolvedNames;
 	ObjectPath.Resolve(ResolvedNames);
@@ -82,6 +87,11 @@ bool operator==(const FObjectHandleDataClassDescriptor& Lhs, const FObjectHandle
 	return (Lhs.PackageName == Rhs.PackageName) && (Lhs.ClassName == Rhs.ClassName);
 }
 
+static inline FName GetNameOrNone(UObject* Object)
+{
+	return Object ? Object->GetFName() : NAME_None;
+}
+
 namespace ObjectHandle_Private
 {
 	class FPackageId
@@ -123,7 +133,7 @@ namespace ObjectHandle_Private
 	union FObjectId
 	{
 	private:
-		uint32 RawData;
+		uint32 RawData = 0;
 	public:
 		struct
 		{
@@ -136,6 +146,11 @@ namespace ObjectHandle_Private
 
 		bool operator==(EObjectId Id) { return RawData == static_cast<uint32>(Id); }
 		bool operator!=(EObjectId Id) { return RawData != static_cast<uint32>(Id); }
+
+		inline uint32 ToIndex() const
+		{
+			return Components.ObjectPathId - 1;
+		}
 	};
 
 	static_assert(sizeof(FObjectId) == sizeof(uint32), "FObjectId type must always compile to something equivalent to a uint32 size.");
@@ -145,7 +160,6 @@ namespace ObjectHandle_Private
 		FMinimalName PackageName;
 		TArray<FObjectPathId> ObjectPaths;
 		TArray<FObjectHandleDataClassDescriptor> DataClassDescriptors;
-		TMap<FObjectPathId, FObjectId> PathToObjectId;
 		FRWLock Lock;
 	};
 
@@ -159,10 +173,17 @@ namespace ObjectHandle_Private
 		FRWLock Lock; // @TODO: OBJPTR: Want to change this to a striped lock per object bucket to allow more concurrency when adding and looking up objects in a package
 		TMap<FMinimalName, FPackageId> NameToPackageId;
 		TArray<FObjectHandlePackageData> PackageData;
+
+		TArray<FPackedObjectRef> ObjectIndexToPackedObjectRef;
 	} GObjectHandleIndex;
 
 #if UE_WITH_OBJECT_HANDLE_LATE_RESOLVE
 
+	void InitObjectHandles(int32 MaxObjects)
+	{
+		GObjectHandleIndex.ObjectIndexToPackedObjectRef.SetNumZeroed(MaxObjects);
+	}
+
 	static FAutoConsoleCommand CmdPrintUnresolvedObjects(
 		TEXT("LazyLoad.PrintUnresolvedObjects"),
 		TEXT("Prints a list of all unresolved objects from the object handle index."),
@@ -248,6 +269,7 @@ namespace ObjectHandle_Private
 		FRWScopeLock GlobalLockScope(GObjectHandleIndex.Lock, SLT_ReadOnly);
 		FPackageId* FoundPackageId = GObjectHandleIndex.NameToPackageId.Find(MinimalName);
 		FObjectHandlePackageData* PackageData = nullptr;
+		bool PackageCreated = false;
 		if (!FoundPackageId)
 		{
 			GlobalLockScope.ReleaseReadOnlyLockAndAcquireWriteLock_USE_WITH_CAUTION();
@@ -260,6 +282,7 @@ namespace ObjectHandle_Private
 				PackageData = &GObjectHandleIndex.PackageData.AddDefaulted_GetRef();
 				PackageData->PackageName = NameToMinimalName(PackageName);
 				GCoreObjectHandlePackageDebug = reinterpret_cast<FObjectHandlePackageDebugData*>(GObjectHandleIndex.PackageData.GetData());
+				PackageCreated = true;
 			}
 			else
 			{
@@ -272,43 +295,74 @@ namespace ObjectHandle_Private
 		{
 			PackageData = &GObjectHandleIndex.PackageData[FoundPackageId->ToIndex()];
 		}
-
 		OutPackageId = *FoundPackageId;
 
-		FRWScopeLock LocalLockScope(PackageData->Lock, SLT_ReadOnly);
-		FObjectId* FoundId = PackageData->PathToObjectId.Find(ObjectPath);
-
-		if (FoundId)
+		//optimization for creating a new package.
+		//since the package is new, GlobalLockScope is write locked and nothing can be added or read
+		if (PackageCreated )
 		{
-			check(*FoundId != EObjectId::Invalid);
-			OutObjectId = *FoundId;
+			if (!ClassName.IsNone() && !ClassPackageName.IsNone())
+			{
+				FObjectHandleDataClassDescriptor DataClassDesc{ NameToMinimalName(ClassPackageName), NameToMinimalName(ClassName) };
+				uint32 DataClassDescriptorIndex = PackageData->DataClassDescriptors.Add(DataClassDesc);
+				checkf(((DataClassDescriptorIndex + 1) & ~DataClassDescriptorIdMask) == 0, TEXT("Data class descriptor id overflowed space in ObjectHandle"));
+				OutObjectId.Components.DataClassDescriptorId = DataClassDescriptorIndex + 1;
+			}
+			if (!ObjectPath.IsNone())
+			{
+				int32 PathIndex = PackageData->ObjectPaths.Add(ObjectPath);
+				checkf(((PathIndex + 1) & ~ObjectPathIdMask) == 0, TEXT("Path id overflowed space in ObjectHandle"));
+				OutObjectId.Components.ObjectPathId = PathIndex + 1;
+			}
 			return;
 		}
-
-		LocalLockScope.ReleaseReadOnlyLockAndAcquireWriteLock_USE_WITH_CAUTION();
-		//The PathToObjectId could have been modified when the read lock was released and the write
-		//lock was acquired, so we must check and see if the ObjectPath was added in that window.
-		FoundId = &PackageData->PathToObjectId.FindOrAdd(ObjectPath, EObjectId::Invalid);
-		if (*FoundId != EObjectId::Invalid)
+		else if (ObjectPath.IsNone())
 		{
-			OutObjectId = *FoundId;
+			FRWScopeLock LocalLockScope(PackageData->Lock, SLT_ReadOnly);
+			if (!ClassName.IsNone() && !ClassPackageName.IsNone())
+			{
+				FObjectHandleDataClassDescriptor DataClassDesc{ NameToMinimalName(ClassPackageName), NameToMinimalName(ClassName) };
+				uint32 DataClassDescriptorIndex = PackageData->DataClassDescriptors.AddUnique(DataClassDesc);
+				checkf(((DataClassDescriptorIndex + 1) & ~DataClassDescriptorIdMask) == 0, TEXT("Data class descriptor id overflowed space in ObjectHandle"));
+				OutObjectId.Components.DataClassDescriptorId = DataClassDescriptorIndex + 1;
+			}
 			return;
 		}
-		uint32 PathIndex = PackageData->ObjectPaths.Emplace(ObjectPath);
-		checkf(((PathIndex + 1) & ~ObjectHandle_Private::ObjectPathIdMask) == 0, TEXT("Path id overflowed space in ObjectHandle"));
-		FoundId->Components.ObjectPathId = PathIndex + 1;
 
-		if (!ClassName.IsNone() && !ClassPackageName.IsNone())
+		FRWScopeLock LocalLockScope(PackageData->Lock, SLT_ReadOnly);
+		int32 PathIndex = PackageData->ObjectPaths.Find(ObjectPath); //linear search is fine as typically there is only one.
+		if (PathIndex == INDEX_NONE)
+		{
+			//ObjectPaths could have been modified when the read lock was released and the write
+			//lock was acquired, so we must check and see if the ObjectPath was added in that window.
+			LocalLockScope.ReleaseReadOnlyLockAndAcquireWriteLock_USE_WITH_CAUTION();
+			PathIndex = PackageData->ObjectPaths.AddUnique(ObjectPath);
+
+			if (!ClassName.IsNone() && !ClassPackageName.IsNone())
+			{
+				FObjectHandleDataClassDescriptor DataClassDesc{ NameToMinimalName(ClassPackageName), NameToMinimalName(ClassName) };
+				uint32 DataClassDescriptorIndex = PackageData->DataClassDescriptors.AddUnique(DataClassDesc);
+				checkf(((DataClassDescriptorIndex + 1) & ~DataClassDescriptorIdMask) == 0, TEXT("Data class descriptor id overflowed space in ObjectHandle"));
+				OutObjectId.Components.DataClassDescriptorId = DataClassDescriptorIndex + 1;
+			}
+		}
+		else
 		{
-			// @TODO: OBJPTR: This could be inefficient if there are a high number of references to blueprint data instances
-			//		or references to unique blueprints in a single package.  Evaluate whether that's likely to be
-			//		the case in practice.
-			FObjectHandleDataClassDescriptor DataClassDesc{NameToMinimalName(ClassPackageName), NameToMinimalName(ClassName)};
-			uint32 DataClassDescriptorIndex = PackageData->DataClassDescriptors.AddUnique(DataClassDesc);
-			checkf(((DataClassDescriptorIndex + 1) & ~ObjectHandle_Private::DataClassDescriptorIdMask) == 0, TEXT("Data class descriptor id overflowed space in ObjectHandle"));
-			FoundId->Components.DataClassDescriptorId = DataClassDescriptorIndex + 1;
+			if (!ClassName.IsNone() && !ClassPackageName.IsNone())
+			{
+				FObjectHandleDataClassDescriptor DataClassDesc{ NameToMinimalName(ClassPackageName), NameToMinimalName(ClassName) };
+				uint32 DataClassDescriptorIndex = PackageData->DataClassDescriptors.Find(DataClassDesc);
+				if (DataClassDescriptorIndex != INDEX_NONE)
+				{
+					checkf(((DataClassDescriptorIndex + 1) & ~DataClassDescriptorIdMask) == 0, TEXT("Data class descriptor id overflowed space in ObjectHandle"));
+					OutObjectId.Components.DataClassDescriptorId = DataClassDescriptorIndex + 1;
+				}
+			}
 		}
-		OutObjectId = *FoundId;
+
+		checkf(((PathIndex + 1) & ~ObjectPathIdMask) == 0, TEXT("Path id overflowed space in ObjectHandle"));
+		OutObjectId.Components.ObjectPathId = PathIndex + 1;
+		return;
 	}
 
 	static inline FPackedObjectRef MakePackedObjectRef(FName PackageName, FName ClassPackageName, FName ClassName, FObjectPathId ObjectPath)
@@ -334,8 +388,9 @@ namespace ObjectHandle_Private
 			return;
 		}
 		FObjectHandlePackageData& FoundPackageData = GObjectHandleIndex.PackageData[PackageIndex];
-
 		FRWScopeLock LocalLockScope(FoundPackageData.Lock, SLT_ReadOnly);
+		OutPackageName = FoundPackageData.PackageName;
+
 		if ((ObjectId.Components.ObjectPathId >= static_cast<uint32>(FoundPackageData.ObjectPaths.Num() + 1)) ||
 			(ObjectId.Components.DataClassDescriptorId >= static_cast<uint32>(FoundPackageData.DataClassDescriptors.Num() + 1)))
 		{
@@ -343,8 +398,14 @@ namespace ObjectHandle_Private
 			return;
 		}
 
-		OutPackageName = FoundPackageData.PackageName;
-		OutPathId = FoundPackageData.ObjectPaths[ObjectId.Components.ObjectPathId - 1];
+		if (ObjectId.Components.ObjectPathId != 0)
+		{
+			OutPathId = FoundPackageData.ObjectPaths[ObjectId.Components.ObjectPathId - 1];
+		}
+		else
+		{
+			OutPathId = FObjectPathId();
+		}
 
 		if (ObjectId.Components.DataClassDescriptorId > 0)
 		{
@@ -360,6 +421,121 @@ namespace ObjectHandle_Private
 		}
 	}
 
+#if UE_WITH_OBJECT_HANDLE_LATE_RESOLVE
+	void FreeObjectHandle(const UObjectBase& Object)
+	{
+		int32 ObjectIndex = GUObjectArray.ObjectToIndex(&Object);
+		FPackedObjectRef& PackedObjectRef = GObjectHandleIndex.ObjectIndexToPackedObjectRef[ObjectIndex];
+		PackedObjectRef.EncodedRef = 0;
+	}
+
+	void UpdateRenamedObject(const UObject& Object, FName NewName, UObject* NewOuter)
+	{
+		if (!UE::LinkerLoad::IsImportLazyLoadEnabled())
+		{
+			return;
+		}
+		int32 ObjectIndex = GUObjectArray.ObjectToIndex(&Object);
+		FPackedObjectRef PackedObjectRef = GObjectHandleIndex.ObjectIndexToPackedObjectRef[ObjectIndex];
+		if (PackedObjectRef.EncodedRef == 0)
+		{
+			return;
+		}
+
+		check(NewName != NAME_None);
+		
+		FMinimalName MinimalName = NameToMinimalName(Object.GetFName());
+		FObjectId ObjectId;
+		FPackageId PackageId;
+		Unpack(PackedObjectRef, PackageId, ObjectId);
+
+		FRWScopeLock GlobalLockScope(GObjectHandleIndex.Lock, SLT_Write);
+		FObjectHandlePackageData& PackageData = GObjectHandleIndex.PackageData[PackageId.ToIndex()];
+		FRWScopeLock PackageLockScope(PackageData.Lock, SLT_ReadOnly);
+		if (Object.GetClass() == UPackage::StaticClass())
+		{
+			//update the package name at existing index. existing object handles will be correct when unpacked
+			//add in the new name pointing the existing index. new object handles will resolve to the existing index
+			check(NewName != NAME_None);
+			FMinimalName MinNewName = NameToMinimalName(NewName);
+			PackageData.PackageName = MinNewName;
+			GObjectHandleIndex.NameToPackageId.Add(MinNewName, PackageId);
+			GObjectHandleIndex.NameToPackageId.Remove(MinimalName);
+			return;
+		}
+
+		//must make a copy to avoid changing data while we are comparing
+		FObjectPathId OldObjectPath = PackageData.ObjectPaths[ObjectId.ToIndex()];
+		if (OldObjectPath.IsWeakObj())
+		{
+			return;
+		}
+
+		const FMinimalName* OldObjectPathtData = nullptr;
+		int32 OldObjectPathSize = 0;
+		FMinimalName OldName;
+
+		if (OldObjectPath.IsSimple())
+		{
+			OldName = OldObjectPath.GetSimpleName();
+			OldObjectPathtData = &OldName;
+			OldObjectPathSize = 1;
+		}
+		else
+		{
+			//this isn't thread safe. if the FStoredObjectPath moves this will now point to garbage
+			const UE::ObjectPath::Private::FStoredObjectPath& StoredPath = OldObjectPath.GetStoredPath();
+			OldObjectPathtData = StoredPath.GetData();
+			OldObjectPathSize = StoredPath.NumElements;
+		}
+
+		//update all paths that start with OldObjectPath. including itself
+		for (int32 PathIndex = 0; PathIndex < PackageData.ObjectPaths.Num(); ++PathIndex)
+		{
+			auto& CurrentObjectPath = PackageData.ObjectPaths[PathIndex];
+			if (CurrentObjectPath.IsWeakObj())
+			{
+				continue;
+			}
+			else if (CurrentObjectPath == OldObjectPath)
+			{
+				CurrentObjectPath.MakeWeakObjPtr(Object);
+				continue;
+			}
+			else if (CurrentObjectPath.IsSimple())
+			{
+				//nothing to do for simple paths. they won't be an inner of the object
+				continue;
+			}
+			const UE::ObjectPath::Private::FStoredObjectPath& ThisPath = CurrentObjectPath.GetStoredPath();
+			if (OldObjectPathSize < ThisPath.NumElements && CompareItems(ThisPath.GetData(), OldObjectPathtData, OldObjectPathSize))
+			{
+				FObjectPathId::ResolvedNameContainerType ResolvedNames;
+				CurrentObjectPath.Resolve(ResolvedNames);
+
+				//resolve all object along the path and convert the object paths to weak objects
+				const UObject* CurrentObject = Object.GetPackage();
+				for (int32 ObjectPathIndex = 0; ObjectPathIndex < ResolvedNames.Num() && CurrentObject; ++ObjectPathIndex)
+				{
+					CurrentObject = StaticFindObjectFastInternal(nullptr, CurrentObject, ResolvedNames[ObjectPathIndex]);
+				}
+
+				// not sure how this could happen, null out the object path
+				if (!CurrentObject)
+				{
+					CurrentObjectPath.Index = 0;
+					CurrentObjectPath.Number = 0; 
+				}
+				else
+				{
+					CurrentObjectPath.MakeWeakObjPtr(*CurrentObject);
+				}
+			}
+		}
+
+	}
+#endif
+
 	static inline FObjectRef MakeObjectRef(FPackedObjectRef PackedObjectRef)
 	{
 		TRACE_CPUPROFILER_EVENT_SCOPE(ObjectHandle_Private::MakeObjectRef);
@@ -373,14 +549,14 @@ namespace ObjectHandle_Private
 		FMinimalName ClassPackageName;
 		FMinimalName ClassName;
 		GetObjectDataFromId(PackageId, ObjectId, PackageName, PathId, ClassPackageName, ClassName);
-		return FObjectRef{MinimalNameToName(PackageName), MinimalNameToName(ClassPackageName), MinimalNameToName(ClassName), PathId};
+		if (PathId.IsWeakObj())
+		{
+			return MakeObjectRef(PathId.GetWeakObjPtr().Get());
+		}
+		return FObjectRef{ MinimalNameToName(PackageName), MinimalNameToName(ClassPackageName), MinimalNameToName(ClassName), PathId };
 	}
 };
 
-static inline FName GetNameOrNone(UObject* Object)
-{
-	return Object ? Object->GetFName() : NAME_None;
-}
 
 FObjectRef MakeObjectRef(const UObject* Object)
 {
@@ -389,13 +565,9 @@ FObjectRef MakeObjectRef(const UObject* Object)
 		return {NAME_None, NAME_None, NAME_None, FObjectPathId()};
 	}
 
-#if WITH_EDITORONLY_DATA	
-	UObject* ClassGeneratedBy = Object->GetClass()->ClassGeneratedBy;
-	UPackage* ClassGeneratedByPackage = ClassGeneratedBy ? ClassGeneratedBy->GetOutermost() : nullptr;
-	return FObjectRef {GetNameOrNone(Object->GetOutermost()), GetNameOrNone(ClassGeneratedByPackage), GetNameOrNone(ClassGeneratedBy), FObjectPathId(Object)};
-#else
-	return FObjectRef{ GetNameOrNone(Object->GetOutermost()), NAME_None, NAME_None, FObjectPathId(Object) };
-#endif
+	FName ClassName = Object->GetClass()->GetFName();
+	FName ClassPackageName = GetNameOrNone(Object->GetClass()->GetOuter());
+	return FObjectRef {GetNameOrNone(Object->GetOutermost()), ClassPackageName, ClassName, FObjectPathId(Object)};
 }
 
 FObjectRef MakeObjectRef(FPackedObjectRef PackedObjectRef)
@@ -408,23 +580,41 @@ FObjectRef MakeObjectRef(FPackedObjectRef PackedObjectRef)
 	return ObjectHandle_Private::MakeObjectRef(PackedObjectRef);
 }
 
+#if UE_WITH_OBJECT_HANDLE_LATE_RESOLVE
 FPackedObjectRef MakePackedObjectRef(const UObject* Object)
 {
 	if (!Object)
 	{
-		return {0};
+		return { 0 };
+	}
+
+	if (UE::LinkerLoad::FindLoadBehavior(*Object->GetClass()) == UE::LinkerLoad::EImportBehavior::Eager)
+	{
+		return { 0 };
+	}
+
+	int32 ObjectIndex = GUObjectArray.ObjectToIndex(Object);
+	FPackedObjectRef& PackedObjectRef = ObjectHandle_Private::GObjectHandleIndex.ObjectIndexToPackedObjectRef[ObjectIndex];
+	if (PackedObjectRef.EncodedRef != 0)
+	{
+		return PackedObjectRef;
 	}
 
 	FName PackageName = GetNameOrNone(Object->GetOutermost());
-#if WITH_EDITORONLY_DATA
-	UObject* ClassGeneratedBy = Object->GetClass()->ClassGeneratedBy;
-	UPackage* ClassGeneratedByPackage = ClassGeneratedBy ? ClassGeneratedBy->GetOutermost() : nullptr;
-	return ObjectHandle_Private::MakePackedObjectRef(PackageName, GetNameOrNone(ClassGeneratedByPackage), GetNameOrNone(ClassGeneratedBy), FObjectPathId(Object));
-#else
-	return ObjectHandle_Private::MakePackedObjectRef(PackageName, NAME_None, NAME_None, FObjectPathId(Object));
-#endif
+
+	UObject* Class = Object->GetClass();
+	FName ClassPackageName = GetNameOrNone(Class->GetOutermost());
+	PackedObjectRef = ObjectHandle_Private::MakePackedObjectRef(PackageName, ClassPackageName, GetNameOrNone(Class), FObjectPathId(Object));
+	return PackedObjectRef;
 }
 
+FPackedObjectRef GetPackedObjectRef(const UObject& Object)
+{
+	int32 ObjectIndex = GUObjectArray.ObjectToIndex(&Object);
+	return ObjectHandle_Private::GObjectHandleIndex.ObjectIndexToPackedObjectRef[ObjectIndex];
+}
+#endif
+
 FPackedObjectRef MakePackedObjectRef(const FObjectRef& ObjectRef)
 {
 	if (IsObjectRefNull(ObjectRef))
diff --git a/Engine/Source/Runtime/CoreUObject/Private/UObject/ObjectPathId.cpp b/Engine/Source/Runtime/CoreUObject/Private/UObject/ObjectPathId.cpp
--- a/Engine/Source/Runtime/CoreUObject/Private/UObject/ObjectPathId.cpp
+++ b/Engine/Source/Runtime/CoreUObject/Private/UObject/ObjectPathId.cpp
@@ -9,6 +9,7 @@
 #include "ProfilingDebugging/CpuProfilerTrace.h"
 #include "UObject/Linker.h"
 #include "UObject/Object.h"
+#include "UObject/ObjectKey.h"
 #include "UObject/Package.h"
 
 FCriticalSection GComplexPathLock;  // Could be changed to an RWLock later if needed
@@ -37,6 +38,14 @@ struct assert_equality
 assert_equality<sizeof(TArray<FMinimalName, TInlineAllocator<3>>), 0> B;
 */
 
+namespace UE::CoreUObject::Private
+{
+	FObjectKey MakeObjectKey(int32 ObjectIndex, int32 ObjectSerialNumber)
+	{
+		return FObjectKey(ObjectIndex, ObjectSerialNumber);
+	}
+}
+
 namespace UE::ObjectPath::Private
 {
 	FStoredObjectPath::FStoredObjectPath(TConstArrayView<FMinimalName> InNames)
@@ -63,6 +72,34 @@ namespace UE::ObjectPath::Private
 		}
 	}
 
+	FStoredObjectPath::FStoredObjectPath(const FMinimalName* Parent, int32 NumParent, const FMinimalName* Child, int32 NumChild)
+	{
+		// Copy into storage in reverse
+		NumElements = NumParent + NumChild;
+
+		FMinimalName* Dest = nullptr;
+
+		if (NumElements > NumInlineElements)
+		{
+			Long = new FMinimalName[NumElements];
+			Dest = Long;
+		}
+		else
+		{
+			Dest = Short;
+		}
+
+		for (int32 i = 0; i < NumParent; ++i)
+		{
+			Dest[i] = Parent[i];
+		}
+
+		for (int32 i = 0; i < NumChild; ++i)
+		{
+			Dest[NumParent + i] = Child[i];
+		}
+	}
+
 	FStoredObjectPath::~FStoredObjectPath()
 	{
 		if (NumElements > NumInlineElements)
@@ -105,7 +142,7 @@ namespace UE::ObjectPath::Private
 }
 
 template <typename NameProducerType>
-void StoreObjectPathId(NameProducerType& NameProducer, uint64 SimplePathFlag, uint64& OutObjectPathId)
+void FObjectPathId::StoreObjectPathId(NameProducerType& NameProducer)
 {
 	FName Name = NameProducer();
 
@@ -117,14 +154,13 @@ void StoreObjectPathId(NameProducerType& NameProducer, uint64 SimplePathFlag, ui
 	FName OuterName = NameProducer();
 	if (OuterName == NAME_None)
 	{
-		int32 Number = Name.GetNumber();
-		// Check if the number is in range by left shift + right shift
-		if (static_cast<int32>((static_cast<uint32>(Number) << 1) >> 1) == Number)
+		int32 NameNumber = Name.GetNumber();
+
+		//verify that the second most significant bits are not set
+		if ( !(NameNumber & SimpleNameMask) && !(NameNumber & WeakObjectMask))
 		{
-			//Simple path scenario
-			OutObjectPathId = static_cast<uint64>(Name.GetComparisonIndex().ToUnstableInt()) << 32 |
-								static_cast<uint64>(static_cast<uint32>(static_cast<uint32>(Number) << 1)) |
-								SimplePathFlag;
+			Index = Name.GetComparisonIndex().ToUnstableInt();
+			Number = NameNumber | SimpleNameMask;
 			return;
 		}
 	}
@@ -150,14 +186,17 @@ void StoreObjectPathId(NameProducerType& NameProducer, uint64 SimplePathFlag, ui
 		uint32 PotentialPathId = It.Value();
 		if (GComplexPaths[PotentialPathId-1] == PathToStore)
 		{
-			OutObjectPathId = static_cast<uint32>(PotentialPathId << 1);
+			Index = PotentialPathId;
+			Number = 0;
 			return;
 		}
 	}
 
+	//add one because a path id of 0 indicates null
 	uint32 NewId = GComplexPaths.Emplace(MoveTemp(PathToStore)) + 1;
 	GComplexPathHashToId.Add(Key, NewId);
-	OutObjectPathId = static_cast<uint32>(NewId << 1);
+	Index = NewId;
+	Number = 0;
 	GCoreComplexObjectPathDebug = GComplexPaths.GetData();
 }
 
@@ -183,7 +222,7 @@ FObjectPathId::FObjectPathId(const UObject* Object)
 		const UObject* CurrentObject;
 	} NamePathProducer(Object);
 
-	StoreObjectPathId(NamePathProducer, static_cast<uint64>(EPathId::FlagSimple), PathId);
+	StoreObjectPathId(NamePathProducer);
 }
 
 FObjectPathId::FObjectPathId(const FObjectImport& Import, const FLinkerTables& LinkerTables)
@@ -226,7 +265,7 @@ FName FObjectPathId::MakeImportPathIdAndPackageName(const FObjectImport& Import,
 		const FLinkerTables& LinkerTables;
 	} NamePathProducer(Import, LinkerTables);
 
-	StoreObjectPathId(NamePathProducer, static_cast<uint64>(EPathId::FlagSimple), OutPathId.PathId);
+	OutPathId.StoreObjectPathId(NamePathProducer);
 	return NamePathProducer.GetPackageName();
 }
 
@@ -276,13 +315,13 @@ private:
 FObjectPathId::FObjectPathId(FWideStringView StringPath)
 {
 	TStringViewNamePathProducer<WIDECHAR> NamePathProducer(StringPath);
-	StoreObjectPathId(NamePathProducer, static_cast<uint64>(EPathId::FlagSimple), PathId);
+	StoreObjectPathId(NamePathProducer);
 }
 
 FObjectPathId::FObjectPathId(FAnsiStringView StringPath)
 {
 	TStringViewNamePathProducer<ANSICHAR> NamePathProducer(StringPath);
-	StoreObjectPathId(NamePathProducer, static_cast<uint64>(EPathId::FlagSimple), PathId);
+	StoreObjectPathId(NamePathProducer);
 }
 
 void FObjectPathId::Resolve(ResolvedNameContainerType& OutContainer) const
@@ -294,20 +333,59 @@ void FObjectPathId::Resolve(ResolvedNameContainerType& OutContainer) const
 		return;
 	}
 
-	if (static_cast<uint64>(PathId) & static_cast<uint64>(EPathId::FlagSimple))
+	if (Number & SimpleNameMask)
 	{
-		uint32 Number = static_cast<uint32>(PathId & 0x00000000FFFFFFE) >> 1;
-		uint32 Index = static_cast<uint32>(PathId >> 32);
-		FNameEntryId EntryId = FNameEntryId::FromUnstableInt(Index);
-		OutContainer.Emplace(EntryId, EntryId, Number);
+		//truncate Path to int and shift back to get the number
+		uint32 NameNumber = Number & ~SimpleNameMask;
+		uint32 NameIndex = Index;
+		FNameEntryId EntryId = FNameEntryId::FromUnstableInt(NameIndex);
+		OutContainer.Emplace(EntryId, EntryId, NameNumber);
 		return; 
 	}
 
 	FScopeLock ComplexPathScopeLock(&GComplexPathLock);
-	const UE::ObjectPath::Private::FStoredObjectPath& FoundContainer = GComplexPaths[(PathId >> 1) - 1];
+	const UE::ObjectPath::Private::FStoredObjectPath& FoundContainer = GComplexPaths[Index - 1];
 	
 	for (FMinimalName Name : FoundContainer.GetView())
 	{
 		OutContainer.Emplace(MinimalNameToName(Name));
 	}
 }
+FMinimalName FObjectPathId::GetSimpleName() const
+{
+	check(IsSimple());
+
+	uint32 NameNumber = Number & ~SimpleNameMask;
+	uint32 NameIndex = Index;
+	FNameEntryId EntryId = FNameEntryId::FromUnstableInt(NameIndex);
+	FMinimalName Name = NameToMinimalName(FName(EntryId, EntryId, NameNumber));
+	return Name;
+}
+
+void FObjectPathId::MakeWeakObjPtr(const UObject& Object)
+{
+	int32 ObjectIndex = GUObjectArray.ObjectToIndex(&Object);
+	int32 ObjectSerialNumber = GUObjectArray.AllocateSerialNumber(ObjectIndex);
+
+	//check that ObjectSerialNumber does not have WeakObjMask bit set
+	check(!(ObjectSerialNumber & WeakObjectMask));
+	Index = ObjectIndex;
+	Number = ObjectSerialNumber | WeakObjectMask;
+}
+
+FWeakObjectPtr FObjectPathId::GetWeakObjPtr() const
+{
+	check(IsWeakObj());
+
+	int32 ObjectIndex = Index;
+	int32 ObjectSerial = Number & ~WeakObjectMask;
+	FObjectKey Key = UE::CoreUObject::Private::MakeObjectKey(ObjectIndex, ObjectSerial);
+	return Key.ResolveObjectPtr();
+}
+
+const UE::ObjectPath::Private::FStoredObjectPath& FObjectPathId::GetStoredPath() const
+{
+	check(IsValid() && !IsNone());
+	const UE::ObjectPath::Private::FStoredObjectPath& FoundContainer = GComplexPaths[Index - 1];
+	return FoundContainer;
+}
\ No newline at end of file
diff --git a/Engine/Source/Runtime/CoreUObject/Private/UObject/UObjectBase.cpp b/Engine/Source/Runtime/CoreUObject/Private/UObject/UObjectBase.cpp
--- a/Engine/Source/Runtime/CoreUObject/Private/UObject/UObjectBase.cpp
+++ b/Engine/Source/Runtime/CoreUObject/Private/UObject/UObjectBase.cpp
@@ -151,6 +151,9 @@ UObjectBase::~UObjectBase()
 		// Validate it.
 		check(IsValidLowLevel());
 		check(GetFName() == NAME_None);
+#if UE_WITH_OBJECT_HANDLE_LATE_RESOLVE
+		ObjectHandle_Private::FreeObjectHandle(*this);
+#endif 
 		GUObjectArray.FreeUObjectIndex(this);
 	}
 
@@ -1055,6 +1058,9 @@ void UObjectBaseInit()
 
 	GUObjectAllocator.AllocatePermanentObjectPool(SizeOfPermanentObjectPool);
 	GUObjectArray.AllocateObjectPool(MaxUObjects, MaxObjectsNotConsideredByGC, bPreAllocateUObjectArray);
+#if UE_WITH_OBJECT_HANDLE_LATE_RESOLVE
+	ObjectHandle_Private::InitObjectHandles(MaxUObjects);
+#endif
 
 	void InitNoPendingKill();
 	InitNoPendingKill();
diff --git a/Engine/Source/Runtime/CoreUObject/Private/UObject/UObjectGlobals.cpp b/Engine/Source/Runtime/CoreUObject/Private/UObject/UObjectGlobals.cpp
--- a/Engine/Source/Runtime/CoreUObject/Private/UObject/UObjectGlobals.cpp
+++ b/Engine/Source/Runtime/CoreUObject/Private/UObject/UObjectGlobals.cpp
@@ -1012,6 +1012,12 @@ UPackage* CreatePackage(const TCHAR* PackageName )
 			else
 			{
 				Result = NewObject<UPackage>(nullptr, NewPackageName, RF_Public);
+#if UE_WITH_OBJECT_HANDLE_LATE_RESOLVE
+				if (Result->GetClass()->GetFName() != NAME_None)
+				{
+					MakePackedObjectRef(Result);
+				}
+#endif
 #if WITH_EDITOR
 				FStringView PackageMountPoint = FPathViews::GetMountPointNameFromPath(InName);
 				uint32 DefaultPackageFlags = GCreatePackageDefaultFlagsMap.Find(PackageMountPoint);
diff --git a/Engine/Source/Runtime/CoreUObject/Public/UObject/ObjectHandle.h b/Engine/Source/Runtime/CoreUObject/Public/UObject/ObjectHandle.h
--- a/Engine/Source/Runtime/CoreUObject/Public/UObject/ObjectHandle.h
+++ b/Engine/Source/Runtime/CoreUObject/Public/UObject/ObjectHandle.h
@@ -127,7 +127,10 @@ struct FPackedObjectRef
 
 inline bool IsPackedObjectRefNull(FPackedObjectRef ObjectRef) { return !ObjectRef.EncodedRef; }
 
+#if UE_WITH_OBJECT_HANDLE_LATE_RESOLVE
 COREUOBJECT_API FPackedObjectRef MakePackedObjectRef(const UObject* Object);
+#endif 
+
 COREUOBJECT_API FPackedObjectRef MakePackedObjectRef(const FObjectRef& ObjectRef);
 COREUOBJECT_API UObject* ResolvePackedObjectRef(FPackedObjectRef ObjectRef, uint32 LoadFlags = LOAD_None);
 COREUOBJECT_API UClass* ResolvePackedObjectRefClass(FPackedObjectRef ObjectRef, uint32 LoadFlags = LOAD_None);
@@ -138,6 +141,13 @@ COREUOBJECT_API UClass* ResolvePackedObjectRefClass(FPackedObjectRef ObjectRef,
  */
 #if UE_WITH_OBJECT_HANDLE_LATE_RESOLVE
 
+namespace ObjectHandle_Private
+{
+	void UpdateRenamedObject(const UObject& Obj, FName NewName, UObject* NewOuter);
+	void InitObjectHandles(int32 Size);
+	void FreeObjectHandle(const UObjectBase& Object);
+}
+
 struct FObjectHandleInternal;
 using FObjectHandle = FObjectHandleInternal;
 
@@ -174,7 +184,11 @@ inline FPackedObjectRef ReadObjectHandlePackedObjectRef(FObjectHandle Handle);
 
 #if UE_WITH_OBJECT_HANDLE_LATE_RESOLVE
 inline FObjectRef MakeObjectRef(FObjectHandle Handle);
+
+UE_DEPRECATED(5.2, "Use `bool TryMakePackedObjectRef(FObjectHandle, FPackedObjectRef&)` instead")
 inline FPackedObjectRef MakePackedObjectRef(FObjectHandle Handle);
+inline FPackedObjectRef GetPackedObjectRef(FObjectHandle Handle);
+COREUOBJECT_API FPackedObjectRef GetPackedObjectRef(const UObject& Object);
 
 struct FObjectHandleInternal
 {
@@ -182,23 +196,81 @@ struct FObjectHandleInternal
 
 	friend inline bool operator==(FObjectHandle LHS, FObjectHandle RHS)
 	{
-		if (IsObjectHandleResolved(LHS) == IsObjectHandleResolved(RHS))
+		bool LhsResolved = IsObjectHandleResolved(LHS);
+		bool RhsResolved = IsObjectHandleResolved(RHS);
+
+		//if both resolved or both unresolved compare the uintptr
+		if (LhsResolved == RhsResolved)
 		{
 			return LHS.PointerOrRef == RHS.PointerOrRef;
 		}
+		
+		//only one side can be resolved
+		if (LhsResolved)
+		{
+			//both sides can't be null as resolved status would have be true for both
+			const UObject* Obj = ReadObjectHandlePointerNoCheck(LHS);
+			if (!Obj)
+			{
+				return false;
+			}
+
+			//if packed ref empty then can't be equal as RHS is an unresolved pointer
+			FPackedObjectRef PackedLhs = GetPackedObjectRef(*Obj);
+			if (PackedLhs.EncodedRef == 0)
+			{
+				return false; 
+			}
+			return PackedLhs.EncodedRef == RHS.PointerOrRef;
+				
+		}
 		else
 		{
-			return MakePackedObjectRef(LHS) == MakePackedObjectRef(RHS);
+			//both sides can't be null as resolved status would have be true for both
+			const UObject* Obj = ReadObjectHandlePointerNoCheck(RHS);
+			if (!Obj)
+			{
+				return false;
+			}
+
+			//if packed ref empty then can't be equal as RHS is an unresolved pointer
+			FPackedObjectRef PackedRhs = GetPackedObjectRef(*Obj);
+			if (PackedRhs.EncodedRef == 0)
+			{
+				return false;
+			}
+			return PackedRhs.EncodedRef == LHS.PointerOrRef;
 		}
+		
 	}
+
 	friend inline bool operator!=(FObjectHandle LHS, FObjectHandle RHS)
 	{
 		return !(LHS == RHS);
 	}
 	friend inline uint32 GetTypeHash(FObjectHandleInternal Handle)
 	{
-		checkf(IsObjectHandleResolved(Handle), TEXT("Cannot hash an unresolved handle."));
-		return GetTypeHash(ReadObjectHandlePointerNoCheck(Handle));
+		if (Handle.PointerOrRef == 0)
+		{
+			return 0;
+		}
+
+		if (IsObjectHandleResolved(Handle))
+		{
+			const UObject* Obj = ReadObjectHandlePointerNoCheck(Handle);
+			FPackedObjectRef PackedObjectRef = GetPackedObjectRef(*Obj);
+			if (PackedObjectRef.EncodedRef == 0)
+			{
+				return GetTypeHash(Obj);
+			}
+			return GetTypeHash(PackedObjectRef.EncodedRef);
+		}
+		return GetTypeHash(Handle.PointerOrRef);
+	}
+
+	explicit inline operator bool() const
+	{
+		return PointerOrRef != 0;
 	}
 };
 #endif
@@ -296,7 +368,7 @@ struct FObjectHandlePackageDebugData
 	FMinimalName PackageName;
 	FScriptArray ObjectPaths;
 	FScriptArray DataClassDescriptors;
-	uint8 _Padding[sizeof(FRWLock) + sizeof(FScriptMap)];
+	uint8 _Padding[sizeof(FRWLock)];
 };
 
 struct FObjectHandleDataClassDescriptor
@@ -519,15 +591,4 @@ inline FObjectRef MakeObjectRef(FObjectHandle Handle)
 	}
 }
 
-inline FPackedObjectRef MakePackedObjectRef(FObjectHandle Handle)
-{
-	if (IsObjectHandleResolved(Handle))
-	{
-		return MakePackedObjectRef(ReadObjectHandlePointerNoCheck(Handle));
-	}
-	else
-	{
-		return ReadObjectHandlePackedObjectRefNoCheck(Handle);
-	}
-}
 #endif
diff --git a/Engine/Source/Runtime/CoreUObject/Public/UObject/ObjectKey.h b/Engine/Source/Runtime/CoreUObject/Public/UObject/ObjectKey.h
--- a/Engine/Source/Runtime/CoreUObject/Public/UObject/ObjectKey.h
+++ b/Engine/Source/Runtime/CoreUObject/Public/UObject/ObjectKey.h
@@ -7,6 +7,13 @@
 
 #include <type_traits>
 
+struct FObjectKey;
+
+namespace UE::CoreUObject::Private
+{
+	FObjectKey MakeObjectKey(int32 ObjectIndex, int32 ObjectSerialNumber);
+}
+
 /** FObjectKey is an immutable, copyable key which can be used to uniquely identify an object for the lifetime of the application */
 struct FObjectKey
 {
@@ -128,6 +135,12 @@ public:
 	}
 
 private:
+	FObjectKey(int32 Index, int32 Serial)
+		: ObjectIndex(Index)
+		, ObjectSerialNumber(Serial)
+	{ }
+
+	friend FObjectKey UE::CoreUObject::Private::MakeObjectKey(int32 ObjectIndex, int32 ObjectSerialNumber);
 
 	int32		ObjectIndex;
 	int32		ObjectSerialNumber;
diff --git a/Engine/Source/Runtime/CoreUObject/Public/UObject/ObjectPathId.h b/Engine/Source/Runtime/CoreUObject/Public/UObject/ObjectPathId.h
--- a/Engine/Source/Runtime/CoreUObject/Public/UObject/ObjectPathId.h
+++ b/Engine/Source/Runtime/CoreUObject/Public/UObject/ObjectPathId.h
@@ -14,6 +14,15 @@
 class FLinkerTables;
 class UObject;
 struct FObjectImport;
+class FObjectPathId;
+struct FObjectRef;
+struct FPackedObjectRef;
+
+namespace ObjectHandle_Private
+{
+	void UpdateRenamedObject(const UObject& Obj, FName NewName, UObject* NewOuter);
+	static FObjectRef MakeObjectRef(FPackedObjectRef PackedObjectRef);
+}
 
 // Declared in the header so the type exists in debug info for debugger visualization, not a public part of the API.
 namespace UE::ObjectPath::Private
@@ -30,6 +39,7 @@ namespace UE::ObjectPath::Private
 		};
 
 		FStoredObjectPath(TConstArrayView<FMinimalName> InNames);
+		FStoredObjectPath(const FMinimalName* Parent, int32 NumParent, const FMinimalName* Child, int32 NumChild);
 		~FStoredObjectPath();
 		FStoredObjectPath(const FStoredObjectPath&) = delete;
 		FStoredObjectPath(FStoredObjectPath&&);
@@ -37,6 +47,7 @@ namespace UE::ObjectPath::Private
 		FStoredObjectPath& operator=(FStoredObjectPath&&);
 
 		TConstArrayView<FMinimalName> GetView() const;
+		const FMinimalName* GetData() const { return NumElements > NumInlineElements ? Long : &Short[0]; }
 	};
 }
 
@@ -63,20 +74,25 @@ public:
 	COREUOBJECT_API explicit FObjectPathId(const UObject* Object);
 	COREUOBJECT_API FObjectPathId(const FObjectImport& Import, const FLinkerTables& LinkerTables);
 
-	enum EInvalid {Invalid = 0};
-	explicit COREUOBJECT_API FObjectPathId(EInvalid): PathId(static_cast<uint64>(EPathId::FlagSimple))
+	enum EInvalid { Invalid = 0 };
+
+	explicit COREUOBJECT_API FObjectPathId(EInvalid)
+		: Index(0)
+		, Number(SimpleNameMask | WeakObjectMask)
 	{
 	}
-
+	
 	COREUOBJECT_API static FName MakeImportPathIdAndPackageName(const FObjectImport& Import, const FLinkerTables& LinkerTables, FObjectPathId& OutPathId);
 
 	explicit FObjectPathId(FWideStringView StringPath);
 	explicit FObjectPathId(FAnsiStringView StringPath);
 
 	FObjectPathId& operator=(const FObjectPathId& Other) = default;
-	inline bool operator==(const FObjectPathId& Other) const { return PathId == Other.PathId; }
-	inline bool IsNone() const { return PathId == static_cast<uint64>(EPathId::None); }
-	inline bool IsValid() const { return PathId != static_cast<uint64>(EPathId::FlagSimple); }
+	inline bool operator==(const FObjectPathId& Other) const { return Index == Other.Index && Number == Other.Number; }
+	inline bool IsNone() const { return Index == 0 && Number == 0; }
+	inline bool IsValid() const { return Number != (SimpleNameMask | WeakObjectMask); }
+	inline bool IsSimple() const { return IsValid() && Number & SimpleNameMask; }
+	inline bool IsWeakObj() const { return IsValid() && Number & WeakObjectMask; }
 
 	// @TODO: OBJPTR: Is there a better interface to output the resolved name without:
 	// 		 1) Requiring multiple lookups in the table in case of complex paths
@@ -86,16 +102,26 @@ public:
 	COREUOBJECT_API void Resolve(ResolvedNameContainerType& OutContainer) const;
 
 private:
-	enum class EPathId : uint64
-	{
-		None = 0,
-		FlagSimple = 0x01,
-	};
 
-	uint64 PathId = static_cast<uint64>(EPathId::None);
+	FMinimalName GetSimpleName() const;
+	void MakeWeakObjPtr(const UObject& Object);
+	const UE::ObjectPath::Private::FStoredObjectPath& GetStoredPath() const;
+	FWeakObjectPtr GetWeakObjPtr() const;
+
+	template <typename NameProducerType>
+	void StoreObjectPathId(NameProducerType& NameProducer);
+
+	friend void ObjectHandle_Private::UpdateRenamedObject(const UObject& Obj, FName NewName, UObject* NewOuter);
+	friend FObjectRef ObjectHandle_Private::MakeObjectRef(FPackedObjectRef PackedObjectRef);
+	
+	static constexpr uint32 WeakObjectMask = ~((~0u) >> 1);       //most significant bit
+	static constexpr uint32 SimpleNameMask = WeakObjectMask >> 1; //second most significant bits
+
+	uint32 Index = 0;
+	uint32 Number = 0;
 
 	friend FORCEINLINE uint32 GetTypeHash(FObjectPathId Value)
 	{
-		return GetTypeHash(Value.PathId);
+		return HashCombineFast(GetTypeHash(Value.Index), GetTypeHash(Value.Number));
 	}
 };
diff --git a/Engine/Source/Runtime/CoreUObject/Public/UObject/ObjectPtr.h b/Engine/Source/Runtime/CoreUObject/Public/UObject/ObjectPtr.h
--- a/Engine/Source/Runtime/CoreUObject/Public/UObject/ObjectPtr.h
+++ b/Engine/Source/Runtime/CoreUObject/Public/UObject/ObjectPtr.h
@@ -216,7 +216,6 @@ public:
 private:
 	friend FORCEINLINE uint32 GetTypeHash(const FObjectPtr& Object)
 	{
-		Object.Get();
 		return GetTypeHash(Object.Handle);
 	}
 
diff --git a/Engine/Source/Runtime/CoreUObject/Tests/ObjectHandleTest.cpp b/Engine/Source/Runtime/CoreUObject/Tests/ObjectHandleTest.cpp
--- a/Engine/Source/Runtime/CoreUObject/Tests/ObjectHandleTest.cpp
+++ b/Engine/Source/Runtime/CoreUObject/Tests/ObjectHandleTest.cpp
@@ -2,8 +2,11 @@
 
 #if WITH_LOW_LEVEL_TESTS
 
+#include "ObjectPtrTestClass.h"
 #include "UObject/ObjectHandle.h"
 #include "UObject/ObjectPtr.h"
+#include "UObject/Package.h"
+#include "UObject/ObjectResource.h"
 #include "UObject/MetaData.h"
 #include "HAL/PlatformProperties.h"
 #include "ObjectRefTrackingTestBase.h"
@@ -63,7 +66,7 @@ protected:
 
 	UObject* ConstructAndResolveHandle(const FPackedObjectRef& PackedTargetRef)
 	{
-		bool bValue = IsPackedObjectRefNull(PackedTargetRef);
+		bool bValue = PackedTargetRef.EncodedRef == 0;
 		TEST_FALSE(TEXT("Reference to target is null"), bValue);
 		if (bValue)
 		{
@@ -163,8 +166,8 @@ TEST_CASE_METHOD(FObjectHandleTestBase, "CoreUObject::FObjectHandle::Resolve Eng
 	const FName TestPackageName(TEXT("/Engine/Test/ObjectPtrDefaultSerialize/Transient"));
 	UPackage* TestPackage = NewObject<UPackage>(nullptr, TestPackageName, RF_Transient);
 	TestPackage->AddToRoot();
-	UObject* TestSoftObject = NewObject<UMetaData>(TestPackage, TEXT("DefaultSerializeObject"));
-	UObject* TestSubObject = NewObject<UMetaData>(TestSoftObject, TEXT("SubObject"));
+	UObject* TestSoftObject = NewObject<UObjectPtrTestClass>(TestPackage, TEXT("DefaultSerializeObject"));
+	UObject* TestSubObject = NewObject<UObjectPtrTestClass>(TestSoftObject, TEXT("SubObject"));
 	ON_SCOPE_EXIT{
 		TestPackage->RemoveFromRoot();
 	};
@@ -184,7 +187,7 @@ DISABLED_TEST_CASE_METHOD(FObjectHandleTestBase, "CoreUObject::FObjectHandle::Re
 	const FName TestPackageName(TEXT("/Engine/Test/ObjectPtrDefaultSerialize/Transient"));
 	UPackage* TestPackage = NewObject<UPackage>(nullptr, TestPackageName, RF_Transient);
 	TestPackage->AddToRoot();
-	UObject* TestSoftObject = NewObject<UMetaData>(TestPackage, TEXT("DefaultSerializeObject"));
+	UObject* TestSoftObject = NewObject<UObjectPtrTestClass>(TestPackage, TEXT("DefaultSerializeObject"));
 	ON_SCOPE_EXIT{
 		TestPackage->RemoveFromRoot();
 	};
@@ -198,18 +201,77 @@ TEST_CASE_METHOD(FObjectHandleTestBase, "CoreUObject::FObjectHandle::Resolve Scr
 	TestResolvableNonNull("/Script/CoreUObject", "MetaData");
 }
 
-TEST_CASE_METHOD(FObjectHandleTestBase, "CoreUObject::TObjectPtr::HandleNullGetClass", "[CoreUObject][ObjectHandle][.Engine]")
+TEST_CASE_METHOD(FObjectHandleTestBase, "CoreUObject::TObjectPtr::HandleNullGetClass", "[CoreUObject][ObjectHandle]")
 {
 	TObjectPtr<UObject> Ptr = nullptr;
 	TEST_TRUE(TEXT("TObjectPtr.GetClass should return null on a null object"), Ptr.GetClass() == nullptr);
 }
 
+#if UE_WITH_OBJECT_HANDLE_LATE_RESOLVE
+TEST_CASE("CoreUObject::FObjectHandle::Names")
+{
+	const FName TestPackageName(TEXT("/Engine/Test/PackageResolve/Transient"));
+	UPackage* TestPackage = NewObject<UPackage>(nullptr, TestPackageName, RF_Transient);
+	TestPackage->AddToRoot();
+	UObject* Obj1 = NewObject<UObjectPtrTestClass>(TestPackage, TEXT("DefaultSerializeObject"));
+	ON_SCOPE_EXIT{
+		TestPackage->RemoveFromRoot();
+	};
+
+	FObjectPtr PackagePtr(MakePackedObjectRef(TestPackage));
+	FObjectPtr Obj1Ptr(MakePackedObjectRef(Obj1));
+
+	CHECK(!PackagePtr.IsResolved());
+	CHECK(TestPackage->GetPathName() == PackagePtr.GetPathName());
+	CHECK(TestPackage->GetFName() == PackagePtr.GetFName());
+	CHECK(TestPackage->GetName() == PackagePtr.GetName());
+	//CHECK(TestPackage->GetFullName() == PackagePtr.GetFullName()); TODO always been broken
+	CHECK(!PackagePtr.IsResolved());
+
+	CHECK(!Obj1Ptr.IsResolved());
+	CHECK(Obj1->GetPathName() == Obj1Ptr.GetPathName());
+	CHECK(Obj1->GetFName() == Obj1Ptr.GetFName());
+	CHECK(Obj1->GetName() == Obj1Ptr.GetName());
+	//CHECK(Obj1->GetFullName() == Obj1Ptr.GetFullName()); TODO always been broken
+	CHECK(!Obj1Ptr.IsResolved());
+}
+#endif
+
+TEST_CASE("CoreUObject::ObjectRef")
+{
+	const FName TestPackageName(TEXT("/Engine/Test/ObjectRef/Transient"));
+	UPackage* TestPackage = NewObject<UPackage>(nullptr, TestPackageName, RF_Transient);
+	TestPackage->AddToRoot();
+	UObject* Obj1 = NewObject<UObjectPtrTestClass>(TestPackage, TEXT("DefaultSerializeObject"));
+	UObject* Inner1 = NewObject<UObjectPtrTestClass>(Obj1, TEXT("Inner"));
+	ON_SCOPE_EXIT{
+		TestPackage->RemoveFromRoot();
+	};
+
+	{
+		FObjectImport ObjectImport(Obj1);
+		FObjectRef ObjectRef = MakeObjectRef(Obj1);
+
+		CHECK(ObjectImport.ClassPackage == ObjectRef.ClassPackageName);
+		CHECK(ObjectImport.ClassName == ObjectRef.ClassName);
+		CHECK(TestPackage->GetFName() == ObjectRef.PackageName);
+	}
+
+	{
+		FObjectImport ObjectImport(Inner1);
+		FObjectRef ObjectRef = MakeObjectRef(Inner1);
+
+		CHECK(ObjectImport.ClassPackage == ObjectRef.ClassPackageName);
+		CHECK(ObjectImport.ClassName == ObjectRef.ClassName);
+		CHECK(TestPackage->GetFName() == ObjectRef.PackageName);
+	}
+}
 #if UE_WITH_OBJECT_HANDLE_TRACKING
 
 TEST_CASE_METHOD(FObjectHandleTestBase, "CoreUObject::TObjectPtr::Null Behavior", "[CoreUObject][ObjectHandle]")
 {
 	TObjectPtr<UObject> Ptr = nullptr;
-	UMetaData* TestObject = nullptr;
+	UObjectPtrTestClass* TestObject = nullptr;
 
 	uint32 ResolveCount = 0;
 	auto ResolveDelegate = FObjectHandleReferenceResolvedDelegate::CreateLambda([&ResolveCount](const FObjectRef& SourceRef, UPackage* ObjectPackage, UObject* Object)
@@ -276,25 +338,53 @@ TEST_CASE_METHOD(FObjectHandleTestBase, "CoreUObject::TObjectPtr::Null Behavior"
 	TestPackage->AddToRoot();
 
 	const FName TestObjectName(TEXT("MyObject"));
-	TestObject = NewObject<UMetaData>(TestPackage, TestObjectName, RF_Transient);
-	TestObject->AddToRoot();
-
-	//ObjectPtr = FObjectPtr(MakePackedObjectRef(TestObject));
+	TestObject = NewObject<UObjectPtrTestClass>(TestPackage, TestObjectName, RF_Transient);
+	TObjectPtr<UObject> TestNotLazyObject = NewObject<UObjectPtrNotLazyTestClass>(TestPackage, TEXT("NotLazy"), RF_Transient);
+
+	//compare resolved ptr against nullptr
+	TObjectPtr<UObject> ResolvedPtr = TestObject;
+	CHECK(ResolvedPtr.IsResolved());
+	CHECK(Ptr != ResolvedPtr);  CHECK(ResolveCount == 0u);
+	CHECK(ResolvedPtr != Ptr);  CHECK(ResolveCount == 0u);
+	CHECK_FALSE(Ptr == ResolvedPtr);  CHECK(ResolveCount == 0u);
+	CHECK_FALSE(ResolvedPtr == Ptr);  CHECK(ResolveCount == 0u);
+
+	//compare unresolved against nullptr
+	FObjectPtr FPtr(MakePackedObjectRef(TestObject));
+	TObjectPtr<UObject> UnResolvedPtr = *reinterpret_cast<TObjectPtr<UObject>*>(&FPtr);
+	CHECK(!UnResolvedPtr.IsResolved());
+	CHECK_FALSE(Ptr == UnResolvedPtr); CHECK(ResolveCount == 0u);
+	CHECK_FALSE(UnResolvedPtr == Ptr); CHECK(ResolveCount == 0u);
+	CHECK(Ptr != UnResolvedPtr); CHECK(ResolveCount == 0u);
+	CHECK(UnResolvedPtr != Ptr); CHECK(ResolveCount == 0u);
+
+	//compare unresolved against resolved not equal
+	CHECK_FALSE(TestNotLazyObject == UnResolvedPtr); CHECK(ResolveCount == 0u); 
+	CHECK_FALSE(UnResolvedPtr == TestNotLazyObject); CHECK(ResolveCount == 0u);
+	CHECK(TestNotLazyObject != UnResolvedPtr); CHECK(ResolveCount == 0u);
+	CHECK(UnResolvedPtr != TestNotLazyObject); CHECK(ResolveCount == 0u);
+
+	//compare resolved against naked pointer
 	Ptr = TestObject;
-
 	REQUIRE(Ptr.IsResolved());
 	CHECK(Ptr == TestObject); CHECK(ResolveCount == 0u);
 	CHECK(TestObject == Ptr); CHECK(ResolveCount == 0u);
 	CHECK_FALSE(Ptr != TestObject); CHECK(ResolveCount == 0u);
 	CHECK_FALSE(TestObject != Ptr); CHECK(ResolveCount == 0u);
 
+	//compare resolved pointer and unresolved of the same object
+	CHECK(Ptr == UnResolvedPtr); CHECK(ResolveCount == 0u);
+	CHECK(UnResolvedPtr == Ptr); CHECK(ResolveCount == 0u);
+	CHECK_FALSE(Ptr != UnResolvedPtr); CHECK(ResolveCount == 0u);
+	CHECK_FALSE(UnResolvedPtr != Ptr); CHECK(ResolveCount == 0u);
+
 	TestObject = nullptr;
 	CHECK_FALSE(Ptr == TestObject); CHECK(ResolveCount == 0u);
 	CHECK_FALSE(TestObject == Ptr); CHECK(ResolveCount == 0u);
 	CHECK(Ptr != TestObject); CHECK(ResolveCount == 0u);
 	CHECK(TestObject != Ptr); CHECK(ResolveCount == 0u);
 
-	TestObject = static_cast<UMetaData*>(Ptr.Get());
+	TestObject = static_cast<UObjectPtrTestClass*>(Ptr.Get());
 	Ptr = nullptr;
 	CHECK_FALSE(Ptr == TestObject); CHECK(ResolveCount == 0u);
 	CHECK_FALSE(TestObject == Ptr); CHECK(ResolveCount == 0u);
@@ -315,4 +405,4 @@ TEST_CASE_METHOD(FObjectHandleTestBase, "CoreUObject::FObjectHandle::Resolve Mal
 }
 #endif // UE_WITH_OBJECT_HANDLE_LATE_RESOLVE
 
-#endif
\ No newline at end of file
+#endif
diff --git a/Engine/Source/Runtime/CoreUObject/Tests/ObjectPtrTest.cpp b/Engine/Source/Runtime/CoreUObject/Tests/ObjectPtrTest.cpp
--- a/Engine/Source/Runtime/CoreUObject/Tests/ObjectPtrTest.cpp
+++ b/Engine/Source/Runtime/CoreUObject/Tests/ObjectPtrTest.cpp
@@ -3,17 +3,20 @@
 #include "UObject/ObjectPtr.h"
 
 #include "ObjectRefTrackingTestBase.h"
-
+#include "ObjectPtrTestClass.h"
 #include "Concepts/EqualityComparable.h"
 #include "Serialization/ArchiveCountMem.h"
 #include "Templates/Models.h"
+#include "UObject/LinkerLoadImportBehavior.h"
 #include "UObject/Interface.h"
 #include "UObject/MetaData.h"
 #include "UObject/ObjectRedirector.h"
+#include "UObject/Package.h"
 #include "UObject/SoftObjectPath.h"
-
+#include "Misc/ScopeExit.h"
 #include <type_traits>
 
+
 namespace ObjectPtrTest
 {
 using FMutableObjectPtr = TObjectPtr<UObject>;
@@ -26,8 +29,6 @@ using FConstPackagePtr = TObjectPtr<const UPackage>;
 class UForwardDeclaredObjDerived;
 class FForwardDeclaredNotObjDerived;
 
-using UTestDummyObject = UMetaData;
-
 static_assert(sizeof(FObjectPtr) == sizeof(FObjectHandle), "FObjectPtr type must always compile to something equivalent to an FObjectHandle size.");
 static_assert(sizeof(FObjectPtr) == sizeof(void*), "FObjectPtr type must always compile to something equivalent to a pointer size.");
 static_assert(sizeof(TObjectPtr<UObject>) == sizeof(void*), "TObjectPtr<UObject> type must always compile to something equivalent to a pointer size.");
@@ -135,6 +136,7 @@ static_assert(!TModels<CEqualityComparableWith, FConstObjectPtr, long>::Value, "
 static_assert(!TModels<CEqualityComparableWith, FMutableObjectPtr, long>::Value, "Should not be able to compare equality and inequality bidirectionally between TObjectPtr<UObject> and long");
 #endif // #if !UE_OBJECT_PTR_NONCONFORMANCE_SUPPORT
 
+
 #if WITH_LOW_LEVEL_TESTS
 
 // Ensure that the use of incomplete types doesn't provide a means to bypass type safety on TObjectPtr
@@ -149,7 +151,16 @@ public:
 protected:
 };
 
-TEST_CASE_METHOD(FObjectPtrTestBase, "CoreUObject::TObjectPtr::Null Behavior", "[CoreUObject][ObjectPtr]")
+#if UE_WITH_OBJECT_HANDLE_LATE_RESOLVE
+TEST_CASE("CoreUObject::TObjectPtr::FindLoadBehavior")
+{
+	CHECK(UE::LinkerLoad::FindLoadBehavior(*UObjectPtrTestClass::StaticClass()) == UE::LinkerLoad::EImportBehavior::LazyOnDemand);
+	CHECK(::UE::LinkerLoad::FindLoadBehavior(*UObjectPtrDerrivedTestClass::StaticClass()) == ::UE::LinkerLoad::EImportBehavior::LazyOnDemand);
+	CHECK(::UE::LinkerLoad::FindLoadBehavior(*UObjectPtrNotLazyTestClass::StaticClass()) == ::UE::LinkerLoad::EImportBehavior::Eager);
+}
+#endif
+
+TEST_CASE("CoreUObject::TObjectPtr::Null", "[CoreUObject][ObjectPtr]")
 {
 	TObjectPtr<UObject> NullObjectPtr(nullptr);
 	TEST_TRUE(TEXT("Nullptr should equal a null object pointer"), nullptr == NullObjectPtr);
@@ -165,7 +176,7 @@ TEST_CASE_METHOD(FObjectPtrTestBase, "CoreUObject::TObjectPtr::Default Serialize
 	const FName TestPackageName(TEXT("/Engine/Test/ObjectPtrDefaultSerialize/Transient"));
 	UPackage* TestPackage = NewObject<UPackage>(nullptr, TestPackageName, RF_Transient);
 	TestPackage->AddToRoot();
-	UObject* TestSoftObject = NewObject<UTestDummyObject>(TestPackage, TEXT("DefaultSerializeObject"));
+	UObject* TestSoftObject = NewObject<UObjectPtrTestClass>(TestPackage, TEXT("DefaultSerializeObject"));
 
 	FObjectPtr DefaultSerializeObjectPtr(FObjectRef {FName("/Engine/Test/ObjectPtrDefaultSerialize/Transient"), NAME_None, NAME_None, FObjectPathId("DefaultSerializeObject")});
 
@@ -196,7 +207,7 @@ TEST_CASE_METHOD(FObjectPtrTestBase, "CoreUObject::TObjectPtr::Soft Object Path"
 	const FName TestPackageName(TEXT("/Engine/Test/ObjectPtrSoftObjectPath/Transient"));
 	UPackage* TestPackage = NewObject<UPackage>(nullptr, TestPackageName, RF_Transient);
 	TestPackage->AddToRoot();
-	UObject* TestSoftObject = NewObject<UTestDummyObject>(TestPackage, TEXT("TestSoftObject"));
+	UObject* TestSoftObject = NewObject<UObjectPtrTestClass>(TestPackage, TEXT("TestSoftObject"));
 
 	FObjectPtr DefaultSoftObjPtr(FObjectRef {FName("/Engine/Test/ObjectPtrSoftObjectPath/Transient"), NAME_None, NAME_None, FObjectPathId("TestSoftObject")});
 
@@ -225,28 +236,26 @@ TEST_CASE_METHOD(FObjectPtrTestBase, "CoreUObject::TObjectPtr::Hash Consistency"
 	const FName TestPackage1Name(TEXT("/Engine/Test/ObjectPtrHashConsistency1/Transient"));
 	UPackage* TestPackage1 = NewObject<UPackage>(nullptr, TestPackage1Name, RF_Transient);
 	TestPackage1->AddToRoot();
-	UObject* TestOuter1 = NewObject<UTestDummyObject>(TestPackage1, TEXT("TestOuter1"));
-	UObject* TestOuter2 = NewObject<UTestDummyObject>(TestPackage1, TEXT("TestOuter2"));
-	UObject* TestOuter3 = NewObject<UTestDummyObject>(TestPackage1, TEXT("TestOuter3"));
-	UObject* TestOuter4 = NewObject<UTestDummyObject>(TestPackage1, TEXT("TestOuter4"));
-
-	UObject* TestPublicObject = NewObject<UTestDummyObject>(TestOuter1, TEXT("TestPublicObject"), RF_Public);
-	UObjectRedirector* TestRedirectorToPublicObject1 = NewObject<UObjectRedirector>(TestOuter3, TEXT("TestPublicRedirector1"), RF_Public);
-	TestRedirectorToPublicObject1->DestinationObject = TestPublicObject;
-	UObjectRedirector* TestRedirectorToPublicObject2 = NewObject<UObjectRedirector>(TestOuter4, TEXT("TestPublicRedirector2"), RF_Public);
-	TestRedirectorToPublicObject2->DestinationObject = TestPublicObject;
+	UObject* TestOuter1 = NewObject<UObjectPtrTestClass>(TestPackage1, TEXT("TestOuter1"));
+	UObject* TestOuter2 = NewObject<UObjectPtrTestClass>(TestPackage1, TEXT("TestOuter2"));
+	UObject* TestOuter3 = NewObject<UObjectPtrTestClass>(TestPackage1, TEXT("TestOuter3"));
+	UObject* TestOuter4 = NewObject<UObjectPtrTestClass>(TestPackage1, TEXT("TestOuter4"));
 
+	UObject* TestPublicObject = NewObject<UObjectPtrTestClass>(TestOuter1, TEXT("TestPublicObject"), RF_Public);
+#if UE_WITH_OBJECT_HANDLE_LATE_RESOLVE
+	MakePackedObjectRef(TestPublicObject); //construct a packed object ref the exported object. replicates linker load
+#endif
 	// Perform hash consistency checks on public object reference
 	{
 		// Check that unresolved/resolved pointers produce the same hash
 		FSnapshotObjectRefMetrics ObjectRefMetrics(*this);
 
-		FObjectPtr TestPublicWrappedObjectPtr(FObjectRef{TestPackage1Name, NAME_None, NAME_None, FObjectPathId("TestOuter1.TestPublicObject")});
+		FObjectPtr TestPublicWrappedObjectPtr(FObjectRef{TestPackage1Name, TestPublicObject->GetClass()->GetOutermost()->GetFName(), TestPublicObject->GetClass()->GetFName(), FObjectPathId("TestOuter1.TestPublicObject")});
 		ObjectRefMetrics.TestNumResolves(TEXT("Unexpected resolve count after initializing an FObjectPtr"), UE_WITH_OBJECT_HANDLE_LATE_RESOLVE ? 0 : 1);
 		ObjectRefMetrics.TestNumFailedResolves(TEXT("Unexpected resolve failure after initializing an FObjectPtr"), 0);
 
 		uint32 HashWrapped = GetTypeHash(TestPublicWrappedObjectPtr);
-		ObjectRefMetrics.TestNumResolves(TEXT("Unexpected resolve count after hashing an FObjectPtr"), 1);
+		ObjectRefMetrics.TestNumResolves(TEXT("Unexpected resolve count after hashing an FObjectPtr"), 0);
 		ObjectRefMetrics.TestNumFailedResolves(TEXT("Unexpected resolve failure after hashing an FObjectPtr"), 0);
 
 		TestPublicWrappedObjectPtr.Get();
@@ -254,28 +263,9 @@ TEST_CASE_METHOD(FObjectPtrTestBase, "CoreUObject::TObjectPtr::Hash Consistency"
 		ObjectRefMetrics.TestNumResolves(TEXT("Unexpected resolve count after resolving an FObjectPtr"), 1);
 		ObjectRefMetrics.TestNumFailedResolves(TEXT("Unexpected resolve failure after resolving an FObjectPtr"), 0);
 
-		uint32 HashRaw = GetTypeHash(TestPublicObject);
-		TEST_EQUAL(TEXT("Hash of raw public FObjectPtr should equal hash of wrapped public FObjectPtr"), HashRaw, HashWrapped);
-
-		FObjectPtr TestPublicWrappedRedir1(FObjectRef{TestPackage1Name, NAME_None, NAME_None, FObjectPathId("TestOuter3.TestPublicRedirector1")});
-		ObjectRefMetrics.TestNumResolves(TEXT("Unexpected resolve count after initializing an FObjectPtr"), UE_WITH_OBJECT_HANDLE_LATE_RESOLVE ? 1 : 2);
-		ObjectRefMetrics.TestNumFailedResolves(TEXT("Unexpected resolve failure after initializing an FObjectPtr"), 0);
-
-		uint32 HashWrappedRedir1 = GetTypeHash(TestPublicWrappedRedir1);
-		ObjectRefMetrics.TestNumResolves(TEXT("Unexpected resolve count after hashing an FObjectPtr"), 2);
-		ObjectRefMetrics.TestNumFailedResolves(TEXT("Unexpected resolve failure after hashing an FObjectPtr"), 0);
-
-		TEST_EQUAL(TEXT("Hash of first wrapped public redirector should equal hash of wrapped public FObjectPtr it references"), HashWrapped, HashWrappedRedir1);
-
-		FObjectPtr TestPublicWrappedRedir2(FObjectRef{TestPackage1Name, NAME_None, NAME_None, FObjectPathId("TestOuter4.TestPublicRedirector2")});
-		ObjectRefMetrics.TestNumResolves(TEXT("Unexpected resolve count after initializing an FObjectPtr"), UE_WITH_OBJECT_HANDLE_LATE_RESOLVE ? 2 : 3);
-		ObjectRefMetrics.TestNumFailedResolves(TEXT("Unexpected resolve failure after initializing an FObjectPtr"), 0);
-
-		uint32 HashWrappedRedir2 = GetTypeHash(TestPublicWrappedRedir2);
-		ObjectRefMetrics.TestNumResolves(TEXT("Unexpected resolve count after hashing an FObjectPtr"), 3);
-		ObjectRefMetrics.TestNumFailedResolves(TEXT("Unexpected resolve failure after hashing an FObjectPtr"), 0);
-
-		TEST_EQUAL(TEXT("Hash of first wrapped public redirector should equal hash of wrapped public FObjectPtr it references"), HashWrapped, HashWrappedRedir1);
+		TObjectPtr<UObject> ObjectPtr = TestPublicObject;
+		uint32 ObjectPtrHash = GetTypeHash(ObjectPtr);
+		TEST_EQUAL(TEXT("Hash of FObjectPtr should equal hash of TObjectPtr"), ObjectPtrHash, HashWrapped);
 
 		//Check that renaming an object doesn't change its hash
 		TestPublicObject->Rename(TEXT("TestPublicObjectRenamed"));
@@ -302,10 +292,10 @@ TEST_CASE_METHOD(FObjectPtrTestBase, "CoreUObject::TObjectPtr::Long Path", "[Cor
 		TestPackage1->RemoveFromRoot();
 	};
 
-	UObject* TestObject1 = NewObject<UTestDummyObject>(TestPackage1, TEXT("TestObject1"));
-	UObject* TestObject2 = NewObject<UTestDummyObject>(TestObject1, TEXT("TestObject2"));
-	UObject* TestObject3 = NewObject<UTestDummyObject>(TestObject2, TEXT("TestObject3"));
-	UObject* TestObject4 = NewObject<UTestDummyObject>(TestObject3, TEXT("TestObject4"));
+	UObject* TestObject1 = NewObject<UObjectPtrTestClass>(TestPackage1, TEXT("TestObject1"));
+	UObject* TestObject2 = NewObject<UObjectPtrTestClass>(TestObject1, TEXT("TestObject2"));
+	UObject* TestObject3 = NewObject<UObjectPtrTestClass>(TestObject2, TEXT("TestObject3"));
+	UObject* TestObject4 = NewObject<UObjectPtrTestClass>(TestObject3, TEXT("TestObject4"));
 
 	FObjectPathId LongPath(TestObject4);
 	FObjectPathId::ResolvedNameContainerType ResolvedNames;
@@ -323,7 +313,7 @@ TEST_CASE("CoreUObject::TObjectPtr::GetPathName", "[CoreUObject][ObjectPtr]")
 	const FName TestPackage1Name(TEXT("/Engine/Test/FObjectPtrTestLongPath/Transient"));
 	UPackage* TestPackage1 = NewObject<UPackage>(nullptr, TestPackage1Name, RF_Transient);
 	TestPackage1->AddToRoot();
-	UObject* TestObject1 = NewObject<UTestDummyObject>(TestPackage1, TEXT("TestObject1"));
+	UObject* TestObject1 = NewObject<UObjectPtrTestClass>(TestPackage1, TEXT("TestObject1"));
 
 #if UE_WITH_OBJECT_HANDLE_TRACKING
 	int ResolveCount = 0;
@@ -348,11 +338,24 @@ TEST_CASE("CoreUObject::TObjectPtr::GetPathName", "[CoreUObject][ObjectPtr]")
 	Ptr = TestObject1;
 	CHECK(TestObject1->GetPathName() == Ptr.GetPathName());
 
+	Ptr = TestPackage1;
+	CHECK(TestPackage1->GetPathName() == Ptr.GetPathName());
+
 #if UE_WITH_OBJECT_HANDLE_LATE_RESOLVE
-	FObjectPtr ObjPtr(MakePackedObjectRef(TestObject1));
-	Ptr = *reinterpret_cast<TObjectPtr<UObject>*>(&ObjPtr);
-	REQUIRE(!Ptr.IsResolved());
-	CHECK(TestObject1->GetPathName() == Ptr.GetPathName());
+	{
+		FObjectPtr ObjPtr(MakePackedObjectRef(TestObject1));
+		Ptr = *reinterpret_cast<TObjectPtr<UObject>*>(&ObjPtr);
+		REQUIRE(!Ptr.IsResolved());
+		CHECK(TestObject1->GetPathName() == Ptr.GetPathName());
+		CHECK(TestObject1->GetClass() == Ptr.GetClass());
+	}
+	{
+		FObjectPtr ObjPtr(MakePackedObjectRef(TestPackage1));
+		Ptr = *reinterpret_cast<TObjectPtr<UObject>*>(&ObjPtr);
+		REQUIRE(!Ptr.IsResolved());
+		CHECK(TestPackage1->GetPathName() == Ptr.GetPathName());
+		CHECK(TestPackage1->GetClass() == Ptr.GetClass());
+	}
 #endif 
 
 #if UE_WITH_OBJECT_HANDLE_TRACKING
@@ -360,6 +363,320 @@ TEST_CASE("CoreUObject::TObjectPtr::GetPathName", "[CoreUObject][ObjectPtr]")
 #endif
 }
 
+#if UE_WITH_OBJECT_HANDLE_LATE_RESOLVE
+TEST_CASE("CoreUObject::TObjectPtr::PackageRename")
+{
+	const FName TestPackageName(TEXT("/Engine/Test/TestName/Transient"));
+	UPackage* TestPackage = NewObject<UPackage>(nullptr, TestPackageName, RF_Transient);
+	TestPackage->AddToRoot();
+
+	//register package with the object handle registry
+	MakePackedObjectRef(TestPackage);
+
+	UObject* Obj1 = NewObject<UObjectPtrTestClass>(TestPackage, TEXT("Obj1"));
+	UObject* Inner1 = NewObject<UObjectPtrTestClass>(Obj1, TEXT("Inner1"));
+
+	UObject* Class = Obj1->GetClass();
+
+	Class->GetOutermost()->GetMetaData()->SetValue(Class, TEXT("LoadBehavior"), TEXT("LazyOnDemand"));
+
+	FObjectPtr ObjPtr(MakePackedObjectRef(Obj1));
+	FObjectPtr InnerPtr(MakePackedObjectRef(Inner1));
+
+	TObjectPtr<UObject> BeforeRename = *reinterpret_cast<TObjectPtr<UObject>*>(&ObjPtr);
+	TObjectPtr<UObject> BeforeRenameInner = *reinterpret_cast<TObjectPtr<UObject>*>(&InnerPtr);
+
+	REQUIRE(!BeforeRename.IsResolved()); 
+	CHECK(Obj1->GetPathName() == BeforeRename.GetPathName());
+	REQUIRE(!BeforeRenameInner.IsResolved());
+	CHECK(Inner1->GetPathName() == BeforeRenameInner.GetPathName());
+	
+	TestPackage->Rename(TEXT("/Engine/Test/TestName/NewName"));
+	CHECK(Obj1->GetPathName() == BeforeRename.GetPathName());
+	CHECK(Inner1->GetPathName() == BeforeRenameInner.GetPathName());
+
+	TObjectPtr<UObject> AfterRenameResolved = Obj1;
+	TObjectPtr<UObject> AfterRenameInnerResolved = Inner1;
+	CHECK(BeforeRename == AfterRenameResolved);
+	CHECK(BeforeRenameInner == AfterRenameInnerResolved);
+
+	FObjectPtr ObjPtr2(MakePackedObjectRef(Obj1));
+	FObjectPtr InnerPtr2(MakePackedObjectRef(Inner1));
+	TObjectPtr<UObject> AfterRenameUnresolved = *reinterpret_cast<TObjectPtr<UObject>*>(&ObjPtr2);
+	TObjectPtr<UObject> AfterRenameInnerUnresolved = *reinterpret_cast<TObjectPtr<UObject>*>(&InnerPtr2);
+
+	CHECK(BeforeRename == AfterRenameUnresolved);
+	CHECK(BeforeRenameInner == AfterRenameInnerUnresolved);
+
+	Obj1->Rename(TEXT("RenamedObj"), nullptr);
+	CHECK(Obj1->GetPathName() == BeforeRename.GetPathName());
+	CHECK(Inner1->GetPathName() == BeforeRenameInner.GetPathName());
+	CHECK(BeforeRename == AfterRenameResolved);
+	CHECK(BeforeRenameInner == AfterRenameInnerResolved);
+}
+
+TEST_CASE("CoreUObject::TObjectPtr::InnerRename")
+{
+	const FName TestPackageName(TEXT("/Engine/TestPackage"));
+	UPackage* TestPackage = NewObject<UPackage>(nullptr, TestPackageName, RF_Transient);
+	TestPackage->AddToRoot();
+	UObject* Obj1 = NewObject<UObjectPtrTestClass>(TestPackage, TEXT("Obj1"));
+	UObject* Inner1 = NewObject<UObjectPtrTestClass>(Obj1, TEXT("Inner1"));
+
+	UObject* Obj2 = NewObject<UObjectPtrTestClass>(TestPackage, TEXT("Obj2"));
+	UObject* Inner2 = NewObject<UObjectPtrTestClass>(Obj2, TEXT("Inner2"));
+
+	FObjectPtr ObjPtr1(MakePackedObjectRef(Obj1));
+	FObjectPtr InnerPtr1(MakePackedObjectRef(Inner1));
+
+	FObjectPtr ObjPtr2(MakePackedObjectRef(Obj2));
+	FObjectPtr InnerPtr2(MakePackedObjectRef(Inner2));
+
+	TObjectPtr<UObject> BeforeRename1 = *reinterpret_cast<TObjectPtr<UObject>*>(&ObjPtr1);
+	TObjectPtr<UObject> BeforeRenameInner1 = *reinterpret_cast<TObjectPtr<UObject>*>(&InnerPtr1);
+
+	TObjectPtr<UObject> BeforeRename2 = *reinterpret_cast<TObjectPtr<UObject>*>(&ObjPtr2);
+	TObjectPtr<UObject> BeforeRenameInner2 = *reinterpret_cast<TObjectPtr<UObject>*>(&InnerPtr2);
+
+	REQUIRE(!BeforeRename1.IsResolved());
+	CHECK(Obj1->GetPathName() == BeforeRename1.GetPathName());
+	REQUIRE(!BeforeRenameInner1.IsResolved());
+	CHECK(Inner1->GetPathName() == BeforeRenameInner1.GetPathName());
+
+	REQUIRE(!BeforeRename2.IsResolved());
+	CHECK(Obj2->GetPathName() == BeforeRename2.GetPathName());
+	REQUIRE(!BeforeRenameInner2.IsResolved());
+	CHECK(Inner2->GetPathName() == BeforeRenameInner2.GetPathName());
+
+	Obj1->Rename(TEXT("RenamedObj"), nullptr);
+
+	FObjectPtr AfterRenameObjPtr1(MakePackedObjectRef(Obj1));
+	FObjectPtr AfterRenameInnerPtr1(MakePackedObjectRef(Inner1));
+
+	TObjectPtr<UObject> AfterRenameUnresolved1 = *reinterpret_cast<TObjectPtr<UObject>*>(&AfterRenameObjPtr1);
+	TObjectPtr<UObject> AfterRenameInnerUnresolved1 = *reinterpret_cast<TObjectPtr<UObject>*>(&AfterRenameInnerPtr1);
+
+	TObjectPtr<UObject> AfterRename1 = Obj1;
+	TObjectPtr<UObject> AfterRenameInner1 = Inner1;
+
+	CHECK(Obj1->GetPathName() == BeforeRename1.GetPathName());
+	REQUIRE(!BeforeRenameInner1.IsResolved());
+	CHECK(Inner1->GetPathName() == BeforeRenameInner1.GetPathName());
+
+	CHECK(BeforeRename1 == AfterRename1);
+	CHECK(BeforeRenameInner1 == AfterRenameInner1);
+
+	CHECK(BeforeRename1 == AfterRenameUnresolved1);
+	CHECK(BeforeRenameInner1 == AfterRenameInnerUnresolved1);
+
+	Obj1->Rename(nullptr, Inner2);
+	CHECK(Obj2->GetPathName() == BeforeRename2.GetPathName());
+	REQUIRE(!BeforeRenameInner2.IsResolved());
+	CHECK(Inner2->GetPathName() == BeforeRenameInner2.GetPathName());
+	CHECK(BeforeRename1 == AfterRename1);
+	CHECK(BeforeRename1.GetPathName() == AfterRename1.GetPathName());
+	CHECK(BeforeRenameInner1 == AfterRenameInner1);
+
+	TestPackage->RemoveFromRoot();
+}
+
+TEST_CASE("CoreUObject::TObjectPtr::Move")
+{
+	UPackage* TestPackageA = NewObject<UPackage>(nullptr, TEXT("/Engine/PackageA"), RF_Transient);
+	TestPackageA->AddToRoot();
+	UObject* Obj1 = NewObject<UObjectPtrTestClass>(TestPackageA, TEXT("Obj1"));
+	UObject* Inner1 = NewObject<UObjectPtrTestClass>(Obj1, TEXT("Inner1"));
+	
+	UPackage* TestPackageB = NewObject<UPackage>(nullptr, TEXT("/Engine/PackageB"), RF_Transient);
+	TestPackageB->AddToRoot();
+
+	FObjectPtr FObjPtr(MakePackedObjectRef(Obj1));
+	FObjectPtr FInnerPtr(MakePackedObjectRef(Inner1));
+
+	TObjectPtr<UObject> BeforeRename = *reinterpret_cast<TObjectPtr<UObject>*>(&FObjPtr);
+	TObjectPtr<UObject> BeforeRenameInner = *reinterpret_cast<TObjectPtr<UObject>*>(&FInnerPtr);
+
+	TObjectPtr<UObject> ObjPtr1 = Obj1;
+	TObjectPtr<UObject> InnerPtr1 = Inner1;
+
+	REQUIRE(!BeforeRename.IsResolved());
+	CHECK(BeforeRename == ObjPtr1);
+	CHECK(Obj1->GetPathName() == BeforeRename.GetPathName());
+	REQUIRE(!BeforeRenameInner.IsResolved());
+	CHECK(InnerPtr1 == BeforeRenameInner);
+	CHECK(Inner1->GetPathName() == BeforeRenameInner.GetPathName());
+
+	Obj1->Rename(TEXT("Obj2"), TestPackageB);
+	bool e = BeforeRename == ObjPtr1;
+
+	REQUIRE(!BeforeRename.IsResolved());
+	CHECK(BeforeRename == ObjPtr1);
+	CHECK(Obj1->GetPathName() == BeforeRename.GetPathName());
+	CHECK(InnerPtr1 == BeforeRenameInner);
+	CHECK(Inner1->GetPathName() == BeforeRenameInner.GetPathName());
+
+	TObjectPtr<UObject> AfterRename = Obj1;
+	TObjectPtr<UObject> AfterRenameInner = Inner1;
+	CHECK(BeforeRename == AfterRename);
+	CHECK(BeforeRenameInner == AfterRenameInner);
+
+	TestPackageA->RemoveFromRoot();
+	TestPackageB->RemoveFromRoot();
+}
+
+TEST_CASE("CoreUObject::TObjectPtr::GetTypeHash")
+{
+	int ResolveCount = 0;
+#if UE_WITH_OBJECT_HANDLE_TRACKING
+	auto ResolveDelegate = FObjectHandleReferenceResolvedDelegate::CreateLambda([&ResolveCount](const FObjectRef& SourceRef, UPackage* ObjectPackage, UObject* Object)
+		{
+			++ResolveCount;
+		});
+	auto Handle = AddObjectHandleReferenceResolvedCallback(ResolveDelegate);
+#endif
+
+	UPackage* TestPackageA = NewObject<UPackage>(nullptr, TEXT("/Engine/PackageA"), RF_Transient);
+	TestPackageA->AddToRoot();
+	TObjectPtr<UObject> Obj1 = NewObject<UObjectPtrTestClass>(TestPackageA, TEXT("Obj1"));
+	TObjectPtr<UObject> Inner1 = NewObject<UObjectPtrTestClass>(Obj1, TEXT("Inner1"));
+
+	MakePackedObjectRef(TestPackageA);
+	MakePackedObjectRef(Obj1.Get());
+	MakePackedObjectRef(Inner1.Get());
+
+	TMap<TObjectPtr<UObject>, int> TestMap;
+
+	int32 BeforeKey = GetTypeHash(Obj1);
+	int32 BeforeKeyInner = GetTypeHash(Inner1);
+	int32 BeforeKeyPackage = GetTypeHash(TestPackageA);
+
+	TestMap.Add(Obj1, 0);
+	TestMap.Add(Inner1, 1);
+	TestMap.Add(TestPackageA, 3);
+	
+	//add a bunch of objects
+	for (int i = 0; i < 10; ++i)
+	{
+		TestMap.Add(NewObject<UObjectPtrTestClass>(TestPackageA), TestMap.Num());;
+	}
+	for (int i = 0; i < 10; ++i)
+	{
+		TestMap.Add(NewObject<UObjectPtrNotLazyTestClass>(TestPackageA), i);;
+	}
+
+	TObjectPtr<UObject> NotLazy = NewObject<UObjectPtrNotLazyTestClass>(TestPackageA);
+	int NotLazyValue = TestMap.Num();
+	TestMap.Add(NotLazy, NotLazyValue);
+	CHECK(GetTypeHash(NotLazy) == GetTypeHash(NotLazy.Get()));
+	CHECK(TestMap[NotLazy] == NotLazyValue);
+
+
+	FObjectPtr FPackagePtr(MakePackedObjectRef(TestPackageA));
+	FObjectPtr FObjPtr(MakePackedObjectRef(Obj1));
+	FObjectPtr FInnerPtr(MakePackedObjectRef(Inner1));
+
+	TObjectPtr<UObject> BeforeRename = *reinterpret_cast<TObjectPtr<UObject>*>(&FObjPtr);
+	TObjectPtr<UObject> BeforeRenameInner = *reinterpret_cast<TObjectPtr<UObject>*>(&FInnerPtr);
+	TObjectPtr<UObject> BeforeRenamePackage = *reinterpret_cast<TObjectPtr<UObject>*>(&FPackagePtr);
+
+	int32 BeforeRenameTypeHash = GetTypeHash(BeforeRename);
+	int32 BeforeRenameTypeHashInner = GetTypeHash(BeforeRenameInner);
+	int32 BeforeRenamePackageTypeHash = GetTypeHash(BeforeRenamePackage);
+
+	CHECK(BeforeKey == BeforeRenameTypeHash);
+	CHECK(BeforeKey == BeforeRenameTypeHash);
+	CHECK(BeforeKeyInner == BeforeRenameTypeHashInner);
+	CHECK(ResolveCount == 0);
+		
+	CHECK(TestMap[Obj1] == 0);
+	CHECK(TestMap[BeforeRename] == 0);
+
+	CHECK(TestMap[Inner1] == 1);
+	CHECK(TestMap[BeforeRenameInner] == 1);
+
+	CHECK(TestMap[TestPackageA] == 3);
+	CHECK(TestMap[BeforeRenamePackage] == 3);
+
+	CHECK(!BeforeRename.IsResolved());
+	CHECK(!BeforeRenameInner.IsResolved());
+	CHECK(!BeforeRenamePackage.IsResolved());
+
+	auto TestMapFunc = [&]()
+	{
+		TMap<TObjectPtr<UObject>, bool> Map;
+		Map.Add(Obj1, false);
+		Map.Add(Inner1, true);
+
+		CHECK(ResolveCount == 0);
+		bool* Found = Map.Find(BeforeRename);
+		CHECK(ResolveCount == 0);
+		CHECK(Found != nullptr);
+		CHECK(!*Found);
+
+		Found = Map.Find(BeforeRenameInner);
+		CHECK(ResolveCount == 0);
+		CHECK(Found != nullptr);
+		CHECK(*Found);
+
+		CHECK(ResolveCount == 0);
+		Found = Map.Find(Obj1);
+		CHECK(ResolveCount == 0);
+		CHECK(Found != nullptr);
+		CHECK(!*Found);
+
+		CHECK(ResolveCount == 0);
+		Found = Map.Find(Inner1);
+		CHECK(ResolveCount == 0);
+		CHECK(Found != nullptr);
+		CHECK(*Found);
+	};
+
+	auto TestUnResolvedMapFunc = [&]()
+	{
+		TMap<TObjectPtr<UObject>, bool> Map;
+		Map.Add(BeforeRename, false);
+		Map.Add(BeforeRenameInner, true);
+
+		CHECK(ResolveCount == 0);
+		bool* Found = Map.Find(BeforeRename);
+		CHECK(ResolveCount == 0);
+		CHECK(Found != nullptr);
+		CHECK(!*Found);
+
+		Found = Map.Find(BeforeRenameInner);
+		CHECK(ResolveCount == 0);
+		CHECK(Found != nullptr);
+		CHECK(*Found);
+
+		CHECK(ResolveCount == 0);
+		Found = Map.Find(Obj1);
+		CHECK(ResolveCount == 0);
+		CHECK(Found != nullptr);
+		CHECK(!*Found);
+
+		CHECK(ResolveCount == 0);
+		Found = Map.Find(Inner1);
+		CHECK(ResolveCount == 0);
+		CHECK(Found != nullptr);
+		CHECK(*Found);
+	};
+
+	TestMapFunc();
+	TestUnResolvedMapFunc();
+
+	UPackage* TestPackageB = NewObject<UPackage>(nullptr, TEXT("/Engine/PackageB"), RF_Transient);
+	TestPackageB->AddToRoot();
+
+	Obj1->Rename(nullptr, TestPackageB);
+
+	TestMapFunc();
+	TestUnResolvedMapFunc();
+
+	TestPackageA->RemoveFromRoot();
+	TestPackageB->RemoveFromRoot();
+}
+#endif
+
 // @TODO: OBJPTR: We should have a test that ensures that lazy loading of an object with an external package is handled correctly.
 //				  This should also include external packages in the outer chain of the target object.
 // IMPLEMENT_CUSTOM_SIMPLE_AUTOMATION_TEST(FObjectPtrTestExternalPackages, FObjectPtrTestBase, TEXT(TEST_NAME_ROOT ".ExternalPackages"), ObjectPtrTestFlags)
@@ -373,12 +690,12 @@ TEST_CASE("CoreUObject::TObjectPtr::GetPathName", "[CoreUObject][ObjectPtr]")
 // 	const FName TestPackage1Name(TEXT("/Engine/Test/ObjectPtrExternalPackages1/Transient"));
 // 	UPackage* TestPackage1 = NewObject<UPackage>(nullptr, TestPackage1Name, RF_Transient);
 // 	TestPackage1->AddToRoot();
-// 	UObject* TestOuter1 = NewObject<UTestDummyObject>(TestPackage1, TEXT("TestOuter1"));
-// 	UObject* TestOuter2 = NewObject<UTestDummyObject>(TestPackage1, TEXT("TestOuter2"));
-// 	UObject* TestOuter3 = NewObject<UTestDummyObject>(TestPackage1, TEXT("TestOuter3"));
-// 	UObject* TestOuter4 = NewObject<UTestDummyObject>(TestPackage1, TEXT("TestOuter4"));
+// 	UObject* TestOuter1 = NewObject<UObjectPtrTestClass>(TestPackage1, TEXT("TestOuter1"));
+// 	UObject* TestOuter2 = NewObject<UObjectPtrTestClass>(TestPackage1, TEXT("TestOuter2"));
+// 	UObject* TestOuter3 = NewObject<UObjectPtrTestClass>(TestPackage1, TEXT("TestOuter3"));
+// 	UObject* TestOuter4 = NewObject<UObjectPtrTestClass>(TestPackage1, TEXT("TestOuter4"));
 
-// 	UObject* TestPublicObject = NewObject<UTestDummyObject>(TestOuter4, TEXT("TestPublicObject"), RF_Public);
+// 	UObject* TestPublicObject = NewObject<UObjectPtrTestClass>(TestOuter4, TEXT("TestPublicObject"), RF_Public);
 // 	TestPublicObject->SetExternalPackage(TestExternalPackage1);
 
 // 	TestPackage1->RemoveFromRoot();
diff --git a/Engine/Source/Runtime/CoreUObject/Tests/ObjectPtrTestClass.cpp b/Engine/Source/Runtime/CoreUObject/Tests/ObjectPtrTestClass.cpp
new file mode 100644
--- /dev/null
+++ b/Engine/Source/Runtime/CoreUObject/Tests/ObjectPtrTestClass.cpp
@@ -0,0 +1,24 @@
+// Copyright Epic Games, Inc. All Rights Reserved.
+
+#if WITH_LOW_LEVEL_TESTS
+
+#include "ObjectPtrTestClass.h"
+#include "UObject/Class.h"
+#include "UObject/Package.h"
+#include "UObject/MetaData.h"
+
+IMPLEMENT_CORE_INTRINSIC_CLASS(UObjectPtrTestClass, UObject,
+	{
+		auto MetaData = Class->GetOutermost()->GetMetaData();
+		if (MetaData)
+		{
+			MetaData->SetValue(Class, TEXT("LoadBehavior"), TEXT("LazyOnDemand"));
+		}
+	}
+);
+
+IMPLEMENT_CORE_INTRINSIC_CLASS(UObjectPtrDerrivedTestClass, UObjectPtrTestClass, {});
+
+IMPLEMENT_CORE_INTRINSIC_CLASS(UObjectPtrNotLazyTestClass, UObject, {});
+
+#endif
\ No newline at end of file
diff --git a/Engine/Source/Runtime/CoreUObject/Tests/ObjectPtrTestClass.h b/Engine/Source/Runtime/CoreUObject/Tests/ObjectPtrTestClass.h
new file mode 100644
--- /dev/null
+++ b/Engine/Source/Runtime/CoreUObject/Tests/ObjectPtrTestClass.h
@@ -0,0 +1,40 @@
+// Copyright Epic Games, Inc. All Rights Reserved.
+
+#pragma once
+#if WITH_LOW_LEVEL_TESTS
+
+#include "UObject/Object.h"
+
+
+//simple test class for testing TObjectPtr resolve behavior
+class UObjectPtrTestClass : public UObject
+{
+	DECLARE_CLASS_INTRINSIC(UObjectPtrTestClass, UObject, CLASS_MatchedSerializers, TEXT("/Script/CoreUObject"))
+
+public:
+
+};
+
+
+
+//derived test class
+class UObjectPtrDerrivedTestClass : public UObjectPtrTestClass
+{
+	DECLARE_CLASS_INTRINSIC(UObjectPtrDerrivedTestClass, UObjectPtrTestClass, CLASS_MatchedSerializers, TEXT("/Script/CoreUObject"))
+
+public:
+
+};
+
+
+//non lazy test class
+class UObjectPtrNotLazyTestClass : public UObject
+{
+	DECLARE_CLASS_INTRINSIC(UObjectPtrNotLazyTestClass, UObject, CLASS_MatchedSerializers, TEXT("/Script/CoreUObject"))
+
+public:
+
+};
+
+
+#endif
\ No newline at end of file
diff --git a/Engine/Source/Runtime/CoreUObject/Tests/ObjectRefTrackingTestBase.h b/Engine/Source/Runtime/CoreUObject/Tests/ObjectRefTrackingTestBase.h
--- a/Engine/Source/Runtime/CoreUObject/Tests/ObjectRefTrackingTestBase.h
+++ b/Engine/Source/Runtime/CoreUObject/Tests/ObjectRefTrackingTestBase.h
@@ -30,28 +30,21 @@ public:
 			Test.ConditionalInstallCallbacks();
 		}
 
-		bool TestNumResolves(const TCHAR* What, uint32 ExpectedDelta)
+		void TestNumResolves(const TCHAR* What, uint32 ExpectedDelta)
 		{
 #if UE_WITH_OBJECT_HANDLE_TRACKING
-			bool bValue = OriginalNumResolves + ExpectedDelta == Test.GetNumResolves();
-			TEST_TRUE(What, bValue);
-			return bValue;
+			TEST_EQUAL(What, OriginalNumResolves + ExpectedDelta, Test.GetNumResolves());
 #endif
-			return true;
 		}
 
-		bool TestNumFailedResolves(const TCHAR* What, uint32 ExpectedDelta)
+		void TestNumFailedResolves(const TCHAR* What, uint32 ExpectedDelta)
 		{
 #if UE_WITH_OBJECT_HANDLE_TRACKING
-			bool bValue = OriginalNumFailedResolves + ExpectedDelta == Test.GetNumFailedResolves();
-			TEST_TRUE(What, bValue);
-			return bValue;
+			TEST_EQUAL(What, OriginalNumFailedResolves + ExpectedDelta, Test.GetNumFailedResolves());
 #endif
-			return true;
-
 		}
 
-		bool TestNumReads(const TCHAR* What, uint32 ExpectedDelta, bool bAllowAdditionalReads = false)
+		void TestNumReads(const TCHAR* What, uint32 ExpectedDelta, bool bAllowAdditionalReads = false)
 		{
 #if UE_WITH_OBJECT_HANDLE_TRACKING
 			bool bValue = false;
@@ -65,10 +58,7 @@ public:
 				bValue = OriginalNumReads + ExpectedDelta == Test.GetNumReads();
 				TEST_TRUE(What, bValue);
 			}
-			return bValue;
 #endif
-			return true;
-
 		}
 	private:
 		FObjectRefTrackingTestBase& Test;
