EditorDomain: Remove indeterminism caused by SavedByEngineVersion and CompatibleWithEngineVersion.
#rnx

diff --git a/Engine/Source/Editor/UnrealEd/Private/EditorDomain/EditorDomainUtils.cpp b/Engine/Source/Editor/UnrealEd/Private/EditorDomain/EditorDomainUtils.cpp
--- a/Engine/Source/Editor/UnrealEd/Private/EditorDomain/EditorDomainUtils.cpp
+++ b/Engine/Source/Editor/UnrealEd/Private/EditorDomain/EditorDomainUtils.cpp
@@ -203,7 +203,7 @@ FBlake3Hash GGlobalConstructClassesHash;
 int64 GMaxBulkDataSize = -1;
 
 // Change to a new guid when EditorDomain needs to be invalidated
-const TCHAR* EditorDomainVersion = TEXT("25CC763ED9BB4CEF9823E7B41FA7C771");
+const TCHAR* EditorDomainVersion = TEXT("1F8E29F8458141ADB52D9EBEEA2991F0");
 
 // Identifier of the CacheBuckets for EditorDomain tables
 const TCHAR* EditorDomainPackageBucketName = TEXT("EditorDomainPackage");
diff --git a/Engine/Source/Runtime/CoreUObject/Private/UObject/SavePackage2.cpp b/Engine/Source/Runtime/CoreUObject/Private/UObject/SavePackage2.cpp
--- a/Engine/Source/Runtime/CoreUObject/Private/UObject/SavePackage2.cpp
+++ b/Engine/Source/Runtime/CoreUObject/Private/UObject/SavePackage2.cpp
@@ -1045,6 +1045,13 @@ ESavePackageResult BuildLinker(FSaveContext& SaveContext)
 	SaveContext.GetLinker()->Summary.PersistentGuid = SaveContext.GetPackage()->GetPersistentGuid();
 #endif
 	SaveContext.GetLinker()->Summary.Generations = TArray<FGenerationInfo>{ FGenerationInfo(0, 0) };
+	if (SaveContext.IsProceduralSave())
+	{
+		// Procedural saves should be deterministic, so we have to clear the EngineVersion fields
+		// to avoid indeterminism when it changes
+		SaveContext.GetLinker()->Summary.SavedByEngineVersion = FEngineVersion();
+		SaveContext.GetLinker()->Summary.CompatibleWithEngineVersion = FEngineVersion();
+	}
 
 	FLinkerSave* Linker = SaveContext.GetLinker();
 
