SavePackage WeakObjectPtr: Fix null pointer dereference in WeakObjectPtr serializer.
#rnx

diff --git a/Engine/Source/Runtime/CoreUObject/Private/UObject/SavePackage/PackageHarvester.cpp b/Engine/Source/Runtime/CoreUObject/Private/UObject/SavePackage/PackageHarvester.cpp
--- a/Engine/Source/Runtime/CoreUObject/Private/UObject/SavePackage/PackageHarvester.cpp
+++ b/Engine/Source/Runtime/CoreUObject/Private/UObject/SavePackage/PackageHarvester.cpp
@@ -514,7 +514,7 @@ FArchive& FPackageHarvester::operator<<(struct FWeakObjectPtr& Value)
 	const UObject* CurrentExport = CurrentExportDependencies.CurrentExport;
 
 	// @todo FH: Should we really force weak import in cooked builds?
-	if (IsCooking() || (!Object->IsInPackage(SaveContext.GetPackage()) && (Object->GetOutermostObject() == CurrentExport->GetOutermostObject())))
+	if (IsCooking() || (Object && !Object->IsInPackage(SaveContext.GetPackage()) && (Object->GetOutermostObject() == CurrentExport->GetOutermostObject())))
 	{
 		*this << Object;
 	}
