fix for failing tests on platforms where meta data is not available

diff --git a/Engine/Source/Runtime/CoreUObject/Tests/ObjectPtrTestClass.cpp b/Engine/Source/Runtime/CoreUObject/Tests/ObjectPtrTestClass.cpp
--- a/Engine/Source/Runtime/CoreUObject/Tests/ObjectPtrTestClass.cpp
+++ b/Engine/Source/Runtime/CoreUObject/Tests/ObjectPtrTestClass.cpp
@@ -7,6 +7,8 @@
 #include "UObject/Package.h"
 #include "UObject/MetaData.h"
 
+//can not put the #if inside as the expansion of IMPLEMENT_CORE_INTRINSIC_CLASS fails
+#if WITH_EDITORONLY_DATA
 IMPLEMENT_CORE_INTRINSIC_CLASS(UObjectPtrTestClass, UObject,
 	{
 		auto MetaData = Class->GetOutermost()->GetMetaData();
@@ -16,6 +18,9 @@ IMPLEMENT_CORE_INTRINSIC_CLASS(UObjectPtrTestClass, UObject,
 		}
 	}
 );
+#else
+IMPLEMENT_CORE_INTRINSIC_CLASS(UObjectPtrTestClass, UObject, { });
+#endif
 
 IMPLEMENT_CORE_INTRINSIC_CLASS(UObjectPtrDerrivedTestClass, UObjectPtrTestClass, {});
 
