[CoreUObject/Core]
* Removed CoreMinimal.h include from a ton of headers all the way to Actor.h. This shaved off 15% of the compile time of small files that has uht markup.

diff --git a/Engine/Plugins/Experimental/CinematicPrestreaming/Source/CinematicPrestreaming/Private/CinePrestreamingData.cpp b/Engine/Plugins/Experimental/CinematicPrestreaming/Source/CinematicPrestreaming/Private/CinePrestreamingData.cpp
--- a/Engine/Plugins/Experimental/CinematicPrestreaming/Source/CinematicPrestreaming/Private/CinePrestreamingData.cpp
+++ b/Engine/Plugins/Experimental/CinematicPrestreaming/Source/CinematicPrestreaming/Private/CinePrestreamingData.cpp
@@ -1,6 +1,9 @@
 // Copyright Epic Games, Inc. All Rights Reserved.
 
 #include "CinePrestreamingData.h"
+#include "Misc/FrameNumber.h"
+
+#include UE_INLINE_GENERATED_CPP_BY_NAME(CinePrestreamingData)
 
 void UCinePrestreamingData::GetResourceSizeEx(FResourceSizeEx& CumulativeResourceSize)
 {
diff --git a/Engine/Plugins/Experimental/CinematicPrestreaming/Source/CinematicPrestreaming/Public/CinePrestreamingData.h b/Engine/Plugins/Experimental/CinematicPrestreaming/Source/CinematicPrestreaming/Public/CinePrestreamingData.h
--- a/Engine/Plugins/Experimental/CinematicPrestreaming/Source/CinematicPrestreaming/Public/CinePrestreamingData.h
+++ b/Engine/Plugins/Experimental/CinematicPrestreaming/Source/CinematicPrestreaming/Public/CinePrestreamingData.h
@@ -1,6 +1,7 @@
 // Copyright Epic Games, Inc. All Rights Reserved.
 #pragma once
 
+#include "Misc/DateTime.h"
 #include "UObject/Object.h"
 #include "CinePrestreamingData.generated.h"
 
diff --git a/Engine/Source/Runtime/AdvancedWidgets/Public/Slate/SRadialSlider.h b/Engine/Source/Runtime/AdvancedWidgets/Public/Slate/SRadialSlider.h
--- a/Engine/Source/Runtime/AdvancedWidgets/Public/Slate/SRadialSlider.h
+++ b/Engine/Source/Runtime/AdvancedWidgets/Public/Slate/SRadialSlider.h
@@ -1,6 +1,8 @@
 // Copyright Epic Games, Inc. All Rights Reserved.
 #pragma once
 
+// HEADER_UNIT_SKIP - Deprecated
+
 #include "Widgets/SRadialSlider.h"
 
 UE_DEPRECATED(5.1, "Including Slate/SRadialSlider.h is deprecated. The location of SRadialSlider has changed. Please include Widgets/SRadialSlider.h instead.")
diff --git a/Engine/Source/Runtime/AssetRegistry/Public/AssetRegistryHelpers.h b/Engine/Source/Runtime/AssetRegistry/Public/AssetRegistryHelpers.h
--- a/Engine/Source/Runtime/AssetRegistry/Public/AssetRegistryHelpers.h
+++ b/Engine/Source/Runtime/AssetRegistry/Public/AssetRegistryHelpers.h
@@ -2,6 +2,8 @@
 
 #pragma once
 
+// HEADER_UNIT_SKIP - Deprecated
+
 // Deprecated in UE 5.1.
 #pragma message (__FILE__"(6): warning: #include AssetRegistry/AssetRegistryHelpers.h instead of AssetRegistryHelpers.h")
 
diff --git a/Engine/Source/Runtime/AssetRegistry/Public/AssetRegistryModule.h b/Engine/Source/Runtime/AssetRegistry/Public/AssetRegistryModule.h
--- a/Engine/Source/Runtime/AssetRegistry/Public/AssetRegistryModule.h
+++ b/Engine/Source/Runtime/AssetRegistry/Public/AssetRegistryModule.h
@@ -2,6 +2,8 @@
 
 #pragma once
 
+// HEADER_UNIT_SKIP - Deprecated
+
 // Deprecated in UE 5.1.
 #pragma message (__FILE__"(6): warning: #include AssetRegistry/AssetRegistryModule.h instead of AssetRegistryModule.h")
 
diff --git a/Engine/Source/Runtime/Core/Public/Memory/VirtualStackAllocator.h b/Engine/Source/Runtime/Core/Public/Memory/VirtualStackAllocator.h
--- a/Engine/Source/Runtime/Core/Public/Memory/VirtualStackAllocator.h
+++ b/Engine/Source/Runtime/Core/Public/Memory/VirtualStackAllocator.h
@@ -2,7 +2,6 @@
 
 #pragma once
 
-#include "CoreMinimal.h"
 #include "HAL/Platform.h"
 #include "HAL/PlatformCrt.h"
 #include "HAL/PlatformMemory.h"
@@ -128,4 +127,6 @@ private:
 	void* RecentHighWaterMark = nullptr;
 };
 
-
+#if UE_ENABLE_INCLUDE_ORDER_DEPRECATED_IN_5_2
+#include "CoreMinimal.h"
+#endif
\ No newline at end of file
diff --git a/Engine/Source/Runtime/CoreUObject/Public/Serialization/ArchiveUObject.h b/Engine/Source/Runtime/CoreUObject/Public/Serialization/ArchiveUObject.h
--- a/Engine/Source/Runtime/CoreUObject/Public/Serialization/ArchiveUObject.h
+++ b/Engine/Source/Runtime/CoreUObject/Public/Serialization/ArchiveUObject.h
@@ -2,7 +2,6 @@
 
 #pragma once
 
-#include "CoreMinimal.h"
 #include "Serialization/Archive.h"
 
 struct FLazyObjectPtr;
@@ -34,3 +33,7 @@ public:
 	virtual FArchive& operator<<(FWeakObjectPtr& Value) override { return SerializeWeakObjectPtr(*this, Value); }
 	//~ End FArchive Interface
 };
+
+#if UE_ENABLE_INCLUDE_ORDER_DEPRECATED_IN_5_2
+#include "CoreMinimal.h"
+#endif
\ No newline at end of file
diff --git a/Engine/Source/Runtime/CoreUObject/Public/Serialization/SerializedPropertyScope.h b/Engine/Source/Runtime/CoreUObject/Public/Serialization/SerializedPropertyScope.h
--- a/Engine/Source/Runtime/CoreUObject/Public/Serialization/SerializedPropertyScope.h
+++ b/Engine/Source/Runtime/CoreUObject/Public/Serialization/SerializedPropertyScope.h
@@ -6,7 +6,7 @@
 
 #pragma once
 
-#include "CoreMinimal.h"
+#include "Serialization/Archive.h"
 
 /** Helper class to set and restore serialized property on an archive */
 class COREUOBJECT_API FSerializedPropertyScope
@@ -34,3 +34,7 @@ public:
 		PopProperty();
 	}
 };
+
+#if UE_ENABLE_INCLUDE_ORDER_DEPRECATED_IN_5_2
+#include "CoreMinimal.h"
+#endif
diff --git a/Engine/Source/Runtime/CoreUObject/Public/Templates/Casts.h b/Engine/Source/Runtime/CoreUObject/Public/Templates/Casts.h
--- a/Engine/Source/Runtime/CoreUObject/Public/Templates/Casts.h
+++ b/Engine/Source/Runtime/CoreUObject/Public/Templates/Casts.h
@@ -3,7 +3,6 @@
 #pragma once
 
 #include "Containers/UnrealString.h"
-#include "CoreMinimal.h"
 #include "CoreTypes.h"
 #include "Templates/LosesQualifiersFromTo.h"
 #include "Templates/UnrealTemplate.h"
@@ -507,3 +506,7 @@ namespace UECasts_Private
 }
 
 #define dynamic_cast UECasts_Private::DynamicCast
+
+#if UE_ENABLE_INCLUDE_ORDER_DEPRECATED_IN_5_2
+#include "CoreMinimal.h"
+#endif
diff --git a/Engine/Source/Runtime/CoreUObject/Public/Templates/SubclassOf.h b/Engine/Source/Runtime/CoreUObject/Public/Templates/SubclassOf.h
--- a/Engine/Source/Runtime/CoreUObject/Public/Templates/SubclassOf.h
+++ b/Engine/Source/Runtime/CoreUObject/Public/Templates/SubclassOf.h
@@ -2,7 +2,6 @@
 
 #pragma once
 
-#include "CoreMinimal.h"
 #include "UObject/Class.h"
 
 #include <type_traits>
@@ -176,3 +175,6 @@ FArchive& operator<<(FArchive& Ar, TSubclassOf<T>& SubclassOf)
 	return Ar;
 }
 
+#if UE_ENABLE_INCLUDE_ORDER_DEPRECATED_IN_5_2
+#include "CoreMinimal.h"
+#endif
diff --git a/Engine/Source/Runtime/CoreUObject/Public/UObject/Class.h b/Engine/Source/Runtime/CoreUObject/Public/UObject/Class.h
--- a/Engine/Source/Runtime/CoreUObject/Public/UObject/Class.h
+++ b/Engine/Source/Runtime/CoreUObject/Public/UObject/Class.h
@@ -13,7 +13,6 @@
 #include "Containers/Set.h"
 #include "Containers/StringFwd.h"
 #include "Containers/UnrealString.h"
-#include "CoreMinimal.h"
 #include "CoreTypes.h"
 #include "HAL/CriticalSection.h"
 #include "HAL/PlatformCrt.h"
@@ -4060,4 +4059,8 @@ UE_DECLARE_CORE_VARIANT_TYPE(Box2,		Box2D);
 UE_DECLARE_CORE_VARIANT_TYPE(Ray3,		Ray);
 UE_DECLARE_CORE_VARIANT_TYPE(Sphere3, Sphere);
 
-#undef UE_DECLARE_CORE_VARIANT_TYPE
\ No newline at end of file
+#undef UE_DECLARE_CORE_VARIANT_TYPE
+
+#if UE_ENABLE_INCLUDE_ORDER_DEPRECATED_IN_5_2
+#include "CoreMinimal.h"
+#endif
\ No newline at end of file
diff --git a/Engine/Source/Runtime/CoreUObject/Public/UObject/CoreNative.h b/Engine/Source/Runtime/CoreUObject/Public/UObject/CoreNative.h
--- a/Engine/Source/Runtime/CoreUObject/Public/UObject/CoreNative.h
+++ b/Engine/Source/Runtime/CoreUObject/Public/UObject/CoreNative.h
@@ -6,7 +6,6 @@
 
 #pragma once
 
-#include "CoreMinimal.h"
 #include "HAL/Platform.h"
 #include "UObject/Object.h"
 #include "UObject/Script.h"
@@ -37,3 +36,7 @@ struct FNativeFunctionRegistrar
 
 	static COREUOBJECT_API void RegisterFunctions(class UClass* Class, const FNameNativePtrPair* InArray, int32 NumFunctions);
 };
+
+#if UE_ENABLE_INCLUDE_ORDER_DEPRECATED_IN_5_2
+#include "CoreMinimal.h"
+#endif
\ No newline at end of file
diff --git a/Engine/Source/Runtime/CoreUObject/Public/UObject/CoreNet.h b/Engine/Source/Runtime/CoreUObject/Public/UObject/CoreNet.h
--- a/Engine/Source/Runtime/CoreUObject/Public/UObject/CoreNet.h
+++ b/Engine/Source/Runtime/CoreUObject/Public/UObject/CoreNet.h
@@ -7,7 +7,6 @@
 #include "Containers/Map.h"
 #include "Containers/Set.h"
 #include "Containers/UnrealString.h"
-#include "CoreMinimal.h"
 #include "Delegates/Delegate.h"
 #include "HAL/PlatformMath.h"
 #include "Misc/AssertionMacros.h"
@@ -754,3 +753,7 @@ enum { LAN_BEACON_MAX_PACKET_SIZE = 1024 }; // MTU for the connection
 COREUOBJECT_API void			RPC_ResetLastFailedReason();
 COREUOBJECT_API void			RPC_ValidateFailed( const TCHAR* Reason );
 COREUOBJECT_API const TCHAR *	RPC_GetLastFailedReason();
+
+#if UE_ENABLE_INCLUDE_ORDER_DEPRECATED_IN_5_2
+#include "CoreMinimal.h"
+#endif
diff --git a/Engine/Source/Runtime/CoreUObject/Public/UObject/CoreNetTypes.h b/Engine/Source/Runtime/CoreUObject/Public/UObject/CoreNetTypes.h
--- a/Engine/Source/Runtime/CoreUObject/Public/UObject/CoreNetTypes.h
+++ b/Engine/Source/Runtime/CoreUObject/Public/UObject/CoreNetTypes.h
@@ -2,7 +2,6 @@
 
 #pragma once
 
-#include "CoreMinimal.h"
 #include "UObject/ObjectMacros.h"
 
 /** Whether or not caching of actor/subobject names to the stack should be enabled, for async demo recording crashdumps */
@@ -53,3 +52,7 @@ enum class EChannelCloseReason : uint8
 };
 
 COREUOBJECT_API const TCHAR* LexToString(const EChannelCloseReason Value);
+
+#if UE_ENABLE_INCLUDE_ORDER_DEPRECATED_IN_5_2
+#include "CoreMinimal.h"
+#endif
diff --git a/Engine/Source/Runtime/CoreUObject/Public/UObject/EnumProperty.h b/Engine/Source/Runtime/CoreUObject/Public/UObject/EnumProperty.h
--- a/Engine/Source/Runtime/CoreUObject/Public/UObject/EnumProperty.h
+++ b/Engine/Source/Runtime/CoreUObject/Public/UObject/EnumProperty.h
@@ -4,7 +4,6 @@
 
 #include "Containers/Array.h"
 #include "Containers/UnrealString.h"
-#include "CoreMinimal.h"
 #include "HAL/PlatformMath.h"
 #include "Misc/AssertionMacros.h"
 #include "Serialization/StructuredArchive.h"
@@ -118,3 +117,7 @@ public:
 	FNumericProperty* UnderlyingProp; // The property which represents the underlying type of the enum
 	UEnum* Enum; // The enum represented by this property
 };
+
+#if UE_ENABLE_INCLUDE_ORDER_DEPRECATED_IN_5_2
+#include "CoreMinimal.h"
+#endif
diff --git a/Engine/Source/Runtime/CoreUObject/Public/UObject/Field.h b/Engine/Source/Runtime/CoreUObject/Public/UObject/Field.h
--- a/Engine/Source/Runtime/CoreUObject/Public/UObject/Field.h
+++ b/Engine/Source/Runtime/CoreUObject/Public/UObject/Field.h
@@ -10,7 +10,6 @@ Field.h: Declares FField property system fundamentals
 #include "Containers/Map.h"
 #include "Containers/StringFwd.h"
 #include "Containers/UnrealString.h"
-#include "CoreMinimal.h"
 #include "CoreTypes.h"
 #include "Delegates/Delegate.h"
 #include "HAL/PlatformMath.h"
@@ -1055,3 +1054,7 @@ inline FieldType* FindFProperty(const TCHAR* InFieldPath)
 	FField* FoundField = FindFPropertyByPath(InFieldPath);
 	return CastField<FieldType>(FoundField);
 }
+
+#if UE_ENABLE_INCLUDE_ORDER_DEPRECATED_IN_5_2
+#include "CoreMinimal.h"
+#endif
\ No newline at end of file
diff --git a/Engine/Source/Runtime/CoreUObject/Public/UObject/FieldPath.h b/Engine/Source/Runtime/CoreUObject/Public/UObject/FieldPath.h
--- a/Engine/Source/Runtime/CoreUObject/Public/UObject/FieldPath.h
+++ b/Engine/Source/Runtime/CoreUObject/Public/UObject/FieldPath.h
@@ -9,7 +9,6 @@
 #include "Containers/Array.h"
 #include "Containers/Map.h"
 #include "Containers/UnrealString.h"
-#include "CoreMinimal.h"
 #include "CoreTypes.h"
 #include "HAL/ThreadSafeCounter.h"
 #include "Misc/AssertionMacros.h"
@@ -551,3 +550,6 @@ struct TPropertyPathMapKeyFuncs : public TDefaultMapKeyFuncs<KeyType, ValueType,
 	}
 };
 
+#if UE_ENABLE_INCLUDE_ORDER_DEPRECATED_IN_5_2
+#include "CoreMinimal.h"
+#endif
diff --git a/Engine/Source/Runtime/CoreUObject/Public/UObject/FieldPathProperty.h b/Engine/Source/Runtime/CoreUObject/Public/UObject/FieldPathProperty.h
--- a/Engine/Source/Runtime/CoreUObject/Public/UObject/FieldPathProperty.h
+++ b/Engine/Source/Runtime/CoreUObject/Public/UObject/FieldPathProperty.h
@@ -3,7 +3,6 @@
 
 #include "Containers/Array.h"
 #include "Containers/UnrealString.h"
-#include "CoreMinimal.h"
 #include "CoreTypes.h"
 #include "Serialization/StructuredArchive.h"
 #include "UObject/Field.h"
@@ -81,4 +80,8 @@ PRAGMA_ENABLE_DEPRECATION_WARNINGS
 	// End of FProperty interface
 
 	static FString RedirectFieldPathName(const FString& InPathName);
-};
\ No newline at end of file
+};
+
+#if UE_ENABLE_INCLUDE_ORDER_DEPRECATED_IN_5_2
+#include "CoreMinimal.h"
+#endif
diff --git a/Engine/Source/Runtime/CoreUObject/Public/UObject/GarbageCollection.h b/Engine/Source/Runtime/CoreUObject/Public/UObject/GarbageCollection.h
--- a/Engine/Source/Runtime/CoreUObject/Public/UObject/GarbageCollection.h
+++ b/Engine/Source/Runtime/CoreUObject/Public/UObject/GarbageCollection.h
@@ -9,7 +9,6 @@
 #include "Containers/Array.h"
 #include "Containers/Map.h"
 #include "Containers/UnrealString.h"
-#include "CoreMinimal.h"
 #include "CoreTypes.h"
 #include "HAL/PlatformCrt.h"
 #include "HAL/ThreadSafeBool.h"
@@ -547,3 +546,7 @@ FORCEINLINE bool IsGarbageCollecting()
 * Whether garbage collection is locking the global uobject hash tables
 */
 COREUOBJECT_API bool IsGarbageCollectingAndLockingUObjectHashTables();
+
+#if UE_ENABLE_INCLUDE_ORDER_DEPRECATED_IN_5_2
+#include "CoreMinimal.h"
+#endif
diff --git a/Engine/Source/Runtime/CoreUObject/Public/UObject/GeneratedCppIncludes.h b/Engine/Source/Runtime/CoreUObject/Public/UObject/GeneratedCppIncludes.h
--- a/Engine/Source/Runtime/CoreUObject/Public/UObject/GeneratedCppIncludes.h
+++ b/Engine/Source/Runtime/CoreUObject/Public/UObject/GeneratedCppIncludes.h
@@ -2,7 +2,6 @@
 
 #pragma once
 
-#include "CoreMinimal.h"
 #include "UObject/Object.h"
 #include "UObject/UObjectGlobals.h"
 #include "UObject/CoreNative.h"
@@ -12,3 +11,7 @@
 #include "UObject/EnumProperty.h"
 #include "UObject/TextProperty.h"
 #include "UObject/FieldPathProperty.h"
+
+#if UE_ENABLE_INCLUDE_ORDER_DEPRECATED_IN_5_2
+#include "CoreMinimal.h"
+#endif
diff --git a/Engine/Source/Runtime/CoreUObject/Public/UObject/Interface.h b/Engine/Source/Runtime/CoreUObject/Public/UObject/Interface.h
--- a/Engine/Source/Runtime/CoreUObject/Public/UObject/Interface.h
+++ b/Engine/Source/Runtime/CoreUObject/Public/UObject/Interface.h
@@ -2,7 +2,6 @@
 
 #pragma once
 
-#include "CoreMinimal.h"
 #include "HAL/Platform.h"
 #include "UObject/Object.h"
 #include "UObject/ObjectMacros.h"
@@ -29,3 +28,6 @@ public:
 	typedef UInterface UClassType;
 };
 
+#if UE_ENABLE_INCLUDE_ORDER_DEPRECATED_IN_5_2
+#include "CoreMinimal.h"
+#endif
diff --git a/Engine/Source/Runtime/CoreUObject/Public/UObject/LazyObjectPtr.h b/Engine/Source/Runtime/CoreUObject/Public/UObject/LazyObjectPtr.h
--- a/Engine/Source/Runtime/CoreUObject/Public/UObject/LazyObjectPtr.h
+++ b/Engine/Source/Runtime/CoreUObject/Public/UObject/LazyObjectPtr.h
@@ -8,7 +8,6 @@
 
 #include "Containers/UnrealString.h"
 #include "CoreGlobals.h"
-#include "CoreMinimal.h"
 #include "HAL/Platform.h"
 #include "HAL/ThreadSafeCounter.h"
 #include "Misc/Guid.h"
@@ -403,4 +402,8 @@ FORCEINLINE bool operator!=(TYPE_OF_NULLPTR, const TLazyObjectPtr<T>& Rhs)
 }
 
 template<class T> struct TIsPODType<TLazyObjectPtr<T> > { enum { Value = TIsPODType<FLazyObjectPtr>::Value }; };
-template<class T> struct TIsWeakPointerType<TLazyObjectPtr<T> > { enum { Value = TIsWeakPointerType<FLazyObjectPtr>::Value }; };
\ No newline at end of file
+template<class T> struct TIsWeakPointerType<TLazyObjectPtr<T> > { enum { Value = TIsWeakPointerType<FLazyObjectPtr>::Value }; };
+
+#if UE_ENABLE_INCLUDE_ORDER_DEPRECATED_IN_5_2
+#include "CoreMinimal.h"
+#endif
diff --git a/Engine/Source/Runtime/CoreUObject/Public/UObject/MetaData.h b/Engine/Source/Runtime/CoreUObject/Public/UObject/MetaData.h
--- a/Engine/Source/Runtime/CoreUObject/Public/UObject/MetaData.h
+++ b/Engine/Source/Runtime/CoreUObject/Public/UObject/MetaData.h
@@ -4,7 +4,6 @@
 
 #include "Containers/Map.h"
 #include "Containers/UnrealString.h"
-#include "CoreMinimal.h"
 #include "HAL/Platform.h"
 #include "Serialization/StructuredArchive.h"
 #include "UObject/NameTypes.h"
@@ -233,3 +232,7 @@ private:
 };
 
 #endif //WITH_EDITOR
+
+#if UE_ENABLE_INCLUDE_ORDER_DEPRECATED_IN_5_2
+#include "CoreMinimal.h"
+#endif
diff --git a/Engine/Source/Runtime/CoreUObject/Public/UObject/Object.h b/Engine/Source/Runtime/CoreUObject/Public/UObject/Object.h
--- a/Engine/Source/Runtime/CoreUObject/Public/UObject/Object.h
+++ b/Engine/Source/Runtime/CoreUObject/Public/UObject/Object.h
@@ -6,7 +6,6 @@
 	Object.h: Direct base class for all UE objects
 =============================================================================*/
 
-#include "CoreMinimal.h"
 #include "UObject/Script.h"
 #include "UObject/ObjectMacros.h"
 #include "UObject/UObjectBaseUtility.h"
@@ -1807,3 +1806,7 @@ const T* GetValid(const T* Test)
 /** Callback for editor object selection. This must be in core instead of editor for UObject::IsSelectedInEditor to work */
 extern COREUOBJECT_API TFunction<bool(const UObject*)> GIsObjectSelectedInEditor;
 #endif
+
+#if UE_ENABLE_INCLUDE_ORDER_DEPRECATED_IN_5_2
+#include "CoreMinimal.h"
+#endif
diff --git a/Engine/Source/Runtime/CoreUObject/Public/UObject/ObjectKey.h b/Engine/Source/Runtime/CoreUObject/Public/UObject/ObjectKey.h
--- a/Engine/Source/Runtime/CoreUObject/Public/UObject/ObjectKey.h
+++ b/Engine/Source/Runtime/CoreUObject/Public/UObject/ObjectKey.h
@@ -2,7 +2,6 @@
 
 #pragma once
 
-#include "CoreMinimal.h"
 #include "UObject/WeakObjectPtr.h"
 
 #include <type_traits>
@@ -228,4 +227,8 @@ public:
 
 private:
 	FObjectKey ObjectKey;
-};
\ No newline at end of file
+};
+
+#if UE_ENABLE_INCLUDE_ORDER_DEPRECATED_IN_5_2
+#include "CoreMinimal.h"
+#endif
diff --git a/Engine/Source/Runtime/CoreUObject/Public/UObject/ObjectMacros.h b/Engine/Source/Runtime/CoreUObject/Public/UObject/ObjectMacros.h
--- a/Engine/Source/Runtime/CoreUObject/Public/UObject/ObjectMacros.h
+++ b/Engine/Source/Runtime/CoreUObject/Public/UObject/ObjectMacros.h
@@ -6,9 +6,10 @@
 
 #pragma once
 
-#include "CoreMinimal.h"
+// IWYU pragma: begin_keep
 #include "Misc/EnumClassFlags.h"
 #include "UObject/Script.h"
+// IWYU pragma: end_keep
 
 class FObjectInitializer;
 class FReferenceCollector;
@@ -2173,3 +2174,7 @@ public:
 	**/
 	virtual void Restore() const=0;
 };
+
+#if UE_ENABLE_INCLUDE_ORDER_DEPRECATED_IN_5_2
+#include "CoreMinimal.h"
+#endif
\ No newline at end of file
diff --git a/Engine/Source/Runtime/CoreUObject/Public/UObject/PersistentObjectPtr.h b/Engine/Source/Runtime/CoreUObject/Public/UObject/PersistentObjectPtr.h
--- a/Engine/Source/Runtime/CoreUObject/Public/UObject/PersistentObjectPtr.h
+++ b/Engine/Source/Runtime/CoreUObject/Public/UObject/PersistentObjectPtr.h
@@ -6,7 +6,6 @@
 
 #pragma once
 
-#include "CoreMinimal.h"
 #include "UObject/WeakObjectPtr.h"
 
 /**
@@ -248,3 +247,6 @@ private:
 template <class TObjectID> struct TIsPODType<TPersistentObjectPtr<TObjectID> > { enum { Value = TIsPODType<TObjectID>::Value }; };
 template <class TObjectID> struct TIsWeakPointerType<TPersistentObjectPtr<TObjectID> > { enum { Value = TIsWeakPointerType<FWeakObjectPtr>::Value }; };
 
+#if UE_ENABLE_INCLUDE_ORDER_DEPRECATED_IN_5_2
+#include "CoreMinimal.h"
+#endif
diff --git a/Engine/Source/Runtime/CoreUObject/Public/UObject/PrimaryAssetId.h b/Engine/Source/Runtime/CoreUObject/Public/UObject/PrimaryAssetId.h
--- a/Engine/Source/Runtime/CoreUObject/Public/UObject/PrimaryAssetId.h
+++ b/Engine/Source/Runtime/CoreUObject/Public/UObject/PrimaryAssetId.h
@@ -4,7 +4,6 @@
 
 #include "Containers/StringFwd.h"
 #include "Containers/UnrealString.h"
-#include "CoreMinimal.h"
 #include "HAL/Platform.h"
 #include "Misc/StringBuilder.h"
 #include "Serialization/Archive.h"
@@ -200,4 +199,8 @@ struct FPrimaryAssetId
 	friend struct Z_Construct_UScriptStruct_FPrimaryAssetId_Statics;
 };
 
-COREUOBJECT_API FStringBuilderBase& operator<<(FStringBuilderBase& Builder, const FPrimaryAssetId& Id);
\ No newline at end of file
+COREUOBJECT_API FStringBuilderBase& operator<<(FStringBuilderBase& Builder, const FPrimaryAssetId& Id);
+
+#if UE_ENABLE_INCLUDE_ORDER_DEPRECATED_IN_5_2
+#include "CoreMinimal.h"
+#endif
diff --git a/Engine/Source/Runtime/CoreUObject/Public/UObject/PropertyTag.h b/Engine/Source/Runtime/CoreUObject/Public/UObject/PropertyTag.h
--- a/Engine/Source/Runtime/CoreUObject/Public/UObject/PropertyTag.h
+++ b/Engine/Source/Runtime/CoreUObject/Public/UObject/PropertyTag.h
@@ -6,7 +6,6 @@
 
 #pragma once
 
-#include "CoreMinimal.h"
 #include "CoreTypes.h"
 #include "Misc/Guid.h"
 #include "Serialization/StructuredArchive.h"
@@ -53,3 +52,7 @@ struct FPropertyTag
 	void SerializeTaggedProperty( FArchive& Ar, FProperty* Property, uint8* Value, uint8* Defaults ) const;
 	void SerializeTaggedProperty(FStructuredArchive::FSlot Slot, FProperty* Property, uint8* Value, uint8* Defaults) const;
 };
+
+#if UE_ENABLE_INCLUDE_ORDER_DEPRECATED_IN_5_2
+#include "CoreMinimal.h"
+#endif
diff --git a/Engine/Source/Runtime/CoreUObject/Public/UObject/Script.h b/Engine/Source/Runtime/CoreUObject/Public/UObject/Script.h
--- a/Engine/Source/Runtime/CoreUObject/Public/UObject/Script.h
+++ b/Engine/Source/Runtime/CoreUObject/Public/UObject/Script.h
@@ -6,7 +6,6 @@
 
 #pragma once
 
-#include "CoreMinimal.h"
 #include "HAL/ThreadSingleton.h"
 #include "Internationalization/Text.h"
 #include "Stats/Stats.h"
@@ -677,3 +676,6 @@ COREUOBJECT_API FString ToValidCPPIdentifierChars(TCHAR Char);
 */
 COREUOBJECT_API FString UnicodeToCPPIdentifier(const FString& InName, bool bDeprecated, const TCHAR* Prefix);
 
+#if UE_ENABLE_INCLUDE_ORDER_DEPRECATED_IN_5_2
+#include "CoreMinimal.h"
+#endif
diff --git a/Engine/Source/Runtime/CoreUObject/Public/UObject/ScriptInterface.h b/Engine/Source/Runtime/CoreUObject/Public/UObject/ScriptInterface.h
--- a/Engine/Source/Runtime/CoreUObject/Public/UObject/ScriptInterface.h
+++ b/Engine/Source/Runtime/CoreUObject/Public/UObject/ScriptInterface.h
@@ -6,7 +6,6 @@
 
 #pragma once
 
-#include "CoreMinimal.h"
 #include "UObject/UObjectGlobals.h"
 #include "Templates/Casts.h"
 #include "Templates/UnrealTemplate.h"
@@ -287,3 +286,7 @@ public:
 		return Interface.Serialize(Ar, InterfaceType::UClassType::StaticClass());
 	}
 };
+
+#if UE_ENABLE_INCLUDE_ORDER_DEPRECATED_IN_5_2
+#include "CoreMinimal.h"
+#endif
diff --git a/Engine/Source/Runtime/CoreUObject/Public/UObject/ScriptMacros.h b/Engine/Source/Runtime/CoreUObject/Public/UObject/ScriptMacros.h
--- a/Engine/Source/Runtime/CoreUObject/Public/UObject/ScriptMacros.h
+++ b/Engine/Source/Runtime/CoreUObject/Public/UObject/ScriptMacros.h
@@ -6,12 +6,13 @@
 
 #pragma once
 
-#include "CoreMinimal.h"
+// IWYU pragma: begin_keep
 #include "UObject/Script.h"
 #include "UObject/ScriptInterface.h"
 #include "UObject/UnrealType.h"
 #include "UObject/Stack.h"
 #include "UObject/FieldPathProperty.h"
+// IWYU pragma: end_keep
 
 /*-----------------------------------------------------------------------------
 	Macros.
@@ -109,3 +110,7 @@ enum {MAX_VARIABLE_SIZE = 0x0FFF };
 
 #define P_NATIVE_BEGIN { SCOPED_SCRIPT_NATIVE_TIMER(ScopedNativeCallTimer);
 #define P_NATIVE_END   }
+
+#if UE_ENABLE_INCLUDE_ORDER_DEPRECATED_IN_5_2
+#include "CoreMinimal.h"
+#endif
\ No newline at end of file
diff --git a/Engine/Source/Runtime/CoreUObject/Public/UObject/SoftObjectPath.h b/Engine/Source/Runtime/CoreUObject/Public/UObject/SoftObjectPath.h
--- a/Engine/Source/Runtime/CoreUObject/Public/UObject/SoftObjectPath.h
+++ b/Engine/Source/Runtime/CoreUObject/Public/UObject/SoftObjectPath.h
@@ -8,7 +8,6 @@
 #include "Containers/StringFwd.h"
 #include "Containers/StringView.h"
 #include "Containers/UnrealString.h"
-#include "CoreMinimal.h"
 #include "Containers/StringView.h"
 #include "CoreTypes.h"
 #include "HAL/ThreadSafeCounter.h"
@@ -614,4 +613,8 @@ PRAGMA_DISABLE_DEPRECATION_WARNINGS
 PRAGMA_ENABLE_DEPRECATION_WARNINGS
 		return Out;
 	}
-}
\ No newline at end of file
+}
+
+#if UE_ENABLE_INCLUDE_ORDER_DEPRECATED_IN_5_2
+#include "CoreMinimal.h"
+#endif
diff --git a/Engine/Source/Runtime/CoreUObject/Public/UObject/SoftObjectPtr.h b/Engine/Source/Runtime/CoreUObject/Public/UObject/SoftObjectPtr.h
--- a/Engine/Source/Runtime/CoreUObject/Public/UObject/SoftObjectPtr.h
+++ b/Engine/Source/Runtime/CoreUObject/Public/UObject/SoftObjectPtr.h
@@ -6,7 +6,6 @@
 
 #pragma once
 
-#include "CoreMinimal.h"
 #include "UObject/Object.h"
 #include "Templates/Casts.h"
 #include "UObject/PersistentObjectPtr.h"
@@ -726,3 +725,7 @@ FArchive& operator<<(FArchive& Ar, TSoftClassPtr<TClass>& Ptr)
 	Ptr.Serialize(Ar);
 	return Ar;
 }
+
+#if UE_ENABLE_INCLUDE_ORDER_DEPRECATED_IN_5_2
+#include "CoreMinimal.h"
+#endif
diff --git a/Engine/Source/Runtime/CoreUObject/Public/UObject/SparseDelegate.h b/Engine/Source/Runtime/CoreUObject/Public/UObject/SparseDelegate.h
--- a/Engine/Source/Runtime/CoreUObject/Public/UObject/SparseDelegate.h
+++ b/Engine/Source/Runtime/CoreUObject/Public/UObject/SparseDelegate.h
@@ -5,7 +5,6 @@
 #include "Containers/Array.h"
 #include "Containers/Map.h"
 #include "Containers/SparseArray.h"
-#include "CoreMinimal.h"
 #include "Delegates/Delegate.h"
 #include "HAL/CriticalSection.h"
 #include "HAL/PlatformMath.h"
@@ -460,3 +459,7 @@ private:
 #define DECLARE_DYNAMIC_MULTICAST_SPARSE_DELEGATE_SevenParams( SparseDelegateClass, OwningClass, DelegateName, Param1Type, Param1Name, Param2Type, Param2Name, Param3Type, Param3Name, Param4Type, Param4Name, Param5Type, Param5Name, Param6Type, Param6Name, Param7Type, Param7Name ) BODY_MACRO_COMBINE(CURRENT_FILE_ID,_,__LINE__,_DELEGATE) FUNC_DECLARE_DYNAMIC_MULTICAST_SPARSE_DELEGATE( SparseDelegateClass, OwningClass, DelegateName, FUNC_CONCAT( Param1Type InParam1, Param2Type InParam2, Param3Type InParam3, Param4Type InParam4, Param5Type InParam5, Param6Type InParam6, Param7Type InParam7 ), FUNC_CONCAT( *this, InParam1, InParam2, InParam3, InParam4, InParam5, InParam6, InParam7 ), void, Param1Type, Param2Type, Param3Type, Param4Type, Param5Type, Param6Type, Param7Type )
 #define DECLARE_DYNAMIC_MULTICAST_SPARSE_DELEGATE_EightParams( SparseDelegateClass, OwningClass, DelegateName, Param1Type, Param1Name, Param2Type, Param2Name, Param3Type, Param3Name, Param4Type, Param4Name, Param5Type, Param5Name, Param6Type, Param6Name, Param7Type, Param7Name, Param8Type, Param8Name ) BODY_MACRO_COMBINE(CURRENT_FILE_ID,_,__LINE__,_DELEGATE) FUNC_DECLARE_DYNAMIC_MULTICAST_SPARSE_DELEGATE( SparseDelegateClass, OwningClass, DelegateName, FUNC_CONCAT( Param1Type InParam1, Param2Type InParam2, Param3Type InParam3, Param4Type InParam4, Param5Type InParam5, Param6Type InParam6, Param7Type InParam7, Param8Type InParam8 ), FUNC_CONCAT( *this, InParam1, InParam2, InParam3, InParam4, InParam5, InParam6, InParam7, InParam8 ), void, Param1Type, Param2Type, Param3Type, Param4Type, Param5Type, Param6Type, Param7Type, Param8Type )
 #define DECLARE_DYNAMIC_MULTICAST_SPARSE_DELEGATE_NineParams( SparseDelegateClass, OwningClass, DelegateName, Param1Type, Param1Name, Param2Type, Param2Name, Param3Type, Param3Name, Param4Type, Param4Name, Param5Type, Param5Name, Param6Type, Param6Name, Param7Type, Param7Name, Param8Type, Param8Name, Param9Type, Param9Name ) BODY_MACRO_COMBINE(CURRENT_FILE_ID,_,__LINE__,_DELEGATE) FUNC_DECLARE_DYNAMIC_MULTICAST_SPARSE_DELEGATE( SparseDelegateClass, OwningClass, DelegateName, FUNC_CONCAT( Param1Type InParam1, Param2Type InParam2, Param3Type InParam3, Param4Type InParam4, Param5Type InParam5, Param6Type InParam6, Param7Type InParam7, Param8Type InParam8, Param9Type InParam9 ), FUNC_CONCAT( *this, InParam1, InParam2, InParam3, InParam4, InParam5, InParam6, InParam7, InParam8, InParam9 ), void, Param1Type, Param2Type, Param3Type, Param4Type, Param5Type, Param6Type, Param7Type, Param8Type, Param9Type )
+
+#if UE_ENABLE_INCLUDE_ORDER_DEPRECATED_IN_5_2
+#include "CoreMinimal.h"
+#endif
\ No newline at end of file
diff --git a/Engine/Source/Runtime/CoreUObject/Public/UObject/Stack.h b/Engine/Source/Runtime/CoreUObject/Public/UObject/Stack.h
--- a/Engine/Source/Runtime/CoreUObject/Public/UObject/Stack.h
+++ b/Engine/Source/Runtime/CoreUObject/Public/UObject/Stack.h
@@ -6,7 +6,6 @@
 
 #pragma once
 
-#include "CoreMinimal.h"
 #include "UObject/Script.h"
 #include "Misc/CoreMisc.h"
 #include "Templates/Casts.h"
@@ -477,3 +476,7 @@ FORCEINLINE_DEBUGGABLE TNativeType& FFrame::StepCompiledInRef(void*const Tempora
 
 	return (MostRecentPropertyAddress != NULL) ? *(TNativeType*)(MostRecentPropertyAddress) : *(TNativeType*)TemporaryBuffer;
 }
+
+#if UE_ENABLE_INCLUDE_ORDER_DEPRECATED_IN_5_2
+#include "CoreMinimal.h"
+#endif
diff --git a/Engine/Source/Runtime/CoreUObject/Public/UObject/StructOnScope.h b/Engine/Source/Runtime/CoreUObject/Public/UObject/StructOnScope.h
--- a/Engine/Source/Runtime/CoreUObject/Public/UObject/StructOnScope.h
+++ b/Engine/Source/Runtime/CoreUObject/Public/UObject/StructOnScope.h
@@ -2,7 +2,6 @@
 
 #pragma once
 
-#include "CoreMinimal.h"
 #include "UObject/Class.h"
 #include "UObject/WeakObjectPtr.h"
 #include "Templates/Casts.h"
@@ -375,3 +374,7 @@ FArchive& operator<<(FArchive& Ar, TStructOnScope<T>& Struct)
 	Struct.Serialize(Ar);
 	return Ar;
 }
+
+#if UE_ENABLE_INCLUDE_ORDER_DEPRECATED_IN_5_2
+#include "CoreMinimal.h"
+#endif
diff --git a/Engine/Source/Runtime/CoreUObject/Public/UObject/TextProperty.h b/Engine/Source/Runtime/CoreUObject/Public/UObject/TextProperty.h
--- a/Engine/Source/Runtime/CoreUObject/Public/UObject/TextProperty.h
+++ b/Engine/Source/Runtime/CoreUObject/Public/UObject/TextProperty.h
@@ -2,7 +2,6 @@
 #pragma once
 
 #include "Containers/UnrealString.h"
-#include "CoreMinimal.h"
 #include "CoreTypes.h"
 #include "Internationalization/Text.h"
 #include "Serialization/StructuredArchive.h"
@@ -70,3 +69,7 @@ public:
 
 	static bool Identical_Implementation(const FText& A, const FText& B, uint32 PortFlags);
 };
+
+#if UE_ENABLE_INCLUDE_ORDER_DEPRECATED_IN_5_2
+#include "CoreMinimal.h"
+#endif
diff --git a/Engine/Source/Runtime/CoreUObject/Public/UObject/UObjectAnnotation.h b/Engine/Source/Runtime/CoreUObject/Public/UObject/UObjectAnnotation.h
--- a/Engine/Source/Runtime/CoreUObject/Public/UObject/UObjectAnnotation.h
+++ b/Engine/Source/Runtime/CoreUObject/Public/UObject/UObjectAnnotation.h
@@ -6,7 +6,6 @@
 
 #pragma once
 
-#include "CoreMinimal.h"
 #include "UObject/UObjectArray.h"
 #include "Misc/ScopeLock.h"
 #include "Misc/ScopeRWLock.h"
@@ -1308,3 +1307,7 @@ private:
 	TArray<TBitType> AnnotationArray;
 
 };
+
+#if UE_ENABLE_INCLUDE_ORDER_DEPRECATED_IN_5_2
+#include "CoreMinimal.h"
+#endif
diff --git a/Engine/Source/Runtime/CoreUObject/Public/UObject/UObjectArray.h b/Engine/Source/Runtime/CoreUObject/Public/UObject/UObjectArray.h
--- a/Engine/Source/Runtime/CoreUObject/Public/UObject/UObjectArray.h
+++ b/Engine/Source/Runtime/CoreUObject/Public/UObject/UObjectArray.h
@@ -6,7 +6,6 @@
 
 #pragma once
 
-#include "CoreMinimal.h"
 #include "HAL/ThreadSafeCounter.h"
 #include "Containers/LockFreeList.h"
 #include "UObject/ObjectMacros.h"
@@ -1251,3 +1250,7 @@ struct FIndexToObject
 		return ObjectItem ? ObjectItem->Object : nullptr;
 	}
 };
+
+#if UE_ENABLE_INCLUDE_ORDER_DEPRECATED_IN_5_2
+#include "CoreMinimal.h"
+#endif
diff --git a/Engine/Source/Runtime/CoreUObject/Public/UObject/UObjectBase.h b/Engine/Source/Runtime/CoreUObject/Public/UObject/UObjectBase.h
--- a/Engine/Source/Runtime/CoreUObject/Public/UObject/UObjectBase.h
+++ b/Engine/Source/Runtime/CoreUObject/Public/UObject/UObjectBase.h
@@ -8,7 +8,6 @@
 
 #include "Containers/Map.h"
 #include "Containers/UnrealString.h"
-#include "CoreMinimal.h"
 #include "HAL/LowLevelMemTracker.h"
 #include "HAL/PlatformAtomics.h"
 #include "HAL/PlatformMath.h"
@@ -469,3 +468,6 @@ void UObjectBaseInit();
  */
 void UObjectBaseShutdown();
 
+#if UE_ENABLE_INCLUDE_ORDER_DEPRECATED_IN_5_2
+#include "CoreMinimal.h"
+#endif
\ No newline at end of file
diff --git a/Engine/Source/Runtime/CoreUObject/Public/UObject/UObjectBaseUtility.h b/Engine/Source/Runtime/CoreUObject/Public/UObject/UObjectBaseUtility.h
--- a/Engine/Source/Runtime/CoreUObject/Public/UObject/UObjectBaseUtility.h
+++ b/Engine/Source/Runtime/CoreUObject/Public/UObject/UObjectBaseUtility.h
@@ -9,7 +9,6 @@
 #include "Containers/StringFwd.h"
 #include "Containers/UnrealString.h"
 #include "Containers/VersePathFwd.h"
-#include "CoreMinimal.h"
 #include "HAL/PlatformMath.h"
 #include "Misc/AssertionMacros.h"
 #include "Misc/EnumClassFlags.h"
@@ -1205,3 +1204,7 @@ public:
 		PRAGMA_ENABLE_SHADOW_VARIABLE_WARNINGS
 	#endif
 #endif
+
+#if UE_ENABLE_INCLUDE_ORDER_DEPRECATED_IN_5_2
+#include "CoreMinimal.h"
+#endif
\ No newline at end of file
diff --git a/Engine/Source/Runtime/CoreUObject/Public/UObject/UObjectGlobals.h b/Engine/Source/Runtime/CoreUObject/Public/UObject/UObjectGlobals.h
--- a/Engine/Source/Runtime/CoreUObject/Public/UObject/UObjectGlobals.h
+++ b/Engine/Source/Runtime/CoreUObject/Public/UObject/UObjectGlobals.h
@@ -18,7 +18,6 @@
 #include "Containers/UnrealString.h"
 #include "Containers/VersePathFwd.h"
 #include "CoreGlobals.h"
-#include "CoreMinimal.h"
 #include "CoreTypes.h"
 #include "Delegates/Delegate.h"
 #include "Internationalization/Text.h"
@@ -3622,3 +3621,7 @@ enum class EDataValidationResult : uint8
  * @return	Returns the combined data validation result
  */
 COREUOBJECT_API EDataValidationResult CombineDataValidationResults(EDataValidationResult Result1, EDataValidationResult Result2);
+
+#if UE_ENABLE_INCLUDE_ORDER_DEPRECATED_IN_5_2
+#include "CoreMinimal.h"
+#endif
diff --git a/Engine/Source/Runtime/CoreUObject/Public/UObject/UObjectHash.h b/Engine/Source/Runtime/CoreUObject/Public/UObject/UObjectHash.h
--- a/Engine/Source/Runtime/CoreUObject/Public/UObject/UObjectHash.h
+++ b/Engine/Source/Runtime/CoreUObject/Public/UObject/UObjectHash.h
@@ -11,7 +11,6 @@
 #include "Containers/Map.h"
 #include "Containers/Set.h"
 #include "Containers/UnrealString.h"
-#include "CoreMinimal.h"
 #include "CoreTypes.h"
 #include "Stats/Stats.h"
 #include "Stats/Stats2.h"
@@ -319,4 +318,8 @@ public:
 	{
 		UnlockUObjectHashTables();
 	}
-};
\ No newline at end of file
+};
+
+#if UE_ENABLE_INCLUDE_ORDER_DEPRECATED_IN_5_2
+#include "CoreMinimal.h"
+#endif
diff --git a/Engine/Source/Runtime/CoreUObject/Public/UObject/UObjectMarks.h b/Engine/Source/Runtime/CoreUObject/Public/UObject/UObjectMarks.h
--- a/Engine/Source/Runtime/CoreUObject/Public/UObject/UObjectMarks.h
+++ b/Engine/Source/Runtime/CoreUObject/Public/UObject/UObjectMarks.h
@@ -7,7 +7,6 @@
 #pragma once
 
 #include "Containers/Array.h"
-#include "CoreMinimal.h"
 
 class UObject;
 
@@ -133,3 +132,6 @@ COREUOBJECT_API void GetObjectsWithAllMarks(TArray<UObject *>& Results, EObjectM
  */
 COREUOBJECT_API void GetObjectsWithAnyMarks(TArray<UObject *>& Results, EObjectMark Marks);
 
+#if UE_ENABLE_INCLUDE_ORDER_DEPRECATED_IN_5_2
+#include "CoreMinimal.h"
+#endif
diff --git a/Engine/Source/Runtime/CoreUObject/Public/UObject/UnrealType.h b/Engine/Source/Runtime/CoreUObject/Public/UObject/UnrealType.h
--- a/Engine/Source/Runtime/CoreUObject/Public/UObject/UnrealType.h
+++ b/Engine/Source/Runtime/CoreUObject/Public/UObject/UnrealType.h
@@ -18,7 +18,6 @@
 #include "Containers/StringFwd.h"
 #include "Containers/UnrealString.h"
 #include "CoreGlobals.h"
-#include "CoreMinimal.h"
 #include "HAL/MemoryBase.h"
 #include "HAL/PlatformMath.h"
 #include "HAL/UnrealMemory.h"
@@ -7269,3 +7268,7 @@ class COREUOBJECT_API UMulticastInlineDelegatePropertyWrapper : public UMulticas
 {
 	DECLARE_CLASS_INTRINSIC(UMulticastInlineDelegatePropertyWrapper, UMulticastDelegatePropertyWrapper, CLASS_Transient, TEXT("/Script/CoreUObject"));
 };
+
+#if UE_ENABLE_INCLUDE_ORDER_DEPRECATED_IN_5_2
+#include "CoreMinimal.h"
+#endif
diff --git a/Engine/Source/Runtime/CoreUObject/Public/UObject/WeakObjectPtr.h b/Engine/Source/Runtime/CoreUObject/Public/UObject/WeakObjectPtr.h
--- a/Engine/Source/Runtime/CoreUObject/Public/UObject/WeakObjectPtr.h
+++ b/Engine/Source/Runtime/CoreUObject/Public/UObject/WeakObjectPtr.h
@@ -6,7 +6,6 @@
 
 #pragma once
 
-#include "CoreMinimal.h"
 #include "CoreTypes.h"
 #include "Misc/AssertionMacros.h"
 #include "Templates/UnrealTemplate.h"
@@ -340,3 +339,6 @@ FORCEINLINE uint32 GetTypeHash(const FWeakObjectPtr& WeakObjectPtr)
 	return WeakObjectPtr.GetTypeHash();
 }
 
+#if UE_ENABLE_INCLUDE_ORDER_DEPRECATED_IN_5_2
+#include "CoreMinimal.h"
+#endif
diff --git a/Engine/Source/Runtime/Engine/Classes/AI/Navigation/NavigationDirtyElement.h b/Engine/Source/Runtime/Engine/Classes/AI/Navigation/NavigationDirtyElement.h
--- a/Engine/Source/Runtime/Engine/Classes/AI/Navigation/NavigationDirtyElement.h
+++ b/Engine/Source/Runtime/Engine/Classes/AI/Navigation/NavigationDirtyElement.h
@@ -1,7 +1,8 @@
-﻿// Copyright Epic Games, Inc. All Rights Reserved.
+// Copyright Epic Games, Inc. All Rights Reserved.
 
 #pragma once
 
+#include "Math/Box.h"
 #include "UObject/WeakObjectPtr.h"
 
 class INavRelevantInterface;
diff --git a/Engine/Source/Runtime/Engine/Classes/Components/ChildActorComponent.h b/Engine/Source/Runtime/Engine/Classes/Components/ChildActorComponent.h
--- a/Engine/Source/Runtime/Engine/Classes/Components/ChildActorComponent.h
+++ b/Engine/Source/Runtime/Engine/Classes/Components/ChildActorComponent.h
@@ -3,7 +3,6 @@
 
 #pragma once
 
-#include "CoreMinimal.h"
 #include "UObject/ObjectMacros.h"
 #include "Templates/SubclassOf.h"
 #include "Components/SceneComponent.h"
@@ -220,4 +219,8 @@ struct FActorParentComponentSetter
 	static void Set(AActor* ChildActor, UChildActorComponent* ParentComponent);
 private:
 	friend UChildActorComponent;
-};
\ No newline at end of file
+};
+
+#if UE_ENABLE_INCLUDE_ORDER_DEPRECATED_IN_5_2
+#include "CoreMinimal.h"
+#endif
diff --git a/Engine/Source/Runtime/Engine/Classes/Components/SceneComponent.h b/Engine/Source/Runtime/Engine/Classes/Components/SceneComponent.h
--- a/Engine/Source/Runtime/Engine/Classes/Components/SceneComponent.h
+++ b/Engine/Source/Runtime/Engine/Classes/Components/SceneComponent.h
@@ -2,8 +2,7 @@
 
 #pragma once
 
-#include "CoreMinimal.h"
-#include "UObject/ObjectMacros.h"
+#include "Math/BoxSphereBounds.h"
 #include "UObject/UObjectGlobals.h"
 #if UE_ENABLE_INCLUDE_ORDER_DEPRECATED_IN_5_1
 #include "Engine/EngineTypes.h"
@@ -1768,4 +1767,8 @@ FORCEINLINE_DEBUGGABLE void USceneComponent::BeginScopedMovementUpdate(FScopedMo
 FORCEINLINE_DEBUGGABLE bool USceneComponent::GetShouldUpdatePhysicsVolume() const
 {
 	return bShouldUpdatePhysicsVolume;
-}
\ No newline at end of file
+}
+
+#if UE_ENABLE_INCLUDE_ORDER_DEPRECATED_IN_5_2
+#include "CoreMinimal.h"
+#endif
diff --git a/Engine/Source/Runtime/Engine/Classes/Engine/EngineBaseTypes.h b/Engine/Source/Runtime/Engine/Classes/Engine/EngineBaseTypes.h
--- a/Engine/Source/Runtime/Engine/Classes/Engine/EngineBaseTypes.h
+++ b/Engine/Source/Runtime/Engine/Classes/Engine/EngineBaseTypes.h
@@ -8,7 +8,6 @@
  */
 
 #include "Async/TaskGraphFwd.h"
-#include "CoreMinimal.h"
 #include "UObject/ObjectMacros.h"
 #include "UObject/Object.h"
 #include "UObject/Class.h"
@@ -16,10 +15,6 @@
 #include "Misc/CoreMisc.h"
 #include "Net/Core/Connection/NetEnums.h"
 
-#if UE_ENABLE_INCLUDE_ORDER_DEPRECATED_IN_5_2
-#include "Async/TaskGraphInterfaces.h"
-#endif
-
 #include "EngineBaseTypes.generated.h"
 
 class UActorComponent;
@@ -1003,3 +998,7 @@ class UEngineBaseTypes : public UObject
 
 };
 
+#if UE_ENABLE_INCLUDE_ORDER_DEPRECATED_IN_5_2
+#include "Async/TaskGraphInterfaces.h"
+#include "CoreMinimal.h"
+#endif
diff --git a/Engine/Source/Runtime/Engine/Classes/Engine/EngineTypes.h b/Engine/Source/Runtime/Engine/Classes/Engine/EngineTypes.h
--- a/Engine/Source/Runtime/Engine/Classes/Engine/EngineTypes.h
+++ b/Engine/Source/Runtime/Engine/Classes/Engine/EngineTypes.h
@@ -7,7 +7,6 @@
  *	use in multiple files where the enum can't be mapped to a specific file.
  */
 
-#include "CoreMinimal.h"
 #include "UObject/ObjectMacros.h"
 #include "UObject/Object.h"
 #include "UObject/Class.h"
@@ -3750,3 +3749,7 @@ enum class ELevelCollectionType : uint8
 
 	MAX
 };
+
+#if UE_ENABLE_INCLUDE_ORDER_DEPRECATED_IN_5_2
+#include "CoreMinimal.h"
+#endif
diff --git a/Engine/Source/Runtime/Engine/Classes/Engine/LevelStreaming.h b/Engine/Source/Runtime/Engine/Classes/Engine/LevelStreaming.h
--- a/Engine/Source/Runtime/Engine/Classes/Engine/LevelStreaming.h
+++ b/Engine/Source/Runtime/Engine/Classes/Engine/LevelStreaming.h
@@ -2,7 +2,6 @@
 
 #pragma once
 
-#include "CoreMinimal.h"
 #include "UObject/ObjectMacros.h"
 #include "UObject/Object.h"
 #include "UObject/ScriptMacros.h"
@@ -752,4 +751,8 @@ private:
 	/** Friend classes to manipulate the streaming level more extensively */
 	friend class UEngine;
 	friend class UWorld;
-};
\ No newline at end of file
+};
+
+#if UE_ENABLE_INCLUDE_ORDER_DEPRECATED_IN_5_2
+#include "CoreMinimal.h"
+#endif
diff --git a/Engine/Source/Runtime/Engine/Classes/Engine/NetSerialization.h b/Engine/Source/Runtime/Engine/Classes/Engine/NetSerialization.h
--- a/Engine/Source/Runtime/Engine/Classes/Engine/NetSerialization.h
+++ b/Engine/Source/Runtime/Engine/Classes/Engine/NetSerialization.h
@@ -7,7 +7,6 @@
 
 #pragma once
 
-#include "CoreMinimal.h"
 #include "Stats/Stats.h"
 #include "UObject/ObjectMacros.h"
 #include "UObject/Class.h"
diff --git a/Engine/Source/Runtime/Engine/Classes/Engine/ReplicatedState.h b/Engine/Source/Runtime/Engine/Classes/Engine/ReplicatedState.h
--- a/Engine/Source/Runtime/Engine/Classes/Engine/ReplicatedState.h
+++ b/Engine/Source/Runtime/Engine/Classes/Engine/ReplicatedState.h
@@ -2,7 +2,6 @@
 
 #pragma once
 
-#include "CoreMinimal.h"
 #include "UObject/ObjectMacros.h"
 #include "Engine/NetSerialization.h"
 #include "ReplicatedState.generated.h"
@@ -260,3 +259,7 @@ struct TStructOpsTypeTraits<FRepMovement> : public TStructOpsTypeTraitsBase2<FRe
 		WithNetSharedSerialization = true,
 	};
 };
+
+#if UE_ENABLE_INCLUDE_ORDER_DEPRECATED_IN_5_2
+#include "CoreMinimal.h"
+#endif
diff --git a/Engine/Source/Runtime/Engine/Classes/GameFramework/Actor.h b/Engine/Source/Runtime/Engine/Classes/GameFramework/Actor.h
--- a/Engine/Source/Runtime/Engine/Classes/GameFramework/Actor.h
+++ b/Engine/Source/Runtime/Engine/Classes/GameFramework/Actor.h
@@ -2,7 +2,6 @@
 
 #pragma once
 
-#include "CoreMinimal.h"
 #include "Stats/Stats.h"
 #include "UObject/ObjectMacros.h"
 #include "UObject/UObjectBaseUtility.h"
@@ -4409,3 +4408,7 @@ DEFINE_ACTORDESC_TYPE(AActor, FWorldPartitionActorDesc);
 	void SetActorRelativeRotation(FRotator NewRelativeRotation, bool bSweep=false, FHitResult* OutSweepHitResult=nullptr, ETeleportType Teleport = ETeleportType::None) { Super::SetActorRelativeRotation(NewRelativeRotation, bSweep, OutSweepHitResult, Teleport); } \
 	void SetActorRelativeRotation(const FQuat& NewRelativeRotation, bool bSweep=false, FHitResult* OutSweepHitResult=nullptr, ETeleportType Teleport = ETeleportType::None) { Super::SetActorRelativeRotation(NewRelativeRotation, bSweep, OutSweepHitResult, Teleport); } \
 	void SetActorRelativeTransform(const FTransform& NewRelativeTransform, bool bSweep=false, FHitResult* OutSweepHitResult=nullptr, ETeleportType Teleport = ETeleportType::None) { Super::SetActorRelativeTransform(NewRelativeTransform, bSweep, OutSweepHitResult, Teleport); }
+
+#if UE_ENABLE_INCLUDE_ORDER_DEPRECATED_IN_5_2
+#include "CoreMinimal.h"
+#endif
diff --git a/Engine/Source/Runtime/Engine/Classes/Interfaces/Interface_AssetUserData.h b/Engine/Source/Runtime/Engine/Classes/Interfaces/Interface_AssetUserData.h
--- a/Engine/Source/Runtime/Engine/Classes/Interfaces/Interface_AssetUserData.h
+++ b/Engine/Source/Runtime/Engine/Classes/Interfaces/Interface_AssetUserData.h
@@ -2,7 +2,6 @@
 
 #pragma once
 
-#include "CoreMinimal.h"
 #include "UObject/ObjectMacros.h"
 #include "Templates/SubclassOf.h"
 #include "Templates/Casts.h"
@@ -50,3 +49,6 @@ class IInterface_AssetUserData
 
 };
 
+#if UE_ENABLE_INCLUDE_ORDER_DEPRECATED_IN_5_2
+#include "CoreMinimal.h"
+#endif
diff --git a/Engine/Source/Runtime/Engine/Private/Components/ActorComponent.cpp b/Engine/Source/Runtime/Engine/Private/Components/ActorComponent.cpp
--- a/Engine/Source/Runtime/Engine/Private/Components/ActorComponent.cpp
+++ b/Engine/Source/Runtime/Engine/Private/Components/ActorComponent.cpp
@@ -15,6 +15,7 @@
 #include "Engine/LevelStreamingPersistent.h"
 #include "Engine/NetDriver.h"
 #include "Engine/ActorChannel.h"
+#include "HAL/LowLevelMemStats.h"
 #include "Misc/ScopeRWLock.h"
 #include "Net/UnrealNetwork.h"
 #include "Logging/MessageLog.h"
diff --git a/Engine/Source/Runtime/Engine/Private/PieFixupSerializer.h b/Engine/Source/Runtime/Engine/Private/PieFixupSerializer.h
--- a/Engine/Source/Runtime/Engine/Private/PieFixupSerializer.h
+++ b/Engine/Source/Runtime/Engine/Private/PieFixupSerializer.h
@@ -2,6 +2,7 @@
 
 #pragma once
 
+#include "Containers/Set.h"
 #include "Serialization/ArchiveUObject.h"
 
 class FPIEFixupSerializer : public FArchiveUObject
diff --git a/Engine/Source/Runtime/Engine/Public/ActorTransactionAnnotation.h b/Engine/Source/Runtime/Engine/Public/ActorTransactionAnnotation.h
--- a/Engine/Source/Runtime/Engine/Public/ActorTransactionAnnotation.h
+++ b/Engine/Source/Runtime/Engine/Public/ActorTransactionAnnotation.h
@@ -2,7 +2,6 @@
 
 #pragma once
 
-#include "CoreMinimal.h"
 #include "Engine/EngineTypes.h"
 #include "ComponentInstanceDataCache.h"
 
@@ -95,3 +94,7 @@ private:
 	TArray<FDiffableComponentInfo> DiffableComponentInfos;
 };
 #endif // WITH_EDITOR
+
+#if UE_ENABLE_INCLUDE_ORDER_DEPRECATED_IN_5_2
+#include "CoreMinimal.h"
+#endif
diff --git a/Engine/Source/Runtime/Engine/Public/ComponentInstanceDataCache.h b/Engine/Source/Runtime/Engine/Public/ComponentInstanceDataCache.h
--- a/Engine/Source/Runtime/Engine/Public/ComponentInstanceDataCache.h
+++ b/Engine/Source/Runtime/Engine/Public/ComponentInstanceDataCache.h
@@ -2,7 +2,6 @@
 
 #pragma once
 
-#include "CoreMinimal.h"
 #include "Templates/SubclassOf.h"
 #include "UObject/Object.h"
 #include "UObject/StructOnScope.h"
@@ -261,3 +260,7 @@ private:
 	/** Map of the actor instanced scene component to their transform relative to the root. */
 	TMap< USceneComponent*, FTransform > InstanceComponentTransformToRootMap;
 };
+
+#if UE_ENABLE_INCLUDE_ORDER_DEPRECATED_IN_5_2
+#include "CoreMinimal.h"
+#endif
diff --git a/Engine/Source/Runtime/Engine/Public/EngineLogs.h b/Engine/Source/Runtime/Engine/Public/EngineLogs.h
--- a/Engine/Source/Runtime/Engine/Public/EngineLogs.h
+++ b/Engine/Source/Runtime/Engine/Public/EngineLogs.h
@@ -4,7 +4,7 @@
 
 #pragma once
 
-#include "CoreMinimal.h"
+#include "Logging/LogMacros.h"
 
 class Error;
 
@@ -32,4 +32,8 @@ ENGINE_API DECLARE_LOG_CATEGORY_EXTERN(LogTextureUpload, Log, All);
 ENGINE_API DECLARE_LOG_CATEGORY_EXTERN(LogPlayerManagement, Error, All);
 ENGINE_API DECLARE_LOG_CATEGORY_EXTERN(LogSecurity, Warning, All);
 ENGINE_API DECLARE_LOG_CATEGORY_EXTERN(LogEngineSessionManager, Log, All);
-ENGINE_API DECLARE_LOG_CATEGORY_EXTERN(LogViewport, Log, All);
\ No newline at end of file
+ENGINE_API DECLARE_LOG_CATEGORY_EXTERN(LogViewport, Log, All);
+
+#if UE_ENABLE_INCLUDE_ORDER_DEPRECATED_IN_5_2
+#include "CoreMinimal.h"
+#endif
diff --git a/Engine/Source/Runtime/Engine/Public/Folder.h b/Engine/Source/Runtime/Engine/Public/Folder.h
--- a/Engine/Source/Runtime/Engine/Public/Folder.h
+++ b/Engine/Source/Runtime/Engine/Public/Folder.h
@@ -2,6 +2,7 @@
 
 #pragma once
 
+#include "Misc/Guid.h"
 #include "UObject/ObjectMacros.h"
 #include "UObject/Object.h"
 #include "UObject/ObjectKey.h"
diff --git a/Engine/Source/Runtime/Engine/Public/ViewportClient.h b/Engine/Source/Runtime/Engine/Public/ViewportClient.h
--- a/Engine/Source/Runtime/Engine/Public/ViewportClient.h
+++ b/Engine/Source/Runtime/Engine/Public/ViewportClient.h
@@ -5,6 +5,7 @@
 #include "Engine/EngineBaseTypes.h"
 #include "GenericPlatform/ICursor.h"
 #include "InputCoreTypes.h"
+#include "Misc/DateTime.h"
 
 class FCanvas;
 class FCursorReply;
diff --git a/Engine/Source/Runtime/Engine/Public/WorldPartition/WorldPartitionActorDescType.h b/Engine/Source/Runtime/Engine/Public/WorldPartition/WorldPartitionActorDescType.h
--- a/Engine/Source/Runtime/Engine/Public/WorldPartition/WorldPartitionActorDescType.h
+++ b/Engine/Source/Runtime/Engine/Public/WorldPartition/WorldPartitionActorDescType.h
@@ -1,7 +1,7 @@
 // Copyright Epic Games, Inc. All Rights Reserved.
 #pragma once
 
-#include "CoreMinimal.h"
+#include "Templates/UnrealTypeTraits.h"
 
 class AActor;
 
@@ -19,3 +19,7 @@ struct FWorldPartitionActorDescType<ActorType>			\
 #else
 #define DEFINE_ACTORDESC_TYPE(ActorType, ActorDescType)
 #endif
+
+#if UE_ENABLE_INCLUDE_ORDER_DEPRECATED_IN_5_2
+#include "CoreMinimal.h"
+#endif
diff --git a/Engine/Source/Runtime/InputCore/Classes/InputCoreTypes.h b/Engine/Source/Runtime/InputCore/Classes/InputCoreTypes.h
--- a/Engine/Source/Runtime/InputCore/Classes/InputCoreTypes.h
+++ b/Engine/Source/Runtime/InputCore/Classes/InputCoreTypes.h
@@ -2,7 +2,6 @@
 
 #pragma once
 
-#include "CoreMinimal.h"
 #include "UObject/ObjectMacros.h"
 #include "UObject/Object.h"
 #include "UObject/Class.h"
@@ -784,3 +783,7 @@ class UInputCoreTypes : public UObject
 	GENERATED_BODY()
 
 };
+
+#if UE_ENABLE_INCLUDE_ORDER_DEPRECATED_IN_5_2
+#include "CoreMinimal.h"
+#endif
diff --git a/Engine/Source/Runtime/Net/Core/Public/Net/Core/Connection/NetEnums.h b/Engine/Source/Runtime/Net/Core/Public/Net/Core/Connection/NetEnums.h
--- a/Engine/Source/Runtime/Net/Core/Public/Net/Core/Connection/NetEnums.h
+++ b/Engine/Source/Runtime/Net/Core/Public/Net/Core/Connection/NetEnums.h
@@ -2,7 +2,6 @@
 
 #pragma once
 
-#include "CoreMinimal.h"
 #include "UObject/ObjectMacros.h"
 
 #include "NetEnums.generated.h"
@@ -104,4 +103,6 @@ namespace ESecurityEvent
 	}
 }
 
-
+#if UE_ENABLE_INCLUDE_ORDER_DEPRECATED_IN_5_2
+#include "CoreMinimal.h"
+#endif
diff --git a/Engine/Source/Runtime/RenderCore/Public/RenderCommandFence.h b/Engine/Source/Runtime/RenderCore/Public/RenderCommandFence.h
--- a/Engine/Source/Runtime/RenderCore/Public/RenderCommandFence.h
+++ b/Engine/Source/Runtime/RenderCore/Public/RenderCommandFence.h
@@ -3,7 +3,6 @@
 #pragma once
 
 #include "Async/TaskGraphFwd.h"
-#include "CoreMinimal.h"
 #include "Templates/RefCounting.h"
 
 namespace ENamedThreads { enum Type : int32; }
@@ -46,3 +45,7 @@ private:
 	/** Thread that will trigger the CompletionEvent **/
 	ENamedThreads::Type TriggerThreadIndex;
 };
+
+#if UE_ENABLE_INCLUDE_ORDER_DEPRECATED_IN_5_2
+#include "CoreMinimal.h"
+#endif
