prevent VM execution functions from being inlined when PGOOptimize is enabled on MS platforms (and prevent linker to fail with data generated with FASTGENPROFILE)

#tests A/B perf test regression

diff --git a/Engine/Source/Programs/UnrealBuildTool/Platform/Windows/VCToolChain.cs b/Engine/Source/Programs/UnrealBuildTool/Platform/Windows/VCToolChain.cs
--- a/Engine/Source/Programs/UnrealBuildTool/Platform/Windows/VCToolChain.cs
+++ b/Engine/Source/Programs/UnrealBuildTool/Platform/Windows/VCToolChain.cs
@@ -623,6 +623,12 @@ namespace UnrealBuildTool
 				Arguments.Add("/GL");
 			}
 
+			if (CompileEnvironment.bPGOOptimize)
+			{
+				// Prevent VM execution functions from being inlined - UE-172329
+				AddDefinition(Arguments, "FORCE_NO_INLINE_VM_FUNCTION=1");
+			}
+
 			//
 			//	PC
 			//
diff --git a/Engine/Source/Runtime/CoreUObject/Private/UObject/ScriptCore.cpp b/Engine/Source/Runtime/CoreUObject/Private/UObject/ScriptCore.cpp
--- a/Engine/Source/Runtime/CoreUObject/Private/UObject/ScriptCore.cpp
+++ b/Engine/Source/Runtime/CoreUObject/Private/UObject/ScriptCore.cpp
@@ -135,6 +135,17 @@ static struct F##inst##Registrar \
 	IMPLEMENT_FUNCTION(func) \
 	static uint8 UObject##func##BytecodeTemp = GRegisterNative( BytecodeIndex, &UObject::func );
 
+
+#ifndef FORCE_NO_INLINE_VM_FUNCTION
+	#define FORCE_NO_INLINE_VM_FUNCTION 0
+#endif
+
+#if FORCE_NO_INLINE_VM_FUNCTION
+	#define DEFINE_VM_FUNCTION FORCENOINLINE DEFINE_FUNCTION
+#else
+	#define DEFINE_VM_FUNCTION DEFINE_FUNCTION
+#endif
+
 //////////////////////////////////////////////////////////////////////////
 // FBlueprintCoreDelegates
 
@@ -1042,7 +1053,7 @@ void ProcessScriptFunction(UObject* Context, UFunction* Function, FFrame& Stack,
 	Stack.bAbortingExecution |= NewStack.bAbortingExecution;
 }
 
-DEFINE_FUNCTION(UObject::execCallMathFunction)
+DEFINE_VM_FUNCTION(UObject::execCallMathFunction)
 {
 	UFunction* Function = (UFunction*)Stack.ReadObject();
 	checkSlow(Function);
@@ -2151,6 +2162,7 @@ void UObject::ProcessEvent( UFunction* Function, void* Parms )
 #pragma warning (pop)
 #endif
 
+
 DEFINE_FUNCTION(UObject::execUndefined)
 {
 	const FText LocalizedErrorMessage = FText::Format(
@@ -2162,7 +2174,7 @@ DEFINE_FUNCTION(UObject::execUndefined)
 	Stack.Log(ELogVerbosity::Error, LocalizedErrorMessage.ToString());
 }
 
-DEFINE_FUNCTION(UObject::execLocalVariable)
+DEFINE_VM_FUNCTION(UObject::execLocalVariable)
 {
 	checkSlow(Stack.Object == P_THIS);
 	checkSlow(Stack.Locals != NULL);
@@ -2196,7 +2208,7 @@ DEFINE_FUNCTION(UObject::execLocalVariable)
 }
 IMPLEMENT_VM_FUNCTION( EX_LocalVariable, execLocalVariable );
 
-DEFINE_FUNCTION(UObject::execInstanceVariable)
+DEFINE_VM_FUNCTION(UObject::execInstanceVariable)
 {
 	FProperty* VarProperty = (FProperty*)Stack.ReadObject();
 	Stack.MostRecentProperty = VarProperty;
@@ -2230,7 +2242,7 @@ DEFINE_FUNCTION(UObject::execInstanceVariable)
 }
 IMPLEMENT_VM_FUNCTION( EX_InstanceVariable, execInstanceVariable );
 
-DEFINE_FUNCTION(UObject::execClassSparseDataVariable)
+DEFINE_VM_FUNCTION(UObject::execClassSparseDataVariable)
 {
 	FProperty* VarProperty = (FProperty*)Stack.ReadObject();
 	Stack.MostRecentProperty = VarProperty;
@@ -2264,7 +2276,7 @@ DEFINE_FUNCTION(UObject::execClassSparseDataVariable)
 }
 IMPLEMENT_VM_FUNCTION(EX_ClassSparseDataVariable, execClassSparseDataVariable);
 
-DEFINE_FUNCTION(UObject::execDefaultVariable)
+DEFINE_VM_FUNCTION(UObject::execDefaultVariable)
 {
 	FProperty* VarProperty = (FProperty*)Stack.ReadObject();
 	Stack.MostRecentProperty = VarProperty;
@@ -2314,7 +2326,7 @@ DEFINE_FUNCTION(UObject::execDefaultVariable)
 }
 IMPLEMENT_VM_FUNCTION( EX_DefaultVariable, execDefaultVariable );
 
-DEFINE_FUNCTION(UObject::execLocalOutVariable)
+DEFINE_VM_FUNCTION(UObject::execLocalOutVariable)
 {
 	checkSlow(Stack.Object == P_THIS);
 
@@ -2339,7 +2351,7 @@ DEFINE_FUNCTION(UObject::execLocalOutVariable)
 }
 IMPLEMENT_VM_FUNCTION(EX_LocalOutVariable, execLocalOutVariable);
 
-DEFINE_FUNCTION(UObject::execInterfaceContext)
+DEFINE_VM_FUNCTION(UObject::execInterfaceContext)
 {
 	// get the value of the interface variable
 	FScriptInterface InterfaceValue;
@@ -2353,7 +2365,7 @@ DEFINE_FUNCTION(UObject::execInterfaceContext)
 }
 IMPLEMENT_VM_FUNCTION( EX_InterfaceContext, execInterfaceContext );
 
-DEFINE_FUNCTION(UObject::execClassContext)
+DEFINE_VM_FUNCTION(UObject::execClassContext)
 {
 	// Get class expression.
 	UClass* ClassContext = NULL;
@@ -2404,7 +2416,7 @@ DEFINE_FUNCTION(UObject::execClassContext)
 }
 IMPLEMENT_VM_FUNCTION( EX_ClassContext, execClassContext );
 
-DEFINE_FUNCTION(UObject::execEndOfScript)
+DEFINE_VM_FUNCTION(UObject::execEndOfScript)
 {
 #if WITH_EDITOR
 	if (GIsEditor)
@@ -2422,19 +2434,19 @@ DEFINE_FUNCTION(UObject::execEndOfScript)
 }
 IMPLEMENT_VM_FUNCTION( EX_EndOfScript, execEndOfScript );
 
-DEFINE_FUNCTION(UObject::execNothing)
+DEFINE_VM_FUNCTION(UObject::execNothing)
 {
 	// Do nothing.
 }
 IMPLEMENT_VM_FUNCTION( EX_Nothing, execNothing );
 
-DEFINE_FUNCTION(UObject::execNothingOp4a)
+DEFINE_VM_FUNCTION(UObject::execNothingOp4a)
 {
 	// Do nothing.
 }
 IMPLEMENT_VM_FUNCTION( EX_DeprecatedOp4A, execNothingOp4a );
 
-DEFINE_FUNCTION(UObject::execBreakpoint)
+DEFINE_VM_FUNCTION(UObject::execBreakpoint)
 {
 	if (FBlueprintCoreDelegates::IsDebuggingEnabled())
 	{
@@ -2444,7 +2456,7 @@ DEFINE_FUNCTION(UObject::execBreakpoint)
 }
 IMPLEMENT_VM_FUNCTION( EX_Breakpoint, execBreakpoint );
 
-DEFINE_FUNCTION(UObject::execTracepoint)
+DEFINE_VM_FUNCTION(UObject::execTracepoint)
 {
 	if (FBlueprintCoreDelegates::IsDebuggingEnabled())
 	{
@@ -2454,7 +2466,7 @@ DEFINE_FUNCTION(UObject::execTracepoint)
 }
 IMPLEMENT_VM_FUNCTION( EX_Tracepoint, execTracepoint );
 
-DEFINE_FUNCTION(UObject::execWireTracepoint)
+DEFINE_VM_FUNCTION(UObject::execWireTracepoint)
 {
 	if (FBlueprintCoreDelegates::IsDebuggingEnabled())
 	{
@@ -2464,7 +2476,7 @@ DEFINE_FUNCTION(UObject::execWireTracepoint)
 }
 IMPLEMENT_VM_FUNCTION( EX_WireTracepoint, execWireTracepoint );
 
-DEFINE_FUNCTION(UObject::execInstrumentation)
+DEFINE_VM_FUNCTION(UObject::execInstrumentation)
 {
 #if !UE_BUILD_SHIPPING
 	const EScriptInstrumentation::Type EventType = static_cast<EScriptInstrumentation::Type>(Stack.PeekCode());
@@ -2505,7 +2517,7 @@ DEFINE_FUNCTION(UObject::execInstrumentation)
 }
 IMPLEMENT_VM_FUNCTION( EX_InstrumentationEvent, execInstrumentation );
 
-DEFINE_FUNCTION(UObject::execEndFunctionParms)
+DEFINE_VM_FUNCTION(UObject::execEndFunctionParms)
 {
 	// For skipping over optional function parms without values specified.
 	Stack.Code--;
@@ -2513,7 +2525,7 @@ DEFINE_FUNCTION(UObject::execEndFunctionParms)
 IMPLEMENT_VM_FUNCTION( EX_EndFunctionParms, execEndFunctionParms );
 
 
-DEFINE_FUNCTION(UObject::execJump)
+DEFINE_VM_FUNCTION(UObject::execJump)
 {
 	CHECK_RUNAWAY;
 
@@ -2523,7 +2535,7 @@ DEFINE_FUNCTION(UObject::execJump)
 }
 IMPLEMENT_VM_FUNCTION( EX_Jump, execJump );
 
-DEFINE_FUNCTION(UObject::execComputedJump)
+DEFINE_VM_FUNCTION(UObject::execComputedJump)
 {
 	CHECK_RUNAWAY;
 
@@ -2537,7 +2549,7 @@ DEFINE_FUNCTION(UObject::execComputedJump)
 }
 IMPLEMENT_VM_FUNCTION( EX_ComputedJump, execComputedJump );
 	
-DEFINE_FUNCTION(UObject::execJumpIfNot)
+DEFINE_VM_FUNCTION(UObject::execJumpIfNot)
 {
 	CHECK_RUNAWAY;
 
@@ -2556,7 +2568,7 @@ DEFINE_FUNCTION(UObject::execJumpIfNot)
 }
 IMPLEMENT_VM_FUNCTION( EX_JumpIfNot, execJumpIfNot );
 
-DEFINE_FUNCTION(UObject::execAssert)
+DEFINE_VM_FUNCTION(UObject::execAssert)
 {
 	// Get line number.
 	int32 wLine = Stack.ReadWord();
@@ -2584,7 +2596,7 @@ DEFINE_FUNCTION(UObject::execAssert)
 }
 IMPLEMENT_VM_FUNCTION( EX_Assert, execAssert );
 
-DEFINE_FUNCTION(UObject::execPushExecutionFlow)
+DEFINE_VM_FUNCTION(UObject::execPushExecutionFlow)
 {
 	// Read a code offset and push it onto the flow stack
 	CodeSkipSizeType Offset = Stack.ReadCodeSkipCount();
@@ -2592,7 +2604,7 @@ DEFINE_FUNCTION(UObject::execPushExecutionFlow)
 }
 IMPLEMENT_VM_FUNCTION( EX_PushExecutionFlow, execPushExecutionFlow );
 
-DEFINE_FUNCTION(UObject::execPopExecutionFlow)
+DEFINE_VM_FUNCTION(UObject::execPopExecutionFlow)
 {
 	// Since this is a branch function, check for runaway script execution
 	CHECK_RUNAWAY;
@@ -2611,7 +2623,7 @@ DEFINE_FUNCTION(UObject::execPopExecutionFlow)
 }
 IMPLEMENT_VM_FUNCTION( EX_PopExecutionFlow, execPopExecutionFlow );
 
-DEFINE_FUNCTION(UObject::execPopExecutionFlowIfNot)
+DEFINE_VM_FUNCTION(UObject::execPopExecutionFlowIfNot)
 {
 	// Since this is a branch function, check for runaway script execution
 	CHECK_RUNAWAY;
@@ -2637,7 +2649,7 @@ DEFINE_FUNCTION(UObject::execPopExecutionFlowIfNot)
 }
 IMPLEMENT_VM_FUNCTION( EX_PopExecutionFlowIfNot, execPopExecutionFlowIfNot );
 
-DEFINE_FUNCTION(UObject::execLetValueOnPersistentFrame)
+DEFINE_VM_FUNCTION(UObject::execLetValueOnPersistentFrame)
 {
 #if USE_UBER_GRAPH_PERSISTENT_FRAME
 	Stack.MostRecentProperty = nullptr;
@@ -2659,7 +2671,7 @@ DEFINE_FUNCTION(UObject::execLetValueOnPersistentFrame)
 }
 IMPLEMENT_VM_FUNCTION(EX_LetValueOnPersistentFrame, execLetValueOnPersistentFrame);
 
-DEFINE_FUNCTION(UObject::execSwitchValue)
+DEFINE_VM_FUNCTION(UObject::execSwitchValue)
 {
 	const int32 NumCases = Stack.ReadWord();
 	const CodeSkipSizeType OffsetToEnd = Stack.ReadCodeSkipCount();
@@ -2728,7 +2740,7 @@ DEFINE_FUNCTION(UObject::execSwitchValue)
 }
 IMPLEMENT_VM_FUNCTION(EX_SwitchValue, execSwitchValue);
 
-DEFINE_FUNCTION(UObject::execArrayGetByRef)
+DEFINE_VM_FUNCTION(UObject::execArrayGetByRef)
 {
 	// Get variable address.
 	Stack.MostRecentPropertyAddress = nullptr;
@@ -2791,7 +2803,7 @@ DEFINE_FUNCTION(UObject::execArrayGetByRef)
 }
 IMPLEMENT_VM_FUNCTION(EX_ArrayGetByRef, execArrayGetByRef);
 
-DEFINE_FUNCTION(UObject::execLet)
+DEFINE_VM_FUNCTION(UObject::execLet)
 {
 	Stack.MostRecentProperty = nullptr;
 	FProperty* LocallyKnownProperty = Stack.ReadPropertyUnchecked();
@@ -2854,7 +2866,7 @@ DEFINE_FUNCTION(UObject::execLet)
 }
 IMPLEMENT_VM_FUNCTION( EX_Let, execLet );
 
-DEFINE_FUNCTION(UObject::execLetObj)
+DEFINE_VM_FUNCTION(UObject::execLetObj)
 {
 	// Get variable address.
 	Stack.MostRecentPropertyAddress = nullptr;
@@ -2901,7 +2913,7 @@ DEFINE_FUNCTION(UObject::execLetObj)
 }
 IMPLEMENT_VM_FUNCTION( EX_LetObj, execLetObj );
 
-DEFINE_FUNCTION(UObject::execLetWeakObjPtr)
+DEFINE_VM_FUNCTION(UObject::execLetWeakObjPtr)
 {
 	// Get variable address.
 	Stack.MostRecentPropertyAddress = nullptr;
@@ -2949,7 +2961,7 @@ DEFINE_FUNCTION(UObject::execLetWeakObjPtr)
 }
 IMPLEMENT_VM_FUNCTION( EX_LetWeakObjPtr, execLetWeakObjPtr );
 
-DEFINE_FUNCTION(UObject::execLetBool)
+DEFINE_VM_FUNCTION(UObject::execLetBool)
 {
 	Stack.MostRecentPropertyAddress = nullptr;
 	Stack.MostRecentPropertyContainer = nullptr;
@@ -3004,7 +3016,7 @@ DEFINE_FUNCTION(UObject::execLetBool)
 IMPLEMENT_VM_FUNCTION( EX_LetBool, execLetBool );
 
 
-DEFINE_FUNCTION(UObject::execLetDelegate)
+DEFINE_VM_FUNCTION(UObject::execLetDelegate)
 {
 	// Get variable address.
 	Stack.MostRecentPropertyAddress = nullptr;
@@ -3024,7 +3036,7 @@ DEFINE_FUNCTION(UObject::execLetDelegate)
 IMPLEMENT_VM_FUNCTION( EX_LetDelegate, execLetDelegate );
 
 
-DEFINE_FUNCTION(UObject::execLetMulticastDelegate)
+DEFINE_VM_FUNCTION(UObject::execLetMulticastDelegate)
 {
 	// Get variable address.
 	Stack.MostRecentPropertyAddress = nullptr;
@@ -3045,7 +3057,7 @@ DEFINE_FUNCTION(UObject::execLetMulticastDelegate)
 IMPLEMENT_VM_FUNCTION( EX_LetMulticastDelegate, execLetMulticastDelegate );
 
 
-DEFINE_FUNCTION(UObject::execSelf)
+DEFINE_VM_FUNCTION(UObject::execSelf)
 {
 	// Get Self actor for this context.
 	if (RESULT_PARAM != nullptr)
@@ -3065,13 +3077,13 @@ DEFINE_FUNCTION(UObject::execSelf)
 }
 IMPLEMENT_VM_FUNCTION( EX_Self, execSelf );
 
-DEFINE_FUNCTION(UObject::execContext)
+DEFINE_VM_FUNCTION(UObject::execContext)
 {
 	P_THIS->ProcessContextOpcode(Stack, RESULT_PARAM, /*bCanFailSilently=*/ false);
 }
 IMPLEMENT_VM_FUNCTION( EX_Context, execContext );
 
-DEFINE_FUNCTION(UObject::execContext_FailSilent)
+DEFINE_VM_FUNCTION(UObject::execContext_FailSilent)
 {
 	P_THIS->ProcessContextOpcode(Stack, RESULT_PARAM, /*bCanFailSilently=*/ true);
 }
@@ -3157,7 +3169,7 @@ void UObject::ProcessContextOpcode( FFrame& Stack, RESULT_DECL, bool bCanFailSil
 	}
 }
 
-DEFINE_FUNCTION(UObject::execStructMemberContext)
+DEFINE_VM_FUNCTION(UObject::execStructMemberContext)
 {
 	// Get the structure element we care about
 	FProperty* StructProperty = Stack.ReadProperty();
@@ -3208,28 +3220,28 @@ DEFINE_FUNCTION(UObject::execStructMemberContext)
 }
 IMPLEMENT_VM_FUNCTION( EX_StructMemberContext, execStructMemberContext );
 
-DEFINE_FUNCTION(UObject::execVirtualFunction)
+DEFINE_VM_FUNCTION(UObject::execVirtualFunction)
 {
 	// Call the virtual function.
 	P_THIS->CallFunction( Stack, RESULT_PARAM, P_THIS->FindFunctionChecked(Stack.ReadName()) );
 }
 IMPLEMENT_VM_FUNCTION( EX_VirtualFunction, execVirtualFunction );
 
-DEFINE_FUNCTION(UObject::execFinalFunction)
+DEFINE_VM_FUNCTION(UObject::execFinalFunction)
 {
 	// Call the final function.
 	P_THIS->CallFunction( Stack, RESULT_PARAM, (UFunction*)Stack.ReadObject() );
 }
 IMPLEMENT_VM_FUNCTION( EX_FinalFunction, execFinalFunction );
 
-DEFINE_FUNCTION(UObject::execLocalVirtualFunction)
+DEFINE_VM_FUNCTION(UObject::execLocalVirtualFunction)
 {
 	// Call the virtual function.
 	ProcessLocalFunction(Context, P_THIS->FindFunctionChecked(Stack.ReadName()), Stack, RESULT_PARAM);
 }
 IMPLEMENT_VM_FUNCTION( EX_LocalVirtualFunction, execLocalVirtualFunction );
 
-DEFINE_FUNCTION(UObject::execLocalFinalFunction)
+DEFINE_VM_FUNCTION(UObject::execLocalFinalFunction)
 {
 	// Call the final function.
 	ProcessLocalFunction(Context, (UFunction*)Stack.ReadObject(), Stack, RESULT_PARAM);
@@ -3291,13 +3303,13 @@ public:
 	}
 };
 
-DEFINE_FUNCTION(UObject::execCallMulticastDelegate)
+DEFINE_VM_FUNCTION(UObject::execCallMulticastDelegate)
 {
 	FCallDelegateHelper::CallMulticastDelegate(Stack);
 }
 IMPLEMENT_VM_FUNCTION( EX_CallMulticastDelegate, execCallMulticastDelegate );
 
-DEFINE_FUNCTION(UObject::execAddMulticastDelegate)
+DEFINE_VM_FUNCTION(UObject::execAddMulticastDelegate)
 {
 	// Get variable address.
 	Stack.MostRecentPropertyAddress = nullptr;
@@ -3318,7 +3330,7 @@ DEFINE_FUNCTION(UObject::execAddMulticastDelegate)
 }
 IMPLEMENT_VM_FUNCTION( EX_AddMulticastDelegate, execAddMulticastDelegate );
 
-DEFINE_FUNCTION(UObject::execRemoveMulticastDelegate)
+DEFINE_VM_FUNCTION(UObject::execRemoveMulticastDelegate)
 {
 	// Get variable address.
 	Stack.MostRecentPropertyAddress = nullptr;
@@ -3339,7 +3351,7 @@ DEFINE_FUNCTION(UObject::execRemoveMulticastDelegate)
 }
 IMPLEMENT_VM_FUNCTION( EX_RemoveMulticastDelegate, execRemoveMulticastDelegate );
 
-DEFINE_FUNCTION(UObject::execClearMulticastDelegate)
+DEFINE_VM_FUNCTION(UObject::execClearMulticastDelegate)
 {
 	// Get the delegate address
 	Stack.MostRecentPropertyAddress = nullptr;
@@ -3357,44 +3369,44 @@ DEFINE_FUNCTION(UObject::execClearMulticastDelegate)
 }
 IMPLEMENT_VM_FUNCTION( EX_ClearMulticastDelegate, execClearMulticastDelegate );
 
-DEFINE_FUNCTION(UObject::execIntConst)
+DEFINE_VM_FUNCTION(UObject::execIntConst)
 {
 	*(int32*)RESULT_PARAM = Stack.ReadInt<int32>();
 }
 IMPLEMENT_VM_FUNCTION( EX_IntConst, execIntConst );
 
-DEFINE_FUNCTION(UObject::execInt64Const)
+DEFINE_VM_FUNCTION(UObject::execInt64Const)
 {
 	*(int64*)RESULT_PARAM = Stack.ReadInt<int64>();
 }
 IMPLEMENT_VM_FUNCTION(EX_Int64Const, execInt64Const);
 
-DEFINE_FUNCTION(UObject::execUInt64Const)
+DEFINE_VM_FUNCTION(UObject::execUInt64Const)
 {
 	*(uint64*)RESULT_PARAM = Stack.ReadInt<uint64>();
 }
 IMPLEMENT_VM_FUNCTION(EX_UInt64Const, execUInt64Const);
 
-DEFINE_FUNCTION(UObject::execSkipOffsetConst)
+DEFINE_VM_FUNCTION(UObject::execSkipOffsetConst)
 {
 	CodeSkipSizeType Literal = Stack.ReadCodeSkipCount();
 	*(int32*)RESULT_PARAM = Literal;
 }
 IMPLEMENT_VM_FUNCTION( EX_SkipOffsetConst, execSkipOffsetConst );
 
-DEFINE_FUNCTION(UObject::execFloatConst)
+DEFINE_VM_FUNCTION(UObject::execFloatConst)
 {
 	*(float*)RESULT_PARAM = Stack.ReadFloat();
 }
 IMPLEMENT_VM_FUNCTION( EX_FloatConst, execFloatConst );
 
-DEFINE_FUNCTION(UObject::execDoubleConst)
+DEFINE_VM_FUNCTION(UObject::execDoubleConst)
 {
 	*(double*)RESULT_PARAM = Stack.ReadInt<double>();
 }
 IMPLEMENT_VM_FUNCTION( EX_DoubleConst, execDoubleConst );
 
-DEFINE_FUNCTION(UObject::execStringConst)
+DEFINE_VM_FUNCTION(UObject::execStringConst)
 {
 	*(FString*)RESULT_PARAM = (ANSICHAR*)Stack.Code;
 	while( *Stack.Code )
@@ -3403,7 +3415,7 @@ DEFINE_FUNCTION(UObject::execStringConst)
 }
 IMPLEMENT_VM_FUNCTION( EX_StringConst, execStringConst );
 
-DEFINE_FUNCTION(UObject::execUnicodeStringConst)
+DEFINE_VM_FUNCTION(UObject::execUnicodeStringConst)
 {
 	FString& ResultStr = *(FString*)RESULT_PARAM;
 	ResultStr = FString((UCS2CHAR*)Stack.Code);
@@ -3419,7 +3431,7 @@ DEFINE_FUNCTION(UObject::execUnicodeStringConst)
 }
 IMPLEMENT_VM_FUNCTION( EX_UnicodeStringConst, execUnicodeStringConst );
 
-DEFINE_FUNCTION(UObject::execTextConst)
+DEFINE_VM_FUNCTION(UObject::execTextConst)
 {
 	// What kind of text are we dealing with?
 	const EBlueprintTextLiteralType TextLiteralType = (EBlueprintTextLiteralType)*Stack.Code++;
@@ -3486,19 +3498,19 @@ DEFINE_FUNCTION(UObject::execTextConst)
 }
 IMPLEMENT_VM_FUNCTION( EX_TextConst, execTextConst );
 
-DEFINE_FUNCTION(UObject::execPropertyConst)
+DEFINE_VM_FUNCTION(UObject::execPropertyConst)
 {
 	*(FProperty**)RESULT_PARAM = (FProperty*)Stack.ReadObject();
 }
 IMPLEMENT_VM_FUNCTION(EX_PropertyConst, execPropertyConst);
 
-DEFINE_FUNCTION(UObject::execObjectConst)
+DEFINE_VM_FUNCTION(UObject::execObjectConst)
 {
 	*(UObject**)RESULT_PARAM = (UObject*)Stack.ReadObject();
 }
 IMPLEMENT_VM_FUNCTION( EX_ObjectConst, execObjectConst );
 
-DEFINE_FUNCTION(UObject::execSoftObjectConst)
+DEFINE_VM_FUNCTION(UObject::execSoftObjectConst)
 {
 	FString LongPath;
 	Stack.Step(Stack.Object, &LongPath);
@@ -3506,7 +3518,7 @@ DEFINE_FUNCTION(UObject::execSoftObjectConst)
 }
 IMPLEMENT_VM_FUNCTION( EX_SoftObjectConst, execSoftObjectConst);
 
-DEFINE_FUNCTION(UObject::execFieldPathConst)
+DEFINE_VM_FUNCTION(UObject::execFieldPathConst)
 {
 	FString StringPath;
 	Stack.Step(Stack.Object, &StringPath);
@@ -3516,14 +3528,14 @@ DEFINE_FUNCTION(UObject::execFieldPathConst)
 }
 IMPLEMENT_VM_FUNCTION(EX_FieldPathConst, execFieldPathConst);
 
-DEFINE_FUNCTION(UObject::execInstanceDelegate)
+DEFINE_VM_FUNCTION(UObject::execInstanceDelegate)
 {
 	FName FunctionName = Stack.ReadName();
 	((FScriptDelegate*)RESULT_PARAM)->BindUFunction( (FunctionName == NAME_None) ? NULL : P_THIS, FunctionName );
 }
 IMPLEMENT_VM_FUNCTION( EX_InstanceDelegate, execInstanceDelegate );
 
-DEFINE_FUNCTION(UObject::execBindDelegate)
+DEFINE_VM_FUNCTION(UObject::execBindDelegate)
 {
 	FName FunctionName = Stack.ReadName();
 
@@ -3545,19 +3557,19 @@ DEFINE_FUNCTION(UObject::execBindDelegate)
 }
 IMPLEMENT_VM_FUNCTION( EX_BindDelegate, execBindDelegate );
 
-DEFINE_FUNCTION(UObject::execNameConst)
+DEFINE_VM_FUNCTION(UObject::execNameConst)
 {
 	*(FName*)RESULT_PARAM = Stack.ReadName();
 }
 IMPLEMENT_VM_FUNCTION( EX_NameConst, execNameConst );
 
-DEFINE_FUNCTION(UObject::execByteConst)
+DEFINE_VM_FUNCTION(UObject::execByteConst)
 {
 	*(uint8*)RESULT_PARAM = *Stack.Code++;
 }
 IMPLEMENT_VM_FUNCTION( EX_ByteConst, execByteConst );
 
-DEFINE_FUNCTION(UObject::execRotationConst)
+DEFINE_VM_FUNCTION(UObject::execRotationConst)
 {
 	((FRotator*)RESULT_PARAM)->Pitch = Stack.ReadDouble();
 	((FRotator*)RESULT_PARAM)->Yaw   = Stack.ReadDouble();
@@ -3565,7 +3577,7 @@ DEFINE_FUNCTION(UObject::execRotationConst)
 }
 IMPLEMENT_VM_FUNCTION( EX_RotationConst, execRotationConst );
 
-DEFINE_FUNCTION(UObject::execVectorConst)
+DEFINE_VM_FUNCTION(UObject::execVectorConst)
 {
 	((FVector*)RESULT_PARAM)->X = Stack.ReadDouble();
 	((FVector*)RESULT_PARAM)->Y = Stack.ReadDouble();
@@ -3573,7 +3585,7 @@ DEFINE_FUNCTION(UObject::execVectorConst)
 }
 IMPLEMENT_VM_FUNCTION( EX_VectorConst, execVectorConst );
 
-DEFINE_FUNCTION(UObject::execVector3fConst)
+DEFINE_VM_FUNCTION(UObject::execVector3fConst)
 {
 	((FVector3f*)RESULT_PARAM)->X = Stack.ReadFloat();
 	((FVector3f*)RESULT_PARAM)->Y = Stack.ReadFloat();
@@ -3581,7 +3593,7 @@ DEFINE_FUNCTION(UObject::execVector3fConst)
 }
 IMPLEMENT_VM_FUNCTION(EX_Vector3fConst, execVector3fConst);
 
-DEFINE_FUNCTION(UObject::execTransformConst)
+DEFINE_VM_FUNCTION(UObject::execTransformConst)
 {
 	// Rotation
 	FQuat TmpRotation;
@@ -3606,7 +3618,7 @@ DEFINE_FUNCTION(UObject::execTransformConst)
 }
 IMPLEMENT_VM_FUNCTION( EX_TransformConst, execTransformConst );
 
-DEFINE_FUNCTION(UObject::execStructConst)
+DEFINE_VM_FUNCTION(UObject::execStructConst)
 {
 	UScriptStruct* ScriptStruct = CastChecked<UScriptStruct>(Stack.ReadObject());
 	int32 SerializedSize = Stack.ReadInt<int32>();
@@ -3639,7 +3651,7 @@ DEFINE_FUNCTION(UObject::execStructConst)
 }
 IMPLEMENT_VM_FUNCTION( EX_StructConst, execStructConst );
 
-DEFINE_FUNCTION(UObject::execSetArray)
+DEFINE_VM_FUNCTION(UObject::execSetArray)
 {
 	// Get the array address
 	Stack.MostRecentPropertyAddress = nullptr;
@@ -3663,7 +3675,7 @@ DEFINE_FUNCTION(UObject::execSetArray)
 }
 IMPLEMENT_VM_FUNCTION( EX_SetArray, execSetArray );
 
-DEFINE_FUNCTION(UObject::execSetSet)
+DEFINE_VM_FUNCTION(UObject::execSetSet)
 {
 	// Get the set address
 	Stack.MostRecentPropertyAddress = nullptr;
@@ -3697,7 +3709,7 @@ DEFINE_FUNCTION(UObject::execSetSet)
 }
 IMPLEMENT_VM_FUNCTION( EX_SetSet, execSetSet );
 
-DEFINE_FUNCTION(UObject::execSetMap)
+DEFINE_VM_FUNCTION(UObject::execSetMap)
 {
 	// Get the map address
 	Stack.MostRecentPropertyAddress = nullptr;
@@ -3732,7 +3744,7 @@ DEFINE_FUNCTION(UObject::execSetMap)
 }
 IMPLEMENT_VM_FUNCTION( EX_SetMap, execSetMap );
 
-DEFINE_FUNCTION(UObject::execArrayConst)
+DEFINE_VM_FUNCTION(UObject::execArrayConst)
 {
 	FProperty* InnerProperty = CastFieldChecked<FProperty>((FField*)Stack.ReadObject());
 	int32 Num = Stack.ReadInt<int32>();
@@ -3752,7 +3764,7 @@ DEFINE_FUNCTION(UObject::execArrayConst)
 }
 IMPLEMENT_VM_FUNCTION(EX_ArrayConst, execArrayConst);
 
-DEFINE_FUNCTION(UObject::execSetConst)
+DEFINE_VM_FUNCTION(UObject::execSetConst)
 {
 	FProperty* InnerProperty = CastFieldChecked<FProperty>((FField*)Stack.ReadObject());
 	int32 Num = Stack.ReadInt<int32>();
@@ -3772,7 +3784,7 @@ DEFINE_FUNCTION(UObject::execSetConst)
 }
 IMPLEMENT_VM_FUNCTION(EX_SetConst, execSetConst);
 
-DEFINE_FUNCTION(UObject::execMapConst)
+DEFINE_VM_FUNCTION(UObject::execMapConst)
 {
 	FProperty* KeyProperty = CastFieldChecked<FProperty>((FField*)Stack.ReadObject());
 	FProperty* ValProperty = CastFieldChecked<FProperty>((FField*)Stack.ReadObject());
@@ -3794,51 +3806,51 @@ DEFINE_FUNCTION(UObject::execMapConst)
 }
 IMPLEMENT_VM_FUNCTION(EX_MapConst, execMapConst);
 
-DEFINE_FUNCTION(UObject::execIntZero)
+DEFINE_VM_FUNCTION(UObject::execIntZero)
 {
 	*(int32*)RESULT_PARAM = 0;
 }
 IMPLEMENT_VM_FUNCTION( EX_IntZero, execIntZero );
 
-DEFINE_FUNCTION(UObject::execIntOne)
+DEFINE_VM_FUNCTION(UObject::execIntOne)
 {
 	*(int32*)RESULT_PARAM = 1;
 }
 IMPLEMENT_VM_FUNCTION( EX_IntOne, execIntOne );
 
-DEFINE_FUNCTION(UObject::execTrue)
+DEFINE_VM_FUNCTION(UObject::execTrue)
 {
 	*(bool*)RESULT_PARAM = true;
 }
 IMPLEMENT_VM_FUNCTION( EX_True, execTrue );
 
-DEFINE_FUNCTION(UObject::execFalse)
+DEFINE_VM_FUNCTION(UObject::execFalse)
 {
 	*(bool*)RESULT_PARAM = false;
 }
 IMPLEMENT_VM_FUNCTION( EX_False, execFalse );
 
-DEFINE_FUNCTION(UObject::execNoObject)
+DEFINE_VM_FUNCTION(UObject::execNoObject)
 {
 	*(UObject**)RESULT_PARAM = NULL;
 }
 IMPLEMENT_VM_FUNCTION( EX_NoObject, execNoObject );
 
-DEFINE_FUNCTION(UObject::execNullInterface)
+DEFINE_VM_FUNCTION(UObject::execNullInterface)
 {
 	FScriptInterface& InterfaceValue = *(FScriptInterface*)RESULT_PARAM;
 	InterfaceValue.SetObject(nullptr);
 }
 IMPLEMENT_VM_FUNCTION( EX_NoInterface, execNullInterface );
 
-DEFINE_FUNCTION(UObject::execIntConstByte)
+DEFINE_VM_FUNCTION(UObject::execIntConstByte)
 {
 	*(int32*)RESULT_PARAM = *Stack.Code++;
 }
 IMPLEMENT_VM_FUNCTION( EX_IntConstByte, execIntConstByte );
 
 
-DEFINE_FUNCTION(UObject::execDynamicCast)
+DEFINE_VM_FUNCTION(UObject::execDynamicCast)
 {
 	// Get "to cast to" class for the dynamic actor class
 	UClass* ClassPtr = (UClass *)Stack.ReadObject();
@@ -3885,7 +3897,7 @@ DEFINE_FUNCTION(UObject::execDynamicCast)
 }
 IMPLEMENT_VM_FUNCTION( EX_DynamicCast, execDynamicCast );
 
-DEFINE_FUNCTION(UObject::execMetaCast)
+DEFINE_VM_FUNCTION(UObject::execMetaCast)
 {
 	UClass* MetaClass = (UClass*)Stack.ReadObject();
 
@@ -3897,20 +3909,20 @@ DEFINE_FUNCTION(UObject::execMetaCast)
 }
 IMPLEMENT_VM_FUNCTION( EX_MetaCast, execMetaCast );
 
-DEFINE_FUNCTION(UObject::execCast)
+DEFINE_VM_FUNCTION(UObject::execCast)
 {
 	int32 B = *(Stack.Code)++;
 	(*GCasts[B])( Stack.Object, Stack, RESULT_PARAM );
 }
 IMPLEMENT_VM_FUNCTION( EX_Cast, execCast );
 
-DEFINE_FUNCTION(UObject::execInterfaceCast)
+DEFINE_VM_FUNCTION(UObject::execInterfaceCast)
 {
 	(*GCasts[CST_ObjectToInterface])(Stack.Object, Stack, RESULT_PARAM);
 }
 IMPLEMENT_VM_FUNCTION( EX_ObjToInterfaceCast, execInterfaceCast );
 
-DEFINE_FUNCTION(UObject::execDoubleToFloatCast)
+DEFINE_VM_FUNCTION(UObject::execDoubleToFloatCast)
 {
 	if (Stack.StepAndCheckMostRecentProperty(Stack.Object, nullptr))
 	{
@@ -3926,7 +3938,7 @@ DEFINE_FUNCTION(UObject::execDoubleToFloatCast)
 }
 IMPLEMENT_CAST_FUNCTION( CST_DoubleToFloat, execDoubleToFloatCast )
 
-DEFINE_FUNCTION(UObject::execFloatToDoubleCast)
+DEFINE_VM_FUNCTION(UObject::execFloatToDoubleCast)
 {
 	if (Stack.StepAndCheckMostRecentProperty(Stack.Object, nullptr))
 	{
@@ -3942,7 +3954,7 @@ DEFINE_FUNCTION(UObject::execFloatToDoubleCast)
 }
 IMPLEMENT_CAST_FUNCTION( CST_FloatToDouble, execFloatToDoubleCast )
 
-DEFINE_FUNCTION(UObject::execObjectToBool)
+DEFINE_VM_FUNCTION(UObject::execObjectToBool)
 {
 	UObject* Obj=NULL;
 	Stack.Step( Stack.Object, &Obj );
@@ -3950,7 +3962,7 @@ DEFINE_FUNCTION(UObject::execObjectToBool)
 }
 IMPLEMENT_CAST_FUNCTION( CST_ObjectToBool, execObjectToBool );
 
-DEFINE_FUNCTION(UObject::execInterfaceToBool)
+DEFINE_VM_FUNCTION(UObject::execInterfaceToBool)
 {
 	FScriptInterface Interface;
 	Stack.Step( Stack.Object, &Interface);
@@ -3958,7 +3970,7 @@ DEFINE_FUNCTION(UObject::execInterfaceToBool)
 }
 IMPLEMENT_CAST_FUNCTION( CST_InterfaceToBool, execInterfaceToBool );
 
-DEFINE_FUNCTION(UObject::execObjectToInterface)
+DEFINE_VM_FUNCTION(UObject::execObjectToInterface)
 {
 	FScriptInterface& InterfaceValue = *(FScriptInterface*)RESULT_PARAM;
 
@@ -3984,7 +3996,7 @@ DEFINE_FUNCTION(UObject::execObjectToInterface)
 }
 IMPLEMENT_CAST_FUNCTION( CST_ObjectToInterface, execObjectToInterface );
 
-DEFINE_FUNCTION(UObject::execInterfaceToInterface)
+DEFINE_VM_FUNCTION(UObject::execInterfaceToInterface)
 {
 	FScriptInterface& CastResult = *(FScriptInterface*)RESULT_PARAM;
 
@@ -4012,7 +4024,7 @@ DEFINE_FUNCTION(UObject::execInterfaceToInterface)
 }
 IMPLEMENT_VM_FUNCTION( EX_CrossInterfaceCast, execInterfaceToInterface );
 
-DEFINE_FUNCTION(UObject::execInterfaceToObject)
+DEFINE_VM_FUNCTION(UObject::execInterfaceToObject)
 {
 	// read the interface class off the stack
 	UClass* ObjClassToCastTo = dynamic_cast<UClass*>(Stack.ReadObject());
