Set correct archive name when cooking

#rnx

diff --git a/Engine/Source/Runtime/Core/Public/Serialization/MemoryWriter.h b/Engine/Source/Runtime/Core/Public/Serialization/MemoryWriter.h
--- a/Engine/Source/Runtime/Core/Public/Serialization/MemoryWriter.h
+++ b/Engine/Source/Runtime/Core/Public/Serialization/MemoryWriter.h
@@ -68,6 +68,11 @@ public:
 	 **/
 	virtual FString GetArchiveName() const override
 	{
+		if (ArchiveName != NAME_None)
+		{
+			return ArchiveName.ToString();
+		}
+
 		if constexpr (IndexSize == 64)
 		{
 			return TEXT("FMemoryWriter64");
diff --git a/Engine/Source/Runtime/CoreUObject/Private/UObject/SavePackage2.cpp b/Engine/Source/Runtime/CoreUObject/Private/UObject/SavePackage2.cpp
--- a/Engine/Source/Runtime/CoreUObject/Private/UObject/SavePackage2.cpp
+++ b/Engine/Source/Runtime/CoreUObject/Private/UObject/SavePackage2.cpp
@@ -1895,7 +1895,7 @@ ESavePackageResult WritePackageTextHeader(FStructuredArchive::FRecord& Structure
 		// Use a separate archive for each export to enable rearranging the exports according to the load/dependency order. Seeks/Tell
 		// during export serialize will be relative to the beginning of the export buffer.
 		ExportBuffer.Reset();
-		FMemoryWriter64 ExportArchive(ExportBuffer);
+		FMemoryWriter64 ExportArchive(ExportBuffer, false, false, SaveContext.GetPackage()->GetFName());
 		TGuardValue<FArchive*> GuardSaver(Linker.Saver, &ExportArchive);
 
 		if (Export.Object->HasAnyFlags(RF_ClassDefaultObject))
