Fix for Clang C++ 20 compiler warnings with StaticClassFlags.

diff --git a/Engine/Source/Runtime/CoreUObject/Public/UObject/ObjectMacros.h b/Engine/Source/Runtime/CoreUObject/Public/UObject/ObjectMacros.h
--- a/Engine/Source/Runtime/CoreUObject/Public/UObject/ObjectMacros.h
+++ b/Engine/Source/Runtime/CoreUObject/Public/UObject/ObjectMacros.h
@@ -1709,7 +1709,7 @@ private: \
 	TRequiredAPI static UClass* GetPrivateStaticClass(); \
 public: \
 	/** Bitwise union of #EClassFlags pertaining to this class.*/ \
-	enum {StaticClassFlags=TStaticFlags}; \
+	static constexpr EClassFlags StaticClassFlags=EClassFlags(TStaticFlags); \
 	/** Typedef for the base class ({{ typedef-type }}) */ \
 	typedef TSuperClass Super;\
 	/** Typedef for {{ typedef-type }}. */ \
@@ -2018,7 +2018,7 @@ private:
 				StaticRegisterNatives##TClass, \
 				sizeof(TClass), \
 				alignof(TClass), \
-				(EClassFlags)TClass::StaticClassFlags, \
+				TClass::StaticClassFlags, \
 				TClass::StaticClassCastFlags(), \
 				TClass::StaticConfigName(), \
 				(UClass::ClassConstructorType)InternalConstructor<TClass>, \
