Made FObjectInstancingGraph::InstancePropertyValue take flags for clarity, and so the DoNotCreateNewInstance flag can be passed through to GetInstancedSubobject.
Removed the unused bIsTransient flag.

diff --git a/Engine/Source/Runtime/CoreUObject/Private/UObject/CoreNative.cpp b/Engine/Source/Runtime/CoreUObject/Private/UObject/CoreNative.cpp
--- a/Engine/Source/Runtime/CoreUObject/Private/UObject/CoreNative.cpp
+++ b/Engine/Source/Runtime/CoreUObject/Private/UObject/CoreNative.cpp
@@ -134,10 +134,13 @@ UObject* FObjectInstancingGraph::GetDestinationObject(UObject* SourceObject)
 	return SourceToDestinationMap.FindRef(SourceObject);
 }
 
-UObject* FObjectInstancingGraph::GetInstancedSubobject( UObject* SourceSubobject, UObject* CurrentValue, UObject* CurrentObject, bool bDoNotCreateNewInstance, bool bAllowSelfReference )
+UObject* FObjectInstancingGraph::GetInstancedSubobject( UObject* SourceSubobject, UObject* CurrentValue, UObject* CurrentObject, EInstancePropertyValueFlags Flags )
 {
 	checkSlow(SourceSubobject);
 
+	bool bDoNotCreateNewInstance = !!(Flags & EInstancePropertyValueFlags::DoNotCreateNewInstance);
+	bool bAllowSelfReference     = !!(Flags & EInstancePropertyValueFlags::AllowSelfReference);
+
 	UObject* InstancedSubobject = INVALID_OBJECT;
 
 	if ( SourceSubobject != nullptr && CurrentValue != nullptr && !CurrentValue->IsIn(CurrentObject))
@@ -200,7 +203,7 @@ UObject* FObjectInstancingGraph::GetInstancedSubobject( UObject* SourceSubobject
 							// outer.  In that case - we need to go ahead and instance that outer.
 							if ( SubobjectOuter == nullptr )
 							{
-								SubobjectOuter = GetInstancedSubobject(SourceSubobject->GetOuter(), SourceSubobject->GetOuter(), CurrentObject, bDoNotCreateNewInstance, bAllowSelfReference);
+								SubobjectOuter = GetInstancedSubobject(SourceSubobject->GetOuter(), SourceSubobject->GetOuter(), CurrentObject, Flags);
 
 								checkf(SubobjectOuter && SubobjectOuter != INVALID_OBJECT, TEXT("No corresponding destination object found for '%s' while attempting to instance subobject '%s'"), *SourceSubobject->GetOuter()->GetFullName(), *SourceSubobject->GetFullName());
 							}
@@ -269,8 +272,11 @@ UObject* FObjectInstancingGraph::GetInstancedSubobject( UObject* SourceSubobject
 }
 
 
-UObject* FObjectInstancingGraph::InstancePropertyValue(UObject* SubObjectTemplate, UObject* CurrentValue, UObject* Owner, bool bIsTransient, bool bCausesInstancing, bool bAllowSelfReference)
+UObject* FObjectInstancingGraph::InstancePropertyValue(UObject* SubObjectTemplate, UObject* CurrentValue, UObject* Owner, EInstancePropertyValueFlags Flags)
 {
+	bool bCausesInstancing   = !!(Flags & EInstancePropertyValueFlags::CausesInstancing);
+	bool bAllowSelfReference = !!(Flags & EInstancePropertyValueFlags::AllowSelfReference);
+
 	UObject* NewValue = CurrentValue;
 
 	check(CurrentValue);
@@ -302,7 +308,7 @@ UObject* FObjectInstancingGraph::InstancePropertyValue(UObject* SubObjectTemplat
 			SubObjectTemplate = CurrentValue;
 		}
 
-		UObject* MaybeNewValue = GetInstancedSubobject(SubObjectTemplate, CurrentValue, Owner, bAllowSelfReference, bAllowSelfReference);
+		UObject* MaybeNewValue = GetInstancedSubobject(SubObjectTemplate, CurrentValue, Owner, Flags);
 		if ( MaybeNewValue != INVALID_OBJECT )
 		{
 			NewValue = MaybeNewValue;
diff --git a/Engine/Source/Runtime/CoreUObject/Private/UObject/PropertyBaseObject.cpp b/Engine/Source/Runtime/CoreUObject/Private/UObject/PropertyBaseObject.cpp
--- a/Engine/Source/Runtime/CoreUObject/Private/UObject/PropertyBaseObject.cpp
+++ b/Engine/Source/Runtime/CoreUObject/Private/UObject/PropertyBaseObject.cpp
@@ -57,7 +57,7 @@ void FObjectPropertyBase::InstanceSubobjects(void* Data, void const* DefaultData
 		if ( CurrentValue )
 		{
 			UObject *SubobjectTemplate = DefaultData ? GetObjectPropertyValue((uint8*)DefaultData + ArrayIndex * ElementSize): nullptr;
-			UObject* NewValue = InstanceGraph->InstancePropertyValue(SubobjectTemplate, CurrentValue, InOwner, HasAnyPropertyFlags(CPF_Transient), HasAnyPropertyFlags(CPF_InstancedReference));
+			UObject* NewValue = InstanceGraph->InstancePropertyValue(SubobjectTemplate, CurrentValue, InOwner, HasAnyPropertyFlags(CPF_InstancedReference) ? EInstancePropertyValueFlags::CausesInstancing : EInstancePropertyValueFlags::None);
 			SetObjectPropertyValue((uint8*)Data + ArrayIndex * ElementSize, NewValue);
 		}
 	}
diff --git a/Engine/Source/Runtime/CoreUObject/Private/UObject/PropertyDelegate.cpp b/Engine/Source/Runtime/CoreUObject/Private/UObject/PropertyDelegate.cpp
--- a/Engine/Source/Runtime/CoreUObject/Private/UObject/PropertyDelegate.cpp
+++ b/Engine/Source/Runtime/CoreUObject/Private/UObject/PropertyDelegate.cpp
@@ -54,7 +54,7 @@ void FDelegateProperty::InstanceSubobjects(void* Data, void const* DefaultData,
 				Template = DefaultDelegate.GetUObject();
 			}
 
-			UObject* NewUObject = InstanceGraph->InstancePropertyValue(Template, CurrentUObject, InOwner, HasAnyPropertyFlags(CPF_Transient), false, true);
+			UObject* NewUObject = InstanceGraph->InstancePropertyValue(Template, CurrentUObject, InOwner, EInstancePropertyValueFlags::AllowSelfReference | EInstancePropertyValueFlags::DoNotCreateNewInstance);
 			DestDelegate.BindUFunction(NewUObject, DestDelegate.GetFunctionName());
 		}
 	}
diff --git a/Engine/Source/Runtime/CoreUObject/Private/UObject/PropertyMulticastDelegate.cpp b/Engine/Source/Runtime/CoreUObject/Private/UObject/PropertyMulticastDelegate.cpp
--- a/Engine/Source/Runtime/CoreUObject/Private/UObject/PropertyMulticastDelegate.cpp
+++ b/Engine/Source/Runtime/CoreUObject/Private/UObject/PropertyMulticastDelegate.cpp
@@ -54,7 +54,7 @@ void FMulticastDelegateProperty::InstanceSubobjects(void* Data, void const* Defa
 				{
 					FScriptDelegate& DefaultDelegateInvocation = *DefaultInvocation;
 					UObject *Template = DefaultDelegateInvocation.GetUObject();
-					UObject* NewUObject = InstanceGraph->InstancePropertyValue(Template, CurrentUObject, InOwner, HasAnyPropertyFlags(CPF_Transient), false, true);
+					UObject* NewUObject = InstanceGraph->InstancePropertyValue(Template, CurrentUObject, InOwner, EInstancePropertyValueFlags::AllowSelfReference | EInstancePropertyValueFlags::DoNotCreateNewInstance);
 					DestDelegateInvocation.BindUFunction(NewUObject, DestDelegateInvocation.GetFunctionName());
 				}
 			}
@@ -66,7 +66,7 @@ void FMulticastDelegateProperty::InstanceSubobjects(void* Data, void const* Defa
 
 				if (CurrentUObject)
 				{
-					UObject* NewUObject = InstanceGraph->InstancePropertyValue(NULL, CurrentUObject, InOwner, HasAnyPropertyFlags(CPF_Transient), false, true);
+					UObject* NewUObject = InstanceGraph->InstancePropertyValue(NULL, CurrentUObject, InOwner, EInstancePropertyValueFlags::AllowSelfReference | EInstancePropertyValueFlags::DoNotCreateNewInstance);
 					DestDelegateInvocation.BindUFunction(NewUObject, DestDelegateInvocation.GetFunctionName());
 				}
 			}
@@ -83,7 +83,7 @@ void FMulticastDelegateProperty::InstanceSubobjects(void* Data, void const* Defa
 
 				if (CurrentUObject)
 				{
-					UObject* NewUObject = InstanceGraph->InstancePropertyValue(NULL, CurrentUObject, InOwner, HasAnyPropertyFlags(CPF_Transient), false, true);
+					UObject* NewUObject = InstanceGraph->InstancePropertyValue(NULL, CurrentUObject, InOwner, EInstancePropertyValueFlags::AllowSelfReference | EInstancePropertyValueFlags::DoNotCreateNewInstance);
 					DestDelegateInvocation.BindUFunction(NewUObject, DestDelegateInvocation.GetFunctionName());
 				}
 			}
diff --git a/Engine/Source/Runtime/CoreUObject/Public/UObject/Class.h b/Engine/Source/Runtime/CoreUObject/Public/UObject/Class.h
--- a/Engine/Source/Runtime/CoreUObject/Public/UObject/Class.h
+++ b/Engine/Source/Runtime/CoreUObject/Public/UObject/Class.h
@@ -3512,6 +3512,22 @@ COREUOBJECT_API void GetPrivateStaticClassBody(
 	FObjectInstancingGraph.
 -----------------------------------------------------------------------------*/
 
+enum class EInstancePropertyValueFlags
+{
+	None                   = 0x00,
+
+	// if set, then this property causes an instance to be created, otherwise this is just a pointer to a uobject that should be remapped if the object is instanced for some other property
+	CausesInstancing       = 0x01,
+
+	// if set, instance the reference to the subobjectroot, so far only delegates remap a self reference
+	AllowSelfReference     = 0x02,
+
+	// if set, then we do not create a new instance, but we will reassign one if there is already a mapping in the table
+	DoNotCreateNewInstance = 0x04
+};
+
+ENUM_CLASS_FLAGS(EInstancePropertyValueFlags)
+
 struct COREUOBJECT_API FObjectInstancingGraph
 {
 public:
@@ -3575,7 +3591,35 @@ public:
 	 *
 	 * @return	As with GetInstancedSubobject, above, but also deals with archetype creation and a few other special cases
 	 */
-	class UObject* InstancePropertyValue( class UObject* SourceComponent, class UObject* CurrentValue, class UObject* CurrentObject, bool bIsTransient, bool bCausesInstancing = false, bool bAllowSelfReference = false );
+	UE_DEPRECATED(5.2, "This overload of InstancePropertyValue has been deprecated, please call the overload that takes flags instead.")
+	class UObject* InstancePropertyValue( class UObject* SourceComponent, class UObject* CurrentValue, class UObject* CurrentObject, bool bIsTransient, bool bCausesInstancing = false, bool bAllowSelfReference = false )
+	{
+		EInstancePropertyValueFlags Flags = EInstancePropertyValueFlags::None;
+		if (bCausesInstancing)
+		{
+			Flags |= EInstancePropertyValueFlags::CausesInstancing;
+		}
+		if (bAllowSelfReference)
+		{
+			Flags |= EInstancePropertyValueFlags::AllowSelfReference;
+		}
+		return InstancePropertyValue(SourceComponent, CurrentValue, CurrentObject, Flags);
+	}
+
+	/**
+	 * Returns the component that has SourceComponent as its archetype, instancing the component as necessary.
+	 *
+	 * @param	SourceComponent		the component to find the corresponding component instance for
+	 * @param	CurrentValue		the component currently assigned as the value for the component property
+	 *								being instanced.  Used when updating archetypes to ensure that the new instanced component
+	 *								replaces the existing component instance in memory.
+	 * @param	CurrentObject		the object that owns the component property currently being instanced;  this is NOT necessarily the object
+	 *								that should be the Outer for the new component.
+	 * @param	Flags				reinstancing flags - see EInstancePropertyValueFlags
+	 *
+	 * @return	As with GetInstancedSubobject, above, but also deals with archetype creation and a few other special cases
+	 */
+	class UObject* InstancePropertyValue( class UObject* SourceComponent, class UObject* CurrentValue, class UObject* CurrentObject, EInstancePropertyValueFlags Flags = EInstancePropertyValueFlags::None );
 
 	/**
 	 * Adds a partially built object instance to the map(s) of source objects to their instances.
@@ -3686,8 +3730,7 @@ private:
 	 *								replaces the existing component instance in memory.
 	 * @param	CurrentObject		the object that owns the component property currently being instanced;  this is NOT necessarily the object
 	 *								that should be the Outer for the new component.
-	 * @param	bDoNotCreateNewInstance If true, then we do not create a new instance, but we will reassign one if there is already a mapping in the table
-	 * @param	bAllowSelfReference If true, instance the reference to the subobjectroot, so far only delegates remap a self reference
+	 * @param	Flags				reinstancing flags - see EInstancePropertyValueFlags
 	 *
 	 * @return	if SourceComponent is contained within SourceRoot, returns a pointer to a unique component instance corresponding to
 	 *			SourceComponent if SourceComponent is allowed to be instanced in this context, or NULL if the component isn't allowed to be
@@ -3695,7 +3738,7 @@ private:
 	 *			if SourceComponent is not contained by SourceRoot, return INVALID_OBJECT, indicating that the that has SourceComponent as its ObjectArchetype, or NULL if SourceComponent is not contained within
 	 *			SourceRoot.
 	 */
-	class UObject* GetInstancedSubobject( class UObject* SourceSubobject, class UObject* CurrentValue, class UObject* CurrentObject, bool bDoNotCreateNewInstance, bool bAllowSelfReference );
+	class UObject* GetInstancedSubobject( class UObject* SourceSubobject, class UObject* CurrentValue, class UObject* CurrentObject, EInstancePropertyValueFlags Flags );
 
 	/**
 	 * The root of the object tree that is the source used for instancing components;
