Change ArchetypeAnnotation to use FUObjectAnnotationChunked instead of FUObjectAnnotationDense (originally implemented in 18810065).
Reduces memory from 3.2 MB to 2.2 MB for an internal project with > 500 000 objects, and prevents large and slow reallocs.
Note that the object archetype may already be cached if GetArchetype is called before CacheArchetypeForObject for an object where the outer does not have RF_NeedLoad set (e.g. UPackages), and this causes the same annotation to be added twice.
Make FUObjectAnnotationChunked::AddAnnotation behavior consistent with FUObjectAnnotationDense by removing the existence check and allowing overwrites.

#rnx

diff --git a/Engine/Source/Runtime/CoreUObject/Private/UObject/UObjectArchetype.cpp b/Engine/Source/Runtime/CoreUObject/Private/UObject/UObjectArchetype.cpp
--- a/Engine/Source/Runtime/CoreUObject/Private/UObject/UObjectArchetype.cpp
+++ b/Engine/Source/Runtime/CoreUObject/Private/UObject/UObjectArchetype.cpp
@@ -48,7 +48,7 @@ struct FArchetypeInfo
 	int32 ArchetypeIndex;
 };
 
-static FUObjectAnnotationDense<FArchetypeInfo, true> ArchetypeAnnotation;
+static FUObjectAnnotationChunked<FArchetypeInfo, true, 8192> ArchetypeAnnotation;
 
 #endif // UE_CACHE_ARCHETYPE
 
diff --git a/Engine/Source/Runtime/CoreUObject/Public/UObject/UObjectAnnotation.h b/Engine/Source/Runtime/CoreUObject/Public/UObject/UObjectAnnotation.h
--- a/Engine/Source/Runtime/CoreUObject/Public/UObject/UObjectAnnotation.h
+++ b/Engine/Source/Runtime/CoreUObject/Public/UObject/UObjectAnnotation.h
@@ -567,10 +567,12 @@ class FUObjectAnnotationChunked : public FUObjectArray::FUObjectDeleteListener
 			CurrentAllocatedMemory += NumAnnotationsPerChunk * sizeof(TAnnotation);
 			MaxAllocatedMemory = FMath::Max(CurrentAllocatedMemory, MaxAllocatedMemory);
 		}
-		check(Chunk.Items[WithinChunkIndex].IsDefault());
-		Chunk.Num++;
-		check(Chunk.Num <= NumAnnotationsPerChunk);
-		NumAnnotations++;
+		if (Chunk.Items[WithinChunkIndex].IsDefault())
+		{
+			Chunk.Num++;
+			check(Chunk.Num <= NumAnnotationsPerChunk);
+			NumAnnotations++;
+		}
 
 		return Chunk.Items[WithinChunkIndex];
 	}
