FArchiveCountMem
- Made a couple of simple getter member functions const
- Added missing virtual and override specifiers

diff --git a/Engine/Source/Runtime/CoreUObject/Public/Serialization/ArchiveCountMem.h b/Engine/Source/Runtime/CoreUObject/Public/Serialization/ArchiveCountMem.h
--- a/Engine/Source/Runtime/CoreUObject/Public/Serialization/ArchiveCountMem.h
+++ b/Engine/Source/Runtime/CoreUObject/Public/Serialization/ArchiveCountMem.h
@@ -13,8 +13,8 @@ class FArchiveCountMem : public FArchiveUObject
 {
 public:
 	FArchiveCountMem( UObject* Src, bool bFilterEditorOnly = false )
-	:	Num(0)
-	,	Max(0)
+		: Num(0)
+		, Max(0)
 	{
 		ArIsCountingMemory = true;
 		ArIsFilterEditorOnly = bFilterEditorOnly;
@@ -23,15 +23,16 @@ public:
 			Src->Serialize( *this );
 		}
 	}
-	SIZE_T GetNum()
+	SIZE_T GetNum() const
 	{
 		return Num;
 	}
-	SIZE_T GetMax()
+	SIZE_T GetMax() const
 	{
 		return Max;
 	}
-	void CountBytes( SIZE_T InNum, SIZE_T InMax )
+
+	virtual void CountBytes( SIZE_T InNum, SIZE_T InMax ) override
 	{
 		Num += InNum;
 		Max += InMax;
@@ -42,7 +43,7 @@ public:
 	 *
 	 * This is overridden for the specific Archive Types
 	 **/
-	virtual FString GetArchiveName() const { return TEXT("FArchiveCountMem"); }
+	virtual FString GetArchiveName() const override { return TEXT("FArchiveCountMem"); }
 
 protected:
 	SIZE_T Num, Max;
