Fix UBSAN detection of underaligned access

The involved types don't need the alignment, the intent is just to put things on separate cache lines. Instead of passing the correct alignment when allocating Contexts I switched to manual padding, since MB2/3 fall back to expensive external allocations rather than using small block bins for >16 alignment.

#[fyi] andrew.scheidecker
#rnx

diff --git a/Engine/Source/Runtime/CoreUObject/Private/UObject/GarbageCollection.cpp b/Engine/Source/Runtime/CoreUObject/Private/UObject/GarbageCollection.cpp
--- a/Engine/Source/Runtime/CoreUObject/Private/UObject/GarbageCollection.cpp
+++ b/Engine/Source/Runtime/CoreUObject/Private/UObject/GarbageCollection.cpp
@@ -5661,14 +5661,22 @@ void DumpMemoryStats(FOutputDevice& OutputDevice)
 // Coordinates worker starting, tail spinning and stopping
 class FWorkCoordinator
 {
-	struct alignas(PLATFORM_CACHE_LINE_SIZE) FStealableContext : public std::atomic<FWorkerContext*> {};
+	class FStealableContext : public std::atomic<FWorkerContext*>
+	{
+		// Put each atomic on its own cache line to avoid false sharing
+		char CacheLinePadding[PLATFORM_CACHE_LINE_SIZE - sizeof(std::atomic<FWorkerContext*>)];
+	};
+	static_assert(sizeof(FStealableContext) == PLATFORM_CACHE_LINE_SIZE);
+
 	const TArrayView<FStealableContext> Contexts;
-	alignas(PLATFORM_CACHE_LINE_SIZE) std::atomic<int32> NumUsedContexts {0};
+	// Ensure the immutable Contexts is on a different cache line than NumUsedContexts 
+	char CacheLinePadding[PLATFORM_CACHE_LINE_SIZE - sizeof(Contexts)]; 
+
+	std::atomic<int32> NumUsedContexts {0};
 	std::atomic<int32> NumWorkless{0}; // Number of workless workers spinning to steal work
 	std::atomic<int32> NumStopDirectly; // Stop a few workless workers directly to pick up non-GC tasks
 	std::atomic<int32> NumStopped{0}; // Synchronize when all workers are done
 
-
 public:
 	FWorkCoordinator(TArrayView<FWorkerContext*> InContexts, int32 NumTaskgraphWorkers)
 	: Contexts(reinterpret_cast<FStealableContext*>(FMemory::Malloc(sizeof(FStealableContext) * InContexts.Num())), InContexts.Num())
