Fixed discrepancy between if (ScriptInterface) and if (ScriptInterface != nullptr) when we have a non-null object with a null interface.

diff --git a/Engine/Source/Runtime/CoreUObject/Public/UObject/ScriptInterface.h b/Engine/Source/Runtime/CoreUObject/Public/UObject/ScriptInterface.h
--- a/Engine/Source/Runtime/CoreUObject/Public/UObject/ScriptInterface.h
+++ b/Engine/Source/Runtime/CoreUObject/Public/UObject/ScriptInterface.h
@@ -237,6 +237,18 @@ public:
 		return GetInterface() != Other.GetInterface() || GetObject() != Other.GetObject();
 	}
 
+	/**
+	 * Comparison operator, taking a nullptr
+	 */
+	FORCEINLINE bool operator==(TYPE_OF_NULLPTR) const
+	{
+		return GetInterface() == nullptr;
+	}
+	FORCEINLINE bool operator!=(TYPE_OF_NULLPTR) const
+	{
+		return GetInterface() != nullptr;
+	}
+
 	/**
 	 * Member access operator.  Provides transparent access to the interface pointer contained by this TScriptInterface
 	 */
