Added a csv scope for garbage reference tracking to not pollute GC stats

diff --git a/Engine/Source/Runtime/CoreUObject/Private/UObject/GarbageCollection.cpp b/Engine/Source/Runtime/CoreUObject/Private/UObject/GarbageCollection.cpp
--- a/Engine/Source/Runtime/CoreUObject/Private/UObject/GarbageCollection.cpp
+++ b/Engine/Source/Runtime/CoreUObject/Private/UObject/GarbageCollection.cpp
@@ -4137,6 +4137,7 @@ static void CollectGarbageInternal(EObjectFlags KeepFlags, bool bPerformFullPurg
 
 			if (GRedoReachabilityToTrackGarbage && GGarbageReferenceTrackingEnabled > 0)
 			{
+				CSV_SCOPED_TIMING_STAT_EXCLUSIVE(GarbageCollection_Debug);
 				SCOPED_NAMED_EVENT(FRealtimeGC_PerformReachabilityAnalysisRerun, FColor::Orange);
 				DECLARE_SCOPE_CYCLE_COUNTER(TEXT("FRealtimeGC::PerformReachabilityAnalysisRerun"), STAT_FArchiveRealtimeGC_PerformReachabilityAnalysisRerun, STATGROUP_GC);
 				const double StartTime = FPlatformTime::Seconds();
