Add a way to poll if a FPackageTrailer is valid or not.

#rnx

diff --git a/Engine/Source/Runtime/CoreUObject/Public/UObject/PackageTrailer.h b/Engine/Source/Runtime/CoreUObject/Public/UObject/PackageTrailer.h
--- a/Engine/Source/Runtime/CoreUObject/Public/UObject/PackageTrailer.h
+++ b/Engine/Source/Runtime/CoreUObject/Public/UObject/PackageTrailer.h
@@ -431,6 +431,15 @@ public:
 	FPackageTrailer(FPackageTrailer&& Other) = default;
 	FPackageTrailer& operator=(FPackageTrailer&& Other) = default;
 
+	/** 
+	 * Returns true if the trailer contains actual data from a package file and false if it just contains the defaults of an unloaded
+	 * trailer.
+	 */
+	[[nodiscard]] bool IsValid() const
+	{
+		return Header.Tag == FHeader::HeaderTag;
+	}
+
 	/** 
 	 * Serializes the trailer from the given archive assuming that the seek position of the archive is already at the correct position
 	 * for the trailer.
