Put engine version back to 5.2.0 and provide fixup code for assets saved incorrectly with 5.2.1
#rnx

diff --git a/Engine/Source/Runtime/CoreUObject/Private/UObject/PackageFileSummary.cpp b/Engine/Source/Runtime/CoreUObject/Private/UObject/PackageFileSummary.cpp
--- a/Engine/Source/Runtime/CoreUObject/Private/UObject/PackageFileSummary.cpp
+++ b/Engine/Source/Runtime/CoreUObject/Private/UObject/PackageFileSummary.cpp
@@ -4,6 +4,7 @@
 #include "UObject/Linker.h"
 #include "Serialization/StructuredArchive.h"
 #include "UObject/UObjectGlobals.h"
+#include "Runtime/Launch/Resources/Version.h"
 
 FPackageFileSummary::FPackageFileSummary()
 {
@@ -43,6 +44,17 @@ static void FixCorruptEngineVersion(const FPackageFileVersion& ObjectVersion, FE
 	{
 		Version.Set(4, 26, 0, Version.GetChangelist(), Version.GetBranch());
 	}
+
+	// Temporarily support the loading of 5.2.1 assets in 5.2.0 as version.h was incorrectly merged back from Release-5.1 when 5.1.1 was set. 
+	// This must be removed from Release-5.2 when the engine version is bumped to 5.2.1 and can be removed from UE5/Main when it is bumped to 5.3.0
+	// Also remove version.h #include when removing this block
+	if (Version.GetMajor() == 5
+		&& Version.GetMinor() == 2
+		&& Version.GetPatch() == 1)
+	{
+		static_assert(ENGINE_MAJOR_VERSION == 5 && ENGINE_MINOR_VERSION == 2 && ENGINE_PATCH_VERSION == 0);
+		Version.Set(5, 2, 0, Version.GetChangelist(), Version.GetBranch());
+	}
 }
 
 void operator<<(FStructuredArchive::FSlot Slot, FPackageFileSummary& Sum)
diff --git a/Engine/Source/Runtime/Launch/Resources/Version.h b/Engine/Source/Runtime/Launch/Resources/Version.h
--- a/Engine/Source/Runtime/Launch/Resources/Version.h
+++ b/Engine/Source/Runtime/Launch/Resources/Version.h
@@ -56,7 +56,7 @@
 // newer than a 5.11.* version, regardless of the changelist that it was built with)
 #define ENGINE_MAJOR_VERSION	5
 #define ENGINE_MINOR_VERSION	2
-#define ENGINE_PATCH_VERSION	1
+#define ENGINE_PATCH_VERSION	0
 
 // Macros for encoding strings
 #define VERSION_TEXT(x) TEXT(x)
