Fix spurious warning
#rnx

diff --git a/Engine/Source/Runtime/CoreUObject/Private/UObject/PackageFileSummary.cpp b/Engine/Source/Runtime/CoreUObject/Private/UObject/PackageFileSummary.cpp
--- a/Engine/Source/Runtime/CoreUObject/Private/UObject/PackageFileSummary.cpp
+++ b/Engine/Source/Runtime/CoreUObject/Private/UObject/PackageFileSummary.cpp
@@ -52,7 +52,7 @@ static void FixCorruptEngineVersion(const FPackageFileVersion& ObjectVersion, FE
 		&& Version.GetMinor() == 2
 		&& Version.GetPatch() == 1)
 	{
-		static_assert(ENGINE_MAJOR_VERSION == 5 && ENGINE_MINOR_VERSION == 2 && ENGINE_PATCH_VERSION == 0);
+		static_assert(ENGINE_MAJOR_VERSION == 5 && ENGINE_MINOR_VERSION == 2 && ENGINE_PATCH_VERSION == 0); //-V501
 		Version.Set(5, 2, 0, Version.GetChangelist(), Version.GetBranch());
 	}
 }
