Bulkdata:
Fixed issue with trying to stream bulkdata from utxt files
Fixed issue with not marking io requests as failed in the package resource backend in some cases

diff --git a/Engine/Source/Runtime/CoreUObject/Private/UObject/PackageResourceIoDispatcherBackend.cpp b/Engine/Source/Runtime/CoreUObject/Private/UObject/PackageResourceIoDispatcherBackend.cpp
--- a/Engine/Source/Runtime/CoreUObject/Private/UObject/PackageResourceIoDispatcherBackend.cpp
+++ b/Engine/Source/Runtime/CoreUObject/Private/UObject/PackageResourceIoDispatcherBackend.cpp
@@ -237,7 +237,7 @@ bool FPackageResourceIoBackend::Resolve(FIoRequestImpl* Request)
 
 				Request->SetResult(Buffer);
 			}
-			else if (bWasCancelled == false)
+			else
 			{
 				Request->SetFailed();
 			}
