FReferenceFinder: Optimize searching of objects that have 100k references.
#rn Minor, GarbageCollection

diff --git a/Engine/Source/Runtime/CoreUObject/Private/UObject/GarbageCollection.cpp b/Engine/Source/Runtime/CoreUObject/Private/UObject/GarbageCollection.cpp
--- a/Engine/Source/Runtime/CoreUObject/Private/UObject/GarbageCollection.cpp
+++ b/Engine/Source/Runtime/CoreUObject/Private/UObject/GarbageCollection.cpp
@@ -2869,6 +2869,7 @@ FReferenceFinder::FReferenceFinder(TArray<UObject*>& InObjectArray, UObject* InO
 	, bSerializeRecursively(false)
 	, bShouldIgnoreTransient(bInShouldIgnoreTransient)
 {
+	ObjectSet.Append(MakeArrayView(const_cast<const UObject**>(ObjectArray.GetData()), ObjectArray.Num()));
 	bSerializeRecursively = bInSerializeRecursively && LimitOuter != NULL;
 	if (InOuter)
 	{
@@ -2901,18 +2902,19 @@ void FReferenceFinder::HandleObjectReference( UObject*& InObject, const UObject*
 		{
 			// Many places that use FReferenceFinder expect the object to not be const.
 			UObject* Object = const_cast<UObject*>(InObject);
-			// do not attempt to serialize objects that have already been 
-			if ( ObjectArray.Contains( Object ) == false )
+			// do not add or recursively serialize objects that have already been added
+			bool bAlreadyExists;
+			ObjectSet.Add(Object, &bAlreadyExists);
+			if (!bAlreadyExists)
 			{
-				check( Object->IsValidLowLevel() );
-				ObjectArray.Add( Object );
-			}
+				check(Object->IsValidLowLevel());
+				ObjectArray.Add(Object);
 
-			// check this object for any potential object references
-			if ( bSerializeRecursively == true && !SerializedObjects.Find(Object) )
-			{
-				SerializedObjects.Add(Object);
-				FindReferences(Object, const_cast<UObject*>(InReferencingObject), const_cast<FProperty*>(InReferencingProperty));
+				// check this object for any potential object references
+				if (bSerializeRecursively)
+				{
+					FindReferences(Object, const_cast<UObject*>(InReferencingObject), const_cast<FProperty*>(InReferencingProperty));
+				}
 			}
 		}
 	}
diff --git a/Engine/Source/Runtime/CoreUObject/Public/UObject/UObjectGlobals.h b/Engine/Source/Runtime/CoreUObject/Public/UObject/UObjectGlobals.h
--- a/Engine/Source/Runtime/CoreUObject/Public/UObject/UObjectGlobals.h
+++ b/Engine/Source/Runtime/CoreUObject/Public/UObject/UObjectGlobals.h
@@ -2682,8 +2682,8 @@ protected:
 
 	/** Stored reference to array of objects we add object references to. */
 	TArray<UObject*>&		ObjectArray;
-	/** List of objects that have been recursively serialized. */
-	TSet<const UObject*>	SerializedObjects;
+	/** Set that duplicates ObjectArray. Keeps ObjectArray unique and avoids duplicate recursive serializing. */
+	TSet<const UObject*>	ObjectSet;
 	/** Only objects within this outer will be considered, nullptr value indicates that outers are disregarded. */
 	UObject*		LimitOuter;
 	/** Property that is referencing the current object */
