23728823: Fix OOB crash in FUObjectAnnotationChunked::FreeAnnotation when deleting a UObject without an annotation but having an internal index that is higher than for the currently max annotated UObject.
Existing bug exposed by 23643652.

#rnx

diff --git a/Engine/Source/Runtime/CoreUObject/Public/UObject/UObjectAnnotation.h b/Engine/Source/Runtime/CoreUObject/Public/UObject/UObjectAnnotation.h
--- a/Engine/Source/Runtime/CoreUObject/Public/UObject/UObjectAnnotation.h
+++ b/Engine/Source/Runtime/CoreUObject/Public/UObject/UObjectAnnotation.h
@@ -585,26 +585,35 @@ class FUObjectAnnotationChunked : public FUObjectArray::FUObjectDeleteListener
 		const int32 ChunkIndex = Index / NumAnnotationsPerChunk;
 		const int32 WithinChunkIndex = Index % NumAnnotationsPerChunk;
 
+		if (ChunkIndex >= Chunks.Num())
+		{
+			return;
+		}
+
 		TAnnotationChunk& Chunk = Chunks[ChunkIndex];
-		if (Chunk.Items != nullptr)
+		if (!Chunk.Items)
 		{
-			if (!Chunk.Items[WithinChunkIndex].IsDefault())
-			{
-				Chunk.Items[WithinChunkIndex] = TAnnotation();
-				Chunk.Num--;
-				check(Chunk.Num >= 0);
-				if (Chunk.Num == 0)
-				{
-					delete[] Chunk.Items;
-					Chunk.Items = nullptr;
-					const uint32 ChunkMemory = NumAnnotationsPerChunk * sizeof(TAnnotation);
-					check(CurrentAllocatedMemory >= ChunkMemory);
-					CurrentAllocatedMemory -= ChunkMemory;
-				}
-				NumAnnotations--;
-			}
-			check(NumAnnotations >= 0);
+			return;
+		}
+
+		if (Chunk.Items[WithinChunkIndex].IsDefault())
+		{
+			return;
+		}
+
+		Chunk.Items[WithinChunkIndex] = TAnnotation();
+		Chunk.Num--;
+		check(Chunk.Num >= 0);
+		if (Chunk.Num == 0)
+		{
+			delete[] Chunk.Items;
+			Chunk.Items = nullptr;
+			const uint32 ChunkMemory = NumAnnotationsPerChunk * sizeof(TAnnotation);
+			check(CurrentAllocatedMemory >= ChunkMemory);
+			CurrentAllocatedMemory -= ChunkMemory;
 		}
+		NumAnnotations--;
+		check(NumAnnotations >= 0);
 	}
 
 	/**
