code cleanup of ObjectHandle
made ObjectPathId to private
trimmed down the public API as much as possible

https://p4-swarm.epicgames.net/reviews/23658342

diff --git a/Engine/Source/Runtime/Core/Private/Misc/CoreGlobals.cpp b/Engine/Source/Runtime/Core/Private/Misc/CoreGlobals.cpp
--- a/Engine/Source/Runtime/Core/Private/Misc/CoreGlobals.cpp
+++ b/Engine/Source/Runtime/Core/Private/Misc/CoreGlobals.cpp
@@ -242,13 +242,18 @@ bool					GExitPurge						= false;
 
 FChunkedFixedUObjectArray* GCoreObjectArrayForDebugVisualizers = nullptr;
 
-UE::ObjectPath::Private::FStoredObjectPath* GCoreComplexObjectPathDebug = nullptr;
-FObjectHandlePackageDebugData* GCoreObjectHandlePackageDebug = nullptr;
+namespace UE::CoreUObject::Private
+{
+	struct FStoredObjectPath;
+	struct FObjectHandlePackageDebugData;
+}
+UE::CoreUObject::Private::FStoredObjectPath* GCoreComplexObjectPathDebug = nullptr;
+UE::CoreUObject::Private::FObjectHandlePackageDebugData* GCoreObjectHandlePackageDebug = nullptr;
 #if PLATFORM_UNIX
 uint8** CORE_API GNameBlocksDebug = FNameDebugVisualizer::GetBlocks();
 FChunkedFixedUObjectArray*& CORE_API GObjectArrayForDebugVisualizers = GCoreObjectArrayForDebugVisualizers;
-UE::ObjectPath::Private::FStoredObjectPath*& GComplexObjectPathDebug = GCoreComplexObjectPathDebug;
-FObjectHandlePackageDebugData*& CORE_API GObjectHandlePackageDebug = GCoreObjectHandlePackageDebug;
+UE::CoreUObject::Private::FStoredObjectPath*& GComplexObjectPathDebug = GCoreComplexObjectPathDebug;
+UE::CoreUObject::Private::FObjectHandlePackageDebugData*& CORE_API GObjectHandlePackageDebug = GCoreObjectHandlePackageDebug;
 #endif
 
 /** Game name, used for base game directory and ini among other things										*/
diff --git a/Engine/Source/Runtime/Core/Public/CoreGlobals.h b/Engine/Source/Runtime/Core/Public/CoreGlobals.h
--- a/Engine/Source/Runtime/Core/Public/CoreGlobals.h
+++ b/Engine/Source/Runtime/Core/Public/CoreGlobals.h
@@ -758,17 +758,16 @@ class FChunkedFixedUObjectArray;
 
 extern CORE_API FChunkedFixedUObjectArray* GCoreObjectArrayForDebugVisualizers;
 
-/** Array to help visualize object paths in the debugger */
-namespace UE::ObjectPath::Private
+namespace UE::CoreUObject::Private
 {
+	/** Array to help visualize object paths in the debugger */
 	struct FStoredObjectPath;
+	
+	/** Array to help visualize object handles in the debugger */
+	struct FObjectHandlePackageDebugData;
 }
-extern CORE_API UE::ObjectPath::Private::FStoredObjectPath* GCoreComplexObjectPathDebug;
-
-/** Array to help visualize object handles in the debugger */
-struct FObjectHandlePackageDebugData;
-
-extern CORE_API FObjectHandlePackageDebugData* GCoreObjectHandlePackageDebug;
+extern CORE_API UE::CoreUObject::Private::FStoredObjectPath* GCoreComplexObjectPathDebug;
+extern CORE_API UE::CoreUObject::Private::FObjectHandlePackageDebugData* GCoreObjectHandlePackageDebug;
 
 /** @return True if running cook-on-the-fly. */
 bool CORE_API IsRunningCookOnTheFly();
diff --git a/Engine/Source/Runtime/Core/Public/Modules/Boilerplate/ModuleBoilerplate.h b/Engine/Source/Runtime/Core/Public/Modules/Boilerplate/ModuleBoilerplate.h
--- a/Engine/Source/Runtime/Core/Public/Modules/Boilerplate/ModuleBoilerplate.h
+++ b/Engine/Source/Runtime/Core/Public/Modules/Boilerplate/ModuleBoilerplate.h
@@ -70,8 +70,8 @@ class FChunkedFixedUObjectArray;
 	#define UE4_VISUALIZERS_HELPERS \
 		uint8** GNameBlocksDebug = FNameDebugVisualizer::GetBlocks(); \
 		FChunkedFixedUObjectArray*& GObjectArrayForDebugVisualizers = GCoreObjectArrayForDebugVisualizers; \
-		UE::ObjectPath::Private::FStoredObjectPath*& GComplexObjectPathDebug = GCoreComplexObjectPathDebug; \
-		FObjectHandlePackageDebugData*& GObjectHandlePackageDebug = GCoreObjectHandlePackageDebug;
+		UE::CoreUObject::Private::FStoredObjectPath*& GComplexObjectPathDebug = GCoreComplexObjectPathDebug; \
+		UE::CoreUObject::Private::FObjectHandlePackageDebugData*& GObjectHandlePackageDebug = GCoreObjectHandlePackageDebug;
 #endif
 
 // in DLL builds, these are done per-module, otherwise we just need one in the application
diff --git a/Engine/Source/Runtime/CoreUObject/Private/Serialization/DuplicateDataReader.cpp b/Engine/Source/Runtime/CoreUObject/Private/Serialization/DuplicateDataReader.cpp
--- a/Engine/Source/Runtime/CoreUObject/Private/Serialization/DuplicateDataReader.cpp
+++ b/Engine/Source/Runtime/CoreUObject/Private/Serialization/DuplicateDataReader.cpp
@@ -83,7 +83,7 @@ FArchive& FDuplicateDataReader::operator<<(FObjectPtr& ObjectPtr)
 
 	if (!IsObjectHandleNull(Handle) && IsObjectHandleResolved(Handle))
 	{
-		UObject* SourceObject = ReadObjectHandlePointerNoCheck(Handle);
+		UObject* SourceObject = UE::CoreUObject::Private::ReadObjectHandlePointerNoCheck(Handle);
 		FDuplicatedObject ObjectInfo = DuplicatedObjectAnnotation.GetAnnotation(SourceObject);
 		if (!ObjectInfo.IsDefault())
 		{
diff --git a/Engine/Source/Runtime/CoreUObject/Private/UObject/LinkerLoad.cpp b/Engine/Source/Runtime/CoreUObject/Private/UObject/LinkerLoad.cpp
--- a/Engine/Source/Runtime/CoreUObject/Private/UObject/LinkerLoad.cpp
+++ b/Engine/Source/Runtime/CoreUObject/Private/UObject/LinkerLoad.cpp
@@ -5258,7 +5258,7 @@ UObject* FLinkerLoad::CreateExport( int32 Index )
 			//this is to have a reverse map of UObject to FPackedObjectRef
 			if (UE::LinkerLoad::IsImportLazyLoadEnabled())
 			{
-				MakePackedObjectRef(Export.Object);
+				UE::CoreUObject::Private::MakePackedObjectRef(Export.Object);
 			}
 #endif
 		}
diff --git a/Engine/Source/Runtime/CoreUObject/Private/UObject/LinkerLoadImportBehavior.cpp b/Engine/Source/Runtime/CoreUObject/Private/UObject/LinkerLoadImportBehavior.cpp
--- a/Engine/Source/Runtime/CoreUObject/Private/UObject/LinkerLoadImportBehavior.cpp
+++ b/Engine/Source/Runtime/CoreUObject/Private/UObject/LinkerLoadImportBehavior.cpp
@@ -5,6 +5,7 @@
 #include "UObject/ConstructorHelpers.h"
 #include "UObject/LinkerLoad.h"
 #include "UObject/ObjectResource.h"
+#include "UObject/ObjectPathId.h"
 #include "UObject/Package.h"
 #include "Misc/AssetRegistryInterface.h"
 #include "Misc/CommandLine.h"
@@ -80,6 +81,8 @@ EImportBehavior GetPropertyImportLoadBehavior(const FObjectImport& Import, const
 // recursively handles FAssetData redirectors
 static bool HandleRedirector(const IAssetRegistryInterface& AssetRegistry, const FAssetData& InAssetData, TSet<FSoftObjectPath>& SeenPaths, FObjectPtr& OutObjectPtr)
 {
+	using namespace UE::CoreUObject::Private;
+
 	FString RedirectedPath;
 	if (!InAssetData.GetTagValue("DestinationObject", RedirectedPath))
 	{
@@ -116,14 +119,16 @@ static bool HandleRedirector(const IAssetRegistryInterface& AssetRegistry, const
 
 	FNameBuilder NameBuilder;
 	AssetData.AssetName.AppendString(NameBuilder);
-	FObjectRef ImportRef = { AssetData.PackageName, AssetData.AssetClassPath.GetPackageName(), AssetData.AssetClassPath.GetAssetName(), FObjectPathId(NameBuilder) };
-	FObjectPtr Ptr(ImportRef);
+	FObjectRef ImportRef( AssetData.PackageName, AssetData.AssetClassPath.GetPackageName(), AssetData.AssetClassPath.GetAssetName(), FObjectPathId(NameBuilder) );
+	FPackedObjectRef PackedObjectRef = MakePackedObjectRef(ImportRef);
+	FObjectPtr Ptr({ PackedObjectRef.EncodedRef });
 	OutObjectPtr = Ptr;
 	return true;
 }
 
 static bool TryLazyLoad(const IAssetRegistryInterface& AssetRegistry, const FSoftObjectPath& ObjectPath, FObjectPtr& OutObjectPtr)
 {
+	using namespace UE::CoreUObject::Private;
 	OutObjectPtr = nullptr;
 	
 	FAssetData AssetData;
@@ -146,13 +151,15 @@ static bool TryLazyLoad(const IAssetRegistryInterface& AssetRegistry, const FSof
 	AssetData.AssetName.AppendString(NameBuilder);
 
 	FObjectRef ImportRef = { AssetData.PackageName, AssetData.AssetClassPath.GetPackageName(), AssetData.AssetClassPath.GetAssetName(), FObjectPathId(NameBuilder) };
-	FObjectPtr Ptr(ImportRef);
+	FPackedObjectRef PackedObjectRef = MakePackedObjectRef(ImportRef);
+	FObjectPtr Ptr({ PackedObjectRef.EncodedRef });
 	OutObjectPtr = Ptr;
 	return true;
 }
 
 bool TryLazyImport(const IAssetRegistryInterface& AssetRegistry, const FObjectImport& Import, const FLinkerLoad& LinkerLoad, FObjectPtr& ObjectPtr)
 {
+	using namespace UE::CoreUObject::Private;
 	EImportBehavior Behavior = GetPropertyImportLoadBehavior(Import, LinkerLoad);
 	if (Behavior != EImportBehavior::LazyOnDemand)
 	{
@@ -179,15 +186,16 @@ bool TryLazyImport(const IAssetRegistryInterface& AssetRegistry, const FObjectIm
 		}
 
 		FObjectPathId ObjectPath;
-		FObjectRef ImportRef{ Import.ObjectName, Import.ClassPackage, Import.ClassName, ObjectPath };
-		ObjectPtr = FObjectPtr(ImportRef);
+		FObjectRef ImportRef(Import.ObjectName, Import.ClassPackage, Import.ClassName, ObjectPath);
+		FPackedObjectRef PackedObjectRef = MakePackedObjectRef(ImportRef);
+		ObjectPtr = FObjectPtr({ PackedObjectRef.EncodedRef });
 		return true;
 	}
 
 	//build the a full objectpath for the AssetRegistry
 	FObjectPathId ObjectPath;
 	FName PackageName = FObjectPathId::MakeImportPathIdAndPackageName(Import, LinkerLoad, ObjectPath);
-	FObjectRef ImportRef{ PackageName, Import.ClassPackage, Import.ClassName, ObjectPath };
+	FObjectRef ImportRef(PackageName, Import.ClassPackage, Import.ClassName, ObjectPath);
 
 	TStringBuilder<FName::StringBufferSize> PathName;
 	ImportRef.AppendPathName(PathName);
diff --git a/Engine/Source/Runtime/CoreUObject/Private/UObject/Obj.cpp b/Engine/Source/Runtime/CoreUObject/Private/UObject/Obj.cpp
--- a/Engine/Source/Runtime/CoreUObject/Private/UObject/Obj.cpp
+++ b/Engine/Source/Runtime/CoreUObject/Private/UObject/Obj.cpp
@@ -313,7 +313,7 @@ bool UObject::Rename( const TCHAR* InName, UObject* NewOuter, ERenameFlags Flags
 			}
 		}
 #if UE_WITH_OBJECT_HANDLE_LATE_RESOLVE
-		ObjectHandle_Private::UpdateRenamedObject(*this, NewName, NewOuter);
+		UE::CoreUObject::Private::UpdateRenamedObject(*this, NewName, NewOuter);
 #endif
 		LowLevelRename(NewName, NewOuter);
 	}
diff --git a/Engine/Source/Runtime/CoreUObject/Private/UObject/ObjectHandle.cpp b/Engine/Source/Runtime/CoreUObject/Private/UObject/ObjectHandle.cpp
deleted file mode 100644
diff --git a/Engine/Source/Runtime/CoreUObject/Private/UObject/ObjectHandleIndex.cpp b/Engine/Source/Runtime/CoreUObject/Private/UObject/ObjectHandleIndex.cpp
new file mode 100644
--- /dev/null
+++ b/Engine/Source/Runtime/CoreUObject/Private/UObject/ObjectHandleIndex.cpp
@@ -0,0 +1,533 @@
+// Copyright Epic Games, Inc. All Rights Reserved.
+
+#include "UObject/ObjectHandle.h"
+#include "Async/TaskGraphInterfaces.h"
+#include "HAL/CriticalSection.h"
+#include "HAL/PlatformAtomics.h"
+#include "Misc/PackageName.h"
+#include "Misc/ScopeRWLock.h"
+#include "ProfilingDebugging/CpuProfilerTrace.h"
+#include "UObject/Class.h"
+#include "UObject/Linker.h"
+#include "UObject/LinkerLoad.h"
+#include "UObject/LinkerLoadImportBehavior.h"
+#include "UObject/ObjectRedirector.h"
+#include "UObject/Package.h"
+#include "UObject/UObjectHash.h"
+#include "UObject/ObjectPathId.h"
+
+#if UE_WITH_OBJECT_HANDLE_LATE_RESOLVE
+
+
+
+static inline FName GetNameOrNone(UObject* Object)
+{
+	return Object ? Object->GetFName() : NAME_None;
+}
+
+namespace UE::CoreUObject::Private
+{
+
+	bool operator==(const UE::CoreUObject::Private::FObjectHandleDataClassDescriptor& Lhs, const UE::CoreUObject::Private::FObjectHandleDataClassDescriptor& Rhs)
+	{
+		return (Lhs.PackageName == Rhs.PackageName) && (Lhs.ClassName == Rhs.ClassName);
+	}
+
+	class FPackageId
+	{
+		static constexpr uint32 InvalidId = ~uint32(0u);
+		uint32 Id = InvalidId;
+
+		inline explicit FPackageId(int32 InId) : Id(InId) {}
+
+	public:
+		FPackageId() = default;
+
+		inline static FPackageId FromIndex(uint32 Index)
+		{
+			return FPackageId(Index);
+		}
+
+		inline bool IsValid() const
+		{
+			return Id != InvalidId;
+		}
+
+		inline uint32 ToIndex() const
+		{
+			check(Id != InvalidId);
+			return Id;
+		}
+
+		inline bool operator==(FPackageId Other) const
+		{
+			return Id == Other.Id;
+		}
+	};
+	enum class EObjectId : uint32
+	{
+		Invalid = 0,
+	};
+
+	union FObjectId
+	{
+	private:
+		uint32 RawData = 0;
+	public:
+		struct
+		{
+			uint32 DataClassDescriptorId : 8;
+			uint32 ObjectPathId : 24;
+		} Components;
+
+		FObjectId() = default;
+		FObjectId(EObjectId Id) : RawData(static_cast<uint32>(Id)) {}
+
+		bool operator==(EObjectId Id) { return RawData == static_cast<uint32>(Id); }
+		bool operator!=(EObjectId Id) { return RawData != static_cast<uint32>(Id); }
+
+		inline uint32 ToIndex() const
+		{
+			return Components.ObjectPathId - 1;
+		}
+	};
+
+	static_assert(sizeof(FObjectId) == sizeof(uint32), "FObjectId type must always compile to something equivalent to a uint32 size.");
+
+	struct FObjectHandlePackageData
+	{
+		FMinimalName PackageName;
+		TArray<FObjectPathId> ObjectPaths;
+		TArray<FObjectHandleDataClassDescriptor> DataClassDescriptors;
+		FRWLock Lock;
+	};
+
+	static_assert(STRUCT_OFFSET(FObjectHandlePackageData, PackageName) == STRUCT_OFFSET(FObjectHandlePackageDebugData, PackageName), "FObjectHandlePackageData and FObjectHandlePackageDebugData must match in position of PackageNameField.");
+	static_assert(STRUCT_OFFSET(FObjectHandlePackageData, ObjectPaths) == STRUCT_OFFSET(FObjectHandlePackageDebugData, ObjectPaths), "FObjectHandlePackageData and FObjectHandlePackageDebugData must match in position of ObjectPaths.");
+	static_assert(STRUCT_OFFSET(FObjectHandlePackageData, DataClassDescriptors) == STRUCT_OFFSET(FObjectHandlePackageDebugData, DataClassDescriptors), "FObjectHandlePackageData and FObjectHandlePackageDebugData must match in position of DataClassDescriptors.");
+	static_assert(sizeof(FObjectHandlePackageData) == sizeof(FObjectHandlePackageDebugData), "FObjectHandlePackageData and FObjectHandlePackageDebugData must match in size.");
+
+	struct FObjectHandleIndex
+	{
+		FRWLock Lock; // @TODO: OBJPTR: Want to change this to a striped lock per object bucket to allow more concurrency when adding and looking up objects in a package
+		TMap<FMinimalName, FPackageId> NameToPackageId;
+		TArray<FObjectHandlePackageData> PackageData;
+
+		TArray<FPackedObjectRef> ObjectIndexToPackedObjectRef;
+	} GObjectHandleIndex;
+
+	void InitObjectHandles(int32 MaxObjects)
+	{
+		GObjectHandleIndex.ObjectIndexToPackedObjectRef.SetNumZeroed(MaxObjects);
+	}
+
+	static FAutoConsoleCommand CmdPrintUnresolvedObjects(
+		TEXT("LazyLoad.PrintUnresolvedObjects"),
+		TEXT("Prints a list of all unresolved objects from the object handle index."),
+		FConsoleCommandWithArgsDelegate::CreateLambda([](const TArray<FString>& Args)
+			{
+				FRWScopeLock GlobalLockScope(GObjectHandleIndex.Lock, SLT_ReadOnly);
+				const int32 TotalObjectHandlePackages = GObjectHandleIndex.PackageData.Num();
+				TArray<FName> LoadedPackages;
+				TArray<FName> UnloadedPackages;
+				for (int32 PackageIndex = 0; PackageIndex < TotalObjectHandlePackages; PackageIndex++)
+				{
+					FName PackageName;
+					{
+						FObjectHandlePackageData& PackageData = GObjectHandleIndex.PackageData[PackageIndex];
+						FRWScopeLock LocalLockScope(PackageData.Lock, SLT_ReadOnly);
+						PackageName = MinimalNameToName(PackageData.PackageName);
+					}
+
+					UPackage* Package = FindObjectFast<UPackage>(nullptr, PackageName);
+					if (Package)
+					{
+						LoadedPackages.Add(PackageName);
+					}
+					else
+					{
+						UnloadedPackages.Add(PackageName);
+					}
+				}
+
+				UnloadedPackages.Sort(FNameLexicalLess());
+				for (FName PackageName : UnloadedPackages)
+				{
+					UE_LOG(LogUObjectGlobals, Log, TEXT("Unloaded %s"), *PackageName.ToString());
+				}
+
+				LoadedPackages.Sort(FNameLexicalLess());
+				for (FName PackageName : LoadedPackages)
+				{
+					UE_LOG(LogUObjectGlobals, Log, TEXT("Loaded %s"), *PackageName.ToString());
+				}
+
+				UE_LOG(LogUObjectGlobals, Log, TEXT("Unloaded Packages (%d out of %d)"), UnloadedPackages.Num(), LoadedPackages.Num() + UnloadedPackages.Num());
+			})
+	);
+
+	static inline FPackedObjectRef Pack(FPackageId PackageId, FObjectId ObjectId)
+	{
+		checkf(PackageId.ToIndex() <= 0x7FFFFFFF, TEXT("Package count exceeded the space permitted within packed object references.  This implies over 2 billion packages are in use."));
+		return { static_cast<UPTRINT>(PackageId.ToIndex()) << PackageIdShift |
+				static_cast<UPTRINT>(ObjectId.Components.DataClassDescriptorId) << DataClassDescriptorIdShift |
+				static_cast<UPTRINT>(ObjectId.Components.ObjectPathId) << ObjectPathIdShift |
+				1 };
+	}
+
+	static inline void Unpack(FPackedObjectRef PackedObjectRef, FPackageId& OutPackageId, FObjectId& OutObjectId)
+	{
+		checkf((PackedObjectRef.EncodedRef & 1) == 1, TEXT("Packed object reference is malformed."));
+		OutObjectId.Components.ObjectPathId = (static_cast<uint32>((PackedObjectRef.EncodedRef >> ObjectPathIdShift) & ObjectPathIdMask));
+		OutObjectId.Components.DataClassDescriptorId = (static_cast<uint32>((PackedObjectRef.EncodedRef >> DataClassDescriptorIdShift) & DataClassDescriptorIdMask));
+		OutPackageId = FPackageId::FromIndex(static_cast<uint32>((PackedObjectRef.EncodedRef >> PackageIdShift) & PackageIdMask));
+	}
+
+	static void MakeReferenceIds(FName PackageName, FName ClassPackageName, FName ClassName, FObjectPathId ObjectPath, FPackageId& OutPackageId, FObjectId& OutObjectId)
+	{
+		TRACE_CPUPROFILER_EVENT_SCOPE(UE::CoreUObject::Private::MakeReferenceIds);
+		FMinimalName MinimalName = NameToMinimalName(PackageName);
+
+		//Biases for read-only locking by default at the expense of having to do a second map search if a write is necessary
+		//Doing one search with either a read or locked write seems impossible with the TMap interface at the moment.
+		FRWScopeLock GlobalLockScope(GObjectHandleIndex.Lock, SLT_ReadOnly);
+		FPackageId* FoundPackageId = GObjectHandleIndex.NameToPackageId.Find(MinimalName);
+		FObjectHandlePackageData* PackageData = nullptr;
+		bool PackageCreated = false;
+		if (!FoundPackageId)
+		{
+			GlobalLockScope.ReleaseReadOnlyLockAndAcquireWriteLock_USE_WITH_CAUTION();
+			//Has to be FindOrAdd, as the NameToPackageId may have changed between relinquishing the read lock
+			//and acquiring the write lock.
+			FPackageId NextId = FPackageId::FromIndex(GObjectHandleIndex.PackageData.Num());
+			FoundPackageId = &GObjectHandleIndex.NameToPackageId.FindOrAdd(MinimalName, NextId);
+			if (*FoundPackageId == NextId)
+			{
+				PackageData = &GObjectHandleIndex.PackageData.AddDefaulted_GetRef();
+				PackageData->PackageName = NameToMinimalName(PackageName);
+				GCoreObjectHandlePackageDebug = reinterpret_cast<UE::CoreUObject::Private::FObjectHandlePackageDebugData*>(GObjectHandleIndex.PackageData.GetData());
+				PackageCreated = true;
+			}
+			else
+			{
+				PackageData = &GObjectHandleIndex.PackageData[FoundPackageId->ToIndex()];
+			}
+			//Can't reasonably switch back to a read lock here because the downgrade would have a window where 
+			//the global map could be modified and invalidate the pointer we're holding to an element in it.
+		}
+		else
+		{
+			PackageData = &GObjectHandleIndex.PackageData[FoundPackageId->ToIndex()];
+		}
+		OutPackageId = *FoundPackageId;
+
+		//optimization for creating a new package.
+		//since the package is new, GlobalLockScope is write locked and nothing can be added or read
+		if (PackageCreated)
+		{
+			if (!ClassName.IsNone() && !ClassPackageName.IsNone())
+			{
+				FObjectHandleDataClassDescriptor DataClassDesc{ NameToMinimalName(ClassPackageName), NameToMinimalName(ClassName) };
+				uint32 DataClassDescriptorIndex = PackageData->DataClassDescriptors.Add(DataClassDesc);
+				checkf(((DataClassDescriptorIndex + 1) & ~DataClassDescriptorIdMask) == 0, TEXT("Data class descriptor id overflowed space in ObjectHandle"));
+				OutObjectId.Components.DataClassDescriptorId = DataClassDescriptorIndex + 1;
+			}
+			if (!ObjectPath.IsNone())
+			{
+				int32 PathIndex = PackageData->ObjectPaths.Add(ObjectPath);
+				checkf(((PathIndex + 1) & ~ObjectPathIdMask) == 0, TEXT("Path id overflowed space in ObjectHandle"));
+				OutObjectId.Components.ObjectPathId = PathIndex + 1;
+			}
+			return;
+		}
+		else if (ObjectPath.IsNone())
+		{
+			FRWScopeLock LocalLockScope(PackageData->Lock, SLT_ReadOnly);
+			if (!ClassName.IsNone() && !ClassPackageName.IsNone())
+			{
+				FObjectHandleDataClassDescriptor DataClassDesc{ NameToMinimalName(ClassPackageName), NameToMinimalName(ClassName) };
+				uint32 DataClassDescriptorIndex = PackageData->DataClassDescriptors.AddUnique(DataClassDesc);
+				checkf(((DataClassDescriptorIndex + 1) & ~DataClassDescriptorIdMask) == 0, TEXT("Data class descriptor id overflowed space in ObjectHandle"));
+				OutObjectId.Components.DataClassDescriptorId = DataClassDescriptorIndex + 1;
+			}
+			return;
+		}
+
+		FRWScopeLock LocalLockScope(PackageData->Lock, SLT_ReadOnly);
+		int32 PathIndex = PackageData->ObjectPaths.Find(ObjectPath); //linear search is fine as typically there is only one.
+		if (PathIndex == INDEX_NONE)
+		{
+			//ObjectPaths could have been modified when the read lock was released and the write
+			//lock was acquired, so we must check and see if the ObjectPath was added in that window.
+			LocalLockScope.ReleaseReadOnlyLockAndAcquireWriteLock_USE_WITH_CAUTION();
+			PathIndex = PackageData->ObjectPaths.AddUnique(ObjectPath);
+
+			if (!ClassName.IsNone() && !ClassPackageName.IsNone())
+			{
+				FObjectHandleDataClassDescriptor DataClassDesc{ NameToMinimalName(ClassPackageName), NameToMinimalName(ClassName) };
+				uint32 DataClassDescriptorIndex = PackageData->DataClassDescriptors.AddUnique(DataClassDesc);
+				checkf(((DataClassDescriptorIndex + 1) & ~DataClassDescriptorIdMask) == 0, TEXT("Data class descriptor id overflowed space in ObjectHandle"));
+				OutObjectId.Components.DataClassDescriptorId = DataClassDescriptorIndex + 1;
+			}
+		}
+		else
+		{
+			if (!ClassName.IsNone() && !ClassPackageName.IsNone())
+			{
+				FObjectHandleDataClassDescriptor DataClassDesc{ NameToMinimalName(ClassPackageName), NameToMinimalName(ClassName) };
+				uint32 DataClassDescriptorIndex = PackageData->DataClassDescriptors.Find(DataClassDesc);
+				if (DataClassDescriptorIndex != INDEX_NONE)
+				{
+					checkf(((DataClassDescriptorIndex + 1) & ~DataClassDescriptorIdMask) == 0, TEXT("Data class descriptor id overflowed space in ObjectHandle"));
+					OutObjectId.Components.DataClassDescriptorId = DataClassDescriptorIndex + 1;
+				}
+			}
+		}
+
+		checkf(((PathIndex + 1) & ~ObjectPathIdMask) == 0, TEXT("Path id overflowed space in ObjectHandle"));
+		OutObjectId.Components.ObjectPathId = PathIndex + 1;
+		return;
+	}
+
+	static inline FPackedObjectRef MakePackedObjectRef(FName PackageName, FName ClassPackageName, FName ClassName, FObjectPathId ObjectPath)
+	{
+		FPackageId PackageId;
+		FObjectId ObjectId;
+		MakeReferenceIds(PackageName, ClassPackageName, ClassName, ObjectPath, PackageId, ObjectId);
+		return Pack(PackageId, ObjectId);
+	}
+
+	static void GetObjectDataFromId(FPackageId PackageId, FObjectId ObjectId, FMinimalName& OutPackageName, FObjectPathId& OutPathId, FMinimalName& OutClassPackageName, FMinimalName& OutClassName)
+	{
+		if ((ObjectId == EObjectId::Invalid) || !PackageId.IsValid())
+		{
+			return;
+		}
+
+		FRWScopeLock GlobalLockScope(GObjectHandleIndex.Lock, SLT_ReadOnly);
+		const uint32 PackageIndex = PackageId.ToIndex();
+		if (PackageIndex >= static_cast<uint32>(GObjectHandleIndex.PackageData.Num()))
+		{
+			//checkf(false, TEXT("FObjectHandle: PackageIndex invalid.  This ObjectHandle is from malformed data."));
+			return;
+		}
+		FObjectHandlePackageData& FoundPackageData = GObjectHandleIndex.PackageData[PackageIndex];
+		FRWScopeLock LocalLockScope(FoundPackageData.Lock, SLT_ReadOnly);
+		OutPackageName = FoundPackageData.PackageName;
+
+		if ((ObjectId.Components.ObjectPathId >= static_cast<uint32>(FoundPackageData.ObjectPaths.Num() + 1)) ||
+			(ObjectId.Components.DataClassDescriptorId >= static_cast<uint32>(FoundPackageData.DataClassDescriptors.Num() + 1)))
+		{
+			//checkf(false, TEXT("FObjectHandle: ObjectId.Components.ObjectPathId or ObjectId.Components.DataClassDescriptorId invalid.  This ObjectHandle is from malformed data."));
+			return;
+		}
+
+		if (ObjectId.Components.ObjectPathId != 0)
+		{
+			OutPathId = FoundPackageData.ObjectPaths[ObjectId.Components.ObjectPathId - 1];
+		}
+		else
+		{
+			OutPathId = FObjectPathId();
+		}
+
+		if (ObjectId.Components.DataClassDescriptorId > 0)
+		{
+			FObjectHandleDataClassDescriptor& Desc = FoundPackageData.DataClassDescriptors[ObjectId.Components.DataClassDescriptorId - 1];
+			OutClassPackageName = Desc.PackageName;
+			OutClassName = Desc.ClassName;
+		}
+		else
+		{
+			//@TODO OBJPTR: Is this an error case we should check on?
+			//checkf(false, TEXT("FObjectHandle: ObjectId.Components.ObjectPathId or ObjectId.Components.DataClassDescriptorId invalid.  This ObjectHandle is from malformed data."));
+			return;
+		}
+	}
+
+	void FreeObjectHandle(const UObjectBase& Object)
+	{
+		int32 ObjectIndex = GUObjectArray.ObjectToIndex(&Object);
+		FPackedObjectRef& PackedObjectRef = GObjectHandleIndex.ObjectIndexToPackedObjectRef[ObjectIndex];
+		PackedObjectRef.EncodedRef = 0;
+	}
+
+	void UpdateRenamedObject(const UObject& Object, FName NewName, UObject* NewOuter)
+	{
+		if (!UE::LinkerLoad::IsImportLazyLoadEnabled())
+		{
+			return;
+		}
+		int32 ObjectIndex = GUObjectArray.ObjectToIndex(&Object);
+		FPackedObjectRef PackedObjectRef = GObjectHandleIndex.ObjectIndexToPackedObjectRef[ObjectIndex];
+		if (PackedObjectRef.EncodedRef == 0)
+		{
+			return;
+		}
+
+		check(NewName != NAME_None);
+
+		FMinimalName MinimalName = NameToMinimalName(Object.GetFName());
+		FObjectId ObjectId;
+		FPackageId PackageId;
+		Unpack(PackedObjectRef, PackageId, ObjectId);
+
+		FRWScopeLock GlobalLockScope(GObjectHandleIndex.Lock, SLT_Write);
+		FObjectHandlePackageData& PackageData = GObjectHandleIndex.PackageData[PackageId.ToIndex()];
+		FRWScopeLock PackageLockScope(PackageData.Lock, SLT_ReadOnly);
+		if (Object.GetClass() == UPackage::StaticClass())
+		{
+			//update the package name at existing index. existing object handles will be correct when unpacked
+			//add in the new name pointing the existing index. new object handles will resolve to the existing index
+			check(NewName != NAME_None);
+			FMinimalName MinNewName = NameToMinimalName(NewName);
+			PackageData.PackageName = MinNewName;
+			GObjectHandleIndex.NameToPackageId.Add(MinNewName, PackageId);
+			GObjectHandleIndex.NameToPackageId.Remove(MinimalName);
+			return;
+		}
+
+		//must make a copy to avoid changing data while we are comparing
+		FObjectPathId OldObjectPath = PackageData.ObjectPaths[ObjectId.ToIndex()];
+		if (OldObjectPath.IsWeakObj())
+		{
+			return;
+		}
+
+		const FMinimalName* OldObjectPathtData = nullptr;
+		int32 OldObjectPathSize = 0;
+		FMinimalName OldName;
+
+		if (OldObjectPath.IsSimple())
+		{
+			OldName = OldObjectPath.GetSimpleName();
+			OldObjectPathtData = &OldName;
+			OldObjectPathSize = 1;
+		}
+		else
+		{
+			//this isn't thread safe. if the FStoredObjectPath moves this will now point to garbage
+			const UE::CoreUObject::Private::FStoredObjectPath& StoredPath = OldObjectPath.GetStoredPath();
+			OldObjectPathtData = StoredPath.GetData();
+			OldObjectPathSize = StoredPath.NumElements;
+		}
+
+		//update all paths that start with OldObjectPath. including itself
+		for (int32 PathIndex = 0; PathIndex < PackageData.ObjectPaths.Num(); ++PathIndex)
+		{
+			auto& CurrentObjectPath = PackageData.ObjectPaths[PathIndex];
+			if (CurrentObjectPath.IsWeakObj())
+			{
+				continue;
+			}
+			else if (CurrentObjectPath == OldObjectPath)
+			{
+				CurrentObjectPath.MakeWeakObjPtr(Object);
+				continue;
+			}
+			else if (CurrentObjectPath.IsSimple())
+			{
+				//nothing to do for simple paths. they won't be an inner of the object
+				continue;
+			}
+			const UE::CoreUObject::Private::FStoredObjectPath& ThisPath = CurrentObjectPath.GetStoredPath();
+			if (OldObjectPathSize < ThisPath.NumElements && CompareItems(ThisPath.GetData(), OldObjectPathtData, OldObjectPathSize))
+			{
+				FObjectPathId::ResolvedNameContainerType ResolvedNames;
+				CurrentObjectPath.Resolve(ResolvedNames);
+
+				//resolve all object along the path and convert the object paths to weak objects
+				const UObject* CurrentObject = Object.GetPackage();
+				for (int32 ObjectPathIndex = 0; ObjectPathIndex < ResolvedNames.Num() && CurrentObject; ++ObjectPathIndex)
+				{
+					CurrentObject = StaticFindObjectFastInternal(nullptr, CurrentObject, ResolvedNames[ObjectPathIndex]);
+				}
+
+				// not sure how this could happen, null out the object path
+				if (!CurrentObject)
+				{
+					CurrentObjectPath.Reset();
+				}
+				else
+				{
+					CurrentObjectPath.MakeWeakObjPtr(*CurrentObject);
+				}
+			}
+		}
+
+	}
+
+	FPackedObjectRef MakePackedObjectRef(const FObjectRef& ObjectRef)
+	{
+		if (ObjectRef.IsNull())
+		{
+			return { 0 };
+		}
+
+		return MakePackedObjectRef(ObjectRef.PackageName, ObjectRef.ClassPackageName, ObjectRef.ClassName, ObjectRef.GetObjectPath());
+	}
+
+
+
+	FObjectRef MakeObjectRef(FPackedObjectRef PackedObjectRef)
+	{
+		if (PackedObjectRef.IsNull())
+		{
+			return FObjectRef(NAME_None, NAME_None, NAME_None, UE::CoreUObject::Private::FObjectPathId());
+		}
+
+		TRACE_CPUPROFILER_EVENT_SCOPE(UE::CoreUObject::Private::MakeObjectRef);
+		FObjectId ObjectId;
+		FPackageId PackageId;
+		Unpack(PackedObjectRef, PackageId, ObjectId);
+
+		// Default reference must be invalid if GetObjectDataFromId doesn't populate the reference fields.
+		FMinimalName PackageName;
+		FObjectPathId PathId(FObjectPathId::Invalid);
+		FMinimalName ClassPackageName;
+		FMinimalName ClassName;
+		GetObjectDataFromId(PackageId, ObjectId, PackageName, PathId, ClassPackageName, ClassName);
+		if (PathId.IsWeakObj())
+		{
+			return FObjectRef(PathId.GetWeakObjPtr().Get());
+		}
+		return FObjectRef(MinimalNameToName(PackageName), MinimalNameToName(ClassPackageName), MinimalNameToName(ClassName), PathId);
+	}
+
+
+	FPackedObjectRef MakePackedObjectRef(const UObject* Object)
+	{
+		if (!Object)
+		{
+			return { 0 };
+		}
+
+		if (UE::LinkerLoad::FindLoadBehavior(*Object->GetClass()) == UE::LinkerLoad::EImportBehavior::Eager)
+		{
+			return { 0 };
+		}
+
+		int32 ObjectIndex = GUObjectArray.ObjectToIndex(Object);
+		FPackedObjectRef& PackedObjectRef = UE::CoreUObject::Private::GObjectHandleIndex.ObjectIndexToPackedObjectRef[ObjectIndex];
+		if (PackedObjectRef.EncodedRef != 0)
+		{
+			return PackedObjectRef;
+		}
+
+		FName PackageName = GetNameOrNone(Object->GetOutermost());
+
+		UObject* Class = Object->GetClass();
+		FName ClassPackageName = GetNameOrNone(Class->GetOutermost());
+		PackedObjectRef = UE::CoreUObject::Private::MakePackedObjectRef(PackageName, ClassPackageName, GetNameOrNone(Class), FObjectPathId(Object));
+		return PackedObjectRef;
+	}
+
+	FPackedObjectRef GetPackedObjectRef(const UObject& Object)
+	{
+		int32 ObjectIndex = GUObjectArray.ObjectToIndex(&Object);
+		return UE::CoreUObject::Private::GObjectHandleIndex.ObjectIndexToPackedObjectRef[ObjectIndex];
+	}
+}
+
+
+#endif
\ No newline at end of file
diff --git a/Engine/Source/Runtime/CoreUObject/Private/UObject/ObjectHandleTracking.cpp b/Engine/Source/Runtime/CoreUObject/Private/UObject/ObjectHandleTracking.cpp
new file mode 100644
--- /dev/null
+++ b/Engine/Source/Runtime/CoreUObject/Private/UObject/ObjectHandleTracking.cpp
@@ -0,0 +1,96 @@
+// Copyright Epic Games, Inc. All Rights Reserved.
+
+#include "UObject/ObjectHandleTracking.h"
+
+#if UE_WITH_OBJECT_HANDLE_TRACKING
+
+COREUOBJECT_API std::atomic<int32> UE::CoreUObject::Private::ObjectHandleEventIndex;
+COREUOBJECT_API UE::CoreUObject::Private::FObjectHandleEvents UE::CoreUObject::Private::ObjectHandleEvents[2];
+
+namespace UE::CoreUObject::Private
+{
+	FObjectHandleEvents& BeginWritingEvents()
+	{
+		// We only allow adding/removing the callbacks from a single thread.
+		check(IsInGameThread());
+
+		// Get what the last and new index will be.
+		const int32 LastIndex = ObjectHandleEventIndex;
+		const int32 NewIndex = LastIndex == 0 ? 1 : 0;
+
+		// Make sure we're not still using the events, this should definitely not be true, because
+		// we spin lock at the end of any of these functions to make sure all using has stopped.
+		check(!ObjectHandleEvents[NewIndex].IsUsing());
+
+		// We need to take whatever delegates are setup on the previous index and copy them over to new set.
+		ObjectHandleEvents[NewIndex].ObjectHandleReadEvent = ObjectHandleEvents[LastIndex].ObjectHandleReadEvent;
+		ObjectHandleEvents[NewIndex].ClassReferenceResolvedEvent = ObjectHandleEvents[LastIndex].ClassReferenceResolvedEvent;
+		ObjectHandleEvents[NewIndex].ObjectHandleReferenceResolvedEvent = ObjectHandleEvents[LastIndex].ObjectHandleReferenceResolvedEvent;
+		ObjectHandleEvents[NewIndex].ObjectHandleReferenceLoadedEvent = ObjectHandleEvents[LastIndex].ObjectHandleReferenceLoadedEvent;
+
+		// Now toggle which buffer we're using.
+		ObjectHandleEventIndex = NewIndex;
+
+		// Ok, now we need to spin lock until nobody is using the LastIndex, this will prevent us from
+		// calling this function right after this one and assuming the previous set isn't being used,
+		// similarly while we were doing this, the BeginReadingEvents() function spun lock after marked
+		// that they were using the event set, so now that we've set the new event index, if we spin lock
+		// here, that should ensure that any existing readers finish up what they were doing.
+		while (ObjectHandleEvents[LastIndex].IsUsing())
+		{
+			// Spin!
+		}
+
+		return ObjectHandleEvents[NewIndex];
+	}
+}
+
+FDelegateHandle AddObjectHandleReadCallback(FObjectHandleReadDelegate Callback)
+{
+	UE::CoreUObject::Private::FObjectHandleEvents& Events = UE::CoreUObject::Private::BeginWritingEvents();
+	return Events.ObjectHandleReadEvent.Add(Callback);
+}
+
+void RemoveObjectHandleReadCallback(FDelegateHandle DelegateHandle)
+{
+	UE::CoreUObject::Private::FObjectHandleEvents& Events = UE::CoreUObject::Private::BeginWritingEvents();
+	Events.ObjectHandleReadEvent.Remove(DelegateHandle);
+}
+
+FDelegateHandle AddObjectHandleClassResolvedCallback(FObjectHandleClassResolvedDelegate Callback)
+{
+	UE::CoreUObject::Private::FObjectHandleEvents& Events = UE::CoreUObject::Private::BeginWritingEvents();
+	return Events.ClassReferenceResolvedEvent.Add(Callback);
+}
+
+void RemoveObjectHandleClassResolvedCallback(FDelegateHandle DelegateHandle)
+{
+	UE::CoreUObject::Private::FObjectHandleEvents& Events = UE::CoreUObject::Private::BeginWritingEvents();
+	Events.ClassReferenceResolvedEvent.Remove(DelegateHandle);
+}
+
+FDelegateHandle AddObjectHandleReferenceResolvedCallback(FObjectHandleReferenceResolvedDelegate Callback)
+{
+	UE::CoreUObject::Private::FObjectHandleEvents& Events = UE::CoreUObject::Private::BeginWritingEvents();
+	return Events.ObjectHandleReferenceResolvedEvent.Add(Callback);
+}
+
+void RemoveObjectHandleReferenceResolvedCallback(FDelegateHandle DelegateHandle)
+{
+	UE::CoreUObject::Private::FObjectHandleEvents& Events = UE::CoreUObject::Private::BeginWritingEvents();
+	Events.ObjectHandleReferenceResolvedEvent.Remove(DelegateHandle);
+}
+
+FDelegateHandle AddObjectHandleReferenceLoadedCallback(FObjectHandleReferenceLoadedDelegate Callback)
+{
+	UE::CoreUObject::Private::FObjectHandleEvents& Events = UE::CoreUObject::Private::BeginWritingEvents();
+	return Events.ObjectHandleReferenceLoadedEvent.Add(Callback);
+}
+
+void RemoveObjectHandleReferenceLoadedCallback(FDelegateHandle DelegateHandle)
+{
+	UE::CoreUObject::Private::FObjectHandleEvents& Events = UE::CoreUObject::Private::BeginWritingEvents();
+	Events.ObjectHandleReferenceLoadedEvent.Remove(DelegateHandle);
+}
+
+#endif // UE_WITH_OBJECT_HANDLE_TRACKING
diff --git a/Engine/Source/Runtime/CoreUObject/Private/UObject/ObjectPathId.cpp b/Engine/Source/Runtime/CoreUObject/Private/UObject/ObjectPathId.cpp
--- a/Engine/Source/Runtime/CoreUObject/Private/UObject/ObjectPathId.cpp
+++ b/Engine/Source/Runtime/CoreUObject/Private/UObject/ObjectPathId.cpp
@@ -26,7 +26,7 @@ FCriticalSection GComplexPathLock;  // Could be changed to an RWLock later if ne
 //our own storage at all.
 TMultiMap<uint32, uint32> GComplexPathHashToId;
 
-TArray<UE::ObjectPath::Private::FStoredObjectPath> GComplexPaths;
+TArray<UE::CoreUObject::Private::FStoredObjectPath> GComplexPaths;
 /*
 template <int a, int b>
 struct assert_equality
@@ -44,10 +44,7 @@ namespace UE::CoreUObject::Private
 	{
 		return FObjectKey(ObjectIndex, ObjectSerialNumber);
 	}
-}
 
-namespace UE::ObjectPath::Private
-{
 	FStoredObjectPath::FStoredObjectPath(TConstArrayView<FMinimalName> InNames)
 	{
 		// Copy into storage in reverse
@@ -139,253 +136,254 @@ namespace UE::ObjectPath::Private
 	{
 		return A.NumElements == B.NumElements && CompareItems(A.GetView().GetData(), B.GetView().GetData(), A.NumElements);
 	}
-}
-
-template <typename NameProducerType>
-void FObjectPathId::StoreObjectPathId(NameProducerType& NameProducer)
-{
-	FName Name = NameProducer();
-
-	if (Name == NAME_None)
-	{
-		return;
-	}
 
-	FName OuterName = NameProducer();
-	if (OuterName == NAME_None)
+	template <typename NameProducerType>
+	void FObjectPathId::StoreObjectPathId(NameProducerType& NameProducer)
 	{
-		int32 NameNumber = Name.GetNumber();
+		FName Name = NameProducer();
 
-		//verify that the second most significant bits are not set
-		if ( !(NameNumber & SimpleNameMask) && !(NameNumber & WeakObjectMask))
+		if (Name == NAME_None)
 		{
-			Index = Name.GetComparisonIndex().ToUnstableInt();
-			Number = NameNumber | SimpleNameMask;
 			return;
 		}
-	}
 
-	//Complex path scenario
-	TArray<FMinimalName, TInlineAllocator<3>> MinimalNames;
-	uint32 Key = GetTypeHash(Name.ToUnstableInt());
-	MinimalNames.Emplace(NameToMinimalName(Name));
-	while (OuterName != NAME_None)
-	{
-		Name = OuterName;
-		OuterName = NameProducer();
+		FName OuterName = NameProducer();
+		if (OuterName == NAME_None)
+		{
+			int32 NameNumber = Name.GetNumber();
+
+			//verify that the second most significant bits are not set
+			if (!(NameNumber & SimpleNameMask) && !(NameNumber & WeakObjectMask))
+			{
+				Index = Name.GetComparisonIndex().ToUnstableInt();
+				Number = NameNumber | SimpleNameMask;
+				return;
+			}
+		}
+
+		//Complex path scenario
+		TArray<FMinimalName, TInlineAllocator<3>> MinimalNames;
+		uint32 Key = GetTypeHash(Name.ToUnstableInt());
 		MinimalNames.Emplace(NameToMinimalName(Name));
-		Key = HashCombine(Key, GetTypeHash(Name.ToUnstableInt()));
-	}
+		while (OuterName != NAME_None)
+		{
+			Name = OuterName;
+			OuterName = NameProducer();
+			MinimalNames.Emplace(NameToMinimalName(Name));
+			Key = HashCombine(Key, GetTypeHash(Name.ToUnstableInt()));
+		}
 
-	// Reverse path so we can compare it to stored paths
-	UE::ObjectPath::Private::FStoredObjectPath PathToStore(MinimalNames);
+		// Reverse path so we can compare it to stored paths
+		FStoredObjectPath PathToStore(MinimalNames);
 
-	FScopeLock ComplexPathScopeLock(&GComplexPathLock);
-	for (typename TMultiMap<uint32, uint32>::TConstKeyIterator It(GComplexPathHashToId, Key); It; ++It)
-	{
-		uint32 PotentialPathId = It.Value();
-		if (GComplexPaths[PotentialPathId-1] == PathToStore)
+		FScopeLock ComplexPathScopeLock(&GComplexPathLock);
+		for (typename TMultiMap<uint32, uint32>::TConstKeyIterator It(GComplexPathHashToId, Key); It; ++It)
 		{
-			Index = PotentialPathId;
-			Number = 0;
-			return;
+			uint32 PotentialPathId = It.Value();
+			if (GComplexPaths[PotentialPathId - 1] == PathToStore)
+			{
+				Index = PotentialPathId;
+				Number = 0;
+				return;
+			}
 		}
+
+		//add one because a path id of 0 indicates null
+		uint32 NewId = GComplexPaths.Emplace(MoveTemp(PathToStore)) + 1;
+		GComplexPathHashToId.Add(Key, NewId);
+		Index = NewId;
+		Number = 0;
+		GCoreComplexObjectPathDebug = GComplexPaths.GetData();
 	}
 
-	//add one because a path id of 0 indicates null
-	uint32 NewId = GComplexPaths.Emplace(MoveTemp(PathToStore)) + 1;
-	GComplexPathHashToId.Add(Key, NewId);
-	Index = NewId;
-	Number = 0;
-	GCoreComplexObjectPathDebug = GComplexPaths.GetData();
-}
+	FObjectPathId::FObjectPathId(const UObject* Object)
+	{
+		struct FLoadedObjectNamePathProducer
+		{
+			FLoadedObjectNamePathProducer(const UObject* InObject)
+				: CurrentObject(InObject)
+			{}
 
-FObjectPathId::FObjectPathId(const UObject* Object)
-{
-	struct FLoadedObjectNamePathProducer
+			FName operator()(void)
+			{
+				if ((CurrentObject == nullptr) || (CurrentObject->GetClass() == UPackage::StaticClass()))
+				{
+					return NAME_None;
+				}
+				FName RetVal = CurrentObject->GetFName();
+				CurrentObject = CurrentObject->GetOuter();
+				return RetVal;
+			}
+		private:
+			const UObject* CurrentObject;
+		} NamePathProducer(Object);
+
+		StoreObjectPathId(NamePathProducer);
+	}
+
+	FObjectPathId::FObjectPathId(const FObjectImport& Import, const FLinkerTables& LinkerTables)
 	{
-		FLoadedObjectNamePathProducer(const UObject* InObject)
-		: CurrentObject(InObject)
-		{}
+		MakeImportPathIdAndPackageName(Import, LinkerTables, *this);
+	}
 
-		FName operator()(void)
+	FName FObjectPathId::MakeImportPathIdAndPackageName(const FObjectImport& Import, const FLinkerTables& LinkerTables, FObjectPathId& OutPathId)
+	{
+		TRACE_CPUPROFILER_EVENT_SCOPE(FObjectPathId::MakeImportPathIdAndPackageName);
+
+		struct FImportObjectNamePathProducer
 		{
-			if ((CurrentObject == nullptr) || (CurrentObject->GetClass() == UPackage::StaticClass()))
+			FImportObjectNamePathProducer(const FObjectImport& InImport, const FLinkerTables& InLinkerTables)
+				: CurrentImport(&InImport)
+				, LinkerTables(InLinkerTables)
+			{}
+
+			FName operator()(void)
 			{
-				return NAME_None;
+				if ((CurrentImport == nullptr) || CurrentImport->OuterIndex.IsNull())
+				{
+					return NAME_None;
+				}
+				FName RetVal = CurrentImport->ObjectName;
+				CurrentImport = &LinkerTables.Imp(CurrentImport->OuterIndex);
+				return RetVal;
 			}
-			FName RetVal = CurrentObject->GetFName();
-			CurrentObject = CurrentObject->GetOuter();
-			return RetVal;
-		}
-	private:
-		const UObject* CurrentObject;
-	} NamePathProducer(Object);
 
-	StoreObjectPathId(NamePathProducer);
-}
+			FName GetPackageName()
+			{
+				if (CurrentImport && CurrentImport->OuterIndex.IsNull())
+				{
+					return CurrentImport->ObjectName;
+				}
+				return NAME_None;
+			}
+		private:
+			const FObjectImport* CurrentImport;
+			const FLinkerTables& LinkerTables;
+		} NamePathProducer(Import, LinkerTables);
 
-FObjectPathId::FObjectPathId(const FObjectImport& Import, const FLinkerTables& LinkerTables)
-{
-	MakeImportPathIdAndPackageName(Import, LinkerTables, *this);
-}
+		OutPathId.StoreObjectPathId(NamePathProducer);
+		return NamePathProducer.GetPackageName();
+	}
 
-FName FObjectPathId::MakeImportPathIdAndPackageName(const FObjectImport& Import, const FLinkerTables& LinkerTables, FObjectPathId& OutPathId)
-{
-	TRACE_CPUPROFILER_EVENT_SCOPE(FObjectPathId::MakeImportPathIdAndPackageName);
-	
-	struct FImportObjectNamePathProducer
+	template <typename CharType>
+	struct TStringViewNamePathProducer
 	{
-		FImportObjectNamePathProducer(const FObjectImport& InImport, const FLinkerTables& InLinkerTables)
-		: CurrentImport(&InImport)
-		, LinkerTables(InLinkerTables)
+		TStringViewNamePathProducer(TStringView<CharType> StringPath)
+			: CurrentStringView(StringPath)
 		{}
 
 		FName operator()(void)
 		{
-			if ((CurrentImport == nullptr) || CurrentImport->OuterIndex.IsNull())
+			if (CurrentStringView.IsEmpty())
 			{
 				return NAME_None;
 			}
-			FName RetVal = CurrentImport->ObjectName;
-			CurrentImport = &LinkerTables.Imp(CurrentImport->OuterIndex);
+
+			int32 FoundIndex = INDEX_NONE;
+			if (!FindLastSeparator(FoundIndex))
+			{
+				FName RetVal(CurrentStringView);
+				CurrentStringView.Reset();
+				return RetVal;
+			}
+
+			FName RetVal(CurrentStringView.RightChop(FoundIndex + 1));
+			CurrentStringView.RemoveSuffix(CurrentStringView.Len() - FoundIndex);
 			return RetVal;
 		}
+	private:
+		static inline bool IsPathIdSeparator(CharType Char) { return Char == '.' || Char == ':'; }
 
-		FName GetPackageName()
+		bool FindLastSeparator(int32& OutIndex)
 		{
-			if (CurrentImport && CurrentImport->OuterIndex.IsNull())
+			if (const CharType* Separator = Algo::FindLastByPredicate(CurrentStringView, IsPathIdSeparator))
 			{
-				return CurrentImport->ObjectName;
+				OutIndex = UE_PTRDIFF_TO_INT32(Separator - CurrentStringView.GetData());
+				return true;
 			}
-			return NAME_None;
+			OutIndex = INDEX_NONE;
+			return false;
 		}
-	private:
-		const FObjectImport* CurrentImport;
-		const FLinkerTables& LinkerTables;
-	} NamePathProducer(Import, LinkerTables);
 
-	OutPathId.StoreObjectPathId(NamePathProducer);
-	return NamePathProducer.GetPackageName();
-}
+		TStringView<CharType> CurrentStringView;
+	};
 
-template <typename CharType>
-struct TStringViewNamePathProducer
-{
-	TStringViewNamePathProducer(TStringView<CharType> StringPath)
-	: CurrentStringView(StringPath)
-	{}
+	FObjectPathId::FObjectPathId(FWideStringView StringPath)
+	{
+		TStringViewNamePathProducer<WIDECHAR> NamePathProducer(StringPath);
+		StoreObjectPathId(NamePathProducer);
+	}
 
-	FName operator()(void)
+	FObjectPathId::FObjectPathId(FAnsiStringView StringPath)
 	{
-		if (CurrentStringView.IsEmpty())
+		TStringViewNamePathProducer<ANSICHAR> NamePathProducer(StringPath);
+		StoreObjectPathId(NamePathProducer);
+	}
+
+	void FObjectPathId::Resolve(ResolvedNameContainerType& OutContainer) const
+	{
+		check(IsValid());
+
+		if (IsNone())
 		{
-			return NAME_None;
+			return;
 		}
 
-		int32 FoundIndex = INDEX_NONE;
-		if (!FindLastSeparator(FoundIndex))
+		if (Number & SimpleNameMask)
 		{
-			FName RetVal(CurrentStringView);
-			CurrentStringView.Reset();
-			return RetVal;
+			//truncate Path to int and shift back to get the number
+			uint32 NameNumber = Number & ~SimpleNameMask;
+			uint32 NameIndex = Index;
+			FNameEntryId EntryId = FNameEntryId::FromUnstableInt(NameIndex);
+			OutContainer.Emplace(EntryId, EntryId, NameNumber);
+			return;
 		}
 
-		FName RetVal(CurrentStringView.RightChop(FoundIndex + 1));
-		CurrentStringView.RemoveSuffix(CurrentStringView.Len() - FoundIndex);
-		return RetVal;
-	}
-private:
-	static inline bool IsPathIdSeparator(CharType Char) { return Char == '.' || Char == ':'; }
+		FScopeLock ComplexPathScopeLock(&GComplexPathLock);
+		const FStoredObjectPath& FoundContainer = GComplexPaths[Index - 1];
 
-	bool FindLastSeparator(int32& OutIndex)
-	{
-		if (const CharType* Separator = Algo::FindLastByPredicate(CurrentStringView, IsPathIdSeparator))
+		for (FMinimalName Name : FoundContainer.GetView())
 		{
-			OutIndex = UE_PTRDIFF_TO_INT32(Separator - CurrentStringView.GetData());
-			return true;
+			OutContainer.Emplace(MinimalNameToName(Name));
 		}
-		OutIndex = INDEX_NONE;
-		return false;
 	}
 
-	TStringView<CharType> CurrentStringView;
-};
-
-FObjectPathId::FObjectPathId(FWideStringView StringPath)
-{
-	TStringViewNamePathProducer<WIDECHAR> NamePathProducer(StringPath);
-	StoreObjectPathId(NamePathProducer);
-}
-
-FObjectPathId::FObjectPathId(FAnsiStringView StringPath)
-{
-	TStringViewNamePathProducer<ANSICHAR> NamePathProducer(StringPath);
-	StoreObjectPathId(NamePathProducer);
-}
-
-void FObjectPathId::Resolve(ResolvedNameContainerType& OutContainer) const
-{
-	check(IsValid());
-	
-	if (IsNone())
+	FMinimalName FObjectPathId::GetSimpleName() const
 	{
-		return;
-	}
+		check(IsSimple());
 
-	if (Number & SimpleNameMask)
-	{
-		//truncate Path to int and shift back to get the number
 		uint32 NameNumber = Number & ~SimpleNameMask;
 		uint32 NameIndex = Index;
 		FNameEntryId EntryId = FNameEntryId::FromUnstableInt(NameIndex);
-		OutContainer.Emplace(EntryId, EntryId, NameNumber);
-		return; 
+		FMinimalName Name = NameToMinimalName(FName(EntryId, EntryId, NameNumber));
+		return Name;
 	}
 
-	FScopeLock ComplexPathScopeLock(&GComplexPathLock);
-	const UE::ObjectPath::Private::FStoredObjectPath& FoundContainer = GComplexPaths[Index - 1];
-	
-	for (FMinimalName Name : FoundContainer.GetView())
+	void FObjectPathId::MakeWeakObjPtr(const UObject& Object)
 	{
-		OutContainer.Emplace(MinimalNameToName(Name));
-	}
-}
-FMinimalName FObjectPathId::GetSimpleName() const
-{
-	check(IsSimple());
+		int32 ObjectIndex = GUObjectArray.ObjectToIndex(&Object);
+		int32 ObjectSerialNumber = GUObjectArray.AllocateSerialNumber(ObjectIndex);
 
-	uint32 NameNumber = Number & ~SimpleNameMask;
-	uint32 NameIndex = Index;
-	FNameEntryId EntryId = FNameEntryId::FromUnstableInt(NameIndex);
-	FMinimalName Name = NameToMinimalName(FName(EntryId, EntryId, NameNumber));
-	return Name;
-}
-
-void FObjectPathId::MakeWeakObjPtr(const UObject& Object)
-{
-	int32 ObjectIndex = GUObjectArray.ObjectToIndex(&Object);
-	int32 ObjectSerialNumber = GUObjectArray.AllocateSerialNumber(ObjectIndex);
-
-	//check that ObjectSerialNumber does not have WeakObjMask bit set
-	check(!(ObjectSerialNumber & WeakObjectMask));
-	Index = ObjectIndex;
-	Number = ObjectSerialNumber | WeakObjectMask;
-}
+		//check that ObjectSerialNumber does not have WeakObjMask bit set
+		check(!(ObjectSerialNumber & WeakObjectMask));
+		Index = ObjectIndex;
+		Number = ObjectSerialNumber | WeakObjectMask;
+	}
 
-FWeakObjectPtr FObjectPathId::GetWeakObjPtr() const
-{
-	check(IsWeakObj());
+	FWeakObjectPtr FObjectPathId::GetWeakObjPtr() const
+	{
+		check(IsWeakObj());
 
-	int32 ObjectIndex = Index;
-	int32 ObjectSerial = Number & ~WeakObjectMask;
-	FObjectKey Key = UE::CoreUObject::Private::MakeObjectKey(ObjectIndex, ObjectSerial);
-	return Key.ResolveObjectPtr();
-}
+		int32 ObjectIndex = Index;
+		int32 ObjectSerial = Number & ~WeakObjectMask;
+		FObjectKey Key = UE::CoreUObject::Private::MakeObjectKey(ObjectIndex, ObjectSerial);
+		return Key.ResolveObjectPtr();
+	}
 
-const UE::ObjectPath::Private::FStoredObjectPath& FObjectPathId::GetStoredPath() const
-{
-	check(IsValid() && !IsNone());
-	const UE::ObjectPath::Private::FStoredObjectPath& FoundContainer = GComplexPaths[Index - 1];
-	return FoundContainer;
+	const FStoredObjectPath& FObjectPathId::GetStoredPath() const
+	{
+		check(IsValid() && !IsNone());
+		const FStoredObjectPath& FoundContainer = GComplexPaths[Index - 1];
+		return FoundContainer;
+	}
 }
\ No newline at end of file
diff --git a/Engine/Source/Runtime/CoreUObject/Private/UObject/ObjectPathId.h b/Engine/Source/Runtime/CoreUObject/Private/UObject/ObjectPathId.h
new file mode 100644
--- /dev/null
+++ b/Engine/Source/Runtime/CoreUObject/Private/UObject/ObjectPathId.h
@@ -0,0 +1,122 @@
+// Copyright Epic Games, Inc. All Rights Reserved.
+
+#pragma once
+
+#include "Containers/Array.h"
+#include "Containers/ContainerAllocationPolicies.h"
+#include "Containers/ContainersFwd.h"
+#include "Containers/StringFwd.h"
+#include "Containers/StringView.h"
+#include "HAL/Platform.h"
+#include "Templates/TypeHash.h"
+#include "UObject/NameTypes.h"
+
+class FLinkerTables;
+class UObject;
+struct FObjectImport;
+struct FObjectRef;
+
+// Declared in the header so the type exists in debug info for debugger visualization, not a public part of the API.
+namespace UE::CoreUObject::Private
+{
+	struct FPackedObjectRef;
+
+	struct FStoredObjectPath
+	{
+		static constexpr const int32 NumInlineElements = 3;
+		int32 NumElements;
+
+		union
+		{
+			FMinimalName Short[NumInlineElements];
+			FMinimalName* Long;
+		};
+
+		FStoredObjectPath(TConstArrayView<FMinimalName> InNames);
+		FStoredObjectPath(const FMinimalName* Parent, int32 NumParent, const FMinimalName* Child, int32 NumChild);
+		~FStoredObjectPath();
+		FStoredObjectPath(const FStoredObjectPath&) = delete;
+		FStoredObjectPath(FStoredObjectPath&&);
+		FStoredObjectPath& operator=(const FStoredObjectPath&) = delete;
+		FStoredObjectPath& operator=(FStoredObjectPath&&);
+
+		TConstArrayView<FMinimalName> GetView() const;
+		const FMinimalName* GetData() const { return NumElements > NumInlineElements ? Long : &Short[0]; }
+	};
+
+	// @TODO: OBJPTR: Should this be changed to FObjectImportPathId? It is already written to target the specific patterns in import paths,
+	//		 if we go further in the future and do things like store classname/classpackage info, to the path segments, it would be
+	//		 even more specific to imports.
+
+	/**
+	 * FObjectPathId is a type tailored to represent object paths, with a focus on avoiding overhead for common import path patterns
+	 * It is primarily for use by FObjectRef/FPackedObjectRef/FObjectHandle. Most object paths used by FObjectRef will be imports,
+	 * and are frequently a single FName without any outer components. FObjectPathId can transform to/from most single FNames without
+	 * any need for a table lookup. For the less common case of an imported object path composed of more than one FName, or if the imported
+	 * object path has an exceedingly high number component, FObjectPathId transformations will have to do a table lookup.
+	 *
+	 * NOTE: FObjectPathId is NOT case sensitive.  Two FObjectPathIds that were constructed with the same path but of differing case
+	 *		 will compare as equal.  This also means that constructing an FObjectPathId from a path then calling Resolve on it may
+	 *		 produce the same set of names as in the original path, but of different case.  THIS WILL CAUSE NON-DETERMINISM.
+	 */
+	class FObjectPathId
+	{
+	public:
+		FObjectPathId() = default;
+		FObjectPathId(const FObjectPathId& Other) = default;
+		explicit FObjectPathId(const UObject* Object);
+		FObjectPathId(const FObjectImport& Import, const FLinkerTables& LinkerTables);
+
+		enum EInvalid { Invalid = 0 };
+
+		explicit COREUOBJECT_API FObjectPathId(EInvalid)
+			: Index(0)
+			, Number(SimpleNameMask | WeakObjectMask)
+		{
+		}
+
+		static FName MakeImportPathIdAndPackageName(const FObjectImport& Import, const FLinkerTables& LinkerTables, FObjectPathId& OutPathId);
+
+		explicit FObjectPathId(FWideStringView StringPath);
+		explicit FObjectPathId(FAnsiStringView StringPath);
+
+		FObjectPathId& operator=(const FObjectPathId& Other) = default;
+		inline bool operator==(const FObjectPathId& Other) const { return Index == Other.Index && Number == Other.Number; }
+		inline bool IsNone() const { return Index == 0 && Number == 0; }
+		inline bool IsValid() const { return Number != (SimpleNameMask | WeakObjectMask); }
+		inline bool IsSimple() const { return IsValid() && Number & SimpleNameMask; }
+		inline bool IsWeakObj() const { return IsValid() && Number & WeakObjectMask; }
+
+		// @TODO: OBJPTR: Is there a better interface to output the resolved name without:
+		// 		 1) Requiring multiple lookups in the table in case of complex paths
+		//		 2) Requiring string allocation
+		//		 3) Constraining how we might change the internal storage of object paths in the future
+		using ResolvedNameContainerType = TArray<FName, TInlineAllocator<3>>;
+		void Resolve(ResolvedNameContainerType& OutContainer) const;
+		FWeakObjectPtr GetWeakObjPtr() const;
+		FMinimalName GetSimpleName() const;
+		void MakeWeakObjPtr(const UObject& Object);
+		const FStoredObjectPath& GetStoredPath() const;
+
+		void Reset()
+		{
+			Index = 0;
+			Number = 0;
+		}
+	private:
+		
+		template <typename NameProducerType>
+		void StoreObjectPathId(NameProducerType& NameProducer);
+
+		static constexpr uint32 WeakObjectMask = ~((~0u) >> 1);       //most significant bit
+		static constexpr uint32 SimpleNameMask = WeakObjectMask >> 1; //second most significant bits
+
+		uint32 Index = 0;
+		uint32 Number = 0;
+
+		friend FORCEINLINE uint32 GetTypeHash(FObjectPathId Value)
+		{
+			return HashCombineFast(GetTypeHash(Value.Index), GetTypeHash(Value.Number));
+		}
+	};
+}
diff --git a/Engine/Source/Runtime/CoreUObject/Private/UObject/ObjectPtr.cpp b/Engine/Source/Runtime/CoreUObject/Private/UObject/ObjectPtr.cpp
--- a/Engine/Source/Runtime/CoreUObject/Private/UObject/ObjectPtr.cpp
+++ b/Engine/Source/Runtime/CoreUObject/Private/UObject/ObjectPtr.cpp
@@ -6,41 +6,42 @@
 
 FString FObjectPtr::GetPathName() const
 {
-	if (IsResolved() && !IsObjectHandleNull(Handle))
+	FObjectHandle LocalHandle = Handle;
+	if (IsObjectHandleResolved(LocalHandle) && !IsObjectHandleNull(LocalHandle))
 	{
 		return Get()->GetPathName();
 	}
 	else
 	{
-		FObjectRef ObjectRef = MakeObjectRef(Handle);
+		FObjectRef ObjectRef = UE::CoreUObject::Private::MakeObjectRef(UE::CoreUObject::Private::ReadObjectHandlePackedObjectRefNoCheck(LocalHandle));
 		return ObjectRef.GetPathName();
 	}
 }
 
 FName FObjectPtr::GetFName() const
 {
-	if (IsResolved())
+	FObjectHandle LocalHandle = Handle;
+	if (IsObjectHandleResolved(LocalHandle) && !IsObjectHandleNull(LocalHandle))
 	{
-		UObject* ResolvedObject = Get();
-		return ResolvedObject ? ResolvedObject->GetFName() : NAME_None;
+		return Get()->GetFName();
 	}
 	else
 	{
-		FObjectRef ObjectRef = MakeObjectRef(Handle);
+		FObjectRef ObjectRef = UE::CoreUObject::Private::MakeObjectRef(UE::CoreUObject::Private::ReadObjectHandlePackedObjectRefNoCheck(LocalHandle));
 		return ObjectRef.GetFName();
 	}
 }
 
 FString FObjectPtr::GetFullName(EObjectFullNameFlags Flags) const
 {
-	if (IsResolved())
+	FObjectHandle LocalHandle = Handle;
+	if (IsObjectHandleResolved(LocalHandle) && !IsObjectHandleNull(LocalHandle))
 	{
-		// UObjectBaseUtility::GetFullName is safe to call on null objects.
 		return Get()->GetFullName(nullptr, Flags);
 	}
 	else
 	{
-		FObjectRef ObjectRef = MakeObjectRef(Handle);
+		FObjectRef ObjectRef = UE::CoreUObject::Private::MakeObjectRef(UE::CoreUObject::Private::ReadObjectHandlePackedObjectRefNoCheck(LocalHandle));
 		return ObjectRef.GetFullName(Flags);
 	}
 }
diff --git a/Engine/Source/Runtime/CoreUObject/Private/UObject/ObjectRef.cpp b/Engine/Source/Runtime/CoreUObject/Private/UObject/ObjectRef.cpp
new file mode 100644
--- /dev/null
+++ b/Engine/Source/Runtime/CoreUObject/Private/UObject/ObjectRef.cpp
@@ -0,0 +1,302 @@
+// Copyright Epic Games, Inc. All Rights Reserved.
+
+#include "UObject/ObjectRef.h"
+
+#if UE_WITH_OBJECT_HANDLE_LATE_RESOLVE || UE_WITH_OBJECT_HANDLE_TRACKING
+
+#include "Misc/PackageName.h"
+#include "UObject/Package.h"
+#include "UObject/UObjectHash.h"
+#include "UObject/ObjectRedirector.h"
+#include "UObject/LinkerLoad.h"
+#include "Async/TaskGraphInterfaces.h"
+#include "UObject/ObjectPathId.h"
+
+DEFINE_LOG_CATEGORY(LogObjectRef);
+
+FObjectRef::FObjectRef(FName PackageName, FName ClassPackageName, FName ClassName, UE::CoreUObject::Private::FObjectPathId ObjectPath)
+	: PackageName(PackageName)
+	, ClassPackageName(ClassPackageName)
+	, ClassName(ClassName)
+	, ObjectPathId(*reinterpret_cast<uint64*>(&ObjectPath))
+{
+}
+
+FObjectRef::FObjectRef(const UObject* Object)
+	: PackageName(NAME_None)
+	, ClassPackageName(NAME_None)
+	, ClassName(NAME_None)
+	, ObjectPathId(0)
+{
+	if (Object)
+	{
+		ClassName = Object->GetClass()->GetFName();
+
+		UObject* ClassOuter = Object->GetClass()->GetOuter();
+		ClassPackageName = ClassOuter ? ClassOuter->GetFName() : NAME_None;
+
+		UObject* ObjectOuter = Object->GetOutermost();
+		PackageName = ObjectOuter ? ObjectOuter->GetFName() : NAME_None;
+		UE::CoreUObject::Private::FObjectPathId ObjectPath(Object);
+
+		ObjectPathId = *reinterpret_cast<uint64*>(&ObjectPath);
+	}
+}
+
+UE::CoreUObject::Private::FObjectPathId FObjectRef::GetObjectPath() const
+{
+	static_assert(sizeof(uint64) == sizeof(UE::CoreUObject::Private::FObjectPathId));
+	return *reinterpret_cast<const UE::CoreUObject::Private::FObjectPathId*>(&ObjectPathId);
+}
+
+FName FObjectRef::GetFName() const
+{
+	if (PackageName == NAME_None)
+	{
+		return NAME_None;
+	}
+	else if (ObjectPathId == 0)
+	{
+		return PackageName;
+	}
+
+	UE::CoreUObject::Private::FObjectPathId ObjectPath = GetObjectPath();
+	UE::CoreUObject::Private::FObjectPathId::ResolvedNameContainerType ResolvedNames;
+	ObjectPath.Resolve(ResolvedNames);
+
+	if (ResolvedNames.Num() > 0)
+	{
+		const FName& Name = ResolvedNames[ResolvedNames.Num() - 1];
+		return Name;
+	}
+
+	return NAME_None;
+}
+
+void FObjectRef::AppendClassPathName(FStringBuilderBase& OutClassPathNameBuilder, EObjectFullNameFlags Flags) const
+{
+	if (EnumHasAllFlags(Flags, EObjectFullNameFlags::IncludeClassPackage))
+	{
+		ClassPackageName.AppendString(OutClassPathNameBuilder);
+		OutClassPathNameBuilder << TEXT(".");
+	}
+
+	ClassName.AppendString(OutClassPathNameBuilder);
+}
+
+void FObjectRef::AppendPathName(FStringBuilderBase& OutPathNameBuilder) const
+{
+	if (PackageName == NAME_None)
+	{
+		OutPathNameBuilder.Append(TEXT("None"));
+		return;
+	}
+
+	PackageName.ToString(OutPathNameBuilder);
+
+	UE::CoreUObject::Private::FObjectPathId ObjectPath = GetObjectPath();
+	UE::CoreUObject::Private::FObjectPathId::ResolvedNameContainerType ResolvedNames;
+	ObjectPath.Resolve(ResolvedNames);
+
+	for (int Index = 0; Index < ResolvedNames.Num(); ++Index)
+	{
+		if (Index == 1)
+		{
+			OutPathNameBuilder << SUBOBJECT_DELIMITER_CHAR;
+		}
+		else
+		{
+			OutPathNameBuilder << TEXT(".");
+		}
+
+		const FName& Name = ResolvedNames[Index];
+		Name.AppendString(OutPathNameBuilder);
+	}
+}
+
+
+static inline UPackage* FindOrLoadPackage(FName PackageName, int32 LoadFlags, OUT bool& bWasPackageLoaded)
+{
+	bWasPackageLoaded = false;
+
+	// @TODO: OBJPTR: Want to replicate the functional path of an import here.  See things like FindImportFast in BlueprintSupport.cpp
+	// 		 for additional behavior that we're not handling here yet.
+	FName* ScriptPackageName = FPackageName::FindScriptPackageName(PackageName);
+	UPackage* TargetPackage = (UPackage*)StaticFindObjectFastInternal(UPackage::StaticClass(), nullptr, PackageName);
+	if (UObjectRedirector* Redirector = dynamic_cast<UObjectRedirector*>(TargetPackage))
+	{
+		TargetPackage = (UPackage*)Redirector->DestinationObject;
+	}
+	if (!ScriptPackageName && !TargetPackage)
+	{
+		// @TODO: OBJPTR: When using the "external package" feature, we will have objects that have a differing package path vs "outer hierarchy" path
+		//				  The package path should be used when loading.  The "outer hierarchy" path may need to be used when finding existing objects in memory.
+		//				  This will need further evaluation and testing before lazy load can be enabled.
+		// @TODO: OBJPTR: Instancing context may be important to consider when loading the package.
+		if (FLinkerLoad::IsKnownMissingPackage(PackageName))
+		{
+			return nullptr;
+		}
+		LoadFlags |= LOAD_NoWarn | LOAD_NoVerify; //This does nothing? | LOAD_DisableDependencyPreloading;
+		TargetPackage = LoadPackage(nullptr, *PackageName.ToString(), LoadFlags);
+		bWasPackageLoaded = true;
+	}
+	return TargetPackage;
+}
+
+UClass* FObjectRef::ResolveObjectRefClass(uint32 LoadFlags /*= LOAD_None*/) const
+{
+	TRACE_CPUPROFILER_EVENT_SCOPE(ResolveObjectRef);
+	UClass* ClassObject = nullptr;
+	UPackage* ClassPackage = nullptr;
+	if (!ClassPackageName.IsNone())
+	{
+		bool bWasPackageLoaded = false;
+		ClassPackage = FindOrLoadPackage(ClassPackageName, LoadFlags, bWasPackageLoaded);
+
+		if (!ClassName.IsNone())
+		{
+			ClassObject = (UClass*)StaticFindObjectFastInternal(UClass::StaticClass(), ClassPackage, ClassName);
+			if (UObjectRedirector* Redirector = dynamic_cast<UObjectRedirector*>(ClassObject))
+			{
+				ClassObject = (UClass*)Redirector->DestinationObject;
+				if (ClassObject != nullptr)
+				{
+					ClassPackage = ClassObject->GetPackage();
+				}
+			}
+
+			if (ClassObject)
+			{
+				if (ClassObject->HasAnyFlags(RF_NeedLoad) && ClassPackage->GetLinker())
+				{
+					ClassPackage->GetLinker()->Preload(ClassObject);
+				}
+				ClassObject->GetDefaultObject(); // build the CDO if it isn't already built
+			}
+		}
+	}
+
+	UE::CoreUObject::Private::OnClassReferenceResolved(*this, ClassPackage, ClassObject);
+	return ClassObject;
+}
+
+class FFullyLoadPackageOnHandleResolveTask
+{
+public:
+	FFullyLoadPackageOnHandleResolveTask(UPackage* InPackage) : Package(InPackage)
+	{
+	}
+
+	FORCEINLINE TStatId GetStatId() const { RETURN_QUICK_DECLARE_CYCLE_STAT(FFullyLoadPackageOnHandleResolveTask, STATGROUP_TaskGraphTasks); }
+	ENamedThreads::Type GetDesiredThread() { return ENamedThreads::Type::GameThread; }
+	static ESubsequentsMode::Type GetSubsequentsMode() { return ESubsequentsMode::FireAndForget; }
+
+	void DoTask(ENamedThreads::Type CurrentThread, const FGraphEventRef& MyCompletionGraphEvent)
+	{
+		Package->GetLinker()->LoadAllObjects(true);
+	}
+
+private:
+	UPackage* Package = nullptr;
+};
+
+UObject* FObjectRef::Resolve(uint32 LoadFlags /*= LOAD_None*/) const
+{
+	TRACE_CPUPROFILER_EVENT_SCOPE(ResolveObjectRef);
+
+	UE::CoreUObject::Private::FObjectPathId ObjectPath = GetObjectPath();
+	if (IsNull() || !ObjectPath.IsValid())
+	{
+		UE::CoreUObject::Private::OnReferenceResolved(*this, nullptr, nullptr);
+		return nullptr;
+	}
+
+	bool bWasObjectOrPackageLoaded = false;
+
+	UPackage* TargetPackage = FindOrLoadPackage(PackageName, LoadFlags, bWasObjectOrPackageLoaded);
+
+	if (!TargetPackage)
+	{
+		UE::CoreUObject::Private::OnReferenceResolved(*this, nullptr, nullptr);
+		return nullptr;
+	}
+
+	UE::CoreUObject::Private::FObjectPathId::ResolvedNameContainerType ResolvedNames;
+	ObjectPath.Resolve(ResolvedNames);
+
+	UObject* CurrentObject = TargetPackage;
+	for (int32 ObjectPathIndex = 0; ObjectPathIndex < ResolvedNames.Num(); ++ObjectPathIndex)
+	{
+		UObject* PreviousOuter = CurrentObject;
+		CurrentObject = StaticFindObjectFastInternal(nullptr, CurrentObject, ResolvedNames[ObjectPathIndex]);
+		if (UObjectRedirector* Redirector = dynamic_cast<UObjectRedirector*>(CurrentObject))
+		{
+			CurrentObject = Redirector->DestinationObject;
+			if (CurrentObject != nullptr)
+			{
+				TargetPackage = CurrentObject->GetPackage();
+			}
+		}
+
+		if (!CurrentObject && !TargetPackage->IsFullyLoaded() && TargetPackage->GetLinker() && TargetPackage->GetLinker()->IsLoading())
+		{
+			if (IsInAsyncLoadingThread() || IsInGameThread())
+			{
+				TargetPackage->GetLinker()->LoadAllObjects(true);
+				bWasObjectOrPackageLoaded = true;
+			}
+			else
+			{
+				// Shunt the load request to happen on the game thread and block on its completion.  This is a deadlock risk!  The game thread may be blocked waiting on this thread.
+				UE_LOG(LogObjectRef, Warning, TEXT("Resolve of object in package '%s' from a non-game thread was shunted to the game thread."), *PackageName.ToString());
+				TGraphTask<FFullyLoadPackageOnHandleResolveTask>::CreateTask().ConstructAndDispatchWhenReady(TargetPackage)->Wait();
+				bWasObjectOrPackageLoaded = true;
+			}
+
+			CurrentObject = StaticFindObjectFastInternal(nullptr, PreviousOuter, ResolvedNames[ObjectPathIndex]);
+			if (UObjectRedirector* Redirector = dynamic_cast<UObjectRedirector*>(CurrentObject))
+			{
+				CurrentObject = Redirector->DestinationObject;
+				if (CurrentObject != nullptr)
+				{
+					TargetPackage = CurrentObject->GetPackage();
+				}
+			}
+		}
+
+		if (!CurrentObject)
+		{
+			UE::CoreUObject::Private::OnReferenceResolved(*this, TargetPackage, nullptr);
+			return nullptr;
+		}
+	}
+
+	if (CurrentObject->HasAnyFlags(RF_NeedLoad) && TargetPackage->GetLinker())
+	{
+		if (IsInAsyncLoadingThread() || IsInGameThread())
+		{
+			TargetPackage->GetLinker()->LoadAllObjects(true);
+			bWasObjectOrPackageLoaded = true;
+		}
+		else
+		{
+			// Shunt the load request to happen on the game thread and block on its completion.  This is a deadlock risk!  The game thread may be blocked waiting on this thread.
+			UE_LOG(LogObjectRef, Warning, TEXT("Resolve of object in package '%s' from a non-game thread was shunted to the game thread."), *PackageName.ToString());
+			TGraphTask<FFullyLoadPackageOnHandleResolveTask>::CreateTask().ConstructAndDispatchWhenReady(TargetPackage)->Wait();
+			bWasObjectOrPackageLoaded = true;
+		}
+	}
+	UE::CoreUObject::Private::OnReferenceResolved(*this, TargetPackage, CurrentObject);
+
+	if (bWasObjectOrPackageLoaded)
+	{
+		UE::CoreUObject::Private::OnReferenceLoaded(*this, TargetPackage, CurrentObject);
+	}
+
+	return CurrentObject;
+}
+
+
+
+#endif
\ No newline at end of file
diff --git a/Engine/Source/Runtime/CoreUObject/Private/UObject/UObjectBase.cpp b/Engine/Source/Runtime/CoreUObject/Private/UObject/UObjectBase.cpp
--- a/Engine/Source/Runtime/CoreUObject/Private/UObject/UObjectBase.cpp
+++ b/Engine/Source/Runtime/CoreUObject/Private/UObject/UObjectBase.cpp
@@ -152,7 +152,7 @@ UObjectBase::~UObjectBase()
 		check(IsValidLowLevel());
 		check(GetFName() == NAME_None);
 #if UE_WITH_OBJECT_HANDLE_LATE_RESOLVE
-		ObjectHandle_Private::FreeObjectHandle(*this);
+		UE::CoreUObject::Private::FreeObjectHandle(*this);
 #endif 
 		GUObjectArray.FreeUObjectIndex(this);
 	}
@@ -1059,7 +1059,7 @@ void UObjectBaseInit()
 	GUObjectAllocator.AllocatePermanentObjectPool(SizeOfPermanentObjectPool);
 	GUObjectArray.AllocateObjectPool(MaxUObjects, MaxObjectsNotConsideredByGC, bPreAllocateUObjectArray);
 #if UE_WITH_OBJECT_HANDLE_LATE_RESOLVE
-	ObjectHandle_Private::InitObjectHandles(MaxUObjects);
+	UE::CoreUObject::Private::InitObjectHandles(MaxUObjects);
 #endif
 
 	void InitNoPendingKill();
diff --git a/Engine/Source/Runtime/CoreUObject/Private/UObject/UObjectGlobals.cpp b/Engine/Source/Runtime/CoreUObject/Private/UObject/UObjectGlobals.cpp
--- a/Engine/Source/Runtime/CoreUObject/Private/UObject/UObjectGlobals.cpp
+++ b/Engine/Source/Runtime/CoreUObject/Private/UObject/UObjectGlobals.cpp
@@ -1015,7 +1015,7 @@ UPackage* CreatePackage(const TCHAR* PackageName )
 #if UE_WITH_OBJECT_HANDLE_LATE_RESOLVE
 				if (Result->GetClass()->GetFName() != NAME_None)
 				{
-					MakePackedObjectRef(Result);
+					UE::CoreUObject::Private::MakePackedObjectRef(Result);
 				}
 #endif
 #if WITH_EDITOR
diff --git a/Engine/Source/Runtime/CoreUObject/Public/Templates/Casts.h b/Engine/Source/Runtime/CoreUObject/Public/Templates/Casts.h
--- a/Engine/Source/Runtime/CoreUObject/Public/Templates/Casts.h
+++ b/Engine/Source/Runtime/CoreUObject/Public/Templates/Casts.h
@@ -260,10 +260,10 @@ FORCEINLINE To* Cast(const TObjectPtr<From>& InSrc)
 	}
 	else if constexpr (TIsIInterface<To>::Value)
 	{
-		UObject* SrcObj = ResolveObjectHandleNoRead(Src.GetHandleRef());
+		UObject* SrcObj = UE::CoreUObject::Private::ResolveObjectHandleNoRead(Src.GetHandleRef());
 		if (SrcObj)
 		{
-			ObjectHandle_Private::OnHandleRead(SrcObj);
+			UE::CoreUObject::Private::OnHandleRead(SrcObj);
 			return (To*)Src.Get()->GetInterfaceAddress(To::UClassType::StaticClass());
 		}
 	}
@@ -283,10 +283,10 @@ FORCEINLINE To* ExactCast(const TObjectPtr<From>& Src)
 {
 	static_assert(sizeof(To) > 0, "Attempting to cast to an incomplete type");
 
-	UObject* SrcObj = ResolveObjectHandleNoRead(((const FObjectPtr&)Src).GetHandleRef());
+	UObject* SrcObj = UE::CoreUObject::Private::ResolveObjectHandleNoRead(((const FObjectPtr&)Src).GetHandleRef());
 	if (SrcObj && (SrcObj->GetClass() == To::StaticClass()))
 	{
-		ObjectHandle_Private::OnHandleRead(SrcObj);
+		UE::CoreUObject::Private::OnHandleRead(SrcObj);
 		return (To*)SrcObj;
 	}
 	return nullptr;
@@ -322,8 +322,8 @@ FORCEINLINE To* CastChecked(const TObjectPtr<From>& Src, ECastCheckedType::Type
 	}
 	else if constexpr (TIsIInterface<To>::Value)
 	{
-		UObject* SrcObj = ResolveObjectHandleNoRead(((const FObjectPtr&)Src).GetHandleRef());
-		ObjectHandle_Private::OnHandleRead(SrcObj);
+		UObject* SrcObj = UE::CoreUObject::Private::ResolveObjectHandleNoRead(((const FObjectPtr&)Src).GetHandleRef());
+		UE::CoreUObject::Private::OnHandleRead(SrcObj);
 		return (To*)((const FObjectPtr&)Src).Get()->GetInterfaceAddress(To::UClassType::StaticClass());
 	}
 	else
diff --git a/Engine/Source/Runtime/CoreUObject/Public/UObject/ObjectHandle.h b/Engine/Source/Runtime/CoreUObject/Public/UObject/ObjectHandle.h
--- a/Engine/Source/Runtime/CoreUObject/Public/UObject/ObjectHandle.h
+++ b/Engine/Source/Runtime/CoreUObject/Public/UObject/ObjectHandle.h
@@ -7,149 +7,31 @@
 #include "CoreTypes.h"
 #include "HAL/CriticalSection.h"
 #include "HAL/Platform.h"
-#include "Logging/LogMacros.h"
 #include "Misc/AssertionMacros.h"
 #include "Templates/TypeHash.h"
 #include "UObject/NameTypes.h"
 #include "UObject/Object.h"
 #include "UObject/ObjectMacros.h"
-#include "UObject/ObjectPathId.h"
+#include "UObject/ObjectRef.h"
+#include "UObject/PackedObjectRef.h"
+#include "UObject/ObjectHandleTracking.h"
+
 
 class UClass;
 class UPackage;
 
-#define UE_WITH_OBJECT_HANDLE_LATE_RESOLVE WITH_EDITORONLY_DATA
-#define UE_WITH_OBJECT_HANDLE_TRACKING WITH_EDITORONLY_DATA
-
-DECLARE_LOG_CATEGORY_EXTERN(LogObjectHandle, Log, All);
-
-/**
- * FObjectRef represents a heavyweight reference that contains the specific pieces of information needed to reference an object
- * (or null) that may or may not be loaded yet.  The expectation is that given the imports of a package we have loaded, we should
- * be able to create an FObjectRef to objects in a package we haven't yet loaded.  For this reason, FObjectRef has to be derivable
- * from the serialized (not transient) contents of an FObjectImport.
- */
-struct FObjectRef
-{
-	FName PackageName;
-	FName ClassPackageName;
-	FName ClassName;
-	FObjectPathId ObjectPath;
-
-	bool operator==(const FObjectRef&& Other) const
-	{
-		return (PackageName == Other.PackageName) && (ObjectPath == Other.ObjectPath) && (ClassPackageName == Other.ClassPackageName) && (ClassName == Other.ClassName);
-	}
-
-	/** Returns the path to the class for this object. */
-	FString GetClassPathName(EObjectFullNameFlags Flags = EObjectFullNameFlags::IncludeClassPackage) const
-	{
-		TStringBuilder<FName::StringBufferSize> Result;
-		AppendClassPathName(Result, Flags);
-		return FString(Result);
-	}
-
-	/** Appends the path to the class for this object to the builder, does not reset builder. */
-	COREUOBJECT_API void AppendClassPathName(FStringBuilderBase& OutClassPathNameBuilder, EObjectFullNameFlags Flags = EObjectFullNameFlags::IncludeClassPackage) const;
-
-	/** Returns the name of the object in the form: ObjectName */
-	COREUOBJECT_API FName GetFName() const;
-
-	/** Returns the full path for the object in the form: ObjectPath */
-	FString GetPathName() const
-	{
-		TStringBuilder<FName::StringBufferSize> Result;
-		AppendPathName(Result);
-		return FString(Result);
-	}
-
-	/** Appends the path to the object to the builder, does not reset builder. */
-	COREUOBJECT_API void AppendPathName(FStringBuilderBase& OutPathNameBuilder) const;
-
-	/** Returns the full name for the object in the form: Class ObjectPath */
-	FString GetFullName(EObjectFullNameFlags Flags = EObjectFullNameFlags::None) const
-	{
-		TStringBuilder<256> FullName;
-		GetFullName(FullName, Flags);
-		return FString(FullName);
-	}
-
-	/** Populates OutFullNameBuilder with the full name for the object in the form: Class ObjectPath */
-	void GetFullName(FStringBuilderBase& OutFullNameBuilder, EObjectFullNameFlags Flags = EObjectFullNameFlags::None) const
-	{
-		OutFullNameBuilder.Reset();
-		AppendClassPathName(OutFullNameBuilder, Flags);
-		OutFullNameBuilder.AppendChar(TEXT(' '));
-		AppendPathName(OutFullNameBuilder);
-	}
-
-	/** Returns the name for the asset in the form: Class'ObjectPath' */
-	FString GetExportTextName() const
-	{
-		TStringBuilder<256> ExportTextName;
-		GetExportTextName(ExportTextName);
-		return FString(ExportTextName);
-	}
-
-	/** Populates OutExportTextNameBuilder with the name for the object in the form: Class'ObjectPath' */
-	void GetExportTextName(FStringBuilderBase& OutExportTextNameBuilder) const
-	{
-		OutExportTextNameBuilder.Reset();
-		AppendClassPathName(OutExportTextNameBuilder);
-		OutExportTextNameBuilder.AppendChar(TEXT('\''));
-		AppendPathName(OutExportTextNameBuilder);
-		OutExportTextNameBuilder.AppendChar(TEXT('\''));
-	}
-};
-
-inline bool IsObjectRefNull(const FObjectRef& ObjectRef) { return ObjectRef.PackageName.IsNone() && ObjectRef.ObjectPath.IsNone(); }
-
-COREUOBJECT_API FObjectRef MakeObjectRef(const UObject* Object);
-COREUOBJECT_API FObjectRef MakeObjectRef(struct FPackedObjectRef ObjectRef);
-COREUOBJECT_API UObject* ResolveObjectRef(const FObjectRef& ObjectRef, uint32 LoadFlags = LOAD_None);
-COREUOBJECT_API UClass* ResolveObjectRefClass(const FObjectRef& ObjectRef, uint32 LoadFlags = LOAD_None);
-
-
-/**
- * FPackedObjectRef represents a lightweight reference that can fit in the space of a pointer and be able to refer to an object
- * (or null) that may or may not be loaded without pointing to its location in memory (even if it is currently loaded).
- */
-struct FPackedObjectRef
-{
-	// Must be 0 for a reference to null.
-	// The least significant bit must always be 1 in a non-null reference.
-	UPTRINT EncodedRef;
-
-	inline bool operator==(FPackedObjectRef RHS) const { return EncodedRef == RHS.EncodedRef; }
-	inline bool operator!=(FPackedObjectRef RHS) const { return EncodedRef != RHS.EncodedRef; }
-	friend inline uint32 GetTypeHash(FPackedObjectRef ObjectRef) { return GetTypeHash(ObjectRef.EncodedRef); }
-};
-
-inline bool IsPackedObjectRefNull(FPackedObjectRef ObjectRef) { return !ObjectRef.EncodedRef; }
-
-#if UE_WITH_OBJECT_HANDLE_LATE_RESOLVE
-COREUOBJECT_API FPackedObjectRef MakePackedObjectRef(const UObject* Object);
-#endif 
-
-COREUOBJECT_API FPackedObjectRef MakePackedObjectRef(const FObjectRef& ObjectRef);
-COREUOBJECT_API UObject* ResolvePackedObjectRef(FPackedObjectRef ObjectRef, uint32 LoadFlags = LOAD_None);
-COREUOBJECT_API UClass* ResolvePackedObjectRefClass(FPackedObjectRef ObjectRef, uint32 LoadFlags = LOAD_None);
-
 /**
  * FObjectHandle is either a packed object ref or the resolved pointer to an object.  Depending on configuration
  * when you create a handle, it may immediately be resolved to a pointer.
  */
 #if UE_WITH_OBJECT_HANDLE_LATE_RESOLVE
 
-namespace ObjectHandle_Private
+namespace UE::CoreUObject::Private
 {
-	void UpdateRenamedObject(const UObject& Obj, FName NewName, UObject* NewOuter);
-	void InitObjectHandles(int32 Size);
-	void FreeObjectHandle(const UObjectBase& Object);
+	struct FObjectHandlePrivate;
 }
 
-struct FObjectHandleInternal;
-using FObjectHandle = FObjectHandleInternal;
+using FObjectHandle = UE::CoreUObject::Private::FObjectHandlePrivate;
 
 #else
 
@@ -158,437 +40,290 @@ using FObjectHandle = UObject*;
 
 #endif
 
-inline FObjectHandle MakeObjectHandle(FPackedObjectRef ObjectRef);
-inline FObjectHandle MakeObjectHandle(const FObjectRef& ObjectRef);
-inline FObjectHandle MakeObjectHandle(UObject* Object);
-
 inline bool IsObjectHandleNull(FObjectHandle Handle);
 inline bool IsObjectHandleResolved(FObjectHandle Handle);
 
-/** Read the handle as a pointer without checking if it is resolved. Invalid to call for unresolved handles. */
-inline UObject* ReadObjectHandlePointerNoCheck(FObjectHandle Handle);
-/** Read the handle as a packed object ref without checking if it is unresolved. Invalid to call for resolved handles. */
-inline FPackedObjectRef ReadObjectHandlePackedObjectRefNoCheck(FObjectHandle Handle);
-
-inline UObject* ResolveObjectHandle(FObjectHandle& Handle);
-inline UObject* ResolveObjectHandleNoRead(FObjectHandle& Handle);
-
-/** Resolves an ObjectHandle without checking if already resolved. Invalid to call for resolved handles */
-inline UObject* ResolveObjectHandleNoReadNoCheck(FObjectHandle& Handle);
-inline UClass* ResolveObjectHandleClass(FObjectHandle Handle);
-
-/** Read the handle as a pointer if resolved, and otherwise return null. */
-inline UObject* ReadObjectHandlePointer(FObjectHandle Handle);
-/** Read the handle as a packed object ref if unresolved, and otherwise return the null packed object ref. */
-inline FPackedObjectRef ReadObjectHandlePackedObjectRef(FObjectHandle Handle);
-
-#if UE_WITH_OBJECT_HANDLE_LATE_RESOLVE
-inline FObjectRef MakeObjectRef(FObjectHandle Handle);
-
-UE_DEPRECATED(5.2, "Use `bool TryMakePackedObjectRef(FObjectHandle, FPackedObjectRef&)` instead")
-inline FPackedObjectRef MakePackedObjectRef(FObjectHandle Handle);
-inline FPackedObjectRef GetPackedObjectRef(FObjectHandle Handle);
-COREUOBJECT_API FPackedObjectRef GetPackedObjectRef(const UObject& Object);
-
-struct FObjectHandleInternal
+//Private functions that forced public due to inlining.
+namespace UE::CoreUObject::Private
 {
-	UPTRINT PointerOrRef;
+#if UE_WITH_OBJECT_HANDLE_LATE_RESOLVE
 
-	friend inline bool operator==(FObjectHandle LHS, FObjectHandle RHS)
+	struct FObjectHandlePrivate
 	{
-		bool LhsResolved = IsObjectHandleResolved(LHS);
-		bool RhsResolved = IsObjectHandleResolved(RHS);
+		//Stores either FPackedObjectRef or a UObject*
+		UPTRINT PointerOrRef;
 
-		//if both resolved or both unresolved compare the uintptr
-		if (LhsResolved == RhsResolved)
+		explicit inline operator bool() const
 		{
-			return LHS.PointerOrRef == RHS.PointerOrRef;
+			return PointerOrRef != 0;
 		}
-		
-		//only one side can be resolved
-		if (LhsResolved)
-		{
-			//both sides can't be null as resolved status would have be true for both
-			const UObject* Obj = ReadObjectHandlePointerNoCheck(LHS);
-			if (!Obj)
-			{
-				return false;
-			}
-
-			//if packed ref empty then can't be equal as RHS is an unresolved pointer
-			FPackedObjectRef PackedLhs = GetPackedObjectRef(*Obj);
-			if (PackedLhs.EncodedRef == 0)
-			{
-				return false; 
-			}
-			return PackedLhs.EncodedRef == RHS.PointerOrRef;
-				
-		}
-		else
-		{
-			//both sides can't be null as resolved status would have be true for both
-			const UObject* Obj = ReadObjectHandlePointerNoCheck(RHS);
-			if (!Obj)
-			{
-				return false;
-			}
-
-			//if packed ref empty then can't be equal as RHS is an unresolved pointer
-			FPackedObjectRef PackedRhs = GetPackedObjectRef(*Obj);
-			if (PackedRhs.EncodedRef == 0)
-			{
-				return false;
-			}
-			return PackedRhs.EncodedRef == LHS.PointerOrRef;
-		}
-		
-	}
+	};
 
-	friend inline bool operator!=(FObjectHandle LHS, FObjectHandle RHS)
-	{
-		return !(LHS == RHS);
-	}
-	friend inline uint32 GetTypeHash(FObjectHandleInternal Handle)
-	{
-		if (Handle.PointerOrRef == 0)
-		{
-			return 0;
-		}
+	/* Returns the packed object ref for this object IF one exists otherwise returns a null PackedObjectRef */
+	COREUOBJECT_API FPackedObjectRef GetPackedObjectRef(const UObject& Object);
 
-		if (IsObjectHandleResolved(Handle))
-		{
-			const UObject* Obj = ReadObjectHandlePointerNoCheck(Handle);
-			FPackedObjectRef PackedObjectRef = GetPackedObjectRef(*Obj);
-			if (PackedObjectRef.EncodedRef == 0)
-			{
-				return GetTypeHash(Obj);
-			}
-			return GetTypeHash(PackedObjectRef.EncodedRef);
-		}
-		return GetTypeHash(Handle.PointerOrRef);
-	}
+	/* Creates and ObjectRef from a packed object ref*/
+	COREUOBJECT_API FObjectRef MakeObjectRef(FPackedObjectRef Handle);
 
-	explicit inline operator bool() const
-	{
-		return PointerOrRef != 0;
-	}
-};
 #endif
 
+	///these functions are always defined regardless of UE_WITH_OBJECT_HANDLE_LATE_RESOLVE value
 
-/**
- * FObjectHandles can optionally support tracking.  Because of the low level nature of object handles, anything that
- * registers itself for these callbacks should ensure that it is:
- * 1) error free (ie: should not cause exceptions even in unusual circumstances)
- * 2) fault tolerant (ie: could be called at a time when an exception has happened)
- * 3) thread-safe (ie: could be called from any thread)
- * 4) high performance (ie: will be called many times)
- */
-#if UE_WITH_OBJECT_HANDLE_TRACKING
-
- /**
-  * Callback notifying when an object value is read from a handle.  Fired regardless of whether the handle
-  * was resolved as part of the read operation or not and whether the object being read is null or not.
-  *
-  * @param ReadObject	The object that was read from a handle.
-  */
-DECLARE_DELEGATE_OneParam(FObjectHandleReadDelegate, UObject* ReadObject);
+	/* Makes a resolved FObjectHandle from an UObject. */
+	inline FObjectHandle MakeObjectHandle(UObject* Object);
 
-/**
- * Callback notifying when a class is resolved from an object handle or object reference.
- * Classes are resolved either independently for a given handle/reference or as part of each object resolve.
- *
- * @param SourceRef		The object reference (either user provided or internally computed from a handle) from which the class was resolved.
- * @param ClassPackage	The package containing the resolved class.
- * @param Class			The resolved class.
- */
-DECLARE_DELEGATE_ThreeParams(FObjectHandleClassResolvedDelegate, const FObjectRef& SourceRef, UPackage* ObjectPackage, UClass* Class);
+	/* Returns the UObject from Handle and the handle is updated cache the resolved UObject */
+	inline UObject* ResolveObjectHandle(FObjectHandle& Handle);
 
- /**
-  * Callback notifying when a object handle is resolved.
-  *
-  * @param SourceRef		The object reference (either user provided or internally computed from a handle) from which the class was resolved.
-  * @param ClassPackage	The package containing the resolved class.
-  * @param Class			The resolved class.
-  */
-DECLARE_DELEGATE_ThreeParams(FObjectHandleReferenceResolvedDelegate, const FObjectRef& SourceRef, UPackage* ObjectPackage, UObject* Object);
+	/* Returns the UClass for UObject store in Handle. Handle is not resolved */
+	inline UClass* ResolveObjectHandleClass(FObjectHandle Handle);
 
-/**
- * Callback notifying when an object was loaded through an object handle.  Will not notify you about global object loads, just ones that occur
- * as the byproduct of resolving an ObjectHandle.
- *
- * @param SourceRef		The object reference (either user provided or internally computed from a handle) from which the class was resolved.
- * @param ClassPackage	The package containing the resolved class.
- * @param Class			The resolved class.
- */
-DECLARE_DELEGATE_ThreeParams(FObjectHandleReferenceLoadedDelegate, const FObjectRef& SourceRef, UPackage* ObjectPackage, UObject* Object);
+	/* Returns the UObject from Handle and the handle is updated cache the resolved UObject.
+	 * Does not cause ObjectHandleTracking to fire a read event
+	 */
+	inline UObject* ResolveObjectHandleNoRead(FObjectHandle& Handle);
 
-/**
- * Installs a new callback for notifications that an object value has been read from a handle.
- *
- * @param Function		The new handle read callback to install.
- * @return				The DelegateHandle so that you can remove the callback at a later date.
- */
-COREUOBJECT_API FDelegateHandle AddObjectHandleReadCallback(FObjectHandleReadDelegate Delegate);
-COREUOBJECT_API void RemoveObjectHandleReadCallback(FDelegateHandle DelegateHandle);
-
-/**
- * Installs a new callback for notifications that a class has been resolved from an object handle or object reference.
- *
- * @param Function		The new class resolved callback to install.
- * @return				The DelegateHandle so that you can remove the callback at a later date.
- */
-COREUOBJECT_API FDelegateHandle AddObjectHandleClassResolvedCallback(FObjectHandleClassResolvedDelegate Callback);
-COREUOBJECT_API void RemoveObjectHandleClassResolvedCallback(FDelegateHandle DelegateHandle);
-
-/**
- * Installs a new callback for notifications that an object has been resolved from an object handle or object reference.
- *
- * @param Function		The new object resolved callback to install.
- * @return				The DelegateHandle so that you can remove the callback at a later date.
- */
-COREUOBJECT_API FDelegateHandle AddObjectHandleReferenceResolvedCallback(FObjectHandleReferenceResolvedDelegate Callback);
-COREUOBJECT_API void RemoveObjectHandleReferenceResolvedCallback(FDelegateHandle DelegateHandle);
-
-/**
- * Installs a new callback for notifications that an object has been loaded from an object handle or object reference.
- *
- * @param Function		The new object resolved callback to install.
- * @return				The DelegateHandle so that you can remove the callback at a later date.
- */
-COREUOBJECT_API FDelegateHandle AddObjectHandleReferenceLoadedCallback(FObjectHandleReferenceLoadedDelegate Callback);
-COREUOBJECT_API void RemoveObjectHandleReferenceLoadedCallback(FDelegateHandle DelegateHandle);
-#endif
+	/** Resolves an ObjectHandle without checking if already resolved. Invalid to call for resolved handles */
+	inline UObject* ResolveObjectHandleNoReadNoCheck(FObjectHandle& Handle);
 
+	/** Read the handle as a pointer without checking if it is resolved. Invalid to call for unresolved handles. */
+	inline UObject* ReadObjectHandlePointerNoCheck(FObjectHandle Handle);
+}
 
-///////////////////////////////////////////////////////////////////////////////////////////////////////////////////////
+/////////////////////////////////////////////////////////////////////////////////////////////////////////////////////////
 
-struct FObjectHandlePackageDebugData
+/* return true if handle is null */
+inline bool IsObjectHandleNull(FObjectHandle Handle)
 {
-	FMinimalName PackageName;
-	FScriptArray ObjectPaths;
-	FScriptArray DataClassDescriptors;
-	uint8 _Padding[sizeof(FRWLock)];
-};
+#if UE_WITH_OBJECT_HANDLE_LATE_RESOLVE
+	return !Handle.PointerOrRef;
+#else
+	return !Handle;
+#endif
+}
 
-struct FObjectHandleDataClassDescriptor
+/* checks if a handle is resolved. 
+ * nullptr counts as resolved
+ * all handles are resolved when late resolved is off
+ */
+inline bool IsObjectHandleResolved(FObjectHandle Handle)
 {
-	FMinimalName PackageName;
-	FMinimalName ClassName;
-};
+#if UE_WITH_OBJECT_HANDLE_LATE_RESOLVE
+	return !(Handle.PointerOrRef & 1);
+#else
+	return true;
+#endif
+}
 
-namespace ObjectHandle_Private
+/////////////////////////////////////////////////////////////////////////////////////////////////////////////////////////
+#if UE_WITH_OBJECT_HANDLE_LATE_RESOLVE
+inline bool operator==(UE::CoreUObject::Private::FObjectHandlePrivate LHS, UE::CoreUObject::Private::FObjectHandlePrivate RHS)
 {
-	inline constexpr uint32 ObjectPathIdShift = 1;
-	inline constexpr uint32 ObjectPathIdMask = 0x00FF'FFFF;
-
-	inline constexpr uint32 DataClassDescriptorIdShift = 25;
-	inline constexpr uint32 DataClassDescriptorIdMask = 0x0000'00FF;
-
-	inline constexpr uint32 PackageIdShift = 33;
-	inline constexpr uint32 PackageIdMask = 0x7FFF'FFFF;
+	using namespace UE::CoreUObject::Private;
 
-#if UE_WITH_OBJECT_HANDLE_TRACKING
-	DECLARE_MULTICAST_DELEGATE_OneParam(FObjectHandleReadEvent, UObject* Object);
-	DECLARE_MULTICAST_DELEGATE_ThreeParams(FClassReferenceResolvedEvent, const FObjectRef& ObjectRef, UPackage* Package, UClass* Class);
-	DECLARE_MULTICAST_DELEGATE_ThreeParams(FObjectHandleReferenceResolvedEvent, const FObjectRef& SourceRef, UPackage* ObjectPackage, UObject* Object);
-	DECLARE_MULTICAST_DELEGATE_ThreeParams(FObjectHandleReferenceLoadedEvent, const FObjectRef& SourceRef, UPackage* ObjectPackage, UObject* Object);
+	bool LhsResolved = IsObjectHandleResolved(LHS);
+	bool RhsResolved = IsObjectHandleResolved(RHS);
 
-	struct FObjectHandleEvents
+	//if both resolved or both unresolved compare the uintptr
+	if (LhsResolved == RhsResolved)
 	{
-		FObjectHandleReadEvent ObjectHandleReadEvent;
-		FClassReferenceResolvedEvent ClassReferenceResolvedEvent;
-		FObjectHandleReferenceResolvedEvent ObjectHandleReferenceResolvedEvent;
-		FObjectHandleReferenceLoadedEvent ObjectHandleReferenceLoadedEvent;
-
-		inline void BeginUsing() { UsingCount++; }
-		inline void EndUsing() { --UsingCount; }
-		inline bool IsUsing() const { return UsingCount > 0; }
-
-	private:
-		std::atomic<int32> UsingCount;
-	};
-
-	extern COREUOBJECT_API std::atomic<int32> ObjectHandleEventIndex;
-	extern COREUOBJECT_API FObjectHandleEvents ObjectHandleEvents[2];
+		return LHS.PointerOrRef == RHS.PointerOrRef;
+	}
 
-	inline FObjectHandleEvents& BeginReadingEvents()
+	//only one side can be resolved
+	if (LhsResolved)
 	{
-		// Quick spin lock to try and begin using the events, normally this wont actually spin, in like 99.999% of the time.
-		while (true)
+		//both sides can't be null as resolved status would have be true for both
+		const UObject* Obj = ReadObjectHandlePointerNoCheck(LHS);
+		if (!Obj)
 		{
-			// Start by getting the current event index from the double buffer.
-			const int32 InitialEventIndex = ObjectHandleEventIndex;
-			// Grab the event set at that index.
-			FObjectHandleEvents& Events = ObjectHandleEvents[InitialEventIndex];
-			// Begin using the events, this will signal we're using them.
-			Events.BeginUsing();
-			// Ok - now we're going to check that we're *still* using the same index.  If we are
-			// then we know that it didn't change out from under us in between getting Events, and calling BeginUsing().
-			if (InitialEventIndex == ObjectHandleEventIndex)
-			{
-				return Events;
-			}
-			// If the check above failed, then we need to cease using the events, because we're about to try this again.
-			Events.EndUsing();
+			return false;
 		}
-	}
 
-	inline void OnHandleRead(UObject* Object)
-	{
-		FObjectHandleEvents& Events = BeginReadingEvents();
-		Events.ObjectHandleReadEvent.Broadcast(Object);
-		Events.EndUsing();
+		//if packed ref empty then can't be equal as RHS is an unresolved pointer
+		FPackedObjectRef PackedLhs = GetPackedObjectRef(*Obj);
+		if (PackedLhs.EncodedRef == 0)
+		{
+			return false;
+		}
+		return PackedLhs.EncodedRef == RHS.PointerOrRef;
+
 	}
-	
-	inline void OnClassReferenceResolved(const FObjectRef& ObjectRef, UPackage* Package, UClass* Class)
+	else
 	{
-		FObjectHandleEvents& Events = BeginReadingEvents();
-		Events.ClassReferenceResolvedEvent.Broadcast(ObjectRef, Package, Class);
-		Events.EndUsing();
-	}
+		//both sides can't be null as resolved status would have be true for both
+		const UObject* Obj = ReadObjectHandlePointerNoCheck(RHS);
+		if (!Obj)
+		{
+			return false;
+		}
 
-	inline void OnReferenceResolved(const FObjectRef& ObjectRef, UPackage* Package, UObject* Object)
-	{
-		FObjectHandleEvents& Events = BeginReadingEvents();
-		Events.ObjectHandleReferenceResolvedEvent.Broadcast(ObjectRef, Package, Object);
-		Events.EndUsing();
+		//if packed ref empty then can't be equal as RHS is an unresolved pointer
+		FPackedObjectRef PackedRhs = GetPackedObjectRef(*Obj);
+		if (PackedRhs.EncodedRef == 0)
+		{
+			return false;
+		}
+		return PackedRhs.EncodedRef == LHS.PointerOrRef;
 	}
 
-	inline void OnReferenceLoaded(const FObjectRef& ObjectRef, UPackage* Package, UObject* Object)
-	{
-		FObjectHandleEvents& Events = BeginReadingEvents();
-		Events.ObjectHandleReferenceLoadedEvent.Broadcast(ObjectRef, Package, Object);
-		Events.EndUsing();
-	}
-#else
-	inline void OnHandleRead(UObject* Object) {}
-	inline void OnClassReferenceResolved(const FObjectRef& ObjectRef, UPackage* Package, UClass* Class) {}
-	inline void OnReferenceResolved(const FObjectRef& ObjectRef, UPackage* Package, UObject* Object) {}
-	inline void OnReferenceLoaded(const FObjectRef& ObjectRef, UPackage* Package, UObject* Object) {}
-#endif
 }
 
-#if UE_WITH_OBJECT_HANDLE_LATE_RESOLVE
-
-inline bool IsObjectHandleNull(FObjectHandle Handle) { return !Handle.PointerOrRef; }
-inline bool IsObjectHandleResolved(FObjectHandle Handle) { return !(Handle.PointerOrRef & 1); }
-
-inline UObject* ReadObjectHandlePointerNoCheck(FObjectHandle Handle)
+inline bool operator!=(UE::CoreUObject::Private::FObjectHandlePrivate LHS, UE::CoreUObject::Private::FObjectHandlePrivate RHS)
 {
-	return reinterpret_cast<UObject*>(Handle.PointerOrRef);
+	return !(LHS == RHS);
 }
 
-inline FPackedObjectRef ReadObjectHandlePackedObjectRefNoCheck(FObjectHandle Handle)
+inline uint32 GetTypeHash(UE::CoreUObject::Private::FObjectHandlePrivate Handle)
 {
-	return {Handle.PointerOrRef};
-}
-
-inline FObjectHandle MakeObjectHandle(FPackedObjectRef ObjectRef) { return {ObjectRef.EncodedRef}; }
-inline FObjectHandle MakeObjectHandle(const FObjectRef& ObjectRef) { return MakeObjectHandle(MakePackedObjectRef(ObjectRef)); }
-inline FObjectHandle MakeObjectHandle(UObject* Object) { return {UPTRINT(Object)}; }
-
-#else
-
-inline bool IsObjectHandleNull(FObjectHandle Handle) { return !Handle; }
-inline bool IsObjectHandleResolved(FObjectHandle Handle) { return true; }
-
-inline UObject* ReadObjectHandlePointerNoCheck(FObjectHandle Handle) { return Handle; }
-inline FPackedObjectRef ReadObjectHandlePackedObjectRefNoCheck(FObjectHandle Handle) { return FPackedObjectRef(); }
-
-inline FObjectHandle MakeObjectHandle(FPackedObjectRef ObjectRef) { return ResolvePackedObjectRef(ObjectRef); }
-inline FObjectHandle MakeObjectHandle(const FObjectRef& ObjectRef) { return ResolveObjectRef(ObjectRef); }
-inline FObjectHandle MakeObjectHandle(UObject* Object) { return Object; }
-
-#endif
+	using namespace UE::CoreUObject::Private;
 
-inline UObject* ResolveObjectHandleNoRead(FObjectHandle& Handle)
-{
-#if UE_WITH_OBJECT_HANDLE_LATE_RESOLVE
-	FObjectHandle LocalHandle = Handle;
-	if (IsObjectHandleResolved(LocalHandle))
+	if (Handle.PointerOrRef == 0)
 	{
-		UObject* ResolvedObject = ReadObjectHandlePointerNoCheck(LocalHandle);
-		return ResolvedObject;
+		return 0;
 	}
-	else
+
+	if (IsObjectHandleResolved(Handle))
 	{
-		LocalHandle = MakeObjectHandle(ResolvePackedObjectRef(ReadObjectHandlePackedObjectRefNoCheck(LocalHandle)));
-		UObject* ResolvedObject = ReadObjectHandlePointerNoCheck(LocalHandle);
-		Handle = LocalHandle;
-		return ResolvedObject;
+		const UObject* Obj = ReadObjectHandlePointerNoCheck(Handle);
+		FPackedObjectRef PackedObjectRef = GetPackedObjectRef(*Obj);
+		if (PackedObjectRef.EncodedRef == 0)
+		{
+			return GetTypeHash(Obj);
+		}
+		return GetTypeHash(PackedObjectRef.EncodedRef);
 	}
-#else
-	return ReadObjectHandlePointerNoCheck(Handle);
-#endif
+	return GetTypeHash(Handle.PointerOrRef);
 }
+#endif
+
+/////////////////////////////////////////////////////////////////////////////////////////////////////////////////////////
 
-inline UObject* ResolveObjectHandleNoReadNoCheck(FObjectHandle& Handle)
+namespace UE::CoreUObject::Private
 {
+	inline FObjectHandle MakeObjectHandle(UObject* Object)
+	{
 #if UE_WITH_OBJECT_HANDLE_LATE_RESOLVE
-	FObjectHandle LocalHandle = Handle;
-	LocalHandle = MakeObjectHandle(ResolvePackedObjectRef(ReadObjectHandlePackedObjectRefNoCheck(LocalHandle)));
-	UObject* ResolvedObject = ReadObjectHandlePointerNoCheck(LocalHandle);
-	Handle = LocalHandle;
-	return ResolvedObject;
-	
+		{ return { UPTRINT(Object) }; }
 #else
-	return ReadObjectHandlePointerNoCheck(Handle);
+		return Object;
 #endif
-}
+	}
 
-inline UObject* ResolveObjectHandle(FObjectHandle& Handle)
-{
+	inline UObject* ResolveObjectHandle(FObjectHandle& Handle)
+	{
 #if UE_WITH_OBJECT_HANDLE_LATE_RESOLVE || UE_WITH_OBJECT_HANDLE_TRACKING
-	UObject* ResolvedObject = ResolveObjectHandleNoRead(Handle);
-	ObjectHandle_Private::OnHandleRead(ResolvedObject);
-	return ResolvedObject;
+		UObject* ResolvedObject = ResolveObjectHandleNoRead(Handle);
+		UE::CoreUObject::Private::OnHandleRead(ResolvedObject);
+		return ResolvedObject;
 #else
-	return ReadObjectHandlePointerNoCheck(Handle);
+		return ReadObjectHandlePointerNoCheck(Handle);
 #endif
-}
+	}
 
-inline UClass* ResolveObjectHandleClass(FObjectHandle Handle)
-{
-	if (IsObjectHandleResolved(Handle))
+#if UE_WITH_OBJECT_HANDLE_LATE_RESOLVE
+	inline FPackedObjectRef ReadObjectHandlePackedObjectRefNoCheck(FObjectHandle Handle)
+	{
+		return { Handle.PointerOrRef };
+	}
+#endif
+
+	inline UClass* ResolveObjectHandleClass(FObjectHandle Handle)
 	{
+#if UE_WITH_OBJECT_HANDLE_LATE_RESOLVE
+		if (IsObjectHandleResolved(Handle))
+		{
+			UObject* Obj = ReadObjectHandlePointerNoCheck(Handle);
+			return Obj != nullptr ? Obj->GetClass() : nullptr;
+		}
+		else
+		{
+			// @TODO: OBJPTR: This should be cached somewhere instead of resolving on every call
+			FPackedObjectRef PackedObjectRef = ReadObjectHandlePackedObjectRefNoCheck(Handle);
+			FObjectRef ObjectRef = MakeObjectRef(PackedObjectRef);
+			return ObjectRef.ResolveObjectRefClass();
+		}
+#else
 		UObject* Obj = ReadObjectHandlePointerNoCheck(Handle);
 		return Obj != nullptr ? Obj->GetClass() : nullptr;
+#endif
 	}
-	else
+
+	inline UObject* ResolveObjectHandleNoRead(FObjectHandle& Handle)
 	{
-		// @TODO: OBJPTR: This should be cached somewhere instead of resolving on every call
-		return ResolvePackedObjectRefClass(ReadObjectHandlePackedObjectRefNoCheck(Handle));
+#if UE_WITH_OBJECT_HANDLE_LATE_RESOLVE
+		FObjectHandle LocalHandle = Handle;
+		if (IsObjectHandleResolved(LocalHandle))
+		{
+			UObject* ResolvedObject = ReadObjectHandlePointerNoCheck(LocalHandle);
+			return ResolvedObject;
+		}
+		else
+		{
+			FPackedObjectRef PackedObjectRef = ReadObjectHandlePackedObjectRefNoCheck(LocalHandle);
+			FObjectRef ObjectRef = MakeObjectRef(PackedObjectRef);
+			UObject* ResolvedObject = ObjectRef.Resolve();
+			Handle = MakeObjectHandle(ResolvedObject);
+			return ResolvedObject;
+		}
+#else
+		return ReadObjectHandlePointerNoCheck(Handle);
+#endif
 	}
-}
 
-inline UObject* ReadObjectHandlePointer(FObjectHandle Handle)
-{
-	if (IsObjectHandleResolved(Handle))
+	
+	inline UObject* ResolveObjectHandleNoReadNoCheck(FObjectHandle& Handle)
 	{
-		UObject* ResolvedObject = ReadObjectHandlePointerNoCheck(Handle);
-		ObjectHandle_Private::OnHandleRead(ResolvedObject);
+#if UE_WITH_OBJECT_HANDLE_LATE_RESOLVE
+		FObjectHandle LocalHandle = Handle;
+		FPackedObjectRef PackedObjectRef = ReadObjectHandlePackedObjectRefNoCheck(LocalHandle);
+		FObjectRef ObjectRef = MakeObjectRef(PackedObjectRef);
+		UObject* ResolvedObject = ObjectRef.Resolve();
+		LocalHandle = MakeObjectHandle(ResolvedObject);
+		Handle = LocalHandle;
 		return ResolvedObject;
+#else
+		return ReadObjectHandlePointerNoCheck(Handle);
+#endif
 	}
-	return nullptr;
-}
 
-inline FPackedObjectRef ReadObjectHandlePackedObjectRef(FObjectHandle Handle)
-{
-	return !IsObjectHandleResolved(Handle) ? ReadObjectHandlePackedObjectRefNoCheck(Handle) : FPackedObjectRef();
-}
-
-#if UE_WITH_OBJECT_HANDLE_LATE_RESOLVE
-inline FObjectRef MakeObjectRef(FObjectHandle Handle)
-{
-	if (IsObjectHandleResolved(Handle))
+	inline UObject* ReadObjectHandlePointerNoCheck(FObjectHandle Handle)
 	{
-		return MakeObjectRef(ReadObjectHandlePointerNoCheck(Handle));
+#if UE_WITH_OBJECT_HANDLE_LATE_RESOLVE
+		return reinterpret_cast<UObject*>(Handle.PointerOrRef);
+#else
+		return Handle;
+#endif
 	}
-	else
+
+	//Natvis structs
+	struct FObjectHandlePackageDebugData
 	{
-		return MakeObjectRef(ReadObjectHandlePackedObjectRefNoCheck(Handle));
-	}
-}
+		FMinimalName PackageName;
+		FScriptArray ObjectPaths;
+		FScriptArray DataClassDescriptors;
+		uint8 _Padding[sizeof(FRWLock)];
+	};
+
+	struct FObjectHandleDataClassDescriptor
+	{
+		FMinimalName PackageName;
+		FMinimalName ClassName;
+	};
+
+	constexpr uint32 PackageIdShift = 33;
+	constexpr uint32 ObjectPathIdShift = 1;
+	constexpr uint32 ObjectPathIdMask = 0x00FF'FFFF;
+	constexpr uint32 DataClassDescriptorIdShift = 25;
+	constexpr uint32 DataClassDescriptorIdMask = 0x0000'00FF;
+	constexpr uint32 PackageIdMask = 0x7FFF'FFFF;
 
+#if UE_WITH_OBJECT_HANDLE_LATE_RESOLVE
+	//forward declarations
+	void InitObjectHandles(int32 Size);
+	void FreeObjectHandle(const UObjectBase& Object);
+	void UpdateRenamedObject(const UObject& Obj, FName NewName, UObject* NewOuter);
+	UE::CoreUObject::Private::FPackedObjectRef MakePackedObjectRef(const UObject* Object);
 #endif
+}
+
diff --git a/Engine/Source/Runtime/CoreUObject/Public/UObject/ObjectHandleDefines.h b/Engine/Source/Runtime/CoreUObject/Public/UObject/ObjectHandleDefines.h
new file mode 100644
--- /dev/null
+++ b/Engine/Source/Runtime/CoreUObject/Public/UObject/ObjectHandleDefines.h
@@ -0,0 +1,14 @@
+// Copyright Epic Games, Inc. All Rights Reserved.
+
+#pragma once
+
+#include "Misc/CoreMiscDefines.h"
+
+#ifndef UE_WITH_OBJECT_HANDLE_LATE_RESOLVE
+	#define UE_WITH_OBJECT_HANDLE_LATE_RESOLVE WITH_EDITORONLY_DATA
+#endif
+
+
+#ifndef UE_WITH_OBJECT_HANDLE_TRACKING
+	#define UE_WITH_OBJECT_HANDLE_TRACKING WITH_EDITORONLY_DATA
+#endif
diff --git a/Engine/Source/Runtime/CoreUObject/Public/UObject/ObjectHandleTracking.h b/Engine/Source/Runtime/CoreUObject/Public/UObject/ObjectHandleTracking.h
new file mode 100644
--- /dev/null
+++ b/Engine/Source/Runtime/CoreUObject/Public/UObject/ObjectHandleTracking.h
@@ -0,0 +1,180 @@
+// Copyright Epic Games, Inc. All Rights Reserved.
+#pragma once
+
+#include "Delegates/Delegate.h"
+#include "UObject/ObjectHandleDefines.h"
+
+struct FObjectRef;
+
+/**
+ * FObjectHandles can optionally support tracking.  Because of the low level nature of object handles, anything that
+ * registers itself for these callbacks should ensure that it is:
+ * 1) error free (ie: should not cause exceptions even in unusual circumstances)
+ * 2) fault tolerant (ie: could be called at a time when an exception has happened)
+ * 3) thread-safe (ie: could be called from any thread)
+ * 4) high performance (ie: will be called many times)
+ */
+#if UE_WITH_OBJECT_HANDLE_TRACKING
+
+ /**
+  * Callback notifying when an object value is read from a handle.  Fired regardless of whether the handle
+  * was resolved as part of the read operation or not and whether the object being read is null or not.
+  *
+  * @param ReadObject	The object that was read from a handle.
+  */
+DECLARE_DELEGATE_OneParam(FObjectHandleReadDelegate, UObject* ReadObject);
+
+/**
+ * Callback notifying when a class is resolved from an object handle or object reference.
+ * Classes are resolved either independently for a given handle/reference or as part of each object resolve.
+ *
+ * @param SourceRef		The object reference (either user provided or internally computed from a handle) from which the class was resolved.
+ * @param ClassPackage	The package containing the resolved class.
+ * @param Class			The resolved class.
+ */
+DECLARE_DELEGATE_ThreeParams(FObjectHandleClassResolvedDelegate, const FObjectRef& SourceRef, UPackage* ObjectPackage, UClass* Class);
+
+/**
+ * Callback notifying when a object handle is resolved.
+ *
+ * @param SourceRef		The object reference (either user provided or internally computed from a handle) from which the class was resolved.
+ * @param ClassPackage	The package containing the resolved class.
+ * @param Class			The resolved class.
+ */
+DECLARE_DELEGATE_ThreeParams(FObjectHandleReferenceResolvedDelegate, const FObjectRef& SourceRef, UPackage* ObjectPackage, UObject* Object);
+
+/**
+ * Callback notifying when an object was loaded through an object handle.  Will not notify you about global object loads, just ones that occur
+ * as the byproduct of resolving an ObjectHandle.
+ *
+ * @param SourceRef		The object reference (either user provided or internally computed from a handle) from which the class was resolved.
+ * @param ClassPackage	The package containing the resolved class.
+ * @param Class			The resolved class.
+ */
+DECLARE_DELEGATE_ThreeParams(FObjectHandleReferenceLoadedDelegate, const FObjectRef& SourceRef, UPackage* ObjectPackage, UObject* Object);
+
+/**
+ * Installs a new callback for notifications that an object value has been read from a handle.
+ *
+ * @param Function		The new handle read callback to install.
+ * @return				The DelegateHandle so that you can remove the callback at a later date.
+ */
+COREUOBJECT_API FDelegateHandle AddObjectHandleReadCallback(FObjectHandleReadDelegate Delegate);
+COREUOBJECT_API void RemoveObjectHandleReadCallback(FDelegateHandle DelegateHandle);
+
+/**
+ * Installs a new callback for notifications that a class has been resolved from an object handle or object reference.
+ *
+ * @param Function		The new class resolved callback to install.
+ * @return				The DelegateHandle so that you can remove the callback at a later date.
+ */
+COREUOBJECT_API FDelegateHandle AddObjectHandleClassResolvedCallback(FObjectHandleClassResolvedDelegate Callback);
+COREUOBJECT_API void RemoveObjectHandleClassResolvedCallback(FDelegateHandle DelegateHandle);
+
+/**
+ * Installs a new callback for notifications that an object has been resolved from an object handle or object reference.
+ *
+ * @param Function		The new object resolved callback to install.
+ * @return				The DelegateHandle so that you can remove the callback at a later date.
+ */
+COREUOBJECT_API FDelegateHandle AddObjectHandleReferenceResolvedCallback(FObjectHandleReferenceResolvedDelegate Callback);
+COREUOBJECT_API void RemoveObjectHandleReferenceResolvedCallback(FDelegateHandle DelegateHandle);
+
+/**
+ * Installs a new callback for notifications that an object has been loaded from an object handle or object reference.
+ *
+ * @param Function		The new object resolved callback to install.
+ * @return				The DelegateHandle so that you can remove the callback at a later date.
+ */
+COREUOBJECT_API FDelegateHandle AddObjectHandleReferenceLoadedCallback(FObjectHandleReferenceLoadedDelegate Callback);
+COREUOBJECT_API void RemoveObjectHandleReferenceLoadedCallback(FDelegateHandle DelegateHandle);
+
+
+namespace UE::CoreUObject::Private
+{
+	DECLARE_MULTICAST_DELEGATE_OneParam(FObjectHandleReadEvent, UObject* Object);
+	DECLARE_MULTICAST_DELEGATE_ThreeParams(FClassReferenceResolvedEvent, const FObjectRef& ObjectRef, UPackage* Package, UClass* Class);
+	DECLARE_MULTICAST_DELEGATE_ThreeParams(FObjectHandleReferenceResolvedEvent, const FObjectRef& SourceRef, UPackage* ObjectPackage, UObject* Object);
+	DECLARE_MULTICAST_DELEGATE_ThreeParams(FObjectHandleReferenceLoadedEvent, const FObjectRef& SourceRef, UPackage* ObjectPackage, UObject* Object);
+
+	struct FObjectHandleEvents
+	{
+		FObjectHandleReadEvent ObjectHandleReadEvent;
+		FClassReferenceResolvedEvent ClassReferenceResolvedEvent;
+		FObjectHandleReferenceResolvedEvent ObjectHandleReferenceResolvedEvent;
+		FObjectHandleReferenceLoadedEvent ObjectHandleReferenceLoadedEvent;
+
+		inline void BeginUsing() { UsingCount++; }
+		inline void EndUsing() { --UsingCount; }
+		inline bool IsUsing() const { return UsingCount > 0; }
+
+	private:
+		std::atomic<int32> UsingCount;
+	};
+
+	extern COREUOBJECT_API std::atomic<int32> ObjectHandleEventIndex;
+	extern COREUOBJECT_API FObjectHandleEvents ObjectHandleEvents[2];
+
+	inline FObjectHandleEvents& BeginReadingEvents()
+	{
+		// Quick spin lock to try and begin using the events, normally this wont actually spin, in like 99.999% of the time.
+		while (true)
+		{
+			// Start by getting the current event index from the double buffer.
+			const int32 InitialEventIndex = ObjectHandleEventIndex;
+			// Grab the event set at that index.
+			FObjectHandleEvents& Events = ObjectHandleEvents[InitialEventIndex];
+			// Begin using the events, this will signal we're using them.
+			Events.BeginUsing();
+			// Ok - now we're going to check that we're *still* using the same index.  If we are
+			// then we know that it didn't change out from under us in between getting Events, and calling BeginUsing().
+			if (InitialEventIndex == ObjectHandleEventIndex)
+			{
+				return Events;
+			}
+			// If the check above failed, then we need to cease using the events, because we're about to try this again.
+			Events.EndUsing();
+		}
+	}
+
+	inline void OnHandleRead(UObject* Object)
+	{
+		FObjectHandleEvents& Events = BeginReadingEvents();
+		Events.ObjectHandleReadEvent.Broadcast(Object);
+		Events.EndUsing();
+	}
+
+	inline void OnClassReferenceResolved(const FObjectRef& ObjectRef, UPackage* Package, UClass* Class)
+	{
+		FObjectHandleEvents& Events = BeginReadingEvents();
+		Events.ClassReferenceResolvedEvent.Broadcast(ObjectRef, Package, Class);
+		Events.EndUsing();
+	}
+
+	inline void OnReferenceResolved(const FObjectRef& ObjectRef, UPackage* Package, UObject* Object)
+	{
+		FObjectHandleEvents& Events = BeginReadingEvents();
+		Events.ObjectHandleReferenceResolvedEvent.Broadcast(ObjectRef, Package, Object);
+		Events.EndUsing();
+	}
+
+	inline void OnReferenceLoaded(const FObjectRef& ObjectRef, UPackage* Package, UObject* Object)
+	{
+		FObjectHandleEvents& Events = BeginReadingEvents();
+		Events.ObjectHandleReferenceLoadedEvent.Broadcast(ObjectRef, Package, Object);
+		Events.EndUsing();
+	}
+
+}
+#else
+
+namespace UE::CoreUObject::Private
+{
+	inline void OnHandleRead(UObject* Object) { }
+	inline void OnClassReferenceResolved(const FObjectRef& ObjectRef, UPackage* Package, UClass* Class) { }
+	inline void OnReferenceResolved(const FObjectRef& ObjectRef, UPackage* Package, UClass* Class) { }
+	inline void OnReferenceLoaded(const FObjectRef& ObjectRef, UPackage* Package, UObject* Object) { }
+}
+
+
+#endif
\ No newline at end of file
diff --git a/Engine/Source/Runtime/CoreUObject/Public/UObject/ObjectPathId.h b/Engine/Source/Runtime/CoreUObject/Public/UObject/ObjectPathId.h
deleted file mode 100644
diff --git a/Engine/Source/Runtime/CoreUObject/Public/UObject/ObjectPtr.h b/Engine/Source/Runtime/CoreUObject/Public/UObject/ObjectPtr.h
--- a/Engine/Source/Runtime/CoreUObject/Public/UObject/ObjectPtr.h
+++ b/Engine/Source/Runtime/CoreUObject/Public/UObject/ObjectPtr.h
@@ -54,7 +54,7 @@ struct FObjectPtr
 {
 public:
 	FObjectPtr()
-		: Handle(MakeObjectHandle(nullptr))
+		: Handle(UE::CoreUObject::Private::MakeObjectHandle(nullptr))
 	{
 	}
 
@@ -63,39 +63,37 @@ public:
 	}
 
 	FORCEINLINE FObjectPtr(TYPE_OF_NULLPTR)
-		: Handle(MakeObjectHandle(nullptr))
+		: Handle(UE::CoreUObject::Private::MakeObjectHandle(nullptr))
 	{
 	}
 
 	explicit FORCEINLINE FObjectPtr(UObject* Object)
-		: Handle(MakeObjectHandle(Object))
+		: Handle(UE::CoreUObject::Private::MakeObjectHandle(Object))
 	{
 	}
 
 	UE_OBJPTR_DEPRECATED(5.0, "Construction with incomplete type pointer is deprecated.  Please update this code to use MakeObjectPtrUnsafe.")
 	explicit FORCEINLINE FObjectPtr(void* IncompleteObject)
-		: Handle(MakeObjectHandle(reinterpret_cast<UObject*>(IncompleteObject)))
+		: Handle(UE::CoreUObject::Private::MakeObjectHandle(reinterpret_cast<UObject*>(IncompleteObject)))
 	{
 	}
 
-	explicit FORCEINLINE FObjectPtr(const FObjectRef& ObjectRef)
-		: Handle(MakeObjectHandle(ObjectRef))
-	{
-	}
-
-	explicit FORCEINLINE FObjectPtr(const FPackedObjectRef& PackedObjectRef)
-		: Handle(MakeObjectHandle(PackedObjectRef))
+#if UE_WITH_OBJECT_HANDLE_LATE_RESOLVE
+	explicit FORCEINLINE FObjectPtr(FObjectHandle Handle)
+		: Handle(Handle)
 	{
 	}
+#endif
+	
 
 	FORCEINLINE UObject* Get() const
 	{
-		return ResolveObjectHandle(Handle);
+		return UE::CoreUObject::Private::ResolveObjectHandle(Handle);
 	}
 
 	FORCEINLINE UClass* GetClass() const
 	{
-		return ResolveObjectHandleClass(Handle);
+		return UE::CoreUObject::Private::ResolveObjectHandleClass(Handle);
 	}
 
 	FObjectPtr(FObjectPtr&&) = default;
@@ -105,20 +103,20 @@ public:
 
 	FObjectPtr& operator=(UObject* Other)
 	{
-		Handle = MakeObjectHandle(Other);
+		Handle = UE::CoreUObject::Private::MakeObjectHandle(Other);
 		return *this;
 	}
 
 	UE_OBJPTR_DEPRECATED(5.0, "Assignment with incomplete type pointer is deprecated.  Please update this code to use MakeObjectPtrUnsafe.")
 	FObjectPtr& operator=(void* IncompleteOther)
 	{
-		Handle = MakeObjectHandle(reinterpret_cast<UObject*>(IncompleteOther));
+		Handle = UE::CoreUObject::Private::MakeObjectHandle(reinterpret_cast<UObject*>(IncompleteOther));
 		return *this;
 	}
 
 	FObjectPtr& operator=(TYPE_OF_NULLPTR)
 	{
-		Handle = MakeObjectHandle(nullptr);
+		Handle = UE::CoreUObject::Private::MakeObjectHandle(nullptr);
 		return *this;
 	}
 
@@ -134,7 +132,7 @@ public:
 	FORCEINLINE UObject& operator*() const { return *Get(); }
 
 	UE_DEPRECATED(5.1, "IsNull is deprecated, please use operator bool instead.")
-	FORCEINLINE bool IsNull() const { return ResolveObjectHandleNoRead(Handle) == nullptr; }
+	FORCEINLINE bool IsNull() const { return UE::CoreUObject::Private::ResolveObjectHandleNoRead(Handle) == nullptr; }
 	
 	UE_DEPRECATED(5.1, "IsNullNoResolve is deprecated, please use operator bool instead.")
 	FORCEINLINE bool IsNullNoResolve() const { return IsObjectHandleNull(Handle); }
@@ -546,8 +544,8 @@ public:
 	template <typename U, typename V> friend bool ObjectPtr_Private::IsObjectPtrEqualToRawPtrOfRelatedType(const TObjectPtr<U>& Ptr, const V* Other);
 
 private:
-	FORCEINLINE T* GetNoReadNoCheck() const { return (T*)(ResolveObjectHandleNoReadNoCheck(ObjectPtr.GetHandleRef())); }
-	FORCEINLINE T* GetNoResolveNoCheck() const { return (T*)(ReadObjectHandlePointerNoCheck(ObjectPtr.GetHandleRef())); }
+	FORCEINLINE T* GetNoReadNoCheck() const { return (T*)(UE::CoreUObject::Private::ResolveObjectHandleNoReadNoCheck(ObjectPtr.GetHandleRef())); }
+	FORCEINLINE T* GetNoResolveNoCheck() const { return (T*)(UE::CoreUObject::Private::ReadObjectHandlePointerNoCheck(ObjectPtr.GetHandleRef())); }
 
 	// @TODO: OBJPTR: There is a risk of a gap in access tracking here.  The caller may get a mutable pointer, write to it, then
 	//			read from it.  That last read would happen without an access being recorded.  Not sure if there is a good way
diff --git a/Engine/Source/Runtime/CoreUObject/Public/UObject/ObjectRef.h b/Engine/Source/Runtime/CoreUObject/Public/UObject/ObjectRef.h
new file mode 100644
--- /dev/null
+++ b/Engine/Source/Runtime/CoreUObject/Public/UObject/ObjectRef.h
@@ -0,0 +1,124 @@
+// Copyright Epic Games, Inc. All Rights Reserved.
+
+#pragma once
+
+#include "Misc/StringBuilder.h"
+#include "UObject/UObjectBaseUtility.h"
+#include "UObject/ObjectHandleDefines.h"
+
+#if UE_WITH_OBJECT_HANDLE_LATE_RESOLVE || UE_WITH_OBJECT_HANDLE_TRACKING
+
+DECLARE_LOG_CATEGORY_EXTERN(LogObjectRef, Log, All);
+
+
+struct FObjectRef;
+
+namespace UE::CoreUObject::Private
+{
+	struct FPackedObjectRef;
+	class FObjectPathId;
+	FPackedObjectRef MakePackedObjectRef(const FObjectRef& ObjectRef);
+}
+
+/**
+ * FObjectRef represents a heavyweight reference that contains the specific pieces of information needed to reference an object
+ * (or null) that may or may not be loaded yet.  The expectation is that given the imports of a package we have loaded, we should
+ * be able to create an FObjectRef to objects in a package we haven't yet loaded.  For this reason, FObjectRef has to be derivable
+ * from the serialized (not transient) contents of an FObjectImport.
+ */
+struct FObjectRef
+{
+	FObjectRef() = default;
+	FObjectRef(FName PackageName, FName ClassPackageName, FName ClassName, UE::CoreUObject::Private::FObjectPathId ObjectPath);
+	explicit FObjectRef(const UObject* Object);
+
+	bool operator==(const FObjectRef&& Other) const
+	{
+		return (PackageName == Other.PackageName) && (ObjectPathId == Other.ObjectPathId) && (ClassPackageName == Other.ClassPackageName) && (ClassName == Other.ClassName);
+	}
+
+	/** Returns the path to the class for this object. */
+	FString GetClassPathName(EObjectFullNameFlags Flags = EObjectFullNameFlags::IncludeClassPackage) const
+	{
+		TStringBuilder<FName::StringBufferSize> Result;
+		AppendClassPathName(Result, Flags);
+		return FString(Result);
+	}
+
+	/** Appends the path to the class for this object to the builder, does not reset builder. */
+	COREUOBJECT_API void AppendClassPathName(FStringBuilderBase& OutClassPathNameBuilder, EObjectFullNameFlags Flags = EObjectFullNameFlags::IncludeClassPackage) const;
+
+	/** Returns the name of the object in the form: ObjectName */
+	COREUOBJECT_API FName GetFName() const;
+
+	/** Returns the full path for the object in the form: ObjectPath */
+	FString GetPathName() const
+	{
+		TStringBuilder<FName::StringBufferSize> Result;
+		AppendPathName(Result);
+		return FString(Result);
+	}
+
+	/** Appends the path to the object to the builder, does not reset builder. */
+	COREUOBJECT_API void AppendPathName(FStringBuilderBase& OutPathNameBuilder) const;
+
+	/** Returns the full name for the object in the form: Class ObjectPath */
+	FString GetFullName(EObjectFullNameFlags Flags = EObjectFullNameFlags::None) const
+	{
+		TStringBuilder<256> FullName;
+		GetFullName(FullName, Flags);
+		return FString(FullName);
+	}
+
+	/** Populates OutFullNameBuilder with the full name for the object in the form: Class ObjectPath */
+	void GetFullName(FStringBuilderBase& OutFullNameBuilder, EObjectFullNameFlags Flags = EObjectFullNameFlags::None) const
+	{
+		OutFullNameBuilder.Reset();
+		AppendClassPathName(OutFullNameBuilder, Flags);
+		OutFullNameBuilder.AppendChar(TEXT(' '));
+		AppendPathName(OutFullNameBuilder);
+	}
+
+	/** Returns the name for the asset in the form: Class'ObjectPath' */
+	FString GetExportTextName() const
+	{
+		TStringBuilder<256> ExportTextName;
+		GetExportTextName(ExportTextName);
+		return FString(ExportTextName);
+	}
+
+	/** Populates OutExportTextNameBuilder with the name for the object in the form: Class'ObjectPath' */
+	void GetExportTextName(FStringBuilderBase& OutExportTextNameBuilder) const
+	{
+		OutExportTextNameBuilder.Reset();
+		AppendClassPathName(OutExportTextNameBuilder);
+		OutExportTextNameBuilder.AppendChar(TEXT('\''));
+		AppendPathName(OutExportTextNameBuilder);
+		OutExportTextNameBuilder.AppendChar(TEXT('\''));
+	}
+
+	inline bool IsNull() const 
+	{ 
+		return PackageName.IsNone() && ObjectPathId == 0;
+	}
+
+	COREUOBJECT_API UClass* ResolveObjectRefClass(uint32 LoadFlags = LOAD_None) const;
+	COREUOBJECT_API UObject* Resolve(uint32 LoadFlags = LOAD_None) const;
+
+public:
+	FName PackageName;
+	FName ClassPackageName;
+	FName ClassName;
+
+private:
+	uint64 ObjectPathId = 0;
+
+	friend UE::CoreUObject::Private::FPackedObjectRef UE::CoreUObject::Private::MakePackedObjectRef(const FObjectRef& ObjectRef);
+	UE::CoreUObject::Private::FObjectPathId GetObjectPath() const;
+
+};
+
+UE_DEPRECATED(5.2, "Use FObjectRef::IsNull() instead")
+inline bool IsObjectRefNull(const FObjectRef& ObjectRef) { return ObjectRef.IsNull(); }
+
+#endif
\ No newline at end of file
diff --git a/Engine/Source/Runtime/CoreUObject/Public/UObject/PackedObjectRef.h b/Engine/Source/Runtime/CoreUObject/Public/UObject/PackedObjectRef.h
new file mode 100644
--- /dev/null
+++ b/Engine/Source/Runtime/CoreUObject/Public/UObject/PackedObjectRef.h
@@ -0,0 +1,39 @@
+// Copyright Epic Games, Inc. All Rights Reserved.
+
+#pragma once
+#include "UObject/ObjectHandleDefines.h"
+
+#if UE_WITH_OBJECT_HANDLE_LATE_RESOLVE
+
+namespace UE::CoreUObject::Private
+{
+	/**
+	 * FPackedObjectRef represents a lightweight reference that can fit in the space of a pointer and be able to refer to an object
+	 * (or null) that may or may not be loaded without pointing to its location in memory (even if it is currently loaded).
+	 */
+	struct FPackedObjectRef
+	{
+		// Must be 0 for a reference to null.
+		// The least significant bit must always be 1 in a non-null reference.
+		UPTRINT EncodedRef;
+
+		inline bool operator==(FPackedObjectRef RHS) const
+		{
+			return EncodedRef == RHS.EncodedRef;
+		}
+
+		inline bool operator!=(FPackedObjectRef RHS) const
+		{
+			return EncodedRef != RHS.EncodedRef;
+		}
+
+		friend inline uint32 GetTypeHash(FPackedObjectRef ObjectRef)
+		{
+			return GetTypeHash(ObjectRef.EncodedRef);
+		}
+
+		inline bool IsNull() { return EncodedRef == 0; }
+	};
+}
+
+#endif
diff --git a/Engine/Source/Runtime/CoreUObject/Public/UObject/UObjectGlobals.h b/Engine/Source/Runtime/CoreUObject/Public/UObject/UObjectGlobals.h
--- a/Engine/Source/Runtime/CoreUObject/Public/UObject/UObjectGlobals.h
+++ b/Engine/Source/Runtime/CoreUObject/Public/UObject/UObjectGlobals.h
@@ -2375,7 +2375,7 @@ public:
 	template<class UObjectType>
 	void AddReferencedObject(TObjectPtr<UObjectType>& Object, const UObject* ReferencingObject = nullptr, const FProperty* ReferencingProperty = nullptr)
 	{
-		if (IsObjectHandleResolved(Object.GetHandle()))
+		if (Object.IsResolved())
 		{
 			HandleObjectReference(*reinterpret_cast<UObject**>(&Object), ReferencingObject, ReferencingProperty);
 		}
@@ -2393,7 +2393,7 @@ public:
 	{
 		static_assert(sizeof(UObjectType) > 0, "AddReferencedObject: Element must be a pointer to a fully-defined type");
 		static_assert(TPointerIsConvertibleFromTo<UObjectType, const UObjectBase>::Value, "AddReferencedObject: Element must be a pointer to a type derived from UObject");
-		if (IsObjectHandleResolved(Object.GetHandle()))
+		if (Object.IsResolved())
 		{
 			HandleObjectReference(*reinterpret_cast<UObject**>(&Object), ReferencingObject, ReferencingProperty);
 		}
@@ -2447,7 +2447,7 @@ public:
 		static_assert(TPointerIsConvertibleFromTo<UObjectType, const UObjectBase>::Value, "AddReferencedObjects: Elements must be pointers to a type derived from UObject");
 		for (auto& Object : ObjectSet)
 		{
-			if (IsObjectHandleResolved(Object.GetHandle()))
+			if (Object.IsResolved())
 			{
 				HandleObjectReference(*reinterpret_cast<UObject**>(&Object), ReferencingObject, ReferencingProperty);
 			}
@@ -2468,7 +2468,7 @@ public:
 		static_assert(TPointerIsConvertibleFromTo<KeyType, const UObjectBase>::Value, "AddReferencedObjects: Keys must be pointers to a type derived from UObject");
 		for (auto& It : Map)
 		{
-			if (IsObjectHandleResolved(It.Key.GetHandle()))
+			if (It.Key.IsResolved())
 			{
 				HandleObjectReference(*reinterpret_cast<UObject**>(&It.Key), ReferencingObject, ReferencingProperty);
 			}
@@ -2481,7 +2481,7 @@ public:
 		static_assert(TPointerIsConvertibleFromTo<ValueType, const UObjectBase>::Value, "AddReferencedObjects: Values must be pointers to a type derived from UObject");
 		for (auto& It : Map)
 		{
-			if (IsObjectHandleResolved(It.Value.GetHandle()))
+			if (It.Value.IsResolved())
 			{
 				HandleObjectReference(*reinterpret_cast<UObject**>(&It.Value), ReferencingObject, ReferencingProperty);
 			}
@@ -2496,11 +2496,11 @@ public:
 		static_assert(TPointerIsConvertibleFromTo<ValueType, const UObjectBase>::Value, "AddReferencedObjects: Values must be pointers to a type derived from UObject");
 		for (auto& It : Map)
 		{
-			if (IsObjectHandleResolved(It.Key.GetHandle()))
+			if (It.Key.IsResolved())
 			{
 				HandleObjectReference(*reinterpret_cast<UObject**>(&It.Key), ReferencingObject, ReferencingProperty);
 			}
-			if (IsObjectHandleResolved(It.Value.GetHandle()))
+			if (It.Value.IsResolved())
 			{
 				HandleObjectReference(*reinterpret_cast<UObject**>(&It.Value), ReferencingObject, ReferencingProperty);
 			}
diff --git a/Engine/Source/Runtime/CoreUObject/Tests/ObjectHandleTest.cpp b/Engine/Source/Runtime/CoreUObject/Tests/ObjectHandleTest.cpp
--- a/Engine/Source/Runtime/CoreUObject/Tests/ObjectHandleTest.cpp
+++ b/Engine/Source/Runtime/CoreUObject/Tests/ObjectHandleTest.cpp
@@ -12,6 +12,8 @@
 #include "ObjectRefTrackingTestBase.h"
 #include "IO/IoDispatcher.h"
 #include "TestHarness.h"
+#include "UObject/ObjectRef.h"
+#include "UObject/ObjectPathId.h"
 
 static_assert(sizeof(FObjectHandle) == sizeof(void*), "FObjectHandle type must always compile to something equivalent to a pointer size.");
 
@@ -20,122 +22,57 @@ class FObjectHandleTestBase : public FObjectRefTrackingTestBase
 public:
 	
 protected:
-	
-	UObject* ResolveHandle(FObjectHandle& TargetHandle)
-	{
-	#if UE_WITH_OBJECT_HANDLE_LATE_RESOLVE
-		// Late resolved handles cannot be null or resolved at this point
-		bool bValue = IsObjectHandleNull(TargetHandle);
-		TEST_FALSE(TEXT("Handle to target is null"), bValue);
-		if (bValue)
-		{
-			return nullptr;
-		}
-		bValue = IsObjectHandleResolved(TargetHandle);
-		TEST_FALSE(TEXT("Handle to target is resolved"), bValue);
-		if (bValue)
-		{
-			return nullptr;
-		}
-	#else
-		bool bValue = IsObjectHandleResolved(TargetHandle);
-		// Immediately resolved handles may be null (if the target is invalid) and must be resolved at this point
-		TEST_TRUE(TEXT("Handle to target is not resolved"), bValue);
-		if (!bValue)
-		{
-			return nullptr;
-		}
-	#endif
-
-		return ResolveObjectHandle(TargetHandle);
-	}
-
-	UObject* ConstructAndResolveHandle(const ANSICHAR* PackageName, const ANSICHAR* ObjectName, const ANSICHAR* ClassPackageName = nullptr, const ANSICHAR* ClassName = nullptr)
+#if UE_WITH_OBJECT_HANDLE_LATE_RESOLVE
+	void TestResolveFailure(UE::CoreUObject::Private::FPackedObjectRef PackedRef)
 	{
-		FObjectRef TargetRef{FName(PackageName), FName(ClassPackageName), FName(ClassName), FObjectPathId(ObjectName)};
-		bool bValue = IsObjectRefNull(TargetRef);
-		TEST_FALSE(TEXT("Reference to target is null"), bValue);
-		if (bValue)
-		{
-			return nullptr;
-		}
+		FSnapshotObjectRefMetrics ObjectRefMetrics(*this);
+		FObjectHandle TargetHandle = { PackedRef.EncodedRef };
+		UObject* ResolvedObject = FObjectPtr(TargetHandle).Get();
+		ObjectRefMetrics.TestNumResolves(TEXT("NumResolves should be incremented by one after a resolve attempt"), 1);
+		ObjectRefMetrics.TestNumReads(TEXT("NumReads should be incremented by one after a resolve attempt"), 1);
 
-		FObjectHandle TargetHandle = MakeObjectHandle(TargetRef);
-		return ResolveHandle(TargetHandle);
+		CHECK(ResolvedObject == nullptr);
+		ObjectRefMetrics.TestNumFailedResolves(TEXT("NumFailedResolves should be incremented by one after a failed resolve attempt"), 1);
 	}
+#endif
 
-	UObject* ConstructAndResolveHandle(const FPackedObjectRef& PackedTargetRef)
+#if UE_WITH_OBJECT_HANDLE_LATE_RESOLVE || UE_WITH_OBJECT_HANDLE_TRACKING
+	void TestResolvableNonNull(const ANSICHAR* PackageName, const ANSICHAR* ObjectName, bool bExpectSubRefReads)
 	{
-		bool bValue = PackedTargetRef.EncodedRef == 0;
-		TEST_FALSE(TEXT("Reference to target is null"), bValue);
-		if (bValue)
-		{
-			return nullptr;
-		}
 
-		FObjectHandle TargetHandle = MakeObjectHandle(PackedTargetRef);
-		return ResolveHandle(TargetHandle);
-	}
-
-	bool TestResolvableNonNull(const ANSICHAR* PackageName, const ANSICHAR* ObjectName, const ANSICHAR* ClassPackageName = nullptr, const ANSICHAR* ClassName = nullptr, bool bExpectSubRefReads = false)
-	{
 		FSnapshotObjectRefMetrics ObjectRefMetrics(*this);
-		UObject* ResolvedObject = ConstructAndResolveHandle(PackageName, ObjectName, ClassPackageName, ClassName);
+		FObjectRef TargetRef(FName(PackageName), NAME_None, NAME_None, UE::CoreUObject::Private::FObjectPathId(ObjectName));
+		UObject* ResolvedObject = TargetRef.Resolve();
+		FObjectPtr Ptr(ResolvedObject);
+		Ptr.Get();
+		TEST_TRUE(TEXT("expected not null"), ResolvedObject != nullptr);
 		ObjectRefMetrics.TestNumResolves(TEXT("NumResolves should be incremented by one after a resolve attempt"), 1);
 		ObjectRefMetrics.TestNumReads(TEXT("NumReads should be incremented by one after a resolve attempt"), 1, bExpectSubRefReads /*bAllowAdditionalReads*/);
-
-		if (!ResolvedObject)
-		{
-			FAIL_CHECK(FString::Printf(TEXT("Expected '%s.%s' to resolve to non null."), ANSI_TO_TCHAR(PackageName), ANSI_TO_TCHAR(ObjectName)));
-			return false;
-		}
 		ObjectRefMetrics.TestNumFailedResolves(TEXT("NumFailedResolves should not change after a successful resolve attempt"), 0);
-
-		return true;
 	}
 
-	bool TestResolveFailure(const ANSICHAR* PackageName, const ANSICHAR* ObjectName, const ANSICHAR* ClassPackageName = nullptr, const ANSICHAR* ClassName = nullptr)
+	void TestResolveFailure(const ANSICHAR* PackageName, const ANSICHAR* ObjectName)
 	{
 		FSnapshotObjectRefMetrics ObjectRefMetrics(*this);
-		UObject* ResolvedObject = ConstructAndResolveHandle(PackageName, ObjectName, ClassPackageName, ClassName);
+		FObjectRef TargetRef(FName(PackageName), NAME_None, NAME_None, UE::CoreUObject::Private::FObjectPathId(ObjectName));
+		const UObject* ResolvedObject = TargetRef.Resolve();
 		ObjectRefMetrics.TestNumResolves(TEXT("NumResolves should be incremented by one after a resolve attempt"), 1);
 		ObjectRefMetrics.TestNumReads(TEXT("NumReads should be incremented by one after a resolve attempt"), 1);
-
-		if (ResolvedObject)
-		{
-			FAIL_CHECK(FString::Printf(TEXT("Expected '%s.%s' to resolve to null."), ANSI_TO_TCHAR(PackageName), ANSI_TO_TCHAR(ObjectName)));
-			return false;
-		}
-		ObjectRefMetrics.TestNumFailedResolves(TEXT("NumFailedResolves should be incremented by one after a failed resolve attempt"), 1);
-		return true;
-	}
-
-	bool TestResolveFailure(FPackedObjectRef PackedRef)
-	{
-		FSnapshotObjectRefMetrics ObjectRefMetrics(*this);
-		UObject* ResolvedObject = ConstructAndResolveHandle(PackedRef);
-		ObjectRefMetrics.TestNumResolves(TEXT("NumResolves should be incremented by one after a resolve attempt"), 1);
-		ObjectRefMetrics.TestNumReads(TEXT("NumReads should be incremented by one after a resolve attempt"), 1);
-
-		if (ResolvedObject)
-		{
-			FAIL_CHECK(FString::Printf(TEXT("Expected PACKEDREF(%" UPTRINT_X_FMT ") to resolve to null."), PackedRef.EncodedRef));
-			return false;
-		}
+		CHECK(ResolvedObject == nullptr);
 		ObjectRefMetrics.TestNumFailedResolves(TEXT("NumFailedResolves should be incremented by one after a failed resolve attempt"), 1);
-		return true;
 	}
+#endif
 };
 
 TEST_CASE_METHOD(FObjectHandleTestBase, "CoreUObject::FObjectHandle::Null Behavior", "[CoreUObject][ObjectHandle]")
 {
-	FObjectHandle TargetHandle = MakeObjectHandle(nullptr);
+	FObjectHandle TargetHandle = UE::CoreUObject::Private::MakeObjectHandle(nullptr);
 
 	TEST_TRUE(TEXT("Handle to target is null"), IsObjectHandleNull(TargetHandle));
 	TEST_TRUE(TEXT("Handle to target is resolved"), IsObjectHandleResolved(TargetHandle));
 
 	FSnapshotObjectRefMetrics ObjectRefMetrics(*this);
-	UObject* ResolvedObject = ResolveObjectHandle(TargetHandle);
+	UObject* ResolvedObject = UE::CoreUObject::Private::ResolveObjectHandle(TargetHandle);
 
 	TEST_EQUAL(TEXT("Resolved object is equal to original object"), (UObject*)nullptr, ResolvedObject);
 
@@ -146,13 +83,13 @@ TEST_CASE_METHOD(FObjectHandleTestBase, "CoreUObject::FObjectHandle::Null Behavi
 
 TEST_CASE_METHOD(FObjectHandleTestBase, "CoreUObject::FObjectHandle::Pointer Behavior", "[CoreUObject][ObjectHandle]")
 {
-	FObjectHandle TargetHandle = MakeObjectHandle((UObject*)0x0042);
+	FObjectHandle TargetHandle = UE::CoreUObject::Private::MakeObjectHandle((UObject*)0x0042);
 
 	TEST_FALSE(TEXT("Handle to target is null"), IsObjectHandleNull(TargetHandle));
 	TEST_TRUE(TEXT("Handle to target is resolved"), IsObjectHandleResolved(TargetHandle));
 
 	FSnapshotObjectRefMetrics ObjectRefMetrics(*this);
-	UObject* ResolvedObject = ResolveObjectHandle(TargetHandle);
+	UObject* ResolvedObject = UE::CoreUObject::Private::ResolveObjectHandle(TargetHandle);
 
 	TEST_EQUAL(TEXT("Resolved object is equal to original object"), (UObject*)0x0042, ResolvedObject);
 
@@ -161,6 +98,7 @@ TEST_CASE_METHOD(FObjectHandleTestBase, "CoreUObject::FObjectHandle::Pointer Beh
 	ObjectRefMetrics.TestNumReads(TEXT("NumReads should be incremented by one after a resolve attempt on a pointer handle"),1);
 }
 
+#if UE_WITH_OBJECT_HANDLE_LATE_RESOLVE
 TEST_CASE_METHOD(FObjectHandleTestBase, "CoreUObject::FObjectHandle::Resolve Engine Content Target", "[CoreUObject][ObjectHandle]")
 {
 	const FName TestPackageName(TEXT("/Engine/Test/ObjectPtrDefaultSerialize/Transient"));
@@ -172,12 +110,11 @@ TEST_CASE_METHOD(FObjectHandleTestBase, "CoreUObject::FObjectHandle::Resolve Eng
 		TestPackage->RemoveFromRoot();
 	};
 
-	// Confirm we successfully resolve a correct reference to a subobject
-	TestResolvableNonNull("/Engine/Test/ObjectPtrDefaultSerialize/Transient", "DefaultSerializeObject.SubObject", nullptr, nullptr, true);
-
-	TestResolvableNonNull("/Engine/Test/ObjectPtrDefaultSerialize/Transient", "DefaultSerializeObject");
+	TestResolvableNonNull("/Engine/Test/ObjectPtrDefaultSerialize/Transient", "DefaultSerializeObject.SubObject", true);
+	TestResolvableNonNull("/Engine/Test/ObjectPtrDefaultSerialize/Transient", "DefaultSerializeObject", false);
 }
 
+
 // TODO: Disabled until warnings and errors related to loading a non-existent package have been fixed.
 DISABLED_TEST_CASE_METHOD(FObjectHandleTestBase, "CoreUObject::FObjectHandle::Resolve Non Existent Target", "[CoreUObject][ObjectHandle]")
 {
@@ -198,9 +135,11 @@ DISABLED_TEST_CASE_METHOD(FObjectHandleTestBase, "CoreUObject::FObjectHandle::Re
 TEST_CASE_METHOD(FObjectHandleTestBase, "CoreUObject::FObjectHandle::Resolve Script Target", "[CoreUObject][ObjectHandle]")
 {
 	// Confirm we successfully resolve a correct reference to engine content
-	TestResolvableNonNull("/Script/CoreUObject", "MetaData");
+	TestResolvableNonNull("/Script/CoreUObject", "MetaData", true);
 }
 
+#endif
+
 TEST_CASE_METHOD(FObjectHandleTestBase, "CoreUObject::TObjectPtr::HandleNullGetClass", "[CoreUObject][ObjectHandle]")
 {
 	TObjectPtr<UObject> Ptr = nullptr;
@@ -218,8 +157,8 @@ TEST_CASE("CoreUObject::FObjectHandle::Names")
 		TestPackage->RemoveFromRoot();
 	};
 
-	FObjectPtr PackagePtr(MakePackedObjectRef(TestPackage));
-	FObjectPtr Obj1Ptr(MakePackedObjectRef(Obj1));
+	FObjectPtr PackagePtr(MakeUnresolvedHandle(TestPackage));
+	FObjectPtr Obj1Ptr(MakeUnresolvedHandle(Obj1));
 
 	CHECK(!PackagePtr.IsResolved());
 	CHECK(TestPackage->GetPathName() == PackagePtr.GetPathName());
@@ -237,6 +176,8 @@ TEST_CASE("CoreUObject::FObjectHandle::Names")
 }
 #endif
 
+#if UE_WITH_OBJECT_HANDLE_TRACKING || UE_WITH_OBJECT_HANDLE_LATE_RESOLVE
+
 TEST_CASE("CoreUObject::ObjectRef")
 {
 	const FName TestPackageName(TEXT("/Engine/Test/ObjectRef/Transient"));
@@ -250,7 +191,7 @@ TEST_CASE("CoreUObject::ObjectRef")
 
 	{
 		FObjectImport ObjectImport(Obj1);
-		FObjectRef ObjectRef = MakeObjectRef(Obj1);
+		FObjectRef ObjectRef(Obj1);
 
 		CHECK(ObjectImport.ClassPackage == ObjectRef.ClassPackageName);
 		CHECK(ObjectImport.ClassName == ObjectRef.ClassName);
@@ -259,14 +200,17 @@ TEST_CASE("CoreUObject::ObjectRef")
 
 	{
 		FObjectImport ObjectImport(Inner1);
-		FObjectRef ObjectRef = MakeObjectRef(Inner1);
+		FObjectRef ObjectRef(Inner1);
 
 		CHECK(ObjectImport.ClassPackage == ObjectRef.ClassPackageName);
 		CHECK(ObjectImport.ClassName == ObjectRef.ClassName);
 		CHECK(TestPackage->GetFName() == ObjectRef.PackageName);
 	}
 }
-#if UE_WITH_OBJECT_HANDLE_TRACKING
+
+#endif
+
+#if UE_WITH_OBJECT_HANDLE_LATE_RESOLVE
 
 TEST_CASE_METHOD(FObjectHandleTestBase, "CoreUObject::TObjectPtr::Null Behavior", "[CoreUObject][ObjectHandle]")
 {
@@ -303,8 +247,9 @@ TEST_CASE_METHOD(FObjectHandleTestBase, "CoreUObject::TObjectPtr::Null Behavior"
 	CHECK_FALSE(Ptr != TestObject); CHECK(ResolveCount == 0u);
 	CHECK_FALSE(TestObject != Ptr); CHECK(ResolveCount == 0u);
 
-	FObjectRef TargetRef{ FName("SomePackage"), FName("ClassPackageName"), FName("ClassName"), FObjectPathId("ObjectName") };
-	FObjectPtr ObjectPtr(TargetRef);
+	FObjectRef TargetRef(FName("SomePackage"), FName("ClassPackageName"), FName("ClassName"), UE::CoreUObject::Private::FObjectPathId("ObjectName"));
+	UE::CoreUObject::Private::FPackedObjectRef PackedObjectRef = UE::CoreUObject::Private::MakePackedObjectRef(TargetRef);
+	FObjectPtr ObjectPtr({ PackedObjectRef.EncodedRef });
 	REQUIRE(!ObjectPtr.IsResolved()); //make sure not resolved
 
 	//an unresolved pointers compared against nullptr should still not resolve
@@ -350,7 +295,7 @@ TEST_CASE_METHOD(FObjectHandleTestBase, "CoreUObject::TObjectPtr::Null Behavior"
 	CHECK_FALSE(ResolvedPtr == Ptr);  CHECK(ResolveCount == 0u);
 
 	//compare unresolved against nullptr
-	FObjectPtr FPtr(MakePackedObjectRef(TestObject));
+	FObjectPtr FPtr(MakeUnresolvedHandle(TestObject));
 	TObjectPtr<UObject> UnResolvedPtr = *reinterpret_cast<TObjectPtr<UObject>*>(&FPtr);
 	CHECK(!UnResolvedPtr.IsResolved());
 	CHECK_FALSE(Ptr == UnResolvedPtr); CHECK(ResolveCount == 0u);
@@ -400,8 +345,8 @@ TEST_CASE_METHOD(FObjectHandleTestBase, "CoreUObject::TObjectPtr::Null Behavior"
 
 TEST_CASE_METHOD(FObjectHandleTestBase, "CoreUObject::FObjectHandle::Resolve Malformed Handle", "[CoreUObject][ObjectHandle]")
 {
-	TestResolveFailure(FPackedObjectRef { 0xFFFF'FFFF'FFFF'FFFFull });
-	TestResolveFailure(FPackedObjectRef { 0xEFEF'EFEF'EFEF'EFEFull });
+	TestResolveFailure(UE::CoreUObject::Private::FPackedObjectRef { 0xFFFF'FFFF'FFFF'FFFFull });
+	TestResolveFailure(UE::CoreUObject::Private::FPackedObjectRef { 0xEFEF'EFEF'EFEF'EFEFull });
 }
 #endif // UE_WITH_OBJECT_HANDLE_LATE_RESOLVE
 
diff --git a/Engine/Source/Runtime/CoreUObject/Tests/ObjectPtrTest.cpp b/Engine/Source/Runtime/CoreUObject/Tests/ObjectPtrTest.cpp
--- a/Engine/Source/Runtime/CoreUObject/Tests/ObjectPtrTest.cpp
+++ b/Engine/Source/Runtime/CoreUObject/Tests/ObjectPtrTest.cpp
@@ -10,14 +10,14 @@
 #include "UObject/LinkerLoadImportBehavior.h"
 #include "UObject/Interface.h"
 #include "UObject/MetaData.h"
+#include "UObject/ObjectPathId.h"
 #include "UObject/ObjectRedirector.h"
 #include "UObject/Package.h"
 #include "UObject/SoftObjectPath.h"
 #include "Misc/ScopeExit.h"
 #include <type_traits>
 
-
-namespace ObjectPtrTest
+namespace UE::CoreObject::Private::Tests
 {
 using FMutableObjectPtr = TObjectPtr<UObject>;
 using FMutableInterfacePtr = TObjectPtr<UInterface>;
@@ -169,6 +169,7 @@ TEST_CASE("CoreUObject::TObjectPtr::Null", "[CoreUObject][ObjectPtr]")
 	TEST_TRUE(TEXT("Negation of a null object pointer should evaluate to true"), !NullObjectPtr);
 }
 
+#if UE_WITH_OBJECT_HANDLE_LATE_RESOLVE || UE_WITH_OBJECT_HANDLE_TRACKING
 TEST_CASE_METHOD(FObjectPtrTestBase, "CoreUObject::TObjectPtr::Default Serialize", "[CoreUObject][ObjectPtr]")
 {
 	FSnapshotObjectRefMetrics ObjectRefMetrics(*this);
@@ -178,27 +179,34 @@ TEST_CASE_METHOD(FObjectPtrTestBase, "CoreUObject::TObjectPtr::Default Serialize
 	TestPackage->AddToRoot();
 	UObject* TestSoftObject = NewObject<UObjectPtrTestClass>(TestPackage, TEXT("DefaultSerializeObject"));
 
-	FObjectPtr DefaultSerializeObjectPtr(FObjectRef {FName("/Engine/Test/ObjectPtrDefaultSerialize/Transient"), NAME_None, NAME_None, FObjectPathId("DefaultSerializeObject")});
+#if UE_WITH_OBJECT_HANDLE_LATE_RESOLVE
+	FObjectPtr DefaultSerializeObjectPtr(MakeUnresolvedHandle(TestSoftObject));
+#else
+	FObjectPtr DefaultSerializeObjectPtr(TestSoftObject);
+#endif
 
-	ObjectRefMetrics.TestNumResolves(TEXT("Unexpected resolve count after initializing an FObjectPtr"), UE_WITH_OBJECT_HANDLE_LATE_RESOLVE ? 0 : 1);
+	ObjectRefMetrics.TestNumResolves(TEXT("Unexpected resolve count after initializing an FObjectPtr"), 0);
 	ObjectRefMetrics.TestNumFailedResolves(TEXT("Unexpected resolve failure after initializing an FObjectPtr"), 0);
 	ObjectRefMetrics.TestNumReads(TEXT("NumReads should not change when initializing an FObjectPtr"), 0);
 
 	FArchiveUObject Writer;
 	Writer << DefaultSerializeObjectPtr;
 
-	ObjectRefMetrics.TestNumResolves(TEXT("Serializing an FObjectPtr should force it to resolve"), 1);
+	ObjectRefMetrics.TestNumResolves(TEXT("Serializing an FObjectPtr should force it to resolve"), UE_WITH_OBJECT_HANDLE_LATE_RESOLVE ? 1 : 0);
 	ObjectRefMetrics.TestNumFailedResolves(TEXT("Unexpected resolve failure after serializing an FObjectPtr"), 0);
 	ObjectRefMetrics.TestNumReads(TEXT("NumReads should increase after serializing an FObjectPtr"), 1);
 
 	Writer << DefaultSerializeObjectPtr;
 
-	ObjectRefMetrics.TestNumResolves(TEXT("Serializing an FObjectPtr twice should only require it to resolve once"), 1);
+	ObjectRefMetrics.TestNumResolves(TEXT("Serializing an FObjectPtr twice should only require it to resolve once"), UE_WITH_OBJECT_HANDLE_LATE_RESOLVE ? 1 : 0);
 	ObjectRefMetrics.TestNumFailedResolves(TEXT("Unexpected resolve failure after serializing an FObjectPtr"), 0);
 	ObjectRefMetrics.TestNumReads(TEXT("NumReads should increase after serializing an FObjectPtr"), 2);
 
 	TestPackage->RemoveFromRoot();
 }
+#endif
+
+#if UE_WITH_OBJECT_HANDLE_LATE_RESOLVE || UE_WITH_OBJECT_HANDLE_TRACKING
 
 TEST_CASE_METHOD(FObjectPtrTestBase, "CoreUObject::TObjectPtr::Soft Object Path", "[CoreUObject][ObjectPtr]")
 {
@@ -209,20 +217,25 @@ TEST_CASE_METHOD(FObjectPtrTestBase, "CoreUObject::TObjectPtr::Soft Object Path"
 	TestPackage->AddToRoot();
 	UObject* TestSoftObject = NewObject<UObjectPtrTestClass>(TestPackage, TEXT("TestSoftObject"));
 
-	FObjectPtr DefaultSoftObjPtr(FObjectRef {FName("/Engine/Test/ObjectPtrSoftObjectPath/Transient"), NAME_None, NAME_None, FObjectPathId("TestSoftObject")});
-
-	ObjectRefMetrics.TestNumResolves(TEXT("Unexpected resolve count after initializing an FObjectPtr"), UE_WITH_OBJECT_HANDLE_LATE_RESOLVE ? 0 : 1);
+#if UE_WITH_OBJECT_HANDLE_LATE_RESOLVE
+	FObjectPtr DefaultSoftObjPtr(MakeUnresolvedHandle(TestSoftObject));
+#else
+	FObjectPtr DefaultSoftObjPtr(TestSoftObject);
+#endif
+	ObjectRefMetrics.TestNumResolves(TEXT("Unexpected resolve count after initializing an FObjectPtr"), 0);
 	ObjectRefMetrics.TestNumFailedResolves(TEXT("Unexpected resolve failure after initializing an FObjectPtr"), 0);
 	ObjectRefMetrics.TestNumReads(TEXT("NumReads should not change when initializing an FObjectPtr"), 0);
 
 	// Initializing a soft object path from a TObjectPtr that's unresolved should stay unresolved.
 	FSoftObjectPath DefaultSoftObjPath(DefaultSoftObjPtr);
 	ObjectRefMetrics.TestNumResolves(TEXT("Unexpected resolve count after initializing an FSoftObjectPath from an FObjectPtr"), 0);
+	ObjectRefMetrics.TestNumReads(TEXT("NumReads should have changed when initializing a FSoftObjectPath"), UE_WITH_OBJECT_HANDLE_LATE_RESOLVE ? 0 : 1);
 
 	TEST_EQUAL_STR(TEXT("Soft object path constructed from an FObjectPtr does not have the expected path value"), TEXT("/Engine/Test/ObjectPtrSoftObjectPath/Transient.TestSoftObject"), *DefaultSoftObjPath.ToString());
 
 	TestPackage->RemoveFromRoot();
 }
+#endif
 
 TEST_CASE_METHOD(FObjectPtrTestBase, "CoreUObject::TObjectPtr::Forward Declared", "[CoreUObject][ObjectPtr]")
 {
@@ -231,6 +244,8 @@ TEST_CASE_METHOD(FObjectPtrTestBase, "CoreUObject::TObjectPtr::Forward Declared"
 	TEST_TRUE(TEXT("Null forward declared pointer used to construct a TObjectPtr should result in a null TObjectPtr"), !ObjPtrFwd);
 }
 
+#if UE_WITH_OBJECT_HANDLE_LATE_RESOLVE
+
 TEST_CASE_METHOD(FObjectPtrTestBase, "CoreUObject::TObjectPtr::Hash Consistency", "[CoreUObject][ObjectPtr]")
 {
 	const FName TestPackage1Name(TEXT("/Engine/Test/ObjectPtrHashConsistency1/Transient"));
@@ -242,16 +257,19 @@ TEST_CASE_METHOD(FObjectPtrTestBase, "CoreUObject::TObjectPtr::Hash Consistency"
 	UObject* TestOuter4 = NewObject<UObjectPtrTestClass>(TestPackage1, TEXT("TestOuter4"));
 
 	UObject* TestPublicObject = NewObject<UObjectPtrTestClass>(TestOuter1, TEXT("TestPublicObject"), RF_Public);
-#if UE_WITH_OBJECT_HANDLE_LATE_RESOLVE
-	MakePackedObjectRef(TestPublicObject); //construct a packed object ref the exported object. replicates linker load
-#endif
+
+	UE::CoreUObject::Private::MakePackedObjectRef(TestPublicObject); //construct a packed object ref the exported object. replicates linker load
 	// Perform hash consistency checks on public object reference
 	{
 		// Check that unresolved/resolved pointers produce the same hash
 		FSnapshotObjectRefMetrics ObjectRefMetrics(*this);
 
-		FObjectPtr TestPublicWrappedObjectPtr(FObjectRef{TestPackage1Name, TestPublicObject->GetClass()->GetOutermost()->GetFName(), TestPublicObject->GetClass()->GetFName(), FObjectPathId("TestOuter1.TestPublicObject")});
-		ObjectRefMetrics.TestNumResolves(TEXT("Unexpected resolve count after initializing an FObjectPtr"), UE_WITH_OBJECT_HANDLE_LATE_RESOLVE ? 0 : 1);
+#if UE_WITH_OBJECT_HANDLE_LATE_RESOLVE
+		FObjectPtr TestPublicWrappedObjectPtr(MakeUnresolvedHandle(TestPublicObject));
+#else
+		FObjectPtr TestPublicWrappedObjectPtr(TestPublicObject);
+#endif
+		ObjectRefMetrics.TestNumResolves(TEXT("Unexpected resolve count after initializing an FObjectPtr"), 0);
 		ObjectRefMetrics.TestNumFailedResolves(TEXT("Unexpected resolve failure after initializing an FObjectPtr"), 0);
 
 		uint32 HashWrapped = GetTypeHash(TestPublicWrappedObjectPtr);
@@ -280,8 +298,9 @@ TEST_CASE_METHOD(FObjectPtrTestBase, "CoreUObject::TObjectPtr::Hash Consistency"
 
 	TestPackage1->RemoveFromRoot();
 }
+#endif
 
-
+#if UE_WITH_OBJECT_HANDLE_LATE_RESOLVE
 TEST_CASE_METHOD(FObjectPtrTestBase, "CoreUObject::TObjectPtr::Long Path", "[CoreUObject][ObjectPtr]")
 {
 	const FName TestPackage1Name(TEXT("/Engine/Test/FObjectPtrTestLongPath/Transient"));
@@ -297,8 +316,8 @@ TEST_CASE_METHOD(FObjectPtrTestBase, "CoreUObject::TObjectPtr::Long Path", "[Cor
 	UObject* TestObject3 = NewObject<UObjectPtrTestClass>(TestObject2, TEXT("TestObject3"));
 	UObject* TestObject4 = NewObject<UObjectPtrTestClass>(TestObject3, TEXT("TestObject4"));
 
-	FObjectPathId LongPath(TestObject4);
-	FObjectPathId::ResolvedNameContainerType ResolvedNames;
+	UE::CoreUObject::Private::FObjectPathId LongPath(TestObject4);
+	UE::CoreUObject::Private::FObjectPathId::ResolvedNameContainerType ResolvedNames;
 	LongPath.Resolve(ResolvedNames);
 
 	TEST_EQUAL(TEXT("Resolved path from FObjectPathId should have 4 elements"), ResolvedNames.Num(), 4);
@@ -307,6 +326,7 @@ TEST_CASE_METHOD(FObjectPtrTestBase, "CoreUObject::TObjectPtr::Long Path", "[Cor
 	TEST_EQUAL(TEXT("Resolved path from FObjectPathId should have TestObject3 at element 2"), ResolvedNames[2], TestObject3->GetFName());
 	TEST_EQUAL(TEXT("Resolved path from FObjectPathId should have TestObject4 at element 3"), ResolvedNames[3], TestObject4->GetFName());
 }
+#endif
 
 TEST_CASE("CoreUObject::TObjectPtr::GetPathName", "[CoreUObject][ObjectPtr]")
 {
@@ -343,14 +363,14 @@ TEST_CASE("CoreUObject::TObjectPtr::GetPathName", "[CoreUObject][ObjectPtr]")
 
 #if UE_WITH_OBJECT_HANDLE_LATE_RESOLVE
 	{
-		FObjectPtr ObjPtr(MakePackedObjectRef(TestObject1));
+		FObjectPtr ObjPtr(MakeUnresolvedHandle(TestObject1));
 		Ptr = *reinterpret_cast<TObjectPtr<UObject>*>(&ObjPtr);
 		REQUIRE(!Ptr.IsResolved());
 		CHECK(TestObject1->GetPathName() == Ptr.GetPathName());
 		CHECK(TestObject1->GetClass() == Ptr.GetClass());
 	}
 	{
-		FObjectPtr ObjPtr(MakePackedObjectRef(TestPackage1));
+		FObjectPtr ObjPtr(MakeUnresolvedHandle(TestPackage1));
 		Ptr = *reinterpret_cast<TObjectPtr<UObject>*>(&ObjPtr);
 		REQUIRE(!Ptr.IsResolved());
 		CHECK(TestPackage1->GetPathName() == Ptr.GetPathName());
@@ -371,7 +391,7 @@ TEST_CASE("CoreUObject::TObjectPtr::PackageRename")
 	TestPackage->AddToRoot();
 
 	//register package with the object handle registry
-	MakePackedObjectRef(TestPackage);
+	UE::CoreUObject::Private::MakePackedObjectRef(TestPackage);
 
 	UObject* Obj1 = NewObject<UObjectPtrTestClass>(TestPackage, TEXT("Obj1"));
 	UObject* Inner1 = NewObject<UObjectPtrTestClass>(Obj1, TEXT("Inner1"));
@@ -380,8 +400,8 @@ TEST_CASE("CoreUObject::TObjectPtr::PackageRename")
 
 	Class->GetOutermost()->GetMetaData()->SetValue(Class, TEXT("LoadBehavior"), TEXT("LazyOnDemand"));
 
-	FObjectPtr ObjPtr(MakePackedObjectRef(Obj1));
-	FObjectPtr InnerPtr(MakePackedObjectRef(Inner1));
+	FObjectPtr ObjPtr(MakeUnresolvedHandle(Obj1));
+	FObjectPtr InnerPtr(MakeUnresolvedHandle(Inner1));
 
 	TObjectPtr<UObject> BeforeRename = *reinterpret_cast<TObjectPtr<UObject>*>(&ObjPtr);
 	TObjectPtr<UObject> BeforeRenameInner = *reinterpret_cast<TObjectPtr<UObject>*>(&InnerPtr);
@@ -400,8 +420,8 @@ TEST_CASE("CoreUObject::TObjectPtr::PackageRename")
 	CHECK(BeforeRename == AfterRenameResolved);
 	CHECK(BeforeRenameInner == AfterRenameInnerResolved);
 
-	FObjectPtr ObjPtr2(MakePackedObjectRef(Obj1));
-	FObjectPtr InnerPtr2(MakePackedObjectRef(Inner1));
+	FObjectPtr ObjPtr2(MakeUnresolvedHandle(Obj1));
+	FObjectPtr InnerPtr2(MakeUnresolvedHandle(Inner1));
 	TObjectPtr<UObject> AfterRenameUnresolved = *reinterpret_cast<TObjectPtr<UObject>*>(&ObjPtr2);
 	TObjectPtr<UObject> AfterRenameInnerUnresolved = *reinterpret_cast<TObjectPtr<UObject>*>(&InnerPtr2);
 
@@ -426,11 +446,11 @@ TEST_CASE("CoreUObject::TObjectPtr::InnerRename")
 	UObject* Obj2 = NewObject<UObjectPtrTestClass>(TestPackage, TEXT("Obj2"));
 	UObject* Inner2 = NewObject<UObjectPtrTestClass>(Obj2, TEXT("Inner2"));
 
-	FObjectPtr ObjPtr1(MakePackedObjectRef(Obj1));
-	FObjectPtr InnerPtr1(MakePackedObjectRef(Inner1));
+	FObjectPtr ObjPtr1(MakeUnresolvedHandle(Obj1));
+	FObjectPtr InnerPtr1(MakeUnresolvedHandle(Inner1));
 
-	FObjectPtr ObjPtr2(MakePackedObjectRef(Obj2));
-	FObjectPtr InnerPtr2(MakePackedObjectRef(Inner2));
+	FObjectPtr ObjPtr2(MakeUnresolvedHandle(Obj2));
+	FObjectPtr InnerPtr2(MakeUnresolvedHandle(Inner2));
 
 	TObjectPtr<UObject> BeforeRename1 = *reinterpret_cast<TObjectPtr<UObject>*>(&ObjPtr1);
 	TObjectPtr<UObject> BeforeRenameInner1 = *reinterpret_cast<TObjectPtr<UObject>*>(&InnerPtr1);
@@ -450,8 +470,8 @@ TEST_CASE("CoreUObject::TObjectPtr::InnerRename")
 
 	Obj1->Rename(TEXT("RenamedObj"), nullptr);
 
-	FObjectPtr AfterRenameObjPtr1(MakePackedObjectRef(Obj1));
-	FObjectPtr AfterRenameInnerPtr1(MakePackedObjectRef(Inner1));
+	FObjectPtr AfterRenameObjPtr1(MakeUnresolvedHandle(Obj1));
+	FObjectPtr AfterRenameInnerPtr1(MakeUnresolvedHandle(Inner1));
 
 	TObjectPtr<UObject> AfterRenameUnresolved1 = *reinterpret_cast<TObjectPtr<UObject>*>(&AfterRenameObjPtr1);
 	TObjectPtr<UObject> AfterRenameInnerUnresolved1 = *reinterpret_cast<TObjectPtr<UObject>*>(&AfterRenameInnerPtr1);
@@ -490,8 +510,8 @@ TEST_CASE("CoreUObject::TObjectPtr::Move")
 	UPackage* TestPackageB = NewObject<UPackage>(nullptr, TEXT("/Engine/PackageB"), RF_Transient);
 	TestPackageB->AddToRoot();
 
-	FObjectPtr FObjPtr(MakePackedObjectRef(Obj1));
-	FObjectPtr FInnerPtr(MakePackedObjectRef(Inner1));
+	FObjectPtr FObjPtr(MakeUnresolvedHandle(Obj1));
+	FObjectPtr FInnerPtr(MakeUnresolvedHandle(Inner1));
 
 	TObjectPtr<UObject> BeforeRename = *reinterpret_cast<TObjectPtr<UObject>*>(&FObjPtr);
 	TObjectPtr<UObject> BeforeRenameInner = *reinterpret_cast<TObjectPtr<UObject>*>(&FInnerPtr);
@@ -540,9 +560,9 @@ TEST_CASE("CoreUObject::TObjectPtr::GetTypeHash")
 	TObjectPtr<UObject> Obj1 = NewObject<UObjectPtrTestClass>(TestPackageA, TEXT("Obj1"));
 	TObjectPtr<UObject> Inner1 = NewObject<UObjectPtrTestClass>(Obj1, TEXT("Inner1"));
 
-	MakePackedObjectRef(TestPackageA);
-	MakePackedObjectRef(Obj1.Get());
-	MakePackedObjectRef(Inner1.Get());
+	UE::CoreUObject::Private::MakePackedObjectRef(TestPackageA);
+	UE::CoreUObject::Private::MakePackedObjectRef(Obj1.Get());
+	UE::CoreUObject::Private::MakePackedObjectRef(Inner1.Get());
 
 	TMap<TObjectPtr<UObject>, int> TestMap;
 
@@ -571,9 +591,9 @@ TEST_CASE("CoreUObject::TObjectPtr::GetTypeHash")
 	CHECK(TestMap[NotLazy] == NotLazyValue);
 
 
-	FObjectPtr FPackagePtr(MakePackedObjectRef(TestPackageA));
-	FObjectPtr FObjPtr(MakePackedObjectRef(Obj1));
-	FObjectPtr FInnerPtr(MakePackedObjectRef(Inner1));
+	FObjectPtr FPackagePtr(MakeUnresolvedHandle(TestPackageA));
+	FObjectPtr FObjPtr(MakeUnresolvedHandle(Obj1));
+	FObjectPtr FInnerPtr(MakeUnresolvedHandle(Inner1));
 
 	TObjectPtr<UObject> BeforeRename = *reinterpret_cast<TObjectPtr<UObject>*>(&FObjPtr);
 	TObjectPtr<UObject> BeforeRenameInner = *reinterpret_cast<TObjectPtr<UObject>*>(&FInnerPtr);
@@ -783,4 +803,4 @@ TEST_CASE("CoreUObject::TObjectPtr::TestEquals")
 	CHECK(BasePtr == ObjPtr);
 }
 
-#endif 
+#endif
\ No newline at end of file
diff --git a/Engine/Source/Runtime/CoreUObject/Tests/ObjectRefTrackingTestBase.h b/Engine/Source/Runtime/CoreUObject/Tests/ObjectRefTrackingTestBase.h
--- a/Engine/Source/Runtime/CoreUObject/Tests/ObjectRefTrackingTestBase.h
+++ b/Engine/Source/Runtime/CoreUObject/Tests/ObjectRefTrackingTestBase.h
@@ -10,6 +10,14 @@
 #include "UObject/ObjectHandle.h"
 #include "Misc/ScopeLock.h"
 
+#if UE_WITH_OBJECT_HANDLE_LATE_RESOLVE
+inline FObjectHandle MakeUnresolvedHandle(const UObject* Obj)
+{
+	UE::CoreUObject::Private::FPackedObjectRef PackedObjectRef = UE::CoreUObject::Private::MakePackedObjectRef(Obj);
+	return { PackedObjectRef.EncodedRef };
+}
+
+#endif
 
 class FObjectRefTrackingTestBase
 {
@@ -50,8 +58,8 @@ public:
 			bool bValue = false;
 			if (bAllowAdditionalReads)
 			{
-				bValue = Test.GetNumReads() >= OriginalNumReads + ExpectedDelta;
-				TEST_TRUE(What, bValue);
+				INFO(What);
+				CHECK(Test.GetNumReads() >= OriginalNumReads + ExpectedDelta);
 			}
 			else
 			{
@@ -72,7 +80,7 @@ private:
 	static void OnRefResolved(const FObjectRef& ObjectRef, UPackage* Pkg, UObject* Obj)
 	{
 		NumResolves++;
-		if (!IsObjectRefNull(ObjectRef) && !Obj)
+		if (!ObjectRef.IsNull() && !Obj)
 		{
 			NumFailedResolves++;
 		}
