FUObjectToken: Fix OriginalObjectPathName not being set if a LabelOverride was provided in constructor

diff --git a/Engine/Source/Runtime/CoreUObject/Private/Misc/UObjectToken.cpp b/Engine/Source/Runtime/CoreUObject/Private/Misc/UObjectToken.cpp
--- a/Engine/Source/Runtime/CoreUObject/Private/Misc/UObjectToken.cpp
+++ b/Engine/Source/Runtime/CoreUObject/Private/Misc/UObjectToken.cpp
@@ -10,6 +10,11 @@ FUObjectToken::FOnGetDisplayName FUObjectToken::DefaultGetObjectDisplayName;
 FUObjectToken::FUObjectToken( const UObject* InObject,  const FText& InLabelOverride )
 	: ObjectBeingReferenced(InObject)
 {
+	if (InObject)
+	{
+		OriginalObjectPathName = InObject->GetPathName();
+	}
+	
 	if ( !InLabelOverride.IsEmpty() )
 	{
 		CachedText = InLabelOverride;
@@ -19,16 +24,10 @@ FUObjectToken::FUObjectToken( const UObject* InObject,  const FText& InLabelOver
 		if ( DefaultGetObjectDisplayName.IsBound() )
 		{
 			CachedText = DefaultGetObjectDisplayName.Execute(InObject, false);
-
-			if (InObject)
-			{
-				OriginalObjectPathName = InObject->GetPathName();
-			}
 		}
 		else if ( InObject )
 		{
 			CachedText = FText::FromString( InObject->GetName() );
-			OriginalObjectPathName = InObject->GetPathName();
 		}
 		else
 		{
