AssetRegistry: Fix missing PackageName in ObjectName of AssetDatas updated when a package is saved.
#rnx

diff --git a/Engine/Source/Runtime/CoreUObject/Private/UObject/SavePackage/SavePackageUtilities.cpp b/Engine/Source/Runtime/CoreUObject/Private/UObject/SavePackage/SavePackageUtilities.cpp
--- a/Engine/Source/Runtime/CoreUObject/Private/UObject/SavePackage/SavePackageUtilities.cpp
+++ b/Engine/Source/Runtime/CoreUObject/Private/UObject/SavePackage/SavePackageUtilities.cpp
@@ -2123,6 +2123,19 @@ namespace AssetRegistry
 						TagsAndValues.Add(Tag.Name, MoveTemp(Tag.Value));
 					}
 				}
+				// if we do not have a full object path already, build it
+				const bool bFullObjectPath = ObjectPath.StartsWith(TEXT("/"), ESearchCase::CaseSensitive);
+				if (!bFullObjectPath)
+				{
+					// if we do not have a full object path, ensure that we have a top level object for the package and not a sub object
+					if (!ensureMsgf(!ObjectPath.Contains(TEXT("."), ESearchCase::CaseSensitive),
+						TEXT("Cannot make FAssetData for sub object %s in package %s!"), *ObjectPath, *PackageName))
+					{
+						continue;
+					}
+					ObjectPath = PackageName + TEXT(".") + ObjectPath;
+				}
+
 				OutAssetDatas->Emplace(PackageName, ObjectPath, FTopLevelAssetPath(ObjectClassName),
 					MoveTemp(TagsAndValues), Package->GetChunkIDs(), Package->GetPackageFlags());
 			}
