Added the ability to define USE_COMPILED_IN_NATIVES outside of the HACK_HEADER_GENERATOR macro.

diff --git a/Engine/Source/Runtime/CoreUObject/Public/UObject/ObjectMacros.h b/Engine/Source/Runtime/CoreUObject/Public/UObject/ObjectMacros.h
--- a/Engine/Source/Runtime/CoreUObject/Public/UObject/ObjectMacros.h
+++ b/Engine/Source/Runtime/CoreUObject/Public/UObject/ObjectMacros.h
@@ -25,11 +25,13 @@ typedef	uint64 ScriptPointerType;
 #error "not supported in UE"
 #endif
 
+#ifndef USE_COMPILED_IN_NATIVES
 #if HACK_HEADER_GENERATOR 
 #define USE_COMPILED_IN_NATIVES	0
 #else
 #define USE_COMPILED_IN_NATIVES	1
 #endif
+#endif
 
 /** Set this to 0 to disable UObject thread safety features */
 #define THREADSAFE_UOBJECTS 1
