changed ObjectHandleTracking to use Functions instead of delegates as the performance of delegate is poor
https://p4-swarm.epicgames.net/reviews/23751084

diff --git a/Engine/Source/Editor/UnrealEd/Private/Cooker/PackageBuildDependencyTracker.cpp b/Engine/Source/Editor/UnrealEd/Private/Cooker/PackageBuildDependencyTracker.cpp
--- a/Engine/Source/Editor/UnrealEd/Private/Cooker/PackageBuildDependencyTracker.cpp
+++ b/Engine/Source/Editor/UnrealEd/Private/Cooker/PackageBuildDependencyTracker.cpp
@@ -56,12 +56,12 @@ TArray<FBuildDependencyAccessData> FPackageBuildDependencyTracker::GetAccessData
 
 FPackageBuildDependencyTracker::FPackageBuildDependencyTracker()
 {
-	ObjectHandleReadHandle = AddObjectHandleReadCallback(FObjectHandleReadDelegate::CreateStatic(StaticOnObjectHandleRead));
+	ObjectHandleReadHandle = UE::CoreUObject::AddObjectHandleReadCallback(StaticOnObjectHandleRead);
 }
 
 FPackageBuildDependencyTracker::~FPackageBuildDependencyTracker()
 {
-	RemoveObjectHandleReadCallback(ObjectHandleReadHandle);
+	UE::CoreUObject::RemoveObjectHandleReadCallback(ObjectHandleReadHandle);
 }
 
 void FPackageBuildDependencyTracker::StaticOnObjectHandleRead(UObject* ReadObject)
diff --git a/Engine/Source/Editor/UnrealEd/Private/Cooker/PackageBuildDependencyTracker.h b/Engine/Source/Editor/UnrealEd/Private/Cooker/PackageBuildDependencyTracker.h
--- a/Engine/Source/Editor/UnrealEd/Private/Cooker/PackageBuildDependencyTracker.h
+++ b/Engine/Source/Editor/UnrealEd/Private/Cooker/PackageBuildDependencyTracker.h
@@ -47,7 +47,7 @@ private:
 	FName LastReferencer = NAME_None;
 	FBuildDependencyAccessData LastAccessData{ NAME_None, nullptr };
 	TSet<FBuildDependencyAccessData>* LastReferencerSet = nullptr;
-	FDelegateHandle ObjectHandleReadHandle;
+	UE::CoreUObject::FObjectHandleTrackingCallbackId ObjectHandleReadHandle;
 	static FPackageBuildDependencyTracker Singleton;
 };
 
diff --git a/Engine/Source/Runtime/CoreUObject/Private/UObject/ObjectHandleTracking.cpp b/Engine/Source/Runtime/CoreUObject/Private/UObject/ObjectHandleTracking.cpp
--- a/Engine/Source/Runtime/CoreUObject/Private/UObject/ObjectHandleTracking.cpp
+++ b/Engine/Source/Runtime/CoreUObject/Private/UObject/ObjectHandleTracking.cpp
@@ -1,96 +1,151 @@
 // Copyright Epic Games, Inc. All Rights Reserved.
 
 #include "UObject/ObjectHandleTracking.h"
+#include "Misc/ScopeRWLock.h"
 
 #if UE_WITH_OBJECT_HANDLE_TRACKING
 
-COREUOBJECT_API std::atomic<int32> UE::CoreUObject::Private::ObjectHandleEventIndex;
-COREUOBJECT_API UE::CoreUObject::Private::FObjectHandleEvents UE::CoreUObject::Private::ObjectHandleEvents[2];
-
-namespace UE::CoreUObject::Private
+namespace UE::CoreUObject
 {
-	FObjectHandleEvents& BeginWritingEvents()
+	namespace Private
 	{
-		// We only allow adding/removing the callbacks from a single thread.
-		check(IsInGameThread());
-
-		// Get what the last and new index will be.
-		const int32 LastIndex = ObjectHandleEventIndex;
-		const int32 NewIndex = LastIndex == 0 ? 1 : 0;
-
-		// Make sure we're not still using the events, this should definitely not be true, because
-		// we spin lock at the end of any of these functions to make sure all using has stopped.
-		check(!ObjectHandleEvents[NewIndex].IsUsing());
-
-		// We need to take whatever delegates are setup on the previous index and copy them over to new set.
-		ObjectHandleEvents[NewIndex].ObjectHandleReadEvent = ObjectHandleEvents[LastIndex].ObjectHandleReadEvent;
-		ObjectHandleEvents[NewIndex].ClassReferenceResolvedEvent = ObjectHandleEvents[LastIndex].ClassReferenceResolvedEvent;
-		ObjectHandleEvents[NewIndex].ObjectHandleReferenceResolvedEvent = ObjectHandleEvents[LastIndex].ObjectHandleReferenceResolvedEvent;
-		ObjectHandleEvents[NewIndex].ObjectHandleReferenceLoadedEvent = ObjectHandleEvents[LastIndex].ObjectHandleReferenceLoadedEvent;
+		int32 NextHandleId = 0;
+		TArray<TTuple<int32, FObjectHandleReadFunc>> ReadHandleCallbacks;
+		TArray<TTuple<int32, FObjectHandleClassResolvedFunc>> ClassResolvedCallbacks;
+		TArray<TTuple<int32, FObjectHandleReferenceResolvedFunc>> HandleResolvedCallbacks;
+		TArray<TTuple<int32, FObjectHandleReferenceLoadedFunc>> HandleLoadedCallbacks;
+		
+		//using a single lock as currently there is not any contention. could add separate locks for each list. 
+		FRWLock HandleLock;
+
+		void OnHandleRead(UObject* Object)
+		{
+			FReadScopeLock _(HandleLock);
+			for (auto&& Pair : ReadHandleCallbacks)
+			{
+				Pair.Value(Object);
+			}
+		}
 
-		// Now toggle which buffer we're using.
-		ObjectHandleEventIndex = NewIndex;
+		void OnClassReferenceResolved(const FObjectRef& ObjectRef, UPackage* Package, UClass* Class)
+		{
+			FReadScopeLock _(HandleLock);
+			for (auto&& Pair : ClassResolvedCallbacks)
+			{
+				Pair.Value(ObjectRef, Package, Class);
+			}
+		}
 
-		// Ok, now we need to spin lock until nobody is using the LastIndex, this will prevent us from
-		// calling this function right after this one and assuming the previous set isn't being used,
-		// similarly while we were doing this, the BeginReadingEvents() function spun lock after marked
-		// that they were using the event set, so now that we've set the new event index, if we spin lock
-		// here, that should ensure that any existing readers finish up what they were doing.
-		while (ObjectHandleEvents[LastIndex].IsUsing())
+		void OnReferenceResolved(const FObjectRef& ObjectRef, UPackage* Package, UObject* Object)
 		{
-			// Spin!
+			FReadScopeLock _(HandleLock);
+			for (auto&& Pair : HandleResolvedCallbacks)
+			{
+				Pair.Value(ObjectRef, Package, Object);
+			}
 		}
 
-		return ObjectHandleEvents[NewIndex];
+		void OnReferenceLoaded(const FObjectRef& ObjectRef, UPackage* Package, UObject* Object)
+		{
+			FReadScopeLock _(HandleLock);
+			for (auto&& Pair : HandleLoadedCallbacks)
+			{
+				Pair.Value(ObjectRef, Package, Object);
+			}
+		}
 	}
-}
 
-FDelegateHandle AddObjectHandleReadCallback(FObjectHandleReadDelegate Callback)
-{
-	UE::CoreUObject::Private::FObjectHandleEvents& Events = UE::CoreUObject::Private::BeginWritingEvents();
-	return Events.ObjectHandleReadEvent.Add(Callback);
-}
+	FObjectHandleTrackingCallbackId AddObjectHandleReadCallback(FObjectHandleReadFunc Func)
+	{
+		using namespace UE::CoreUObject::Private;
+		FWriteScopeLock _(HandleLock);
+		NextHandleId++;
+		ReadHandleCallbacks.Add({ NextHandleId, Func });
+		return { NextHandleId };
+	}
 
-void RemoveObjectHandleReadCallback(FDelegateHandle DelegateHandle)
-{
-	UE::CoreUObject::Private::FObjectHandleEvents& Events = UE::CoreUObject::Private::BeginWritingEvents();
-	Events.ObjectHandleReadEvent.Remove(DelegateHandle);
-}
+	void RemoveObjectHandleReadCallback(FObjectHandleTrackingCallbackId Handle)
+	{
+		using namespace UE::CoreUObject::Private;
+		FWriteScopeLock _(Private::HandleLock);
+		for (int32 i = ReadHandleCallbacks.Num() - 1; i >= 0; --i)
+		{
+			auto& Pair = ReadHandleCallbacks[i];
+			if (Pair.Key == Handle.Id)
+			{
+				ReadHandleCallbacks.RemoveAt(i);
+			}
+		}
+		
+	}
 
-FDelegateHandle AddObjectHandleClassResolvedCallback(FObjectHandleClassResolvedDelegate Callback)
-{
-	UE::CoreUObject::Private::FObjectHandleEvents& Events = UE::CoreUObject::Private::BeginWritingEvents();
-	return Events.ClassReferenceResolvedEvent.Add(Callback);
-}
+	FObjectHandleTrackingCallbackId AddObjectHandleClassResolvedCallback(FObjectHandleClassResolvedFunc Func)
+	{
+		using namespace UE::CoreUObject::Private;
+		FWriteScopeLock _(HandleLock);
+		NextHandleId++;
+		ClassResolvedCallbacks.Add({ NextHandleId, Func });
+		return { NextHandleId };
+	}
 
-void RemoveObjectHandleClassResolvedCallback(FDelegateHandle DelegateHandle)
-{
-	UE::CoreUObject::Private::FObjectHandleEvents& Events = UE::CoreUObject::Private::BeginWritingEvents();
-	Events.ClassReferenceResolvedEvent.Remove(DelegateHandle);
-}
+	void RemoveObjectHandleClassResolvedCallback(FObjectHandleTrackingCallbackId Handle)
+	{
+		using namespace UE::CoreUObject::Private;
+		FWriteScopeLock _(Private::HandleLock);
+		for (int32 i = ClassResolvedCallbacks.Num() - 1; i >= 0; --i)
+		{
+			auto& Pair = ClassResolvedCallbacks[i];
+			if (Pair.Key == Handle.Id)
+			{
+				ClassResolvedCallbacks.RemoveAt(i);
+			}
+		}
+	}
 
-FDelegateHandle AddObjectHandleReferenceResolvedCallback(FObjectHandleReferenceResolvedDelegate Callback)
-{
-	UE::CoreUObject::Private::FObjectHandleEvents& Events = UE::CoreUObject::Private::BeginWritingEvents();
-	return Events.ObjectHandleReferenceResolvedEvent.Add(Callback);
-}
+	FObjectHandleTrackingCallbackId AddObjectHandleReferenceResolvedCallback(FObjectHandleReferenceResolvedFunc Func)
+	{
+		using namespace UE::CoreUObject::Private;
+		FWriteScopeLock _(HandleLock);
+		NextHandleId++;
+		HandleResolvedCallbacks.Add({ NextHandleId, Func });
+		return { NextHandleId };
+	}
 
-void RemoveObjectHandleReferenceResolvedCallback(FDelegateHandle DelegateHandle)
-{
-	UE::CoreUObject::Private::FObjectHandleEvents& Events = UE::CoreUObject::Private::BeginWritingEvents();
-	Events.ObjectHandleReferenceResolvedEvent.Remove(DelegateHandle);
-}
+	void RemoveObjectHandleReferenceResolvedCallback(FObjectHandleTrackingCallbackId Handle)
+	{
+		using namespace UE::CoreUObject::Private;
+		FWriteScopeLock _(Private::HandleLock);
+		for (int32 i = HandleResolvedCallbacks.Num() - 1; i >= 0; --i)
+		{
+			auto& Pair = HandleResolvedCallbacks[i];
+			if (Pair.Key == Handle.Id)
+			{
+				HandleResolvedCallbacks.RemoveAt(i);
+			}
+		}
+	}
 
-FDelegateHandle AddObjectHandleReferenceLoadedCallback(FObjectHandleReferenceLoadedDelegate Callback)
-{
-	UE::CoreUObject::Private::FObjectHandleEvents& Events = UE::CoreUObject::Private::BeginWritingEvents();
-	return Events.ObjectHandleReferenceLoadedEvent.Add(Callback);
-}
+	FObjectHandleTrackingCallbackId AddObjectHandleReferenceLoadedCallback(FObjectHandleReferenceLoadedFunc Func)
+	{
+		using namespace UE::CoreUObject::Private;
+		FWriteScopeLock _(HandleLock);
+		NextHandleId++;
+		HandleLoadedCallbacks.Add({ NextHandleId, Func });
+		return { NextHandleId };
+	}
 
-void RemoveObjectHandleReferenceLoadedCallback(FDelegateHandle DelegateHandle)
-{
-	UE::CoreUObject::Private::FObjectHandleEvents& Events = UE::CoreUObject::Private::BeginWritingEvents();
-	Events.ObjectHandleReferenceLoadedEvent.Remove(DelegateHandle);
+	void RemoveObjectHandleReferenceLoadedCallback(FObjectHandleTrackingCallbackId Handle)
+	{
+		using namespace UE::CoreUObject::Private;
+		FWriteScopeLock _(Private::HandleLock);
+		for (int32 i = HandleLoadedCallbacks.Num() - 1; i >= 0; --i)
+		{
+			auto& Pair = HandleLoadedCallbacks[i];
+			if (Pair.Key == Handle.Id)
+			{
+				HandleLoadedCallbacks.RemoveAt(i);
+			}
+		}
+	}
 }
-
 #endif // UE_WITH_OBJECT_HANDLE_TRACKING
diff --git a/Engine/Source/Runtime/CoreUObject/Public/UObject/ObjectHandleTracking.h b/Engine/Source/Runtime/CoreUObject/Public/UObject/ObjectHandleTracking.h
--- a/Engine/Source/Runtime/CoreUObject/Public/UObject/ObjectHandleTracking.h
+++ b/Engine/Source/Runtime/CoreUObject/Public/UObject/ObjectHandleTracking.h
@@ -3,168 +3,108 @@
 
 #include "Delegates/Delegate.h"
 #include "UObject/ObjectHandleDefines.h"
+#include "Templates/Function.h"
 
 struct FObjectRef;
 
-/**
- * FObjectHandles can optionally support tracking.  Because of the low level nature of object handles, anything that
- * registers itself for these callbacks should ensure that it is:
- * 1) error free (ie: should not cause exceptions even in unusual circumstances)
- * 2) fault tolerant (ie: could be called at a time when an exception has happened)
- * 3) thread-safe (ie: could be called from any thread)
- * 4) high performance (ie: will be called many times)
- */
-#if UE_WITH_OBJECT_HANDLE_TRACKING
-
- /**
-  * Callback notifying when an object value is read from a handle.  Fired regardless of whether the handle
-  * was resolved as part of the read operation or not and whether the object being read is null or not.
-  *
-  * @param ReadObject	The object that was read from a handle.
-  */
-DECLARE_DELEGATE_OneParam(FObjectHandleReadDelegate, UObject* ReadObject);
-
-/**
- * Callback notifying when a class is resolved from an object handle or object reference.
- * Classes are resolved either independently for a given handle/reference or as part of each object resolve.
- *
- * @param SourceRef		The object reference (either user provided or internally computed from a handle) from which the class was resolved.
- * @param ClassPackage	The package containing the resolved class.
- * @param Class			The resolved class.
- */
-DECLARE_DELEGATE_ThreeParams(FObjectHandleClassResolvedDelegate, const FObjectRef& SourceRef, UPackage* ObjectPackage, UClass* Class);
-
-/**
- * Callback notifying when a object handle is resolved.
- *
- * @param SourceRef		The object reference (either user provided or internally computed from a handle) from which the class was resolved.
- * @param ClassPackage	The package containing the resolved class.
- * @param Class			The resolved class.
- */
-DECLARE_DELEGATE_ThreeParams(FObjectHandleReferenceResolvedDelegate, const FObjectRef& SourceRef, UPackage* ObjectPackage, UObject* Object);
-
-/**
- * Callback notifying when an object was loaded through an object handle.  Will not notify you about global object loads, just ones that occur
- * as the byproduct of resolving an ObjectHandle.
- *
- * @param SourceRef		The object reference (either user provided or internally computed from a handle) from which the class was resolved.
- * @param ClassPackage	The package containing the resolved class.
- * @param Class			The resolved class.
- */
-DECLARE_DELEGATE_ThreeParams(FObjectHandleReferenceLoadedDelegate, const FObjectRef& SourceRef, UPackage* ObjectPackage, UObject* Object);
-
-/**
- * Installs a new callback for notifications that an object value has been read from a handle.
- *
- * @param Function		The new handle read callback to install.
- * @return				The DelegateHandle so that you can remove the callback at a later date.
- */
-COREUOBJECT_API FDelegateHandle AddObjectHandleReadCallback(FObjectHandleReadDelegate Delegate);
-COREUOBJECT_API void RemoveObjectHandleReadCallback(FDelegateHandle DelegateHandle);
-
-/**
- * Installs a new callback for notifications that a class has been resolved from an object handle or object reference.
- *
- * @param Function		The new class resolved callback to install.
- * @return				The DelegateHandle so that you can remove the callback at a later date.
- */
-COREUOBJECT_API FDelegateHandle AddObjectHandleClassResolvedCallback(FObjectHandleClassResolvedDelegate Callback);
-COREUOBJECT_API void RemoveObjectHandleClassResolvedCallback(FDelegateHandle DelegateHandle);
-
-/**
- * Installs a new callback for notifications that an object has been resolved from an object handle or object reference.
- *
- * @param Function		The new object resolved callback to install.
- * @return				The DelegateHandle so that you can remove the callback at a later date.
- */
-COREUOBJECT_API FDelegateHandle AddObjectHandleReferenceResolvedCallback(FObjectHandleReferenceResolvedDelegate Callback);
-COREUOBJECT_API void RemoveObjectHandleReferenceResolvedCallback(FDelegateHandle DelegateHandle);
-
-/**
- * Installs a new callback for notifications that an object has been loaded from an object handle or object reference.
- *
- * @param Function		The new object resolved callback to install.
- * @return				The DelegateHandle so that you can remove the callback at a later date.
- */
-COREUOBJECT_API FDelegateHandle AddObjectHandleReferenceLoadedCallback(FObjectHandleReferenceLoadedDelegate Callback);
-COREUOBJECT_API void RemoveObjectHandleReferenceLoadedCallback(FDelegateHandle DelegateHandle);
-
-
-namespace UE::CoreUObject::Private
+namespace UE::CoreUObject
 {
-	DECLARE_MULTICAST_DELEGATE_OneParam(FObjectHandleReadEvent, UObject* Object);
-	DECLARE_MULTICAST_DELEGATE_ThreeParams(FClassReferenceResolvedEvent, const FObjectRef& ObjectRef, UPackage* Package, UClass* Class);
-	DECLARE_MULTICAST_DELEGATE_ThreeParams(FObjectHandleReferenceResolvedEvent, const FObjectRef& SourceRef, UPackage* ObjectPackage, UObject* Object);
-	DECLARE_MULTICAST_DELEGATE_ThreeParams(FObjectHandleReferenceLoadedEvent, const FObjectRef& SourceRef, UPackage* ObjectPackage, UObject* Object);
+	/**
+	 * FObjectHandles can optionally support tracking.  Because of the low level nature of object handles, anything that
+	 * registers itself for these callbacks should ensure that it is:
+	 * 1) error free (ie: should not cause exceptions even in unusual circumstances)
+	 * 2) fault tolerant (ie: could be called at a time when an exception has happened)
+	 * 3) thread-safe (ie: could be called from any thread)
+	 * 4) high performance (ie: will be called many times)
+	 */
+#if UE_WITH_OBJECT_HANDLE_TRACKING
 
-	struct FObjectHandleEvents
+	 /**
+	  * Callback notifying when an object value is read from a handle.  Fired regardless of whether the handle
+	  * was resolved as part of the read operation or not and whether the object being read is null or not.
+	  *
+	  * @param ReadObject	The object that was read from a handle.
+	  */
+	using FObjectHandleReadFunc = TFunction<void(UObject* ReadObject)>;
+
+	/**
+	 * Callback notifying when a class is resolved from an object handle or object reference.
+	 * Classes are resolved either independently for a given handle/reference or as part of each object resolve.
+	 *
+	 * @param SourceRef		The object reference (either user provided or internally computed from a handle) from which the class was resolved.
+	 * @param ClassPackage	The package containing the resolved class. May be null.
+	 * @param Class			The resolved class. May be null.
+	 */
+	using FObjectHandleClassResolvedFunc = TFunction<void(const FObjectRef& SourceRef, UPackage* ClassPackage, UClass* Class)>;
+	/**
+	 * Callback notifying when a object handle is resolved.
+	 *
+	 * @param SourceRef		The object reference (either user provided or internally computed from a handle) from which the object was resolved.
+	 * @param ObjectPackage	The package containing the resolved object. May be null.
+	 * @param Object		The resolved Object. May be null.
+	 */
+	using FObjectHandleReferenceResolvedFunc = TFunction<void(const FObjectRef& SourceRef, UPackage* ObjectPackage, UObject* Object)>;
+	/**
+	 * Callback notifying when an object was loaded through an object handle.  Will not notify you about global object loads, just ones that occur
+	 * as the byproduct of resolving an ObjectHandle.
+	 *
+	 * @param SourceRef		The object reference (either user provided or internally computed from a handle) from which the object was resolved.
+	 * @param ObjectPackage	The package containing the resolved object. May be null.
+	 * @param Object		The resolved object. May be null.
+	 */
+	using FObjectHandleReferenceLoadedFunc = TFunction<void(const FObjectRef& SourceRef, UPackage* ObjectPackage, UObject* Object)>;
+
+	struct FObjectHandleTrackingCallbackId
 	{
-		FObjectHandleReadEvent ObjectHandleReadEvent;
-		FClassReferenceResolvedEvent ClassReferenceResolvedEvent;
-		FObjectHandleReferenceResolvedEvent ObjectHandleReferenceResolvedEvent;
-		FObjectHandleReferenceLoadedEvent ObjectHandleReferenceLoadedEvent;
-
-		inline void BeginUsing() { UsingCount++; }
-		inline void EndUsing() { --UsingCount; }
-		inline bool IsUsing() const { return UsingCount > 0; }
-
-	private:
-		std::atomic<int32> UsingCount;
+		int32 Id = -1;
 	};
 
-	extern COREUOBJECT_API std::atomic<int32> ObjectHandleEventIndex;
-	extern COREUOBJECT_API FObjectHandleEvents ObjectHandleEvents[2];
-
-	inline FObjectHandleEvents& BeginReadingEvents()
-	{
-		// Quick spin lock to try and begin using the events, normally this wont actually spin, in like 99.999% of the time.
-		while (true)
-		{
-			// Start by getting the current event index from the double buffer.
-			const int32 InitialEventIndex = ObjectHandleEventIndex;
-			// Grab the event set at that index.
-			FObjectHandleEvents& Events = ObjectHandleEvents[InitialEventIndex];
-			// Begin using the events, this will signal we're using them.
-			Events.BeginUsing();
-			// Ok - now we're going to check that we're *still* using the same index.  If we are
-			// then we know that it didn't change out from under us in between getting Events, and calling BeginUsing().
-			if (InitialEventIndex == ObjectHandleEventIndex)
-			{
-				return Events;
-			}
-			// If the check above failed, then we need to cease using the events, because we're about to try this again.
-			Events.EndUsing();
-		}
-	}
-
-	inline void OnHandleRead(UObject* Object)
-	{
-		FObjectHandleEvents& Events = BeginReadingEvents();
-		Events.ObjectHandleReadEvent.Broadcast(Object);
-		Events.EndUsing();
-	}
-
-	inline void OnClassReferenceResolved(const FObjectRef& ObjectRef, UPackage* Package, UClass* Class)
-	{
-		FObjectHandleEvents& Events = BeginReadingEvents();
-		Events.ClassReferenceResolvedEvent.Broadcast(ObjectRef, Package, Class);
-		Events.EndUsing();
-	}
-
-	inline void OnReferenceResolved(const FObjectRef& ObjectRef, UPackage* Package, UObject* Object)
-	{
-		FObjectHandleEvents& Events = BeginReadingEvents();
-		Events.ObjectHandleReferenceResolvedEvent.Broadcast(ObjectRef, Package, Object);
-		Events.EndUsing();
-	}
-
-	inline void OnReferenceLoaded(const FObjectRef& ObjectRef, UPackage* Package, UObject* Object)
-	{
-		FObjectHandleEvents& Events = BeginReadingEvents();
-		Events.ObjectHandleReferenceLoadedEvent.Broadcast(ObjectRef, Package, Object);
-		Events.EndUsing();
-	}
+	/**
+	 * Installs a new callback for notifications that an object value has been read from a handle.
+	 *
+	 * @param Function		The new handle read callback to install.
+	 * @return				The handle so that you can remove the callback at a later date.
+	 */
+	COREUOBJECT_API FObjectHandleTrackingCallbackId AddObjectHandleReadCallback(FObjectHandleReadFunc Delegate);
+	COREUOBJECT_API void RemoveObjectHandleReadCallback(FObjectHandleTrackingCallbackId DelegateHandle);
+
+	/**
+	 * Installs a new callback for notifications that a class has been resolved from an object handle or object reference.
+	 *
+	 * @param Function		The new class resolved callback to install.
+	 * @return				The handle so that you can remove the callback at a later date.
+	 */
+	COREUOBJECT_API FObjectHandleTrackingCallbackId AddObjectHandleClassResolvedCallback(FObjectHandleClassResolvedFunc Callback);
+	COREUOBJECT_API void RemoveObjectHandleClassResolvedCallback(FObjectHandleTrackingCallbackId DelegateHandle);
+
+	/**
+	 * Installs a new callback for notifications that an object has been resolved from an object handle or object reference.
+	 *
+	 * @param Function		The new object resolved callback to install.
+	 * @return				The handle so that you can remove the callback at a later date.
+	 */
+	COREUOBJECT_API FObjectHandleTrackingCallbackId AddObjectHandleReferenceResolvedCallback(FObjectHandleReferenceResolvedFunc Callback);
+	COREUOBJECT_API void RemoveObjectHandleReferenceResolvedCallback(FObjectHandleTrackingCallbackId DelegateHandle);
+
+	/**
+	 * Installs a new callback for notifications that an object has been loaded from an object handle or object reference.
+	 *
+	 * @param Function		The new object resolved callback to install.
+	 * @return				The handle so that you can remove the callback at a later date.
+	 */
+	COREUOBJECT_API FObjectHandleTrackingCallbackId AddObjectHandleReferenceLoadedCallback(FObjectHandleReferenceLoadedFunc Callback);
+	COREUOBJECT_API void RemoveObjectHandleReferenceLoadedCallback(FObjectHandleTrackingCallbackId DelegateHandle);
+
+#endif
+}
 
+#if UE_WITH_OBJECT_HANDLE_TRACKING
+namespace UE::CoreUObject::Private
+{
+	COREUOBJECT_API void OnHandleRead(UObject* Object);
+	COREUOBJECT_API void OnClassReferenceResolved(const FObjectRef& ObjectRef, UPackage* ClassPackage, UClass* Class);
+	COREUOBJECT_API void OnReferenceResolved(const FObjectRef& ObjectRef, UPackage* ObjectPackage, UObject* Object);
+	COREUOBJECT_API void OnReferenceLoaded(const FObjectRef& ObjectRef, UPackage* ObjectPackage, UObject* Object);
 }
 #else
 
diff --git a/Engine/Source/Runtime/CoreUObject/Tests/ObjectHandleTest.cpp b/Engine/Source/Runtime/CoreUObject/Tests/ObjectHandleTest.cpp
--- a/Engine/Source/Runtime/CoreUObject/Tests/ObjectHandleTest.cpp
+++ b/Engine/Source/Runtime/CoreUObject/Tests/ObjectHandleTest.cpp
@@ -218,12 +218,15 @@ TEST_CASE_METHOD(FObjectHandleTestBase, "CoreUObject::TObjectPtr::Null Behavior"
 	UObjectPtrTestClass* TestObject = nullptr;
 
 	uint32 ResolveCount = 0;
-	auto ResolveDelegate = FObjectHandleReferenceResolvedDelegate::CreateLambda([&ResolveCount](const FObjectRef& SourceRef, UPackage* ObjectPackage, UObject* Object)
+	auto ResolveDelegate = [&ResolveCount](const FObjectRef& SourceRef, UPackage* ObjectPackage, UObject* Object)
 		{
 			++ResolveCount;
-		});
-	auto Handle = AddObjectHandleReferenceResolvedCallback(ResolveDelegate);
-
+		};
+	auto Handle = UE::CoreUObject::AddObjectHandleReferenceResolvedCallback(ResolveDelegate);
+	ON_SCOPE_EXIT
+	{
+		UE::CoreUObject::RemoveObjectHandleReferenceResolvedCallback(Handle);
+	};
 	//compare against all flavours of nullptr, should not try and resolve this pointer
 	CHECK(Ptr == nullptr); CHECK(ResolveCount == 0u);
 	CHECK(nullptr == Ptr); CHECK(ResolveCount == 0u);
@@ -335,8 +338,6 @@ TEST_CASE_METHOD(FObjectHandleTestBase, "CoreUObject::TObjectPtr::Null Behavior"
 	CHECK_FALSE(TestObject == Ptr); CHECK(ResolveCount == 0u);
 	CHECK(Ptr != TestObject); CHECK(ResolveCount == 0u);
 	CHECK(TestObject != Ptr); CHECK(ResolveCount == 0u);
-
-	RemoveObjectHandleReferenceResolvedCallback(Handle);
 }
 
 #endif
diff --git a/Engine/Source/Runtime/CoreUObject/Tests/ObjectPtrTest.cpp b/Engine/Source/Runtime/CoreUObject/Tests/ObjectPtrTest.cpp
--- a/Engine/Source/Runtime/CoreUObject/Tests/ObjectPtrTest.cpp
+++ b/Engine/Source/Runtime/CoreUObject/Tests/ObjectPtrTest.cpp
@@ -337,18 +337,18 @@ TEST_CASE("CoreUObject::TObjectPtr::GetPathName", "[CoreUObject][ObjectPtr]")
 
 #if UE_WITH_OBJECT_HANDLE_TRACKING
 	int ResolveCount = 0;
-	auto ResolveDelegate = FObjectHandleReferenceResolvedDelegate::CreateLambda([&ResolveCount](const FObjectRef& SourceRef, UPackage* ObjectPackage, UObject* Object)
+	auto ResolveDelegate = [&ResolveCount](const FObjectRef& SourceRef, UPackage* ObjectPackage, UObject* Object)
 		{
 			++ResolveCount;
-		});
-	auto Handle = AddObjectHandleReferenceResolvedCallback(ResolveDelegate);
+		};
+	auto Handle = UE::CoreUObject::AddObjectHandleReferenceResolvedCallback(ResolveDelegate);
 #endif
 
 	ON_SCOPE_EXIT
 	{
 		TestPackage1->RemoveFromRoot();
 #if UE_WITH_OBJECT_HANDLE_TRACKING
-		RemoveObjectHandleReferenceResolvedCallback(Handle);
+	UE::CoreUObject::RemoveObjectHandleReferenceResolvedCallback(Handle);
 #endif
 	};
 
@@ -548,11 +548,15 @@ TEST_CASE("CoreUObject::TObjectPtr::GetTypeHash")
 {
 	int ResolveCount = 0;
 #if UE_WITH_OBJECT_HANDLE_TRACKING
-	auto ResolveDelegate = FObjectHandleReferenceResolvedDelegate::CreateLambda([&ResolveCount](const FObjectRef& SourceRef, UPackage* ObjectPackage, UObject* Object)
+	auto ResolveDelegate = [&ResolveCount](const FObjectRef& SourceRef, UPackage* ObjectPackage, UObject* Object)
 		{
 			++ResolveCount;
-		});
-	auto Handle = AddObjectHandleReferenceResolvedCallback(ResolveDelegate);
+		};
+	auto Handle = UE::CoreUObject::AddObjectHandleReferenceResolvedCallback(ResolveDelegate);
+	ON_SCOPE_EXIT
+	{
+		UE::CoreUObject::RemoveObjectHandleReferenceResolvedCallback(Handle);
+	};
 #endif
 
 	UPackage* TestPackageA = NewObject<UPackage>(nullptr, TEXT("/Engine/PackageA"), RF_Transient);
diff --git a/Engine/Source/Runtime/CoreUObject/Tests/ObjectRefTrackingTestBase.cpp b/Engine/Source/Runtime/CoreUObject/Tests/ObjectRefTrackingTestBase.cpp
--- a/Engine/Source/Runtime/CoreUObject/Tests/ObjectRefTrackingTestBase.cpp
+++ b/Engine/Source/Runtime/CoreUObject/Tests/ObjectRefTrackingTestBase.cpp
@@ -3,14 +3,261 @@
 #if WITH_LOW_LEVEL_TESTS
 
 #include "ObjectRefTrackingTestBase.h"
-
-#if UE_WITH_OBJECT_HANDLE_TRACKING
-FDelegateHandle FObjectRefTrackingTestBase::ResolvedCallbackHandle;
-FDelegateHandle FObjectRefTrackingTestBase::HandleReadCallbackHandle;
-#endif
+#include "UObject/Package.h"
+#include "ObjectPtrTestClass.h"
+#include "Misc/ScopeExit.h"
+#include "UObject/ObjectPtr.h"
 
 thread_local uint32 FObjectRefTrackingTestBase::NumResolves = 0;
 thread_local uint32 FObjectRefTrackingTestBase::NumFailedResolves = 0;
 thread_local uint32 FObjectRefTrackingTestBase::NumReads = 0;
 
+#if UE_WITH_OBJECT_HANDLE_TRACKING
+UE::CoreUObject::FObjectHandleTrackingCallbackId FObjectRefTrackingTestBase::ResolvedCallbackHandle;
+UE::CoreUObject::FObjectHandleTrackingCallbackId FObjectRefTrackingTestBase::HandleReadCallbackHandle;
+
+TEST_CASE("UE::CoreUObject::ObjectHandleTracking::Callbacks")
+{
+	using namespace UE::CoreUObject;
+	using namespace UE::CoreUObject::Private;
+
+	int32 CallbackCount1 = 0;
+	int32 CallbackCount2 = 0;
+	int32 CallbackCount3 = 0;
+	int32 CallbackCount4 = 0;
+	int32 CallbackCount5 = 0;
+	int32 CallbackCount6 = 0;
+	int32 CallbackCount7 = 0;
+	int32 CallbackCount8 = 0;
+
+	FObjectHandleReadFunc ReadCallback1 = [&](UObject* Obj)
+	{
+		++CallbackCount1;
+	};
+	FObjectHandleReadFunc ReadCallback2 = [&](UObject* Obj)
+	{
+		++CallbackCount2;
+	};
+
+	FObjectHandleClassResolvedFunc ClassResolvedCallback1 = [&](const FObjectRef& SourceRef, UPackage* ObjectPackage, UClass* Class)
+	{
+		++CallbackCount3;
+	};
+	FObjectHandleClassResolvedFunc ClassResolvedCallback2 = [&](const FObjectRef& SourceRef, UPackage* ObjectPackage, UClass* Class)
+	{
+		++CallbackCount4;
+	};
+
+	FObjectHandleReferenceResolvedFunc ReferenceResolvedCallback1 = [&](const FObjectRef& SourceRef, UPackage* ObjectPackage, UObject* Object)
+	{
+		++CallbackCount5;
+	};
+	FObjectHandleReferenceResolvedFunc ReferenceResolvedCallback2 = [&](const FObjectRef& SourceRef, UPackage* ObjectPackage, UObject* Object)
+	{
+		++CallbackCount6;
+	};
+
+	FObjectHandleReferenceLoadedFunc ReferenceLoadedCallback1 = [&](const FObjectRef& SourceRef, UPackage* ObjectPackage, UObject* Object)
+	{
+		++CallbackCount7;
+	};
+	FObjectHandleReferenceLoadedFunc ReferenceLoadedCallback2 = [&](const FObjectRef& SourceRef, UPackage* ObjectPackage, UObject* Object)
+	{
+		++CallbackCount8;
+	};
+	auto ReadCallbackHandle1 = AddObjectHandleReadCallback(ReadCallback1);
+	auto ReadCallbackHandle2 = AddObjectHandleReadCallback(ReadCallback2);
+
+	auto ClassResolvedHandle1 = AddObjectHandleClassResolvedCallback(ClassResolvedCallback1);
+	auto ClassResolvedHandle2 = AddObjectHandleClassResolvedCallback(ClassResolvedCallback2);
+
+	auto ReferenceResolvedHandle1 = AddObjectHandleReferenceResolvedCallback(ReferenceResolvedCallback1);
+	auto ReferenceResolvedHandle2 = AddObjectHandleReferenceResolvedCallback(ReferenceResolvedCallback2);
+
+	auto ReferenceLoadedHandle1 = AddObjectHandleReferenceLoadedCallback(ReferenceLoadedCallback1);
+	auto ReferenceLoadedHandle2 = AddObjectHandleReferenceLoadedCallback(ReferenceLoadedCallback2);
+
+	auto TestCounts = [&](int32 a, int32 b, int32 c, int32 d, int32 e, int32 f, int32 g, int32 h)
+	{
+		CHECK(CallbackCount1 == a);
+		CHECK(CallbackCount2 == b);
+		CHECK(CallbackCount3 == c);
+		CHECK(CallbackCount4 == d);
+		CHECK(CallbackCount5 == e);
+		CHECK(CallbackCount6 == f);
+		CHECK(CallbackCount7 == g);
+		CHECK(CallbackCount8 == h);
+	};
+
+	UE::CoreUObject::Private::OnHandleRead(nullptr);
+	TestCounts(1, 1, 0, 0, 0, 0, 0, 0);
+
+	//remove the first added read callback to verify they can be remove in any order
+	RemoveObjectHandleReadCallback(ReadCallbackHandle1);
+
+	UE::CoreUObject::Private::OnHandleRead(nullptr);
+	TestCounts(1, 2, 0, 0, 0, 0, 0, 0);
+	
+	//add it back
+	ReadCallbackHandle1 = AddObjectHandleReadCallback(ReadCallback1);
+	UE::CoreUObject::Private::OnHandleRead(nullptr);
+	TestCounts(2, 3, 0, 0, 0, 0, 0, 0);
+
+	//remove it again
+	RemoveObjectHandleReadCallback(ReadCallbackHandle1);
+	UE::CoreUObject::Private::OnHandleRead(nullptr);
+	TestCounts(2, 4, 0, 0, 0, 0, 0, 0);
+
+	//remove last callback
+	RemoveObjectHandleReadCallback(ReadCallbackHandle2);
+	UE::CoreUObject::Private::OnHandleRead(nullptr);
+	TestCounts(2, 4, 0, 0, 0, 0, 0, 0);
+
+	//////////////////////////////////////////////////////////////////////////////////
+	UE::CoreUObject::Private::OnClassReferenceResolved(FObjectRef(), nullptr, nullptr);
+	TestCounts(2, 4, 1, 1, 0, 0, 0, 0);
+
+	//remove the first added read callback to verify they can be remove in any order
+	RemoveObjectHandleClassResolvedCallback(ClassResolvedHandle1);
+
+	UE::CoreUObject::Private::OnClassReferenceResolved(FObjectRef(), nullptr, nullptr);
+	TestCounts(2, 4, 1, 2, 0, 0, 0, 0);
+
+	//add it back
+	ClassResolvedHandle1 = AddObjectHandleClassResolvedCallback(ClassResolvedCallback1);
+	UE::CoreUObject::Private::OnClassReferenceResolved(FObjectRef(), nullptr, nullptr);
+	TestCounts(2, 4, 2, 3, 0, 0, 0, 0);
+
+	//remove it again
+	RemoveObjectHandleClassResolvedCallback(ClassResolvedHandle1);
+	UE::CoreUObject::Private::OnClassReferenceResolved(FObjectRef(), nullptr, nullptr);
+	TestCounts(2, 4, 2, 4, 0, 0, 0, 0);
+
+	//remove last callback
+	RemoveObjectHandleClassResolvedCallback(ClassResolvedHandle2);
+	UE::CoreUObject::Private::OnClassReferenceResolved(FObjectRef(), nullptr, nullptr);
+	TestCounts(2, 4, 2, 4, 0, 0, 0, 0);
+
+
+	//////////////////////////////////////////////////////////////////////////////////
+	UE::CoreUObject::Private::OnReferenceResolved(FObjectRef(), nullptr, nullptr);
+	TestCounts(2, 4, 2, 4, 1, 1, 0, 0);
+
+	//remove the first added read callback to verify they can be remove in any order
+	RemoveObjectHandleReferenceResolvedCallback(ReferenceResolvedHandle1);
+
+	UE::CoreUObject::Private::OnReferenceResolved(FObjectRef(), nullptr, nullptr);
+	TestCounts(2, 4, 2, 4, 1, 2, 0, 0);
+
+	//add it back
+	ReferenceResolvedHandle1 = AddObjectHandleReferenceResolvedCallback(ReferenceResolvedCallback1);
+	UE::CoreUObject::Private::OnReferenceResolved(FObjectRef(), nullptr, nullptr);
+	TestCounts(2, 4, 2, 4, 2, 3, 0, 0);
+
+	//remove it again
+	RemoveObjectHandleReferenceResolvedCallback(ReferenceResolvedHandle1);
+	UE::CoreUObject::Private::OnReferenceResolved(FObjectRef(), nullptr, nullptr);
+	TestCounts(2, 4, 2, 4, 2, 4, 0, 0);
+
+	//remove last callback
+	RemoveObjectHandleReferenceResolvedCallback(ReferenceResolvedHandle2);
+	UE::CoreUObject::Private::OnReferenceResolved(FObjectRef(), nullptr, nullptr);
+	TestCounts(2, 4, 2, 4, 2, 4, 0, 0);
+
+	//////////////////////////////////////////////////////////////////////////////////
+	UE::CoreUObject::Private::OnReferenceLoaded(FObjectRef(), nullptr, nullptr);
+	TestCounts(2, 4, 2, 4, 2, 4, 1, 1);
+
+	//remove the first added read callback to verify they can be remove in any order
+	RemoveObjectHandleReferenceLoadedCallback(ReferenceLoadedHandle1);
+
+	UE::CoreUObject::Private::OnReferenceLoaded(FObjectRef(), nullptr, nullptr);
+	TestCounts(2, 4, 2, 4, 2, 4, 1, 2);
+
+	//add it back
+	ReferenceLoadedHandle1 = AddObjectHandleReferenceLoadedCallback(ReferenceLoadedCallback1);
+	UE::CoreUObject::Private::OnReferenceLoaded(FObjectRef(), nullptr, nullptr);
+	TestCounts(2, 4, 2, 4, 2, 4, 2, 3);
+
+	//remove it again
+	RemoveObjectHandleReferenceLoadedCallback(ReferenceLoadedHandle1);
+	UE::CoreUObject::Private::OnReferenceLoaded(FObjectRef(), nullptr, nullptr);
+	TestCounts(2, 4, 2, 4, 2, 4, 2, 4);
+
+	//remove last callback
+	RemoveObjectHandleReferenceLoadedCallback(ReferenceLoadedHandle2);
+	UE::CoreUObject::Private::OnReferenceLoaded(FObjectRef(), nullptr, nullptr);
+	TestCounts(2, 4, 2, 4, 2, 4, 2, 4);
+}
+
+TEST_CASE("UE::CoreUObject::ObjectHandleTracking::Callbacks::Verify")
+{
+	using namespace UE::CoreUObject;
+	using namespace UE::CoreUObject::Private;
+
+	const FName TestPackageName(TEXT("/Engine/Test/ObjectRef/Transient"));
+	UPackage* TestPackage = NewObject<UPackage>(nullptr, TestPackageName, RF_Transient);
+	TestPackage->AddToRoot();
+	UObject* Obj1 = NewObject<UObjectPtrTestClass>(TestPackage, TEXT("DefaultSerializeObject"));
+	ON_SCOPE_EXIT{
+		TestPackage->RemoveFromRoot();
+	};
+
+	TObjectPtr<UObject> ObjPtr = Obj1;
+	int32 ReadCount = 0;
+	auto ReadCallbackHandle1 = AddObjectHandleReadCallback([&](UObject* ReadObject)
+		{
+			++ReadCount;
+	CHECK(Obj1 == ReadObject);
+		});
+	CHECK(ObjPtr.Get() != nullptr);
+	CHECK(ReadCount == 1);
+	
+	FObjectRef ObjectRef(Obj1);
+	int32 ClassCount = 0;
+	auto ClassResolvedHandle1 = AddObjectHandleClassResolvedCallback([&](const FObjectRef& SourceRef, UPackage* ClassPackage, UClass* Class)
+		{
+			++ClassCount;
+			CHECK(Class == UObjectPtrTestClass::StaticClass());
+			CHECK(&ObjectRef == &SourceRef);
+			CHECK(UObjectPtrTestClass::StaticClass()->GetPackage() == ClassPackage);
+		});
+	UClass* Class = ObjectRef.ResolveObjectRefClass();
+	CHECK(ClassCount == 1);
+
+	int32 ResolvedCount = 0;
+	auto ReferenceResolvedHandle1 = AddObjectHandleReferenceResolvedCallback([&](const FObjectRef& SourceRef, UPackage* ObjectPackage, UObject* Object)
+		{
+			++ResolvedCount;
+			CHECK(Obj1 == Object);
+			CHECK(&ObjectRef == &SourceRef);
+			CHECK(TestPackage == ObjectPackage);
+
+		});
+	ObjectRef.Resolve();
+	CHECK(ResolvedCount == 1);
+
+	int32 LoadCount = 0;
+	auto ReferenceLoadedHandle1 = AddObjectHandleReferenceLoadedCallback([&](const FObjectRef& SourceRef, UPackage* ObjectPackage, UObject* Object)
+		{
+			++LoadCount;
+			CHECK(Obj1 == Object);
+			CHECK(&ObjectRef == &SourceRef);
+			CHECK(TestPackage == ObjectPackage);
+		});
+
+	//don't have a good to test this as the object are never loaded from disk in the tests
+	OnReferenceLoaded(ObjectRef, TestPackage, Obj1);
+	CHECK(LoadCount == 1);
+
+	ON_SCOPE_EXIT
+	{
+		RemoveObjectHandleReadCallback(ReadCallbackHandle1);
+		RemoveObjectHandleClassResolvedCallback(ClassResolvedHandle1);
+		RemoveObjectHandleReferenceResolvedCallback(ReferenceResolvedHandle1);
+		RemoveObjectHandleReferenceLoadedCallback(ReferenceLoadedHandle1);
+	};
+	
+}
+#endif
 #endif
\ No newline at end of file
diff --git a/Engine/Source/Runtime/CoreUObject/Tests/ObjectRefTrackingTestBase.h b/Engine/Source/Runtime/CoreUObject/Tests/ObjectRefTrackingTestBase.h
--- a/Engine/Source/Runtime/CoreUObject/Tests/ObjectRefTrackingTestBase.h
+++ b/Engine/Source/Runtime/CoreUObject/Tests/ObjectRefTrackingTestBase.h
@@ -105,8 +105,8 @@ private:
 		if (!bCallbacksInstalled)
 		{
 #if UE_WITH_OBJECT_HANDLE_TRACKING
-			ResolvedCallbackHandle = AddObjectHandleReferenceResolvedCallback(FObjectHandleReferenceResolvedDelegate::CreateStatic(OnRefResolved));
-			HandleReadCallbackHandle = AddObjectHandleReadCallback(FObjectHandleReadDelegate::CreateStatic(OnRefRead));
+			ResolvedCallbackHandle = UE::CoreUObject::AddObjectHandleReferenceResolvedCallback(OnRefResolved);
+			HandleReadCallbackHandle = UE::CoreUObject::AddObjectHandleReadCallback(OnRefRead);
 			// TODO We should unhook these handles somewhere, but i don't want to refactor the test, it's not as if they were
 			// being unhooked before.  So...
 #endif
@@ -115,8 +115,8 @@ private:
 	}
 
 #if UE_WITH_OBJECT_HANDLE_TRACKING
-	static FDelegateHandle ResolvedCallbackHandle;
-	static FDelegateHandle HandleReadCallbackHandle;
+	static UE::CoreUObject::FObjectHandleTrackingCallbackId ResolvedCallbackHandle;
+	static UE::CoreUObject::FObjectHandleTrackingCallbackId HandleReadCallbackHandle;
 #endif
 	static thread_local uint32 NumResolves;
 	static thread_local uint32 NumFailedResolves;
