added more tests to verify that read events are not fired for null checks or equality compares

https://p4-swarm.epicgames.net/reviews/23782690

diff --git a/Engine/Source/Runtime/CoreUObject/Tests/ObjectRefTrackingTestBase.cpp b/Engine/Source/Runtime/CoreUObject/Tests/ObjectRefTrackingTestBase.cpp
--- a/Engine/Source/Runtime/CoreUObject/Tests/ObjectRefTrackingTestBase.cpp
+++ b/Engine/Source/Runtime/CoreUObject/Tests/ObjectRefTrackingTestBase.cpp
@@ -208,11 +208,34 @@ TEST_CASE("UE::CoreUObject::ObjectHandleTracking::Callbacks::Verify")
 	auto ReadCallbackHandle1 = AddObjectHandleReadCallback([&](UObject* ReadObject)
 		{
 			++ReadCount;
-	CHECK(Obj1 == ReadObject);
+			CHECK(Obj1 == ReadObject);
 		});
+	
+	CHECK(ObjPtr == Obj1);
+	CHECK(ReadCount == 0);
+
+	CHECK_FALSE(ObjPtr != Obj1);
+	CHECK(ReadCount == 0);
+
+	CHECK_FALSE(ObjPtr == nullptr);
+	CHECK(ReadCount == 0);
+	
+	CHECK(ObjPtr != nullptr);
+	CHECK(ReadCount == 0);
+
+	bool bIsNull = !ObjPtr;
+	CHECK_FALSE(bIsNull);
+	CHECK(ReadCount == 0);
+
+	if (ObjPtr)
+	{
+		CHECK(ReadCount == 0);
+	}
+
+	//get fires read event
 	CHECK(ObjPtr.Get() != nullptr);
 	CHECK(ReadCount == 1);
-	
+
 	FObjectRef ObjectRef(Obj1);
 	int32 ClassCount = 0;
 	auto ClassResolvedHandle1 = AddObjectHandleClassResolvedCallback([&](const FObjectRef& SourceRef, UPackage* ClassPackage, UClass* Class)
