Fix the generation of duplicate bulkdata identifiers when creating new levels from templates containing bulkdata payloads.

#rnx

- This replaces the solution to the same problem that was submitted in CL 21587137.
- Instead of trying to pass the need to regenerate the guids via FLinkerInstancingContext it is now done via ELoadFlags.
- The fix was originally added to FLinkerInstancingContext as a) at the time a similar fix had been made to instancing b) we felt it would be easier to remove in the future when we change how the guids are calculated making the need to regenerate moot. However this stopped working when 22635658 prevented the instancing context from being passed to objects not being instanced. The concept was a better fix for the flags so it has been moved there.

diff --git a/Engine/Source/Editor/UnrealEd/Private/EditorServer.cpp b/Engine/Source/Editor/UnrealEd/Private/EditorServer.cpp
--- a/Engine/Source/Editor/UnrealEd/Private/EditorServer.cpp
+++ b/Engine/Source/Editor/UnrealEd/Private/EditorServer.cpp
@@ -2706,7 +2706,7 @@ bool UEditorEngine::Map_Load(const TCHAR* Str, FOutputDevice& Ar)
 						FSoftObjectPath(*WriteToString<256>(WorldPackage->GetName(), TEXT("."), ShortWorldPackageName))
 					);
 
-					WorldPackageInstancingContext.SetRegenerateUniqueBulkDataGuids(true);
+					LoadFlags |= LOAD_RegenerateBulkDataGuids;
 
 					WorldPackage = LoadPackage( WorldPackage, *LongTempFname, LoadFlags, nullptr /* InReaderOverride */, &WorldPackageInstancingContext);
 					WorldPackage->SetPackageFlags(PKG_NewlyCreated);
diff --git a/Engine/Source/Runtime/CoreUObject/Private/Blueprint/BlueprintSupport.cpp b/Engine/Source/Runtime/CoreUObject/Private/Blueprint/BlueprintSupport.cpp
--- a/Engine/Source/Runtime/CoreUObject/Private/Blueprint/BlueprintSupport.cpp
+++ b/Engine/Source/Runtime/CoreUObject/Private/Blueprint/BlueprintSupport.cpp
@@ -1430,7 +1430,7 @@ void FLinkerLoad::ResolveDeferredDependencies(UStruct* LoadStruct)
 			// fully loading... we need to make sure metadata gets loaded, etc.)
 			if ((SourcePackage != nullptr) && !SourcePackage->HasAnyFlags(RF_WasLoaded))
 			{
-				uint32 InternalLoadFlags = LoadFlags & (LOAD_NoVerify | LOAD_NoWarn | LOAD_Quiet);
+				uint32 InternalLoadFlags = LoadFlags & (LOAD_NoVerify | LOAD_NoWarn | LOAD_Quiet | LOAD_RegenerateBulkDataGuids);
 				// make sure LoadAllObjects() is called for this package
 				LoadPackageInternal(/*Outer =*/nullptr, SourceLinker->GetPackagePath(), InternalLoadFlags, this, nullptr/*InReaderOverride*/, nullptr/*InstancingContext*/, nullptr /* DiffPackagePath */); //-V595
 			}
diff --git a/Engine/Source/Runtime/CoreUObject/Private/Serialization/EditorBulkData.cpp b/Engine/Source/Runtime/CoreUObject/Private/Serialization/EditorBulkData.cpp
--- a/Engine/Source/Runtime/CoreUObject/Private/Serialization/EditorBulkData.cpp
+++ b/Engine/Source/Runtime/CoreUObject/Private/Serialization/EditorBulkData.cpp
@@ -243,7 +243,7 @@ FString GetDebugNameFromOwner(UObject* Owner)
  */
 bool ShouldGenerateNewIdentifier(FLinkerLoad* LinkerLoad, UObject* Owner)
 {
-	if (LinkerLoad && LinkerLoad->GetInstancingContext().ShouldRegenerateUniqueBulkDataGuids())
+	if (LinkerLoad && LinkerLoad->ShouldRegenerateGuids())
 	{
 		return true;
 	}
diff --git a/Engine/Source/Runtime/CoreUObject/Private/UObject/LinkerLoad.cpp b/Engine/Source/Runtime/CoreUObject/Private/UObject/LinkerLoad.cpp
--- a/Engine/Source/Runtime/CoreUObject/Private/UObject/LinkerLoad.cpp
+++ b/Engine/Source/Runtime/CoreUObject/Private/UObject/LinkerLoad.cpp
@@ -3336,7 +3336,7 @@ bool FLinkerLoad::VerifyImportInner(const int32 ImportIndex, FString& WarningSuf
 		check(Import.ClassName == NAME_Package || Import.HasPackageName());
 
 		UPackage* Package = nullptr;
-		uint32 InternalLoadFlags = LoadFlags & (LOAD_NoVerify | LOAD_NoWarn | LOAD_Quiet);
+		uint32 InternalLoadFlags = LoadFlags & (LOAD_NoVerify | LOAD_NoWarn | LOAD_Quiet | LOAD_RegenerateBulkDataGuids);
 		FUObjectSerializeContext* SerializeContext = GetSerializeContext();
 
 		// Resolve the package name for the import, potentially remapping it, if instancing
diff --git a/Engine/Source/Runtime/CoreUObject/Public/UObject/LinkerInstancingContext.h b/Engine/Source/Runtime/CoreUObject/Public/UObject/LinkerInstancingContext.h
--- a/Engine/Source/Runtime/CoreUObject/Public/UObject/LinkerInstancingContext.h
+++ b/Engine/Source/Runtime/CoreUObject/Public/UObject/LinkerInstancingContext.h
@@ -93,15 +93,11 @@ public:
 		return bSoftObjectPathRemappingEnabled; 
 	}
 
-	void SetRegenerateUniqueBulkDataGuids(bool bFlag)
-	{
-		bRegenerateUniqueBulkDataGuids = bFlag;
-	}
+	UE_DEPRECATED(5.2, "No longer used, pass ELoadFlags::LOAD_RegenerateBulkDataGuids to LoadPackage instead")
+	void SetRegenerateUniqueBulkDataGuids(bool bFlag) { }
 
-	bool ShouldRegenerateUniqueBulkDataGuids() const
-	{
-		return bRegenerateUniqueBulkDataGuids;
-	}
+	UE_DEPRECATED(5.2, "No longer used, check ELoadFlags::LOAD_RegenerateBulkDataGuids in the LoadFlags instead")
+	bool ShouldRegenerateUniqueBulkDataGuids() const { return false; }
 
 	/** Return the instanced package name for a given instanced outer package and an object package name */
 	static FString GetInstancedPackageName(const FString& InOuterPackageName, const FString& InPackageName)
@@ -133,6 +129,4 @@ private:
 	TSet<FName> Tags;
 	/** Remap soft object paths */
 	bool bSoftObjectPathRemappingEnabled = true;
-	/** When true we will generate new unique identifiers for editor bulkdata objects */
-	bool bRegenerateUniqueBulkDataGuids = false;
 };
diff --git a/Engine/Source/Runtime/CoreUObject/Public/UObject/LinkerLoad.h b/Engine/Source/Runtime/CoreUObject/Public/UObject/LinkerLoad.h
--- a/Engine/Source/Runtime/CoreUObject/Public/UObject/LinkerLoad.h
+++ b/Engine/Source/Runtime/CoreUObject/Public/UObject/LinkerLoad.h
@@ -930,6 +930,12 @@ public:
 	 */
 	COREUOBJECT_API void DetachExports();
 
+	/** Should bulkdata identifiers should be regenerated as they are loaded or not */
+	bool ShouldRegenerateGuids() const
+	{
+		return (LoadFlags & LOAD_RegenerateBulkDataGuids) != 0;
+	}
+
 private:
 
 	FORCEINLINE virtual void Seek(int64 InPos) override
diff --git a/Engine/Source/Runtime/CoreUObject/Public/UObject/ObjectMacros.h b/Engine/Source/Runtime/CoreUObject/Public/UObject/ObjectMacros.h
--- a/Engine/Source/Runtime/CoreUObject/Public/UObject/ObjectMacros.h
+++ b/Engine/Source/Runtime/CoreUObject/Public/UObject/ObjectMacros.h
@@ -67,22 +67,22 @@ enum ELoadFlags
 	LOAD_Verify						= 0x00000010,	///< Only verify existance; don't actually load.
 //	LOAD_Unused						= 0x00000020,	///< Allow plain DLLs.
 //	LOAD_Unused						= 0x00000040
-	LOAD_NoVerify					= 0x00000080,   ///< Don't verify imports yet.
+	LOAD_NoVerify					= 0x00000080,	///< Don't verify imports yet.
 	LOAD_IsVerifying				= 0x00000100,	///< Is verifying imports
 	LOAD_SkipLoadImportedPackages	= 0x00000200,	///< Assume that all import packages are already loaded and don't call LoadPackage when creating imports 
-//	LOAD_Unused						= 0x00000400,
+	LOAD_RegenerateBulkDataGuids	= 0x00000400,	///< BulkData identifiers should be regenerated as they are loaded 
 //	LOAD_Unused						= 0x00000800,
 	LOAD_DisableDependencyPreloading = 0x00001000,	///< Bypass dependency preloading system
-	LOAD_Quiet						= 0x00002000,   ///< No log warnings.
+	LOAD_Quiet						= 0x00002000,	///< No log warnings.
 	LOAD_FindIfFail					= 0x00004000,	///< Tries FindObject if a linker cannot be obtained (e.g. package is currently being compiled)
 	LOAD_MemoryReader				= 0x00008000,	///< Loads the file into memory and serializes from there.
 	LOAD_NoRedirects				= 0x00010000,	///< Never follow redirects when loading objects; redirected loads will fail
 	LOAD_ForDiff					= 0x00020000,	///< Loading for diffing in the editor
-	LOAD_PackageForPIE				= 0x00080000,   ///< This package is being loaded for PIE, it must be flagged as such immediately
-	LOAD_DeferDependencyLoads       = 0x00100000,   ///< Do not load external (blueprint) dependencies (instead, track them for deferred loading)
+	LOAD_PackageForPIE				= 0x00080000,	///< This package is being loaded for PIE, it must be flagged as such immediately
+	LOAD_DeferDependencyLoads       = 0x00100000,	///< Do not load external (blueprint) dependencies (instead, track them for deferred loading)
 	LOAD_ForFileDiff				= 0x00200000,	///< Load the package (not for diffing in the editor), instead verify at the two packages serialized output are the same, if they are not then debug break so that you can get the callstack and object information
 	LOAD_DisableCompileOnLoad		= 0x00400000,	///< Prevent this load call from running compile on load for the loaded blueprint (intentionally not recursive, dependencies will still compile on load)
-	LOAD_DisableEngineVersionChecks = 0x00800000,   ///< Prevent this load call from running engine version checks
+	LOAD_DisableEngineVersionChecks = 0x00800000,	///< Prevent this load call from running engine version checks
 };
 
 /** Flags for saving objects/packages, passed into UPackage::SavePackage() as a uint32 */
