Support diff only cooks with separate exports archive

- extracted out callstack offets and stack traces into FArchiveCallstacks
- created FArchiveStackTraceWriter that wraps an archive and writes into a FArchiveCallstacks
- refactored existing FArchiveStackTrace to use the new writer and callstack collection
- added IPackageWriter::CreateLinkerExportArchive in order to create a stack trace archive writer in DiffPackageWriter

#rnx

diff --git a/Engine/Source/Editor/UnrealEd/Private/Cooker/DiffPackageWriter.cpp b/Engine/Source/Editor/UnrealEd/Private/Cooker/DiffPackageWriter.cpp
--- a/Engine/Source/Editor/UnrealEd/Private/Cooker/DiffPackageWriter.cpp
+++ b/Engine/Source/Editor/UnrealEd/Private/Cooker/DiffPackageWriter.cpp
@@ -154,8 +154,12 @@ void FDiffPackageWriter::WritePackageData(const FPackageInfo& Info, FLargeMemory
 	PreviousPackageData.HeaderSize = PreviousInnerData.HeaderSize;
 	PreviousPackageData.StartOffset = PreviousInnerData.StartOffset;
 
+	check(ExportsCallstacks.IsValid());
+
 	if (bDiffCallstack)
 	{
+		Writer.GetCallstacks().Append(*ExportsCallstacks);
+
 		TMap<FName, FArchiveDiffStats> PackageDiffStats;
 		const TCHAR* CutoffString = TEXT("UEditorEngine::Save()");
 		Writer.CompareWith(PreviousPackageData, *Info.LooseFilePath, Info.HeaderSize, CutoffString,
@@ -166,6 +170,9 @@ void FDiffPackageWriter::WritePackageData(const FPackageInfo& Info, FLargeMemory
 	}
 	else
 	{
+		Writer.GetCallstacks().Append(*ExportsCallstacks, Info.HeaderSize);
+		ExportsDiffMapOffset = Info.HeaderSize;
+
 		bIsDifferent = !Writer.GenerateDiffMap(PreviousPackageData, Info.HeaderSize, MaxDiffsToLog, DiffMap);
 	}
 
@@ -188,6 +195,43 @@ TUniquePtr<FLargeMemoryWriter> FDiffPackageWriter::CreateLinkerArchive(FName Pac
 	}
 }
 
+TUniquePtr<FLargeMemoryWriter> FDiffPackageWriter::CreateLinkerExportsArchive(FName PackageName, UObject* Asset)
+{
+	// When cooking, exports are serialized into a separate archive. The serialization callstack offsets
+	// and stack traces are collected into a separate callstack collection and appended to the overall
+	// callstacks for the entire package. DiffOnly cooks saves a package twice. The first pass collects
+	// the serialization offsets without the stack traces and then creates a FArchiveDiffMap. In the second
+	// pass, the diff map is used to collect offsets AND the entire stack trace for mismatching package data.
+	// In the first pass, callstack offsets will be relative to the export archive and adjusted when appending
+	// the callstacks to the overall package callstacks. In the second pass, the offset (ExportsDiffMapOffset) is
+	// known and the callstack will be relative to the beginning of the package.
+	
+	ExportsCallstacks = MakeUnique<FArchiveCallstacks>(Asset);
+	const int64 PreAllocateBytes = 0;
+	const bool bIsPersistent = true;
+
+	if (bDiffCallstack)
+	{
+		return MakeUnique<FArchiveStackTraceMemoryWriter>(
+			*ExportsCallstacks,
+			&DiffMap,
+			ExportsDiffMapOffset,
+			PreAllocateBytes,
+			bIsPersistent,
+			*PackageName.ToString());
+	}
+	else
+	{
+		return MakeUnique<FArchiveStackTraceMemoryWriter>(
+			*ExportsCallstacks,
+			nullptr,
+			0,
+			PreAllocateBytes, 
+			bIsPersistent, 
+			*PackageName.ToString());
+	}
+}
+
 bool FDiffPackageWriter::IsAnotherSaveNeeded(FSavePackageResultStruct& PreviousResult, FSavePackageArgs& SaveArgs)
 {
 	checkf(!Inner->IsAnotherSaveNeeded(PreviousResult, SaveArgs),
diff --git a/Engine/Source/Editor/UnrealEd/Private/Cooker/DiffPackageWriter.h b/Engine/Source/Editor/UnrealEd/Private/Cooker/DiffPackageWriter.h
--- a/Engine/Source/Editor/UnrealEd/Private/Cooker/DiffPackageWriter.h
+++ b/Engine/Source/Editor/UnrealEd/Private/Cooker/DiffPackageWriter.h
@@ -46,6 +46,7 @@ public:
 		return Inner->GetExportsFooterSize();
 	}
 	virtual TUniquePtr<FLargeMemoryWriter> CreateLinkerArchive(FName PackageName, UObject* Asset) override;
+	virtual TUniquePtr<FLargeMemoryWriter> CreateLinkerExportsArchive(FName PackageName, UObject* Asset) override;
 	virtual bool IsPreSaveCompleted() const override
 	{
 		return bDiffCallstack;
@@ -113,10 +114,12 @@ private:
 	void ConditionallyDumpObjects();
 
 	FArchiveDiffMap DiffMap;
+	TUniquePtr<FArchiveCallstacks> ExportsCallstacks;
 	FBeginPackageInfo BeginInfo;
 	TUniquePtr<ICookedPackageWriter> Inner;
 	FString DumpObjListParams;
 	FString PackageFilter;
+	int64 ExportsDiffMapOffset = 0;
 	int32 MaxDiffsToLog = 5;
 	bool bSaveForDiff = false;
 	bool bIgnoreHeaderDiffs = false;
@@ -244,4 +247,4 @@ private:
 	TUniquePtr<ICookedPackageWriter> Inner;
 	EDiffMode DiffMode = EDiffMode::LDM_Consistent;
 	bool bHasStartedSecondSave = false;
-};
\ No newline at end of file
+};
diff --git a/Engine/Source/Runtime/Core/Public/Serialization/PackageWriter.h b/Engine/Source/Runtime/Core/Public/Serialization/PackageWriter.h
--- a/Engine/Source/Runtime/Core/Public/Serialization/PackageWriter.h
+++ b/Engine/Source/Runtime/Core/Public/Serialization/PackageWriter.h
@@ -173,6 +173,9 @@ public:
 	/** Create the FLargeMemoryWriter to which the Header and Exports are written during the save. */
 	COREUOBJECT_API virtual TUniquePtr<FLargeMemoryWriter> CreateLinkerArchive(FName PackageName, UObject* Asset);
 
+	/** Returns an archive to be used when serializing exports. */
+	COREUOBJECT_API virtual TUniquePtr<FLargeMemoryWriter> CreateLinkerExportsArchive(FName PackageName, UObject* Asset);
+
 	/** Report whether PreSave was already called by the PackageWriter before the current UPackage::Save call. */
 	virtual bool IsPreSaveCompleted() const
 	{
diff --git a/Engine/Source/Runtime/CoreUObject/Private/Serialization/ArchiveStackTrace.cpp b/Engine/Source/Runtime/CoreUObject/Private/Serialization/ArchiveStackTrace.cpp
--- a/Engine/Source/Runtime/CoreUObject/Private/Serialization/ArchiveStackTrace.cpp
+++ b/Engine/Source/Runtime/CoreUObject/Private/Serialization/ArchiveStackTrace.cpp
@@ -47,12 +47,11 @@ public:
 
 class FIgnoreDiffManager
 {
-	int32 IgnoreCount;
+	int32 IgnoreCount = 0;
+	int32 DisabledCount = 0;
 
 public:
-	FIgnoreDiffManager()
-		: IgnoreCount(0)
-	{}
+	FIgnoreDiffManager() = default;
 	void Push()
 	{
 		IgnoreCount++;
@@ -66,6 +65,19 @@ public:
 	{
 		return !!IgnoreCount;
 	}
+	void PushDisabled()
+	{
+		DisabledCount++;
+	}
+	void PopDisabled()
+	{
+		DisabledCount--;
+		check(DisabledCount >= 0);
+	}
+	bool ShouldBypassDiff() const
+	{
+		return !!DisabledCount;
+	}
 };
 
 static FIgnoreDiffManager GIgnoreDiffManager;
@@ -88,14 +100,18 @@ FArchiveStackTraceIgnoreScope::~FArchiveStackTraceIgnoreScope()
 	}
 }
 
-FArchiveStackTrace::FCallstackData::FCallstackData()
-	: Callstack(nullptr)
-	, SerializedProp(nullptr)
+FArchiveStackTraceDisabledScope::FArchiveStackTraceDisabledScope()
+{
+	GIgnoreDiffManager.PushDisabled();
+}
+
+FArchiveStackTraceDisabledScope::~FArchiveStackTraceDisabledScope()
 {
+	GIgnoreDiffManager.PopDisabled();
 }
 
-FArchiveStackTrace::FCallstackData::FCallstackData(ANSICHAR* InCallstack, UObject* InSerializedObject, FProperty* InSerializedProperty)
-	: Callstack(InCallstack)
+FArchiveCallstacks::FCallstackData::FCallstackData(TUniquePtr<ANSICHAR[]>&& InCallstack, UObject* InSerializedObject, FProperty* InSerializedProperty)
+	: Callstack(MoveTemp(InCallstack))
 	, SerializedProp(InSerializedProperty)
 {
 	if (InSerializedObject)
@@ -108,7 +124,7 @@ FArchiveStackTrace::FCallstackData::FCallstackData(ANSICHAR* InCallstack, UObjec
 	}
 }
 
-FString FArchiveStackTrace::FCallstackData::ToString(const TCHAR* CallstackCutoffText) const
+FString FArchiveCallstacks::FCallstackData::ToString(const TCHAR* CallstackCutoffText) const
 {
 	FString HumanReadableString;
 
@@ -116,7 +132,7 @@ FString FArchiveStackTrace::FCallstackData::ToString(const TCHAR* CallstackCutof
 	const TCHAR* const LineTerminator = FDiffFormatHelper::Get().LineTerminator;
 	const TCHAR* const Indent = FDiffFormatHelper::Get().Indent;
 	
-	FString StackTraceText = Callstack;
+	FString StackTraceText = Callstack.Get();
 	if (CallstackCutoffText != nullptr)
 	{
 		// If the cutoff string is provided, remove all functions starting with the one specifiec in the cutoff string
@@ -171,207 +187,154 @@ FString FArchiveStackTrace::FCallstackData::ToString(const TCHAR* CallstackCutof
 	return HumanReadableString;
 }
 
-FArchiveStackTrace::FArchiveStackTrace(UObject* InAsset, const TCHAR* InFilename, bool bInCollectCallstacks, const FArchiveDiffMap* InDiffMap)
-	: FLargeMemoryWriter(0, false, InFilename)
-	, Asset(InAsset)
-	, AssetClass(InAsset ? InAsset->GetClass()->GetFName() : NAME_None)
-	, DiffMap(InDiffMap)
-	, bCollectCallstacks(bInCollectCallstacks)
+FArchiveCallstacks::FCallstackData FArchiveCallstacks::FCallstackData::Clone() const
+{
+	TUniquePtr<ANSICHAR[]> CallstackCopy;
+	if (const int32 Len = FCStringAnsi::Strlen(Callstack.Get()); Len > 0)
+	{
+		CallstackCopy = MakeUnique<ANSICHAR[]>(Len);
+		FCStringAnsi::Strcpy(CallstackCopy.Get(), Len, Callstack.Get());
+	}
+
+	FCallstackData Clone(MoveTemp(CallstackCopy), nullptr, SerializedProp);
+	Clone.SerializedObjectName = SerializedObjectName;
+
+	return Clone;
+}
+
+FArchiveCallstacks::FArchiveCallstacks(UObject* InAsset)
+	: Asset(InAsset)
 	, bCallstacksDirty(true)
 	, StackTraceSize(65535)
 	, LastSerializeCallstack(nullptr)
-	, ThreadContext(FUObjectThreadContext::Get())
 {
-	this->SetIsSaving(true);
-
-	StackTrace = (ANSICHAR*)FMemory::Malloc(StackTraceSize);
+	StackTrace = MakeUnique<ANSICHAR[]>(StackTraceSize);
 	StackTrace[0] = 0;
 }
 
-FArchiveStackTrace::~FArchiveStackTrace()
+FName FArchiveCallstacks::GetAssetClass() const
 {
-	FMemory::Free(StackTrace);
-
-	for (TPair<uint32, FCallstackData>& UniqueCallstackPair : UniqueCallstacks)
-	{
-		FMemory::Free(UniqueCallstackPair.Value.Callstack);
-	}
+	return Asset != nullptr ? Asset->GetClass()->GetFName() : NAME_None;
 }
 
-ANSICHAR* FArchiveStackTrace::AddUniqueCallstack(UObject* InSerializedObject, FProperty* InSerializedProperty, uint32& OutCallstackCRC)
+ANSICHAR* FArchiveCallstacks::AddUniqueCallstack(UObject* SerializedObject, FProperty* SerializedProperty, uint32& OutCallstackCRC)
 {
-	ANSICHAR* Callstack = nullptr;
-	if (bCollectCallstacks)
-	{
-		OutCallstackCRC = FCrc::StrCrc32(StackTrace);
+	ANSICHAR* OutCallstack = nullptr;
+	OutCallstackCRC = FCrc::StrCrc32(StackTrace.Get());
 
-		if (FCallstackData* ExistingCallstack = UniqueCallstacks.Find(OutCallstackCRC))
-		{
-			Callstack = ExistingCallstack->Callstack;
-		}
-		else
-		{
-			int32 CallstackSize = FCStringAnsi::Strlen(StackTrace) + 1;
-			Callstack = (ANSICHAR*)FMemory::Malloc(CallstackSize);
-			FCStringAnsi::Strcpy(Callstack, CallstackSize, StackTrace);
-			UniqueCallstacks.Add(OutCallstackCRC, FCallstackData(Callstack,
-				InSerializedObject,
-				InSerializedProperty));
-		}
+	if (FCallstackData* ExistingCallstack = UniqueCallstacks.Find(OutCallstackCRC))
+	{
+		OutCallstack = ExistingCallstack->Callstack.Get();
 	}
 	else
 	{
-		OutCallstackCRC = 0;
+		const int32 Len = FCStringAnsi::Strlen(StackTrace.Get()) + 1;
+		TUniquePtr<ANSICHAR[]> Callstack = MakeUnique<ANSICHAR[]>(Len);
+		FCStringAnsi::Strcpy(Callstack.Get(), Len, StackTrace.Get());
+		FCallstackData& Data = UniqueCallstacks.Add(OutCallstackCRC, FCallstackData(MoveTemp(Callstack), SerializedObject, SerializedProperty));
+
+		OutCallstack = Data.Callstack.Get();
 	}
-	return Callstack;
+
+	return OutCallstack;
 }
 
-void FArchiveStackTrace::Serialize(void* InData, int64 Num)
+void FArchiveCallstacks::Add(
+	int64 CurrentOffset,
+	int64 CurrentTotalSize,
+	UObject* SerializedObject,
+	FProperty* SerializedProperty,
+	TArrayView<const FName> DebugDataStack,
+	bool bSaveStackTrace,
+	int32 StackIgnoreCount)
 {
-	if (Num)
+	if (GIgnoreDiffManager.ShouldBypassDiff())
 	{
-		// Remove this function from the reported stack. Note that all functions that we want to
-		// remove from the stack using this method have to be marked with FORCENOINLINE so they're
-		// there for us to remove and we don't remove meaningful functions above them instead.
-		const int32 StackIgnoreCount = 1;
+		return;
+	}
+
+	TotalSize = CurrentTotalSize;
 
-		static struct FBreakAtOffsetSettings
+	if (bSaveStackTrace && GIgnoreDiffManager.ShouldIgnoreDiff() == false)
+	{
+		StackTrace[0] = '\0';
+		FPlatformStackWalk::StackWalkAndDump(StackTrace.Get(), StackTraceSize, StackIgnoreCount);
+#if WITH_EDITOR && !NO_LOGGING
+		//if we have a debug name stack, plaster it onto the end of the current stack buffer so that it's a part of the unique stack entry.
+		if (DebugDataStack.Num() > 0)
 		{
-			FString PackageToBreakOn;
-			int64 OffsetToBreakOn;
+			FCStringAnsi::Strcat(StackTrace.Get(), StackTraceSize, DebugDataStackMarker);
 
-			FBreakAtOffsetSettings()
-				: OffsetToBreakOn(-1)
-			{
-				if (!FParse::Param(FCommandLine::Get(), TEXT("cooksinglepackage")))
-				{
-					return;
-				}
+			const FString SubIndent = FString(FDiffFormatHelper::Get().Indent) + FString(TEXT("    "));
 
-				FString Package;
-				if (!FParse::Value(FCommandLine::Get(), TEXT("map="), Package))
+			bool bIsIndenting = true;
+			for (const auto& DebugData : DebugDataStack)
+			{
+				if (bIsIndenting)
 				{
-					return;
+					FCStringAnsi::Strcat(StackTrace.Get(), StackTraceSize, TCHAR_TO_ANSI(*SubIndent));
 				}
 
-				int64 Offset;
-				if (!FParse::Value(FCommandLine::Get(), TEXT("diffonlybreakoffset="), Offset) || Offset <= 0)
-				{
-					return;
-				}
+				ANSICHAR DebugName[NAME_SIZE];
+				DebugData.GetPlainANSIString(DebugName);
+				FCStringAnsi::Strcat(StackTrace.Get(), StackTraceSize, DebugName);
 
-				OffsetToBreakOn = Offset;
-				PackageToBreakOn = TEXT("/") + FPackageName::GetShortName(Package);
+				//these are special-cased, as we assume they'll be followed by object/property names and want the names on the same line for readability's sake.
+				const bool bIsPropertyLabel = (DebugData == TEXT("SerializeScriptProperties") || DebugData == TEXT("PropertySerialize") || DebugData == TEXT("SerializeTaggedProperty"));
+				const ANSICHAR* const LineEnd = bIsPropertyLabel ? ": " : "\r\n";
+				FCStringAnsi::Strcat(StackTrace.Get(), StackTraceSize, LineEnd);
+				bIsIndenting = !bIsPropertyLabel;
 			}
-		} BreakAtOffsetSettings;
-
-		int64 CurrentOffset = Tell();
+		}
+#endif
+		// Make sure we compare the new stack trace with the last one in the next if statement
+		bCallstacksDirty = true;
+	}
 
-		if (BreakAtOffsetSettings.OffsetToBreakOn >= 0 && BreakAtOffsetSettings.OffsetToBreakOn >= CurrentOffset && BreakAtOffsetSettings.OffsetToBreakOn < CurrentOffset + Num)
+	if (LastSerializeCallstack == nullptr || (bCallstacksDirty && FCStringAnsi::Strcmp(LastSerializeCallstack, StackTrace.Get()) != 0))
+	{
+		uint32 CallstackCRC = 0;
+		if (CallstackAtOffsetMap.Num() == 0 || CurrentOffset > CallstackAtOffsetMap.Last().Offset)
 		{
-			FString ArcName = GetArchiveName();
-			int32 SubnameIndex = ArcName.Find(BreakAtOffsetSettings.PackageToBreakOn, ESearchCase::IgnoreCase, ESearchDir::FromEnd);
-			if (SubnameIndex >= 0)
+			// New data serialized at the end of archive buffer
+			if (bSaveStackTrace)
 			{
-				int32 SubnameEndIndex = SubnameIndex + BreakAtOffsetSettings.PackageToBreakOn.Len();
-				if (SubnameEndIndex == ArcName.Len() || ArcName[SubnameEndIndex] == TEXT('.'))
-				{
-					UE_DEBUG_BREAK();
-				}
+				LastSerializeCallstack = AddUniqueCallstack(SerializedObject, SerializedProperty, CallstackCRC);
 			}
+			CallstackAtOffsetMap.Add(FCallstackAtOffset {CurrentOffset, CallstackCRC, GIgnoreDiffManager.ShouldIgnoreDiff()});
 		}
-
-		// Walk the stack and dump it to the allocated memory.
-		bool bShouldCollectCallstack = bCollectCallstacks && ShouldLogOffset(CurrentOffset) && !GIgnoreDiffManager.ShouldIgnoreDiff();
-		if (bShouldCollectCallstack)
+		else
 		{
-			StackTrace[0] = '\0';
-			FPlatformStackWalk::StackWalkAndDump(StackTrace, StackTraceSize, StackIgnoreCount);
-#if WITH_EDITOR && !NO_LOGGING
-			//if we have a debug name stack, plaster it onto the end of the current stack buffer so that it's a part of the unique stack entry.
-			if (DebugDataStack.Num() > 0)
+			// This happens usually after Seek() so we need to find the exiting offset or insert a new one
+			int32 CallstackToUpdateIndex = GetCallstackIndexAtOffset(CurrentOffset);
+			check(CallstackToUpdateIndex != -1);
+			FCallstackAtOffset& CallstackToUpdate = CallstackAtOffsetMap[CallstackToUpdateIndex];
+			if (bSaveStackTrace)
 			{
-				FCStringAnsi::Strcat(StackTrace, StackTraceSize, DebugDataStackMarker);
-
-				const FString SubIndent = FString(FDiffFormatHelper::Get().Indent) + FString(TEXT("    "));
-
-				bool bIsIndenting = true;
-				for (const auto& DebugData : DebugDataStack)
-				{
-					if (bIsIndenting)
-					{
-						FCStringAnsi::Strcat(StackTrace, StackTraceSize, TCHAR_TO_ANSI(*SubIndent));
-					}
-
-					ANSICHAR DebugName[NAME_SIZE];
-					DebugData.GetPlainANSIString(DebugName);
-					FCStringAnsi::Strcat(StackTrace, StackTraceSize, DebugName);
-
-					//these are special-cased, as we assume they'll be followed by object/property names and want the names on the same line for readability's sake.
-					const bool bIsPropertyLabel = (DebugData == TEXT("SerializeScriptProperties") || DebugData == TEXT("PropertySerialize") || DebugData == TEXT("SerializeTaggedProperty"));
-					const ANSICHAR* const LineEnd = bIsPropertyLabel ? ": " : "\r\n";
-					FCStringAnsi::Strcat(StackTrace, StackTraceSize, LineEnd);
-					bIsIndenting = !bIsPropertyLabel;
-				}
+				LastSerializeCallstack = AddUniqueCallstack(SerializedObject, SerializedProperty, CallstackCRC);
 			}
-#endif
-			// Make sure we compare the new stack trace with the last one in the next if statement
-			bCallstacksDirty = true;
-		}
-
-		if (LastSerializeCallstack == nullptr || (bCallstacksDirty && FCStringAnsi::Strcmp(LastSerializeCallstack, StackTrace) != 0))
-		{
-			uint32 CallstackCRC = 0;
-			if (CallstackAtOffsetMap.Num() == 0 || CurrentOffset > CallstackAtOffsetMap.Last().Offset)
+			if (CallstackToUpdate.Offset == CurrentOffset)
 			{
-				// New data serialized at the end of archive buffer
-				check(LoadContext);
-				LastSerializeCallstack = AddUniqueCallstack(LoadContext->SerializedObject, GetSerializedProperty(), CallstackCRC);
-				CallstackAtOffsetMap.Add(FCallstactAtOffset(CurrentOffset, CallstackCRC, GIgnoreDiffManager.ShouldIgnoreDiff()));
+				CallstackToUpdate.Callstack = CallstackCRC;
 			}
 			else
 			{
-				// This happens usually after Seek() so we need to find the exiting offset or insert a new one
-				int32 CallstackToUpdateIndex = GetCallstackAtOffset(CurrentOffset, 0);
-				check(CallstackToUpdateIndex != -1);
-				check(LoadContext);
-				FCallstactAtOffset& CallstackToUpdate = CallstackAtOffsetMap[CallstackToUpdateIndex];
-				LastSerializeCallstack = AddUniqueCallstack(LoadContext->SerializedObject, GetSerializedProperty(), CallstackCRC);
-				if (CallstackToUpdate.Offset == CurrentOffset)
-				{
-					CallstackToUpdate.Callstack = CallstackCRC;
-				}
-				else
-				{
-					// Insert a new callstack
-					check(CallstackToUpdate.Offset < CurrentOffset);
-					CallstackAtOffsetMap.Insert(FCallstactAtOffset(CurrentOffset, CallstackCRC, GIgnoreDiffManager.ShouldIgnoreDiff()), CallstackToUpdateIndex + 1);
-				}
+				// Insert a new callstack
+				check(CallstackToUpdate.Offset < CurrentOffset);
+				CallstackAtOffsetMap.Insert(FCallstackAtOffset {CurrentOffset, CallstackCRC, GIgnoreDiffManager.ShouldIgnoreDiff()}, CallstackToUpdateIndex + 1);
 			}
-			check(CallstackCRC != 0 || !bShouldCollectCallstack);
-		}
-		else if (LastSerializeCallstack)
-		{
-			// Skip callstack comparison on next serialize call unless we grab a stack trace
-			bCallstacksDirty = false;
 		}
 	}
-	FLargeMemoryWriter::Serialize(InData, Num);
-}
-
-void FArchiveStackTrace::SetSerializeContext(FUObjectSerializeContext* InLoadContext)
-{
-	LoadContext = InLoadContext;
-}
-
-FUObjectSerializeContext* FArchiveStackTrace::GetSerializeContext()
-{
-	return LoadContext;
+	else if (LastSerializeCallstack)
+	{
+		// Skip callstack comparison on next serialize call unless we grab a stack trace
+		bCallstacksDirty = false;
+	}
 }
 
-int32 FArchiveStackTrace::GetCallstackAtOffset(int64 InOffset, int32 MinOffsetIndex)
+int32 FArchiveCallstacks::GetCallstackIndexAtOffset(int64 Offset, int32 MinOffsetIndex) const
 {
-	if (InOffset < 0 || InOffset > TotalSize() || MinOffsetIndex < 0 || MinOffsetIndex >= CallstackAtOffsetMap.Num())
+	if (Offset < 0 || Offset > TotalSize || MinOffsetIndex < 0 || MinOffsetIndex >= CallstackAtOffsetMap.Num())
 	{
 		return -1;
 	}
@@ -384,11 +347,11 @@ int32 FArchiveStackTrace::GetCallstackAtOffset(int64 InOffset, int32 MinOffsetIn
 	for (; MinOffsetIndex <= MaxOffsetIndex; )
 	{
 		int32 SearchIndex = (MinOffsetIndex + MaxOffsetIndex) / 2;
-		if (CallstackAtOffsetMap[SearchIndex].Offset < InOffset)
+		if (CallstackAtOffsetMap[SearchIndex].Offset < Offset)
 		{
 			MinOffsetIndex = SearchIndex + 1;
 		}
-		else if (CallstackAtOffsetMap[SearchIndex].Offset > InOffset)
+		else if (CallstackAtOffsetMap[SearchIndex].Offset > Offset)
 		{
 			MaxOffsetIndex = SearchIndex - 1;
 		}
@@ -405,20 +368,211 @@ int32 FArchiveStackTrace::GetCallstackAtOffset(int64 InOffset, int32 MinOffsetIn
 		MinOffsetIndex = FMath::Min(MinOffsetIndex, CallstackAtOffsetMap.Num() - 1);
 		for (int32 FirstLowerOffsetIndex = MinOffsetIndex; FirstLowerOffsetIndex >= 0; --FirstLowerOffsetIndex)
 		{
-			if (CallstackAtOffsetMap[FirstLowerOffsetIndex].Offset < InOffset)
+			if (CallstackAtOffsetMap[FirstLowerOffsetIndex].Offset < Offset)
 			{
 				OffsetForCallstackIndex = FirstLowerOffsetIndex;
 				break;
 			}
 		}
 		check(OffsetForCallstackIndex != -1);
-		check(CallstackAtOffsetMap[OffsetForCallstackIndex].Offset < InOffset);
-		check(OffsetForCallstackIndex == (CallstackAtOffsetMap.Num() - 1) || CallstackAtOffsetMap[OffsetForCallstackIndex + 1].Offset > InOffset);
+		check(CallstackAtOffsetMap[OffsetForCallstackIndex].Offset < Offset);
+		check(OffsetForCallstackIndex == (CallstackAtOffsetMap.Num() - 1) || CallstackAtOffsetMap[OffsetForCallstackIndex + 1].Offset > Offset);
 	}
 
 	return OffsetForCallstackIndex;
 }
 
+void FArchiveCallstacks::Append(const FArchiveCallstacks& Other, int64 Offset)
+{
+	for (const FCallstackAtOffset& OtherOffset : Other.CallstackAtOffsetMap)
+	{
+		FCallstackAtOffset& New = CallstackAtOffsetMap.Add_GetRef(OtherOffset);
+		New.Offset += Offset;
+	}
+
+	CallstackAtOffsetMap.Sort([](const FCallstackAtOffset& LHS,const FCallstackAtOffset& RHS)
+	{
+		return LHS.Offset < RHS.Offset;
+	});
+
+	for (const TPair<uint32, FCallstackData>& Kv : Other.UniqueCallstacks)
+	{
+		if (FCallstackData* Existing = UniqueCallstacks.Find(Kv.Key))
+		{
+			if (LastSerializeCallstack == Existing->Callstack.Get())
+			{
+				LastSerializeCallstack = nullptr;
+			}
+			UniqueCallstacks.Remove(Kv.Key);
+		}
+		UniqueCallstacks.Emplace(Kv.Key, Kv.Value.Clone());
+	}
+
+	TotalSize = FMath::Max(TotalSize, Other.TotalSize);
+}
+
+FArchiveStackTraceWriter::FArchiveStackTraceWriter(
+	FArchive& InInner,
+	FArchiveCallstacks& InCallstacks,
+	const FArchiveDiffMap* InDiffMap,
+	int64 InDiffMapOffset)
+		: FArchiveProxy(InInner)
+		, Callstacks(InCallstacks)
+		, DiffMap(InDiffMap)
+		, DiffMapOffset(InDiffMapOffset)
+		, bInnerArchiveDisabled(false)
+{
+}
+
+void FArchiveStackTraceWriter::Serialize(void* Data, int64 Length)
+{
+	static struct FBreakAtOffsetSettings
+	{
+		FString PackageToBreakOn;
+		int64 OffsetToBreakOn;
+
+		FBreakAtOffsetSettings()
+			: OffsetToBreakOn(-1)
+		{
+			if (!FParse::Param(FCommandLine::Get(), TEXT("cooksinglepackage")))
+			{
+				return;
+			}
+
+			FString Package;
+			if (!FParse::Value(FCommandLine::Get(), TEXT("map="), Package))
+			{
+				return;
+			}
+
+			int64 Offset;
+			if (!FParse::Value(FCommandLine::Get(), TEXT("diffonlybreakoffset="), Offset) || Offset <= 0)
+			{
+				return;
+			}
+
+			OffsetToBreakOn = Offset;
+			PackageToBreakOn = TEXT("/") + FPackageName::GetShortName(Package);
+		}
+	} BreakAtOffsetSettings;
+
+	const int64 CurrentOffset = DiffMapOffset + Tell();
+	const int64 CurrentTotalSize = DiffMapOffset + TotalSize();
+
+	if (BreakAtOffsetSettings.OffsetToBreakOn >= 0 && BreakAtOffsetSettings.OffsetToBreakOn >= CurrentOffset && BreakAtOffsetSettings.OffsetToBreakOn < CurrentOffset + Length)
+	{
+		FString ArcName = GetArchiveName();
+		int32 SubnameIndex = ArcName.Find(BreakAtOffsetSettings.PackageToBreakOn, ESearchCase::IgnoreCase, ESearchDir::FromEnd);
+		if (SubnameIndex >= 0)
+		{
+			int32 SubnameEndIndex = SubnameIndex + BreakAtOffsetSettings.PackageToBreakOn.Len();
+			if (SubnameEndIndex == ArcName.Len() || ArcName[SubnameEndIndex] == TEXT('.'))
+			{
+				UE_DEBUG_BREAK();
+			}
+		}
+	}
+
+	if (Length > 0)
+	{
+		check(SerializeContext);
+		const bool bSaveStackTrace = DiffMap != nullptr && DiffMap->ContainsOffset(CurrentOffset);
+		TArrayView<const FName> DebugStack;
+#if WITH_EDITOR
+		DebugStack = DebugDataStack;
+#endif
+		Callstacks.Add(
+			CurrentOffset,
+			CurrentTotalSize,
+			SerializeContext->SerializedObject,
+			GetSerializedProperty(),
+			DebugStack,
+			bSaveStackTrace,
+			StackIgnoreCount);
+	}
+
+	if (bInnerArchiveDisabled == false)
+	{
+		InnerArchive.Serialize(Data, Length);
+	}
+}
+
+void FArchiveStackTraceWriter::SetSerializeContext(FUObjectSerializeContext* Context)
+{
+	SerializeContext = Context;
+
+	if (bInnerArchiveDisabled == false)
+	{
+		InnerArchive.SetSerializeContext(Context);
+	}
+}
+
+FUObjectSerializeContext* FArchiveStackTraceWriter::GetSerializeContext()
+{
+	return SerializeContext;
+}
+
+FArchiveStackTraceMemoryWriter::FArchiveStackTraceMemoryWriter(
+	FArchiveCallstacks& Callstacks,
+	const FArchiveDiffMap* DiffMap,
+	const int64 DiffMapOffset,
+	const int64 PreAllocateBytes,
+	bool bIsPersistent,
+	const TCHAR* Filename)
+		: FLargeMemoryWriter(PreAllocateBytes, bIsPersistent, Filename)
+		, StackTraceWriter(*this, Callstacks, DiffMap, DiffMapOffset)
+{
+	// Hack to prevent recursive calls to serialize when passing in this.
+	StackTraceWriter.SetDisableInnerArchive(true);
+	StackTraceWriter.SetStackIgnoreCount(StackTraceWriter.GetStackIgnoreCount() + 1);
+}
+
+void FArchiveStackTraceMemoryWriter::Serialize(void* Memory, int64 Length)
+{
+	StackTraceWriter.Serialize(Memory, Length);
+	FLargeMemoryWriter::Serialize(Memory, Length);
+}
+
+void FArchiveStackTraceMemoryWriter::SetSerializeContext(FUObjectSerializeContext* Context)
+{
+	StackTraceWriter.SetSerializeContext(Context);
+}
+
+FUObjectSerializeContext* FArchiveStackTraceMemoryWriter::GetSerializeContext()
+{
+	return StackTraceWriter.GetSerializeContext();
+}
+
+FArchiveStackTrace::FArchiveStackTrace(UObject* InAsset, const TCHAR* InFilename, bool bInCollectCallstacks, const FArchiveDiffMap* InDiffMap)
+	: FLargeMemoryWriter(0, false, InFilename)
+	, Callstacks(InAsset)
+	, StackTraceWriter(*this, Callstacks, InDiffMap, 0)
+{
+	// Hack to prevent recursive calls to serialize when passing in this.
+	StackTraceWriter.SetDisableInnerArchive(true);
+	StackTraceWriter.SetStackIgnoreCount(StackTraceWriter.GetStackIgnoreCount() + 1);
+}
+
+FArchiveStackTrace::~FArchiveStackTrace()
+{
+}
+
+void FArchiveStackTrace::Serialize(void* Memory, int64 Length)
+{
+	StackTraceWriter.Serialize(Memory, Length);
+	FLargeMemoryWriter::Serialize(Memory, Length);
+}
+
+void FArchiveStackTrace::SetSerializeContext(FUObjectSerializeContext* Context)
+{
+	StackTraceWriter.SetSerializeContext(Context);
+}
+
+FUObjectSerializeContext* FArchiveStackTrace::GetSerializeContext()
+{
+	return StackTraceWriter.GetSerializeContext();
+}
+
 bool FArchiveStackTrace::LoadPackageIntoMemory(const TCHAR* InFilename, FPackageData& OutPackageData, TUniquePtr<uint8>& OutLoadedBytes)
 {
 	TUniquePtr<FArchive> UAssetFileArchive(IFileManager::Get().CreateFileReader(InFilename));
@@ -502,7 +656,16 @@ namespace
 	}
 }
 
-void FArchiveStackTrace::CompareWithInternal(const FPackageData& SourcePackage, const FPackageData& DestPackage, const TCHAR* AssetFilename, const TCHAR* CallstackCutoffText, const int64 MaxDiffsToLog, int32& InOutDiffsLogged, TMap<FName, FArchiveDiffStats>& OutStats)
+void FArchiveStackTraceWriter::Compare(
+	const FPackageData& SourcePackage,
+	const FPackageData& DestPackage,
+	const FArchiveCallstacks& Callstacks,
+	const FArchiveDiffMap& DiffMap,
+	const TCHAR* AssetFilename,
+	const TCHAR* CallstackCutoffText,
+	const int64 MaxDiffsToLog,
+	int32& InOutDiffsLogged,
+	TMap<FName, FArchiveDiffStats>& OutStats)
 {
 #if !NO_LOGGING
 	const TCHAR* const Indent = FDiffFormatHelper::Get().Indent;
@@ -510,6 +673,7 @@ void FArchiveStackTrace::CompareWithInternal(const FPackageData& SourcePackage,
 	const int64 SourceSize = SourcePackage.Size - SourcePackage.StartOffset;
 	const int64 DestSize = DestPackage.Size - DestPackage.StartOffset;
 	const int64 SizeToCompare = FMath::Min(SourceSize, DestSize);
+	const FName AssetClass = Callstacks.GetAssetClass();
 	
 	if (SourceSize != DestSize)
 	{
@@ -546,9 +710,9 @@ void FArchiveStackTrace::CompareWithInternal(const FPackageData& SourcePackage,
 			}
 		};
 
-		if (ShouldLogOffset(DestAbsoluteOffset))
+		if (DiffMap.ContainsOffset(DestAbsoluteOffset))
 		{
-			int32 DifferenceCallstackoffsetIndex = GetCallstackAtOffset(DestAbsoluteOffset, FMath::Max(LastDifferenceCallstackOffsetIndex, 0));
+			int32 DifferenceCallstackoffsetIndex = Callstacks.GetCallstackIndexAtOffset(DestAbsoluteOffset, FMath::Max(LastDifferenceCallstackOffsetIndex, 0));
 			ON_SCOPE_EXIT
 			{
 				LastDifferenceCallstackOffsetIndex = DifferenceCallstackoffsetIndex;
@@ -565,8 +729,8 @@ void FArchiveStackTrace::CompareWithInternal(const FPackageData& SourcePackage,
 				continue;
 			}
 
-			const FCallstactAtOffset& CallstackAtOffset = CallstackAtOffsetMap[DifferenceCallstackoffsetIndex];
-			const FCallstackData& DifferenceCallstackData = UniqueCallstacks[CallstackAtOffset.Callstack];
+			const FArchiveCallstacks::FCallstackAtOffset& CallstackAtOffset = Callstacks.GetCallstack(DifferenceCallstackoffsetIndex);
+			const FArchiveCallstacks::FCallstackData& DifferenceCallstackData = Callstacks.GetCallstackData(CallstackAtOffset);
 			FString DifferenceCallstackDataText = DifferenceCallstackData.ToString(CallstackCutoffText);
 			if (LastDifferenceCallstackDataText.Compare(DifferenceCallstackDataText, ESearchCase::CaseSensitive) == 0)
 			{
@@ -594,10 +758,10 @@ void FArchiveStackTrace::CompareWithInternal(const FPackageData& SourcePackage,
 								break;
 							}
 
-							int32 CallstackIndex = GetCallstackAtOffset(OffsetX - 1, 0);
-
-							const FCallstactAtOffset& PreviousCallstack = CallstackAtOffsetMap[CallstackIndex];
-							if (UniqueCallstacks[PreviousCallstack.Callstack].SerializedProp != SerProp)
+							const int32 CallstackIndex = Callstacks.GetCallstackIndexAtOffset(OffsetX - 1, 0);
+							const FArchiveCallstacks::FCallstackAtOffset& PreviousCallstack = Callstacks.GetCallstack(CallstackIndex);
+							const FArchiveCallstacks::FCallstackData& PreviousCallstackData = Callstacks.GetCallstackData(PreviousCallstack);
+							if (PreviousCallstackData.SerializedProp != SerProp)
 							{
 								break;
 							}
@@ -631,7 +795,7 @@ void FArchiveStackTrace::CompareWithInternal(const FPackageData& SourcePackage,
 				FString DebugDataStackText;
 #if WITH_EDITOR
 				//check for a debug data stack as part of the unique stack entry, and log it out if we find it.
-				FString FullStackText = DifferenceCallstackData.Callstack;
+				FString FullStackText = DifferenceCallstackData.Callstack.Get();
 				int32 DebugDataIndex = FullStackText.Find(ANSI_TO_TCHAR(DebugDataStackMarker), ESearchCase::CaseSensitive);
 				if (DebugDataIndex > 0)
 				{
@@ -723,10 +887,10 @@ void FArchiveStackTrace::CompareWith(const TCHAR* InFilename, const int64 TotalH
 	CompareWith(SourcePackage, InFilename, TotalHeaderSize, CallstackCutoffText, MaxDiffsToLog, OutStats);
 }
 
-
 void FArchiveStackTrace::CompareWith(const FPackageData& SourcePackage, const TCHAR* FileDisplayName, const int64 TotalHeaderSize,
 	const TCHAR* CallstackCutoffText, const int32 MaxDiffsToLog, TMap<FName, FArchiveDiffStats>&OutStats)
 {
+	const FName AssetClass = Callstacks.GetAssetClass();
 	OutStats.FindOrAdd(AssetClass).NewFileTotalSize = TotalSize();
 	if (SourcePackage.Size == 0)
 	{
@@ -756,11 +920,11 @@ void FArchiveStackTrace::CompareWith(const FPackageData& SourcePackage, const TC
 	DestPackageHeader.HeaderSize = 0;
 	DestPackageHeader.StartOffset = 0;
 
-	CompareWithInternal(SourcePackageHeader, DestPackageHeader, FileDisplayName, CallstackCutoffText, MaxDiffsToLog, NumLoggedDiffs, OutStats);
+	FArchiveStackTraceWriter::Compare(SourcePackageHeader, DestPackageHeader, Callstacks, StackTraceWriter.GetDiffMap(), FileDisplayName, CallstackCutoffText, MaxDiffsToLog, NumLoggedDiffs, OutStats);
 
 	if (TotalHeaderSize > 0 && OutStats.FindOrAdd(AssetClass).NumDiffs > 0)
 	{
-		DumpPackageHeaderDiffs(SourcePackage, DestPackage, FileDisplayName, MaxDiffsToLog);
+		FArchiveStackTraceWriter::DumpPackageHeaderDiffs(SourcePackage, DestPackage, FileDisplayName, MaxDiffsToLog);
 	}
 
 	FPackageData SourcePackageExports = SourcePackage;
@@ -781,7 +945,7 @@ void FArchiveStackTrace::CompareWith(const FPackageData& SourcePackage, const TC
 		AssetName = FileDisplayName;
 	}
 
-	CompareWithInternal(SourcePackageExports, DestPackageExports, *AssetName, CallstackCutoffText, MaxDiffsToLog, NumLoggedDiffs, OutStats);
+	FArchiveStackTraceWriter::Compare(SourcePackageExports, DestPackageExports, Callstacks, StackTraceWriter.GetDiffMap(), *AssetName, CallstackCutoffText, MaxDiffsToLog, NumLoggedDiffs, OutStats);
 
 	// Optionally save out any differences we detected.
 	const FArchiveDiffStats& Stats = OutStats.FindOrAdd(AssetClass);
@@ -843,11 +1007,11 @@ void FArchiveStackTrace::CompareWith(const FPackageData& SourcePackage, const TC
 	}
 }
 
-bool FArchiveStackTrace::GenerateDiffMapInternal(const FPackageData& SourcePackage, const FPackageData& DestPackage, int32 MaxDiffsToFind, FArchiveDiffMap& OutDiffMap)
+bool FArchiveStackTraceWriter::GenerateDiffMap(const FPackageData& SourcePackage, const FPackageData& DestPackage, const FArchiveCallstacks& Callstacks, int32 MaxDiffsToFind, FArchiveDiffMap& OutDiffMap)
 {
 	bool bIdentical = true;
 	int64 LastDifferenceCallstackOffsetIndex = -1;
-	FCallstackData* DifferenceCallstackData = nullptr;
+	FArchiveCallstacks::FCallstackData* DifferenceCallstackData = nullptr;
 
 	const int64 SourceSize = SourcePackage.Size - SourcePackage.StartOffset;
 	const int64 DestSize = DestPackage.Size - DestPackage.StartOffset;
@@ -862,19 +1026,19 @@ bool FArchiveStackTrace::GenerateDiffMapInternal(const FPackageData& SourcePacka
 			bIdentical = false;
 			if (OutDiffMap.Num() < MaxDiffsToFind)
 			{
-				int64 DifferenceCallstackoffsetIndex = GetCallstackAtOffset(DestAbsoluteOffset, FMath::Max<int64>(LastDifferenceCallstackOffsetIndex, 0));
-				if (DifferenceCallstackoffsetIndex >= 0 && DifferenceCallstackoffsetIndex != LastDifferenceCallstackOffsetIndex)
+				const int32 DifferenceCallstackOffsetIndex = Callstacks.GetCallstackIndexAtOffset(DestAbsoluteOffset, FMath::Max<int64>(LastDifferenceCallstackOffsetIndex, 0));
+				if (DifferenceCallstackOffsetIndex >= 0 && DifferenceCallstackOffsetIndex != LastDifferenceCallstackOffsetIndex)
 				{
-					const FCallstactAtOffset& CallstackAtOffset = CallstackAtOffsetMap[DifferenceCallstackoffsetIndex];
+					const FArchiveCallstacks::FCallstackAtOffset& CallstackAtOffset = Callstacks.GetCallstack(DifferenceCallstackOffsetIndex);
 					if (!CallstackAtOffset.bIgnore)
 					{
 						FArchiveDiffInfo OffsetAndSize;
 						OffsetAndSize.Offset = CallstackAtOffset.Offset;
-						OffsetAndSize.Size = GetSerializedDataSizeForOffsetIndex(DifferenceCallstackoffsetIndex);
+						OffsetAndSize.Size = Callstacks.GetSerializedDataSizeForOffsetIndex(DifferenceCallstackOffsetIndex);
 						OutDiffMap.Add(OffsetAndSize);
 					}
 				}
-				LastDifferenceCallstackOffsetIndex = DifferenceCallstackoffsetIndex;
+				LastDifferenceCallstackOffsetIndex = DifferenceCallstackOffsetIndex;
 			}
 		}
 	}
@@ -884,9 +1048,9 @@ bool FArchiveStackTrace::GenerateDiffMapInternal(const FPackageData& SourcePacka
 		bIdentical = false;
 
 		// Add all the remaining callstacks to the diff map
-		for (int32 OffsetIndex = LastDifferenceCallstackOffsetIndex + 1; OffsetIndex < CallstackAtOffsetMap.Num() && OutDiffMap.Num() < MaxDiffsToFind; ++OffsetIndex)
+		for (int32 OffsetIndex = LastDifferenceCallstackOffsetIndex + 1; OffsetIndex < Callstacks.Num() && OutDiffMap.Num() < MaxDiffsToFind; ++OffsetIndex)
 		{
-			const FCallstactAtOffset& CallstackAtOffset = CallstackAtOffsetMap[OffsetIndex];
+			const FArchiveCallstacks::FCallstackAtOffset& CallstackAtOffset = Callstacks.GetCallstack(OffsetIndex);
 			// Compare against the size without start offset as all callstack offsets are absolute (from the merged header + exports file)
 			if (CallstackAtOffset.Offset < DestPackage.Size)
 			{
@@ -894,7 +1058,7 @@ bool FArchiveStackTrace::GenerateDiffMapInternal(const FPackageData& SourcePacka
 				{
 					FArchiveDiffInfo OffsetAndSize;
 					OffsetAndSize.Offset = CallstackAtOffset.Offset;
-					OffsetAndSize.Size = GetSerializedDataSizeForOffsetIndex(OffsetIndex);
+					OffsetAndSize.Size = Callstacks.GetSerializedDataSizeForOffsetIndex(OffsetIndex);
 					OutDiffMap.Add(OffsetAndSize);
 				}
 			}
@@ -947,7 +1111,7 @@ bool FArchiveStackTrace::GenerateDiffMap(const FPackageData& SourcePackage, int6
 		DestPackageHeader.HeaderSize = 0;
 		DestPackageHeader.StartOffset = 0;
 
-		bHeaderIdentical = GenerateDiffMapInternal(SourcePackageHeader, DestPackageHeader, MaxDiffsToFind, OutDiffMap);
+		bHeaderIdentical = FArchiveStackTraceWriter::GenerateDiffMap(SourcePackageHeader, DestPackageHeader, Callstacks, MaxDiffsToFind, OutDiffMap);
 	}
 
 	{
@@ -959,7 +1123,7 @@ bool FArchiveStackTrace::GenerateDiffMap(const FPackageData& SourcePackage, int6
 		DestPackageExports.HeaderSize = 0;
 		DestPackageExports.StartOffset = TotalHeaderSize;
 
-		bExportsIdentical = GenerateDiffMapInternal(SourcePackageExports, DestPackageExports, MaxDiffsToFind, OutDiffMap);
+		bExportsIdentical = FArchiveStackTraceWriter::GenerateDiffMap(SourcePackageExports, DestPackageExports, Callstacks, MaxDiffsToFind, OutDiffMap);
 	}
 
 	bIdentical = bHeaderIdentical && bExportsIdentical;
@@ -993,7 +1157,7 @@ bool FArchiveStackTrace::IsIdentical(const FPackageData& SourcePackage, int64 Bu
 	return bIdentical;
 }
 
-FLinkerLoad* FArchiveStackTrace::CreateLinkerForPackage(FUObjectSerializeContext* LoadContext, const FString& InPackageName, const FString& InFilename, const FPackageData& PackageData)
+FLinkerLoad* FArchiveStackTraceWriter::CreateLinkerForPackage(FUObjectSerializeContext* LoadContext, const FString& InPackageName, const FString& InFilename, const FPackageData& PackageData)
 {
 	// First create a temp package to associate the linker with
 	UPackage* Package = FindObjectFast<UPackage>(nullptr, *InPackageName);
@@ -1370,7 +1534,7 @@ static void DumpTableDifferences(
 #endif // !NO_LOGGING
 }
 
-void FArchiveStackTrace::DumpPackageHeaderDiffs(const FPackageData& SourcePackage, const FPackageData& DestPackage, const FString& AssetFilename, const int32 MaxDiffsToLog)
+void FArchiveStackTraceWriter::DumpPackageHeaderDiffs(const FPackageData& SourcePackage, const FPackageData& DestPackage, const FString& AssetFilename, const int32 MaxDiffsToLog)
 {
 #if !NO_LOGGING
 	FString AssetPathName = FPaths::Combine(*FPaths::GetPath(AssetFilename.Mid(AssetFilename.Find(TEXT(":"), ESearchCase::CaseSensitive) + 1)), *FPaths::GetBaseFilename(AssetFilename));
diff --git a/Engine/Source/Runtime/CoreUObject/Private/Serialization/PackageWriterToSharedBuffer.cpp b/Engine/Source/Runtime/CoreUObject/Private/Serialization/PackageWriterToSharedBuffer.cpp
--- a/Engine/Source/Runtime/CoreUObject/Private/Serialization/PackageWriterToSharedBuffer.cpp
+++ b/Engine/Source/Runtime/CoreUObject/Private/Serialization/PackageWriterToSharedBuffer.cpp
@@ -14,6 +14,12 @@ TUniquePtr<FLargeMemoryWriter> IPackageWriter::CreateLinkerArchive(FName Package
 	return TUniquePtr<FLargeMemoryWriter>(new FLargeMemoryWriter(0, bIsPersistent, *PackageName.ToString()));
 }
 
+TUniquePtr<FLargeMemoryWriter> IPackageWriter::CreateLinkerExportsArchive(FName PackageName, UObject* Asset)
+{
+	const bool bPersistent = true;
+	return MakeUnique<FLargeMemoryWriter>(0, bPersistent, *PackageName.ToString());
+}
+
 static FSharedBuffer IoBufferToSharedBuffer(const FIoBuffer& InBuffer)
 {
 	InBuffer.EnsureOwned();
@@ -138,4 +144,4 @@ void FPackageWriterRecords::ValidateCommit(FPackage& Record, const IPackageWrite
 			TEXT("IPackageWriter->WriteBulkData must not be called with more than one BulkData of the same type."));
 		HasBulkDataType[(int32)BulkRecord.Info.BulkDataType] |= 1 << BulkRecord.Info.MultiOutputIndex;
 	}
-}
\ No newline at end of file
+}
diff --git a/Engine/Source/Runtime/CoreUObject/Private/UObject/SavePackage2.cpp b/Engine/Source/Runtime/CoreUObject/Private/UObject/SavePackage2.cpp
--- a/Engine/Source/Runtime/CoreUObject/Private/UObject/SavePackage2.cpp
+++ b/Engine/Source/Runtime/CoreUObject/Private/UObject/SavePackage2.cpp
@@ -23,6 +23,7 @@
 #include "Misc/ScopedSlowTask.h"
 #include "Misc/SecureHash.h"
 #include "ProfilingDebugging/CookStats.h"
+#include "Serialization/ArchiveProxy.h"
 #include "Serialization/ArchiveUObjectFromStructuredArchive.h"
 #include "Serialization/ArchiveStackTrace.h"
 #include "Serialization/EditorBulkData.h"
@@ -1865,16 +1866,30 @@ ESavePackageResult WritePackageTextHeader(FStructuredArchive::FRecord& Structure
 	return ReturnSuccessOrCancel();
 }
 
-[[nodiscard]] ESavePackageResult WriteCookedExports(TArray64<uint8>& ExportsBuffer, FSaveContext& SaveContext)
+[[nodiscard]] ESavePackageResult WriteCookedExports(FArchive& ExportsArchive, FSaveContext& SaveContext)
 {
+	// Used to make any serialized offset during export serialization relative to the beginning of the export.
+	struct FExportProxyArchive
+		: public FArchiveProxy
+	{
+		FExportProxyArchive(FArchive& Inner)
+			: FArchiveProxy(Inner)
+			, Offset(Inner.Tell()) { }
+
+		virtual void Seek(int64 Pos) { InnerArchive.Seek(Offset + Pos); }
+		virtual int64 Tell() { return InnerArchive.Tell() - Offset; }
+		virtual int64 TotalSize() { return InnerArchive.TotalSize() - Offset; }
+		
+		int64 Offset;
+	};
+
 	SCOPED_SAVETIMER(UPackage_Save_SaveExports);
 
 	check(SaveContext.GetLinker() && SaveContext.GetLinker()->IsCooking());
 	FLinkerSave& Linker = *SaveContext.GetLinker();
 	FScopedSlowTask SlowTask((float)Linker.ExportMap.Num(), FText(), SaveContext.IsUsingSlowTask());
-	TArray64<uint8> ExportBuffer;
 
-	for (int32 i = 0; i < Linker.ExportMap.Num(); i++)
+	for (int32 ExportIndex = 0; ExportIndex < Linker.ExportMap.Num(); ExportIndex++)
 	{
 		if (GWarn->ReceivedUserCancel())
 		{
@@ -1882,7 +1897,7 @@ ESavePackageResult WritePackageTextHeader(FStructuredArchive::FRecord& Structure
 		}
 		SlowTask.EnterProgressFrame();
 
-		FObjectExport& Export = Linker.ExportMap[i];
+		FObjectExport& Export = Linker.ExportMap[ExportIndex];
 		if (Export.Object == nullptr)
 		{
 			continue;
@@ -1891,14 +1906,11 @@ ESavePackageResult WritePackageTextHeader(FStructuredArchive::FRecord& Structure
 		SCOPED_SAVETIMER(UPackage_Save_SaveExport);
 		SCOPED_SAVETIMER_TEXT(*WriteToString<128>(GetClassTraceScope(Export.Object), TEXT("_SaveSerialize")));
 
-		Export.SerialOffset = ExportsBuffer.Num();
-		Linker.CurrentlySavingExport = FPackageIndex::FromExport(i);
+		Export.SerialOffset = ExportsArchive.Tell();
+		Linker.CurrentlySavingExport = FPackageIndex::FromExport(ExportIndex);
 
-		// Use a separate archive for each export to enable rearranging the exports according to the load/dependency order. Seeks/Tell
-		// during export serialize will be relative to the beginning of the export buffer.
-		ExportBuffer.Reset();
-		FMemoryWriter64 ExportArchive(ExportBuffer, false, false, SaveContext.GetPackage()->GetFName());
-		TGuardValue<FArchive*> GuardSaver(Linker.Saver, &ExportArchive);
+		FExportProxyArchive Ar(ExportsArchive);
+		TGuardValue<FArchive*> GuardSaver(Linker.Saver, &Ar);
 
 		if (Export.Object->HasAnyFlags(RF_ClassDefaultObject))
 		{
@@ -1919,8 +1931,7 @@ ESavePackageResult WritePackageTextHeader(FStructuredArchive::FRecord& Structure
 		}
 
 		Linker.CurrentlySavingExport = FPackageIndex();
-		Export.SerialSize = ExportBuffer.Num(); 
-		ExportsBuffer.Append(ExportBuffer);
+		Export.SerialSize = ExportsArchive.Tell() - Export.SerialOffset;
 	}
 
 	return Linker.IsError() ? ESavePackageResult::Error : ReturnSuccessOrCancel();
@@ -2641,9 +2652,10 @@ ESavePackageResult SaveHarvestedRealms(FSaveContext& SaveContext, ESaveRealm Har
 		if (Linker.IsCooking())
 		{
 			// Write the exports into a seperate archive
-			TArray64<uint8> ExportsBuffer;
-			ExportsBuffer.Reserve(4 << 10);
-			SaveContext.Result = WriteCookedExports(ExportsBuffer, SaveContext);
+			TUniquePtr<FLargeMemoryWriter> ExportsArchive = SaveContext.GetPackageWriter()->CreateLinkerExportsArchive(
+				SaveContext.GetPackage()->GetFName(), SaveContext.GetAsset());
+			ExportsArchive->SetSerializeContext(SaveContext.GetSerializeContext());
+			SaveContext.Result = WriteCookedExports(*ExportsArchive, SaveContext);
 
 			if (SaveContext.Result == ESavePackageResult::Success)
 			{
@@ -2652,7 +2664,11 @@ ESavePackageResult SaveHarvestedRealms(FSaveContext& SaveContext, ESaveRealm Har
 				check(DataResourceSize >= 0);
 
 				Linker.Summary.TotalHeaderSize += DataResourceSize;
-				Linker.Serialize(ExportsBuffer.GetData(), ExportsBuffer.Num()); 
+				{
+					// Disables writing stack trace data when appending the exports data
+					FArchiveStackTraceDisabledScope _; 
+					Linker.Serialize(ExportsArchive->GetData(), ExportsArchive->TotalSize());
+				}
 
 				// Adjust the export offsets with the total header size
 				for (FObjectExport& Export : Linker.ExportMap)
diff --git a/Engine/Source/Runtime/CoreUObject/Public/Serialization/ArchiveStackTrace.h b/Engine/Source/Runtime/CoreUObject/Public/Serialization/ArchiveStackTrace.h
--- a/Engine/Source/Runtime/CoreUObject/Public/Serialization/ArchiveStackTrace.h
+++ b/Engine/Source/Runtime/CoreUObject/Public/Serialization/ArchiveStackTrace.h
@@ -44,6 +44,7 @@ struct COREUOBJECT_API FArchiveDiffStats
 	{}
 };
 
+/** Ignores saving the stack trace when collecting serialize offsets .*/
 class COREUOBJECT_API FArchiveStackTraceIgnoreScope
 {
 	const bool bIgnore;
@@ -52,6 +53,17 @@ public:
 	~FArchiveStackTraceIgnoreScope();
 };
 
+/**
+ * Disables collecting both offsets and stack traces when collecting serialize callstacks.
+ * Typically used when appending data from one stack tracing archive to another.
+ */
+class COREUOBJECT_API FArchiveStackTraceDisabledScope
+{
+public:
+	FArchiveStackTraceDisabledScope();
+	~FArchiveStackTraceDisabledScope();
+};
+
 struct COREUOBJECT_API FArchiveDiffInfo
 {
 	int64 Offset;
@@ -84,168 +96,286 @@ struct COREUOBJECT_API FArchiveDiffInfo
 
 class COREUOBJECT_API FArchiveDiffMap : public TArray<FArchiveDiffInfo>
 {
+public:
+	bool ContainsOffset(int64 Offset) const
+	{
+		for (const FArchiveDiffInfo& Diff : *this)
+		{
+			if (Diff.Offset <= Offset && Offset < (Diff.Offset + Diff.Size))
+			{
+				return true;
+			}
+		}
 
+		return false;
+	}
 };
 
-/**
- * Archive that stores a callstack for each of the Serialize calls and has the ability to compare itself to an existing
- * package on disk and dump all the differences to log.
- */
-class COREUOBJECT_API FArchiveStackTrace : public FLargeMemoryWriter
+/** Holds offsets to captured callstacks. */
+class COREUOBJECT_API FArchiveCallstacks
 {
+public:
 	/** Offset and callstack pair */
-	struct FCallstactAtOffset
+	struct FCallstackAtOffset
 	{
 		/** Offset of a Serialize call */
-		int64 Offset;
+		int64 Offset = -1;
 		/** Callstack CRC for the Serialize call */
-		uint32 Callstack;
+		uint32 Callstack = 0;
 		/** Collected inside of skip scope */
-		uint32 bIgnore : 1;
-
-		FCallstactAtOffset()
-			: Offset(-1)
-			, Callstack(0)
-			, bIgnore(false)
-		{}
-
-		FCallstactAtOffset(int64 InOffset, uint32 InCallstack, bool bInIgnore)
-			: Offset(InOffset)
-			, Callstack(InCallstack)
-			, bIgnore(bInIgnore)
-		{
-		}
+		bool bIgnore = false;
 	};
 
 	/** Struct to hold the actual Serialize call callstack and any associated data */
 	struct FCallstackData
 	{
 		/** Full callstack */
-		ANSICHAR* Callstack;
+		TUniquePtr<ANSICHAR[]> Callstack;
 		/** Full name of the currently serialized object */
 		FString SerializedObjectName;
 		/** The currently serialized property */
-		FProperty* SerializedProp;
+		FProperty* SerializedProp = nullptr;
 		/** Name of the currently serialized property */
 		FString SerializedPropertyName;
 
-		FCallstackData();
-		FCallstackData(ANSICHAR* InCallstack, UObject* InSerializedObject, FProperty* InSerializedProperty);
+		FCallstackData() = default;
+		FCallstackData(TUniquePtr<ANSICHAR[]>&& InCallstack, UObject* InSerializedObject, FProperty* InSerializedProperty);
+		FCallstackData(FCallstackData&&) = default;
+		FCallstackData(const FCallstackData&) = delete;
+
+		FCallstackData& operator=(FCallstackData&&) = default;
+		FCallstackData& operator=(const FCallstackData&) = delete;
 
 		/** Converts the callstack and associated data to human readable string */
 		FString ToString(const TCHAR* CallstackCutoffText) const;
+
+		/** Clone the callstack data */
+		FCallstackData Clone() const;
 	};
+	
+	explicit FArchiveCallstacks(UObject* InAsset);
+
+	/** Returns the asset class name. */
+	FName GetAssetClass() const;
+
+	/** Returns the total number of callstacks. */
+	int32 Num() const
+	{
+		return CallstackAtOffsetMap.Num();
+	}
+
+	/** Capture and append the current callstack. */
+	void Add(
+		int64 CurrentOffset,
+		int64 CurrentTotalSize,
+		UObject* SerializedObject,
+		FProperty* SerializedProperty,
+		TArrayView<const FName> DebugDataStack,
+		bool bSaveStackTrace,
+		int32 StackIgnoreCount);
+
+	/** Append other callstacks. */
+	void Append(const FArchiveCallstacks& Other, int64 Offset = 0);
+
+	/** Finds a callstack associated with data at the specified offset */
+	int32 GetCallstackIndexAtOffset(int64 Offset, int32 MinOffsetIndex = 0) const;
+
+	/** Finds a callstack associated with data at the specified offset */
+	const FCallstackAtOffset& GetCallstack(int32 CallstackIndex) const
+	{
+		return CallstackAtOffsetMap[CallstackIndex];
+	}
+	
+	const FCallstackData& GetCallstackData(const FCallstackAtOffset& CallstackOffset) const
+	{
+		return UniqueCallstacks[CallstackOffset.Callstack];
+	}
+
+	/** Returns the size of serialized data at the specified offset. */
+	int64 GetSerializedDataSizeForOffsetIndex(int32 InOffsetIndex) const
+	{
+		if (InOffsetIndex < CallstackAtOffsetMap.Num() - 1)
+		{
+			return CallstackAtOffsetMap[InOffsetIndex + 1].Offset - CallstackAtOffsetMap[InOffsetIndex].Offset;
+		}
+		else
+		{
+			return TotalSize - CallstackAtOffsetMap[InOffsetIndex].Offset;
+		}
+	}
+
+	/** Returns total serialized bytes. */
+	int64 TotalCapturedSize() const
+	{
+		return TotalSize;
+	}
+
+private:
+	/** Adds a unique callstack to UniqueCallstacks map */
+	ANSICHAR* AddUniqueCallstack(UObject* SerializedObject, FProperty* SerializedProperty, uint32& OutCallstackCRC);
 
 	/** The asset being serialized */
 	UObject* Asset;
-	/** The class of asset being serialized */
-	FName AssetClass;
-
 	/** List of offsets and their respective callstacks */
-	TArray<FCallstactAtOffset> CallstackAtOffsetMap;
+	TArray<FCallstackAtOffset> CallstackAtOffsetMap;
 	/** Contains all unique callstacks for all Serialize calls */
 	TMap<uint32, FCallstackData> UniqueCallstacks;
-	/** Contains offsets to gather callstacks for */
-	const FArchiveDiffMap* DiffMap;
-	/** If true the archive will callect callstacks for all offsets or for offsets specified in the DiffMap */
-	bool bCollectCallstacks;
 	/** Optimizes callstack comparison. If true the current and the last callstack should be compared as it may have changed */
 	bool bCallstacksDirty;
-
 	/** Maximum size of the stack trace */
 	const SIZE_T StackTraceSize;
 	/** Buffer for getting the current stack trace */
-	ANSICHAR* StackTrace;
+	TUniquePtr<ANSICHAR[]> StackTrace;
 	/** Callstack associated with the previous Serialize call */
 	ANSICHAR* LastSerializeCallstack;
-	/** Cached thread context */
-	FUObjectThreadContext& ThreadContext;
-	/** Current save/load context */
-	TRefCountPtr<FUObjectSerializeContext> LoadContext;
+	/** Total serialized bytes */
+	int64 TotalSize;
+};
 
+/** Archive proxy that captures callstacks for each serialize call. */
+class COREUOBJECT_API FArchiveStackTraceWriter
+	: public FArchiveProxy
+{
+public:
+	FArchiveStackTraceWriter(
+		FArchive& InInner,
+		FArchiveCallstacks& InCallstacks,
+		const FArchiveDiffMap* InDiffMap = nullptr,
+		int64 InDiffMapStartOffset = 0);
+
+	FORCENOINLINE virtual void Serialize(void* Data, int64 Length) override; // FORCENOINLINE so it can be counted during StackTrace
+	virtual void SetSerializeContext(FUObjectSerializeContext* Context) override;
+	virtual FUObjectSerializeContext* GetSerializeContext() override;
+	
 #if WITH_EDITOR
-	/** Current debug name stack */
-	TArray<FName> DebugDataStack;
-#endif
-
-	/** Adds a unique callstack to UniqueCallstacks map */
-	ANSICHAR* AddUniqueCallstack(UObject* InSerializedObject, FProperty* InSerializedProperty, uint32& OutCallstackCRC);
-
-	/** Finds a callstack associated with data at the specified offset */
-	int32 GetCallstackAtOffset(int64 InOffset, int32 MinOffsetIndex);
+	virtual void PushDebugDataString(const FName& DebugData) override
+	{
+		DebugDataStack.Push(DebugData);
+	}
 
-	int64 GetSerializedDataSizeForOffsetIndex(int32 InOffsetIndex)
+	virtual void PopDebugDataString() override
 	{
-		if (InOffsetIndex < CallstackAtOffsetMap.Num() - 1)
-		{
-			return CallstackAtOffsetMap[InOffsetIndex + 1].Offset - CallstackAtOffsetMap[InOffsetIndex].Offset;
-		}
-		else
-		{
-			return TotalSize() - CallstackAtOffsetMap[InOffsetIndex].Offset;
-		}
+		DebugDataStack.Pop();
 	}
+#endif
 
-	bool ShouldLogOffset(int64 InOffset) const
+	const FArchiveDiffMap& GetDiffMap() const
 	{
-		if (!DiffMap)
-		{
-			return true;
-		}
+		static FArchiveDiffMap Empty;
+		return DiffMap != nullptr ? *DiffMap : Empty;
+	}
 
-		for (const FArchiveDiffInfo& Diff : *DiffMap)
-		{
-			if (Diff.Offset <= InOffset && InOffset < (Diff.Offset + Diff.Size))
-			{
-				return true;
-			}
-		}
-		return false;
+	void SetDisableInnerArchive(bool bDisable)
+	{
+		bInnerArchiveDisabled = bDisable;
 	}
 
-public:
+	int32 GetStackIgnoreCount() const { return StackIgnoreCount; }
+	void SetStackIgnoreCount(const int32 IgnoreCount) { StackIgnoreCount = IgnoreCount; }
 
 	struct FPackageData
 	{
-		uint8* Data;
-		int64 Size;
-		int64 HeaderSize;
-		int64 StartOffset;
-
-		FPackageData()
-			: Size(0)
-			, HeaderSize(0)
-			, StartOffset(0)
-		{}
+		uint8* Data = nullptr;
+		int64 Size = 0;
+		int64 HeaderSize = 0;
+		int64 StartOffset = 0;
 	};
 
+	/** Compares two packages and logs the differences and calltacks. */
+	static void Compare(
+		const FPackageData& SourcePackage,
+		const FPackageData& DestPackage,
+		const FArchiveCallstacks& Callstacks,
+		const FArchiveDiffMap& DiffMap,
+		const TCHAR* AssetFilename,
+		const TCHAR* CallstackCutoffText,
+		const int64 MaxDiffsToLog,
+		int32& InOutDiffsLogged,
+		TMap<FName, FArchiveDiffStats>& OutStats);
+
+	/** Creates map with mismatching callstacks. */
+	static bool GenerateDiffMap(
+		const FPackageData& SourcePackage,
+		const FPackageData& DestPackage,
+		const FArchiveCallstacks& Callstacks,
+		int32 MaxDiffsToFind,
+		FArchiveDiffMap& OutDiffMap);
+
+	/** Logs any mismatching header data. */
+	static void DumpPackageHeaderDiffs(
+		const FPackageData& SourcePackage,
+		const FPackageData& DestPackage,
+		const FString& AssetFilename,
+		const int32 MaxDiffsToLog);
+
+	/** Returns a new linker for loading the specified package. */
+	static FLinkerLoad* CreateLinkerForPackage(
+		FUObjectSerializeContext* LoadContext,
+		const FString& InPackageName,
+		const FString& InFilename,
+		const FPackageData& PackageData);
+
 private:
+	FArchiveCallstacks& Callstacks;
+	const FArchiveDiffMap* DiffMap;
+	TRefCountPtr<FUObjectSerializeContext> SerializeContext;
+#if WITH_EDITOR
+	TArray<FName> DebugDataStack;
+#endif
+	int64 DiffMapOffset;
+	int32 StackIgnoreCount = 2;
+	bool bInnerArchiveDisabled;
+};
 
-	void CompareWithInternal(const FPackageData& SourcePackage, const FPackageData& DestPackage, const TCHAR* AssetFilename, const TCHAR* CallstackCutoffText, const int64 MaxDiffsToLog, int32& InOutDiffsLogged, TMap<FName, FArchiveDiffStats>& OutStats);
-	bool GenerateDiffMapInternal(const FPackageData& SourcePackage, const FPackageData& DestPackage, int32 MaxDiffsToFind, FArchiveDiffMap& OutDiffMap);
-	void DumpPackageHeaderDiffs(const FPackageData& SourcePackage, const FPackageData& DestPackage, const FString& AssetFilename, const int32 MaxDiffsToLog);
-	static FLinkerLoad* CreateLinkerForPackage(FUObjectSerializeContext* LoadContext, const FString& InPackageName, const FString& InFilename, const FPackageData& PackageData);
+/**
+ * Memory backed stack trace writer.
+ */
+class COREUOBJECT_API FArchiveStackTraceMemoryWriter final
+	: public FLargeMemoryWriter
+{
+public:
+	FArchiveStackTraceMemoryWriter(
+		FArchiveCallstacks& Callstacks,
+		const FArchiveDiffMap* DiffMap = nullptr,
+		const int64 DiffMapOffset = 0,
+		const int64 PreAllocateBytes = 0,
+		bool bIsPersistent = false,
+		const TCHAR* Filename = nullptr);
+
+	FORCENOINLINE virtual void Serialize(void* Memory, int64 Length) override; // FORCENOINLINE so it can be counted during StackTrace
+	virtual void SetSerializeContext(FUObjectSerializeContext* Context) override;
+	virtual FUObjectSerializeContext* GetSerializeContext() override;
 
+private:
+	FArchiveStackTraceWriter StackTraceWriter;
+};
+
+/**
+ * Archive that stores a callstack for each of the Serialize calls and has the ability to compare itself to an existing
+ * package on disk and dump all the differences to log.
+ */
+class COREUOBJECT_API FArchiveStackTrace
+	: public FLargeMemoryWriter
+{
 public:
+	using FPackageData = FArchiveStackTraceWriter::FPackageData;
 
 	FArchiveStackTrace(UObject* InAsset, const TCHAR* InFilename, bool bInCollectCallstacks = true, const FArchiveDiffMap* InDiffMap = nullptr);
 	virtual ~FArchiveStackTrace();
 
-	//~ Begin FArchive Interface
-	FORCENOINLINE virtual void Serialize(void* Data, int64 Num) override; // FORCENOINLINE so it can be counted during StackTrace
-	virtual void SetSerializeContext(FUObjectSerializeContext* InLoadContext) override;
-	virtual FUObjectSerializeContext* GetSerializeContext() override;
-#if WITH_EDITOR
-	virtual void PushDebugDataString(const FName& DebugData) override
+	FArchiveCallstacks& GetCallstacks()
 	{
-		DebugDataStack.Add(DebugData);
+		return Callstacks;
 	}
-	virtual void PopDebugDataString() override
+
+	const FArchiveCallstacks& GetCallstacks() const
 	{
-		DebugDataStack.Pop();
+		return Callstacks;
 	}
-#endif
-	//~ End FArchive Interface
+
+	FORCENOINLINE virtual void Serialize(void* Memory, int64 Length) override; // FORCENOINLINE so it can be counted during StackTrace
+	virtual void SetSerializeContext(FUObjectSerializeContext* Context) override;
+	virtual FUObjectSerializeContext* GetSerializeContext() override;
 
 	/** Compares this archive with the given bytes from disk or FPackageData. Dumps all differences to log. */
 	void CompareWith(const TCHAR* InFilename, const int64 TotalHeaderSize, const TCHAR* CallstackCutoffText,
@@ -265,6 +395,10 @@ public:
 	/** Helper function to load package contents into memory. Supports EDL packages. */
 	static bool LoadPackageIntoMemory(const TCHAR* InFilename, FPackageData& OutPackageData,
 		TUniquePtr<uint8>& OutLoadedBytes);
+
+private:
+	FArchiveCallstacks Callstacks;
+	FArchiveStackTraceWriter StackTraceWriter;
 };
 
 class COREUOBJECT_API FArchiveStackTraceReader : public FLargeMemoryReader
@@ -308,4 +442,4 @@ public:
 		return SerializeTrace;
 	}
 	static FArchiveStackTraceReader* CreateFromFile(const TCHAR* InFilename);
-};
\ No newline at end of file
+};
