Fix issue where garbage references referenced via UE::GC::RegisterSlowImplementation would report a random UObject as the referencer, rather than the owner of the AddReferencedObject function.

diff --git a/Engine/Source/Runtime/CoreUObject/Private/UObject/GarbageCollection.cpp b/Engine/Source/Runtime/CoreUObject/Private/UObject/GarbageCollection.cpp
--- a/Engine/Source/Runtime/CoreUObject/Private/UObject/GarbageCollection.cpp
+++ b/Engine/Source/Runtime/CoreUObject/Private/UObject/GarbageCollection.cpp
@@ -2061,8 +2061,9 @@ public:
 		return GetQueue(AROIdx, WorkerIdx).TryPush(Object);
 	}
 
-	FORCEINLINE_DEBUGGABLE void DrainLocalUnbalancedQueues(int32 WorkerIdx, FReferenceCollector& Collector)
+	FORCEINLINE_DEBUGGABLE void DrainLocalUnbalancedQueues(FWorkerContext& Context, FReferenceCollector& Collector)
 	{
+		const int32 WorkerIdx = Context.GetWorkerIndex();
 		for (uint32 AROIdx : UnbalancedAROIndices)
 		{
 			AROFunc ARO = AROs[AROIdx];
@@ -2071,13 +2072,14 @@ public:
 			{
 				for (UObject* Object : Objects)
 				{
+					Context.ReferencingObject = Object;
 					ARO(Object, Collector);
 				}
 			}
 		}
 	}
 
-	bool ProcessAllQueues(int32 WorkerIdx, FReferenceCollector& Collector)
+	bool ProcessAllQueues(FWorkerContext& Context, FReferenceCollector& Collector)
 	{
 		// Drain all queues for one specific slow ARO type at a time and stop if any ARO calls were made
 		// Workers start with different ARO types to reach the slowest ARO type sooner and reduce stealing contention
@@ -2087,12 +2089,12 @@ public:
 		const uint32 NumAROs = (uint32)AROs.Num();
 		
 		const double StopTime = FPlatformTime::Seconds() + 0.1/1000;
-		for (uint32 OffsetIdx = (uint32)WorkerIdx % NumAROs, EndOffsetIdx = OffsetIdx + NumAROs; OffsetIdx < EndOffsetIdx && !bStop; ++OffsetIdx)
+		for (uint32 OffsetIdx = (uint32)Context.GetWorkerIndex() % NumAROs, EndOffsetIdx = OffsetIdx + NumAROs; OffsetIdx < EndOffsetIdx && !bStop; ++OffsetIdx)
 		{
 			int32 AROIdx = OffsetIdx - (OffsetIdx >= NumAROs ? NumAROs : 0);
 			TArrayView<FAROQueue> WorkerQueues = MakeArrayView(&GetQueue(AROIdx, 0), NumWorkers);
-			bStop = IsExtraSlow(AROIdx)	? ProcessQueues<FAROQueue::FFew >(WorkerIdx, WorkerQueues, AROs[AROIdx], Collector, NumCalls)
-										: ProcessQueues<FAROQueue::FMany>(WorkerIdx, WorkerQueues, AROs[AROIdx], Collector, NumCalls);
+			bStop = IsExtraSlow(AROIdx)	? ProcessQueues<FAROQueue::FFew >(Context, WorkerQueues, AROs[AROIdx], Collector, NumCalls)
+										: ProcessQueues<FAROQueue::FMany>(Context, WorkerQueues, AROs[AROIdx], Collector, NumCalls);
 
 			// Don't proceed to next ARO type if we did some work that took a considerable amount of time
 			bStop = bStop || (NumCalls > 0 && FPlatformTime::Seconds() > StopTime);
@@ -2127,18 +2129,19 @@ private:
 
 	// @return if processing should stop since NumCalls reached CallLimit
 	template<class Quantity>
-	static bool ProcessQueues(int32 WorkerIdx, TArrayView<FAROQueue> WorkerQueues, AROFunc ARO, FReferenceCollector& Collector, uint32& NumCalls)
+	static bool ProcessQueues(FWorkerContext& Context, TArrayView<FAROQueue> WorkerQueues, AROFunc ARO, FReferenceCollector& Collector, uint32& NumCalls)
 	{
 		static constexpr uint32 CallLimit = FAROQueue::NumPop<Quantity>;
+		const int32 WorkerIdx = Context.GetWorkerIndex();
 
-		return	ProcessQueuesUsing<&FAROQueue::Pop<Quantity>,	CallLimit>(WorkerQueues.Slice(WorkerIdx, 1),		ARO, Collector, NumCalls) ||
-				ProcessQueuesUsing<&FAROQueue::Steal<Quantity>,	CallLimit>(WorkerQueues.RightChop(WorkerIdx + 1),	ARO, Collector, NumCalls) ||
-				ProcessQueuesUsing<&FAROQueue::Steal<Quantity>,	CallLimit>(WorkerQueues.Slice(0, WorkerIdx),		ARO, Collector, NumCalls);
+		return	ProcessQueuesUsing<&FAROQueue::Pop<Quantity>,	CallLimit>(Context, WorkerQueues.Slice(WorkerIdx, 1),		ARO, Collector, NumCalls) ||
+				ProcessQueuesUsing<&FAROQueue::Steal<Quantity>,	CallLimit>(Context, WorkerQueues.RightChop(WorkerIdx + 1),	ARO, Collector, NumCalls) ||
+				ProcessQueuesUsing<&FAROQueue::Steal<Quantity>,	CallLimit>(Context, WorkerQueues.Slice(0, WorkerIdx),		ARO, Collector, NumCalls);
 	}
 
 	// @return if processing should stop since NumCalls reached or exceeded CallLimit
 	template<TArrayView<UObject*>(FAROQueue::*PopFunc)(), uint32 CallLimit>
-	FORCEINLINE_DEBUGGABLE static bool ProcessQueuesUsing(TArrayView<FAROQueue> Queues, AROFunc ARO, FReferenceCollector& Collector, uint32& NumCalls)
+	FORCEINLINE_DEBUGGABLE static bool ProcessQueuesUsing(FWorkerContext& Context, TArrayView<FAROQueue> Queues, AROFunc ARO, FReferenceCollector& Collector, uint32& NumCalls)
 	{
 		for (FAROQueue& Queue : Queues)
 		{
@@ -2146,6 +2149,7 @@ private:
 			{
 				for (UObject* Object : Objects)
 				{
+					Context.ReferencingObject = Object;
 					ARO(Object, Collector);
 				}
 
@@ -2221,11 +2225,11 @@ bool FSlowARO::TryQueueCall(uint32 SlowAROIndex, UObject* Object, FWorkerContext
 
 void FSlowARO::ProcessUnbalancedCalls(FWorkerContext& Context, FReferenceCollector& Collector)
 {
-	GSlowARO.GetPostInit().DrainLocalUnbalancedQueues(Context.GetWorkerIndex(), Collector);
+	GSlowARO.GetPostInit().DrainLocalUnbalancedQueues(Context, Collector);
 }
 bool FSlowARO::ProcessAllCalls(FWorkerContext& Context, FReferenceCollector& Collector)
 {
-	return GSlowARO.GetPostInit().ProcessAllQueues(Context.GetWorkerIndex(), Collector);
+	return GSlowARO.GetPostInit().ProcessAllQueues(Context, Collector);
 }
 
 void RegisterSlowImplementation(AROFunc ARO, EAROFlags Flags)
diff --git a/Engine/Source/Runtime/CoreUObject/Public/UObject/FastReferenceCollector.h b/Engine/Source/Runtime/CoreUObject/Public/UObject/FastReferenceCollector.h
--- a/Engine/Source/Runtime/CoreUObject/Public/UObject/FastReferenceCollector.h
+++ b/Engine/Source/Runtime/CoreUObject/Public/UObject/FastReferenceCollector.h
@@ -221,6 +221,7 @@ struct alignas(PLATFORM_CACHE_LINE_SIZE) FWorkerContext
 private:
 	template <typename ProcessorType, typename CollectorType>
 	friend class TFastReferenceCollector;
+	friend class FSlowAROManager;
 	
 	// This is set by GC when processing references from the current referencing object
 	UObject* ReferencingObject = nullptr;
@@ -536,9 +537,10 @@ public:
 
 StoleContext:
 		// Process initial references first
+		Context.ReferencingObject = FGCObject::GGCObjectReferencer;
 		for (UObject** InitialReference : Context.InitialNativeReferences)
 		{
-			Dispatcher.HandleKillableReference(*InitialReference, ETokenlessId::Collector, EGCTokenType::Native);
+			Dispatcher.HandleKillableReference(*InitialReference, ETokenlessId::InitialReference, EGCTokenType::Native);
 		}
 
 		TConstArrayView<UObject*> CurrentObjects = Context.InitialObjects;
