Fix some cast/truncation warnings in the async loading code

#rnx

- Changed a ranged for loop in FEventLoadGraph::RemoveNode to be index based to avoid pointer arithmetic to index casting.
- FAsyncPackage::PostLoadDeferredObjects and FAsyncPackage::CreateClusters were not modifying the parameter InOutTimeLimit that was being passed in by reference, so it has been changed to be passed by value to avoid confusion.
- Changed a number of methods that were taking float values to doubles, in most cases double variables were being passed into these methods and they were being compared to doubles inside of the methods, so there didn't seem to be much point in the parameter type being float.

### Double->Float explicit cast
- Need to do this when using the IAsyncReadRequest api.
- Need to do this when using FLinkerLoad
- FAsyncLoadingThread::TickAsyncThread otherwise we'd need to change FAsyncLoadingThread::ProcessAsyncLoading to take doubles, which in turn would require a large number of double->float casting when calling out to the rest of the engine.

diff --git a/Engine/Source/Runtime/CoreUObject/Private/Serialization/AsyncLoading.cpp b/Engine/Source/Runtime/CoreUObject/Private/Serialization/AsyncLoading.cpp
--- a/Engine/Source/Runtime/CoreUObject/Private/Serialization/AsyncLoading.cpp
+++ b/Engine/Source/Runtime/CoreUObject/Private/Serialization/AsyncLoading.cpp
@@ -674,7 +674,7 @@ uint32 FAsyncLoadingThread::AsyncLoadingThreadID = 0;
 FThreadSafeCounter FAsyncLoadingThread::AsyncLoadingTickCounter;
 int32 FAsyncLoadingThread::CurrentAsyncLoadingTickThreadIndex = INDEX_NONE;
 
-static FORCEINLINE bool IsTimeLimitExceeded(double InTickStartTime, bool bUseTimeLimit, float InTimeLimit, const TCHAR* InLastTypeOfWorkPerformed = nullptr, UObject* InLastObjectWorkWasPerformedOn = nullptr)
+static FORCEINLINE bool IsTimeLimitExceeded(double InTickStartTime, bool bUseTimeLimit, double InTimeLimit, const TCHAR* InLastTypeOfWorkPerformed = nullptr, UObject* InLastObjectWorkWasPerformedOn = nullptr)
 {
 	static double LastTestTime = -1.0;
 	bool bTimeLimitExceeded = false;
@@ -935,7 +935,7 @@ void FAsyncLoadingThread::ProcessAsyncPackageRequest(FAsyncPackageDesc* InReques
 	}
 }
 
-int32 FAsyncLoadingThread::CreateAsyncPackagesFromQueue(bool bUseTimeLimit, bool bUseFullTimeLimit, float TimeLimit)
+int32 FAsyncLoadingThread::CreateAsyncPackagesFromQueue(bool bUseTimeLimit, bool bUseFullTimeLimit, double TimeLimit)
 {
 	SCOPED_LOADTIMER(CreateAsyncPackagesFromQueueTime);
 
@@ -1229,15 +1229,14 @@ void FEventLoadGraph::RemoveNode(FEventLoadNodePtr& InNodeToRemove)
 		check(PrereqisiteNode->NumPrerequistes == 0);
 		Swap(NodesToFire, PrereqisiteNode->NodesWaitingForMe);
 
-
-		for (FEventLoadNodePtr& Target : NodesToFire)
+		for (int32 Index = 0; Index < NodesToFire.Num(); ++Index)
 		{
-			FEventLoadNode* DependentNode(&GetNode(Target));
+			FEventLoadNode* DependentNode(&GetNode(NodesToFire[Index]));
 			check(DependentNode->NumPrerequistes > 0);
 			if (!--DependentNode->NumPrerequistes)
 			{
 				DependentNode->bFired = true;
-				IndicesToFire.Add(&Target - &NodesToFire[0]);
+				IndicesToFire.Add(Index);
 			}
 		}
 		if (Array.RemoveNode(NodeToRemove))
@@ -4715,11 +4714,11 @@ bool FAsyncLoadingThread::ShouldProcessPackage(FAsyncPackage* InAsyncPackage, co
 	return InAsyncPackage->IsDependencyOf(FlushRequest.GetId());
 }
 
-EAsyncPackageState::Type FAsyncLoadingThread::ProcessLoadedPackages(bool bUseTimeLimit, bool bUseFullTimeLimit, float TimeLimit, bool& bDidSomething, const FFlushRequest& FlushRequest)
+EAsyncPackageState::Type FAsyncLoadingThread::ProcessLoadedPackages(bool bUseTimeLimit, bool bUseFullTimeLimit, double TimeLimit, bool& bDidSomething, const FFlushRequest& FlushRequest)
 {
 	EAsyncPackageState::Type Result = EAsyncPackageState::Complete;
 
-	double TickStartTime = FPlatformTime::Seconds();
+	const double TickStartTime = FPlatformTime::Seconds();
 
 	{
 #if THREADSAFE_UOBJECTS
@@ -4952,7 +4951,7 @@ EAsyncPackageState::Type FAsyncLoadingThread::ProcessLoadedPackages(bool bUseTim
 
 
 
-EAsyncPackageState::Type FAsyncLoadingThread::TickAsyncLoading(bool bUseTimeLimit, bool bUseFullTimeLimit, float TimeLimit, FFlushRequest FlushRequest)
+EAsyncPackageState::Type FAsyncLoadingThread::TickAsyncLoading(bool bUseTimeLimit, bool bUseFullTimeLimit, double TimeLimit, FFlushRequest FlushRequest)
 {
 	LLM_SCOPE(ELLMTag::AsyncLoading);
 
@@ -4970,8 +4969,8 @@ EAsyncPackageState::Type FAsyncLoadingThread::TickAsyncLoading(bool bUseTimeLimi
 
 	if (!bLoadingSuspended)
 	{
-		double TickStartTime = FPlatformTime::Seconds();
-		double TimeLimitUsedForProcessLoaded = 0;
+		const double TickStartTime = FPlatformTime::Seconds();
+		double TimeLimitUsedForProcessLoaded = 0.0;
 
 		// First make sure there's no objects pending to be unhashed. This is important in uncooked builds since we don't 
 		// detach linkers immediately and we may otherwise end up getting unreachable objects from Linkers in CreateImports
@@ -4999,7 +4998,7 @@ EAsyncPackageState::Type FAsyncLoadingThread::TickAsyncLoading(bool bUseTimeLimi
 
 		if (!bIsMultithreaded && Result != EAsyncPackageState::TimeOut && !IsTimeLimitExceeded(TickStartTime, bUseTimeLimit, TimeLimit, TEXT("ProcessLoadedPackages")))
 		{
-			double RemainingTimeLimit = FMath::Max(0.0, TimeLimit - TimeLimitUsedForProcessLoaded);
+			const double RemainingTimeLimit = FMath::Max(0.0, TimeLimit - TimeLimitUsedForProcessLoaded);
 			AdjustFlushRequest(FlushRequest);
 			Result = TickAsyncThread(bUseTimeLimit, bUseFullTimeLimit, RemainingTimeLimit, bDidSomething, FlushRequest);
 		}
@@ -5206,7 +5205,7 @@ uint32 FAsyncLoadingThread::Run()
 
 				// Peek if we currently have a flush request going on.
 				FFlushRequest FlushRequest = PeekFlushRequest();
-				TickAsyncThread(true, false, 0.033f, bDidSomething, FlushRequest);
+				TickAsyncThread(true, false, 0.033, bDidSomething, FlushRequest);
 			}
 		}
 		else if (!bWasSuspendedLastFrame)
@@ -5253,7 +5252,7 @@ void FAsyncLoadingThread::CheckForCycles()
 }
 
 
-EAsyncPackageState::Type FAsyncLoadingThread::TickAsyncThread(bool bUseTimeLimit, bool bUseFullTimeLimit, float TimeLimit, bool& bDidSomething, FFlushRequest& FlushRequest)
+EAsyncPackageState::Type FAsyncLoadingThread::TickAsyncThread(bool bUseTimeLimit, bool bUseFullTimeLimit, double TimeLimit, bool& bDidSomething, FFlushRequest& FlushRequest)
 {
 	check(!IsInGameThread() || !IsMultithreaded());
 	EAsyncPackageState::Type Result = EAsyncPackageState::Complete;
@@ -5273,7 +5272,7 @@ EAsyncPackageState::Type FAsyncLoadingThread::TickAsyncThread(bool bUseTimeLimit
 				CreateAsyncPackagesFromQueue(bUseTimeLimit, bUseFullTimeLimit, TimeLimit); 
 			}
 			float TimeUsed = (float)(FPlatformTime::Seconds() - TickStartTime);
-			const float RemainingTimeLimit = FMath::Max(0.0f, TimeLimit - TimeUsed);
+			const float RemainingTimeLimit = FMath::Max(0.0f, (float)TimeLimit - TimeUsed);
 			if (IsGarbageCollectionWaiting() || (RemainingTimeLimit <= 0.0f && bUseTimeLimit && !IsMultithreaded()))
 			{
 				Result = EAsyncPackageState::TimeOut;
@@ -6890,7 +6889,7 @@ EAsyncPackageState::Type FAsyncPackage::PostLoadObjects()
 	TOptional<FScopedSlowTask> SlowTask;
 	if (IsInGameThread())
 	{
-		SlowTask.Emplace(1, NSLOCTEXT("AsyncLoading", "PostloadingObjects", "Postloading objects..."));
+		SlowTask.Emplace(1.0f, NSLOCTEXT("AsyncLoading", "PostloadingObjects", "Postloading objects..."));
 	}
 
 	// PostLoad objects.
@@ -6956,7 +6955,7 @@ EAsyncPackageState::Type FAsyncPackage::PostLoadObjects()
 
 void CreateClustersFromPackage(FLinkerLoad* PackageLinker, TArray<UObject*>& OutClusterObjects);
 
-EAsyncPackageState::Type FAsyncPackage::PostLoadDeferredObjects(double InTickStartTime, bool bInUseTimeLimit, float& InOutTimeLimit)
+EAsyncPackageState::Type FAsyncPackage::PostLoadDeferredObjects(double InTickStartTime, bool bInUseTimeLimit, double InOutTimeLimit)
 {
 	SCOPE_CYCLE_COUNTER(STAT_FAsyncPackage_PostLoadObjectsGameThread);
 	SCOPED_LOADTIMER(PostLoadDeferredObjectsTime);
@@ -7111,7 +7110,7 @@ EAsyncPackageState::Type FAsyncPackage::PostLoadDeferredObjects(double InTickSta
 			LastTypeOfWorkPerformed = TEXT("CreateClustersFromPackage");
 			LinkerRoot->AtomicallyClearInternalFlags(EInternalObjectFlags::AsyncLoading);
 			LinkerRoot->MarkAsFullyLoaded();			
-			LinkerRoot->SetLoadTime(FPlatformTime::Seconds() - LoadStartTime);
+			LinkerRoot->SetLoadTime((float)(FPlatformTime::Seconds() - LoadStartTime));
 
 			if (Linker)
 			{
@@ -7128,7 +7127,7 @@ EAsyncPackageState::Type FAsyncPackage::PostLoadDeferredObjects(double InTickSta
 	return Result;
 }
 
-EAsyncPackageState::Type FAsyncPackage::CreateClusters(double InTickStartTime, bool bInUseTimeLimit, float& InOutTimeLimit)
+EAsyncPackageState::Type FAsyncPackage::CreateClusters(double InTickStartTime, bool bInUseTimeLimit, double InOutTimeLimit)
 {
 	LastObjectWorkWasPerformedOn = nullptr;
 	LastTypeOfWorkPerformed = TEXT("CreateClusters");
@@ -7933,11 +7932,11 @@ bool FAsyncArchive::WaitRead(double TimeLimit)
 	if (ReadRequestPtr)
 	{
 		QUICK_SCOPE_CYCLE_COUNTER(STAT_FArchiveAsync2_WaitRead);
-		int64 Offset = ReadRequestOffset;
-		int64 Size = ReadRequestSize;
+		const int64 Offset = ReadRequestOffset;
+		const int64 Size = ReadRequestSize;
 		check(Size > 0);
-		double StartTime = FPlatformTime::Seconds();
-		bool bResult = ReadRequestPtr->WaitCompletion(TimeLimit);
+		const double StartTime = FPlatformTime::Seconds();
+		const bool bResult = ReadRequestPtr->WaitCompletion((float)TimeLimit);
 		LogItem(TEXT("Wait Read"), Offset, Size, StartTime);
 		if (!bResult)
 		{
@@ -8029,14 +8028,14 @@ bool FAsyncArchive::WaitForIntialPhases(double InTimeLimit)
 		double StartTime = FPlatformTime::Seconds();
 		if (SizeRequestPtr)
 		{
-			if (SizeRequestPtr->WaitCompletion(InTimeLimit))
+			if (SizeRequestPtr->WaitCompletion((float)InTimeLimit))
 			{
 				delete SizeRequestPtr;
 				SizeRequestPtr = nullptr;
 			}
 			else
 			{
-				check(InTimeLimit > 0.0f);
+				check(InTimeLimit > 0.0);
 				return false;
 			}
 		}
@@ -8053,7 +8052,7 @@ bool FAsyncArchive::WaitForIntialPhases(double InTimeLimit)
 						return false;
 					}
 				}
-				if (SummaryRequestPtr->WaitCompletion(TimeLimit))
+				if (SummaryRequestPtr->WaitCompletion((float)TimeLimit))
 				{
 					delete SummaryRequestPtr;
 					SummaryRequestPtr = nullptr;
@@ -8075,7 +8074,7 @@ bool FAsyncArchive::WaitForIntialPhases(double InTimeLimit)
 						return false;
 					}
 				}
-				if (SummaryPrecacheRequestPtr->WaitCompletion(TimeLimit))
+				if (SummaryPrecacheRequestPtr->WaitCompletion((float)TimeLimit))
 				{
 					delete SummaryPrecacheRequestPtr;
 					SummaryPrecacheRequestPtr = nullptr;
diff --git a/Engine/Source/Runtime/CoreUObject/Private/Serialization/AsyncLoading.h b/Engine/Source/Runtime/CoreUObject/Private/Serialization/AsyncLoading.h
--- a/Engine/Source/Runtime/CoreUObject/Private/Serialization/AsyncLoading.h
+++ b/Engine/Source/Runtime/CoreUObject/Private/Serialization/AsyncLoading.h
@@ -512,7 +512,7 @@ struct FAsyncPackage : public FGCObject
 	}
 
 	/** Creates GC clusters from loaded objects */
-	EAsyncPackageState::Type CreateClusters(double InTickStartTime, bool bInUseTimeLimit, float& InOutTimeLimit);
+	EAsyncPackageState::Type CreateClusters(double InTickStartTime, bool bInUseTimeLimit, double InOutTimeLimit);
 
 	/** Returns true if this package is a dependency to the explicit request id RequestID. */
 	bool IsDependencyOf(int32 RequestID);
@@ -828,7 +828,7 @@ public:
 	*
 	* @return true if we finished calling PostLoad on all loaded objects and no new ones were created, false otherwise
 	*/
-	EAsyncPackageState::Type PostLoadDeferredObjects(double InTickStartTime, bool bInUseTimeLimit, float& InOutTimeLimit);
+	EAsyncPackageState::Type PostLoadDeferredObjects(double InTickStartTime, bool bInUseTimeLimit, double InOutTimeLimit);
 
 	/** Close any linkers that have been open as a result of synchronous load during async loading */
 	void CloseDelayedLinkers();
diff --git a/Engine/Source/Runtime/CoreUObject/Private/Serialization/AsyncLoading2.cpp b/Engine/Source/Runtime/CoreUObject/Private/Serialization/AsyncLoading2.cpp
--- a/Engine/Source/Runtime/CoreUObject/Private/Serialization/AsyncLoading2.cpp
+++ b/Engine/Source/Runtime/CoreUObject/Private/Serialization/AsyncLoading2.cpp
@@ -942,7 +942,7 @@ public:
 		return ScriptObjects.Num();
 	}
 
-	uint32 GetStoredScriptObjectsAllocatedSize() const
+	SIZE_T GetStoredScriptObjectsAllocatedSize() const
 	{
 		return ScriptObjects.GetAllocatedSize();
 	}
@@ -3104,7 +3104,7 @@ public:
 	* @param FlushTree Package dependency tree to be flushed
 	* @return The current state of async loading
 	*/
-	EAsyncPackageState::Type TickAsyncLoadingFromGameThread(FAsyncLoadingThreadState2& ThreadState, bool bUseTimeLimit, bool bUseFullTimeLimit, float TimeLimit, int32 FlushRequestID, bool& bDidSomething);
+	EAsyncPackageState::Type TickAsyncLoadingFromGameThread(FAsyncLoadingThreadState2& ThreadState, bool bUseTimeLimit, bool bUseFullTimeLimit, double TimeLimit, int32 FlushRequestID, bool& bDidSomething);
 
 	/**
 	* [ASYNC THREAD] Main thread loop
@@ -3289,7 +3289,7 @@ private:
 			const TArray<FEventLoadNode2*>& Nodes = It.Value();
 			for (const FEventLoadNode2* Node : Nodes)
 			{
-				int32 NodeContextId = Node->GetSyncLoadContextId();
+				const uint64 NodeContextId = Node->GetSyncLoadContextId();
 				if (NodeContextId >= SyncLoadContextId)
 				{
 					Class = CurrentClass;
@@ -4876,7 +4876,7 @@ static void ReadAsyncPackageHeader(FAsyncPackageSerializationState& Serializatio
 	const uint64 ArcsDataOffset = ExportBundleHeadersOffset + ExportBundleHeadersSize;
 	const uint64 ArcsDataSize = PackageSummary->HeaderSize - ArcsDataOffset;
 	check(ExportBundleEntriesSize == PackageSummary->GraphDataOffset - PackageSummary->ExportBundleEntriesOffset);
-	HeaderData.ArcsData = MakeArrayView(PackageHeaderDataPtr + ArcsDataOffset, ArcsDataSize);
+	HeaderData.ArcsData = TArrayView64<const uint8>(PackageHeaderDataPtr + ArcsDataOffset, ArcsDataSize);
 	FMemory::Memcpy(HeaderData.ExportBundleHeaders.GetData(), PackageHeaderDataPtr + ExportBundleHeadersOffset, ExportBundleHeadersSize);
 	FMemory::Memcpy(HeaderData.ExportBundleEntries.GetData(), PackageHeaderDataPtr + PackageSummary->ExportBundleEntriesOffset, ExportBundleEntriesSize);
 
@@ -6963,7 +6963,7 @@ EAsyncPackageState::Type FAsyncLoadingThread2::ProcessLoadedPackagesFromGameThre
 	return Result;
 }
 
-EAsyncPackageState::Type FAsyncLoadingThread2::TickAsyncLoadingFromGameThread(FAsyncLoadingThreadState2& ThreadState, bool bUseTimeLimit, bool bUseFullTimeLimit, float TimeLimit, int32 FlushRequestID, bool& bDidSomething)
+EAsyncPackageState::Type FAsyncLoadingThread2::TickAsyncLoadingFromGameThread(FAsyncLoadingThreadState2& ThreadState, bool bUseTimeLimit, bool bUseFullTimeLimit, double TimeLimit, int32 FlushRequestID, bool& bDidSomething)
 {
 	SCOPE_CYCLE_COUNTER(STAT_FAsyncPackage_TickAsyncLoadingGameThread);
 	//TRACE_INT_VALUE(QueuedPackagesCounter, QueuedPackagesCounter);
@@ -8012,7 +8012,7 @@ void FAsyncPackage2::FinishUPackage()
 	{
 		// Mark package as having been fully loaded and update load time.
 		LinkerRoot->MarkAsFullyLoaded();
-		LinkerRoot->SetLoadTime(FPlatformTime::Seconds() - LoadStartTime);
+		LinkerRoot->SetLoadTime((float)(FPlatformTime::Seconds() - LoadStartTime));
 	}
 }
 
@@ -8213,7 +8213,7 @@ void FAsyncLoadingThread2::FlushLoading(int32 RequestId)
 			while (IsAsyncLoadingPackages())
 			{
 				bool bDidSomething = false;
-				EAsyncPackageState::Type Result = TickAsyncLoadingFromGameThread(*GameThreadState, false, false, 0, RequestId, bDidSomething);
+				EAsyncPackageState::Type Result = TickAsyncLoadingFromGameThread(*GameThreadState, false, false, 0.0, RequestId, bDidSomething);
 				if (RequestId != INDEX_NONE && !ContainsRequestID(RequestId))
 				{
 					break;
diff --git a/Engine/Source/Runtime/CoreUObject/Private/Serialization/AsyncLoadingThread.h b/Engine/Source/Runtime/CoreUObject/Private/Serialization/AsyncLoadingThread.h
--- a/Engine/Source/Runtime/CoreUObject/Private/Serialization/AsyncLoadingThread.h
+++ b/Engine/Source/Runtime/CoreUObject/Private/Serialization/AsyncLoadingThread.h
@@ -560,7 +560,7 @@ public:
 	* @param FlushRequest The request to flush a package and its depedencies. May be invalid.
 	* @return The current state of async loading
 	*/
-	EAsyncPackageState::Type TickAsyncLoading(bool bUseTimeLimit, bool bUseFullTimeLimit, float TimeLimit, FFlushRequest FlushRequest = FFlushRequest());
+	EAsyncPackageState::Type TickAsyncLoading(bool bUseTimeLimit, bool bUseFullTimeLimit, double TimeLimit, FFlushRequest FlushRequest = FFlushRequest());
 
 	/**
 	* [ASYNC THREAD] Main thread loop
@@ -570,7 +570,7 @@ public:
 	* @param TimeLimit Maximum amount of time that can be spent in this call [time-slicing].
 	* @param FlushRequest The request to flush a package and its depedencies. May be invalid.
 	*/
-	EAsyncPackageState::Type TickAsyncThread(bool bUseTimeLimit, bool bUseFullTimeLimit, float TimeLimit, bool& bDidSomething, FFlushRequest& FlushRequest);
+	EAsyncPackageState::Type TickAsyncThread(bool bUseTimeLimit, bool bUseFullTimeLimit, double TimeLimit, bool& bDidSomething, FFlushRequest& FlushRequest);
 
 	/** Initializes async loading thread */
 	void InitializeLoading() override;
@@ -686,12 +686,12 @@ private:
 	* @param FlushRequest Package dependency tree to be flushed
 	* @return The current state of async loading
 	*/
-	EAsyncPackageState::Type ProcessLoadedPackages(bool bUseTimeLimit, bool bUseFullTimeLimit, float TimeLimit, bool& bDidSomething, const FFlushRequest& FlushRequest);
+	EAsyncPackageState::Type ProcessLoadedPackages(bool bUseTimeLimit, bool bUseFullTimeLimit, double TimeLimit, bool& bDidSomething, const FFlushRequest& FlushRequest);
 
 	/**
 	* [ASYNC THREAD] Creates async packages from the queued requests
 	*/
-	int32 CreateAsyncPackagesFromQueue(bool bUseTimeLimit, bool bUseFullTimeLimit, float TimeLimit);
+	int32 CreateAsyncPackagesFromQueue(bool bUseTimeLimit, bool bUseFullTimeLimit, double TimeLimit);
 
 	/**
 	* [ASYNC THREAD] Internal helper function for processing a package load request. If dependency preloading is enabled, 
diff --git a/Engine/Source/Runtime/CoreUObject/Private/Serialization/AsyncPackageLoader.cpp b/Engine/Source/Runtime/CoreUObject/Private/Serialization/AsyncPackageLoader.cpp
--- a/Engine/Source/Runtime/CoreUObject/Private/Serialization/AsyncPackageLoader.cpp
+++ b/Engine/Source/Runtime/CoreUObject/Private/Serialization/AsyncPackageLoader.cpp
@@ -722,7 +722,7 @@ void IsTimeLimitExceededPrint(
 	double InTickStartTime,
 	double CurrentTime,
 	double LastTestTime,
-	float InTimeLimit, 
+	double InTimeLimit,
 	const TCHAR* InLastTypeOfWorkPerformed,
 	UObject* InLastObjectWorkWasPerformedOn)
 {
diff --git a/Engine/Source/Runtime/CoreUObject/Public/Serialization/AsyncPackageLoader.h b/Engine/Source/Runtime/CoreUObject/Public/Serialization/AsyncPackageLoader.h
--- a/Engine/Source/Runtime/CoreUObject/Public/Serialization/AsyncPackageLoader.h
+++ b/Engine/Source/Runtime/CoreUObject/Public/Serialization/AsyncPackageLoader.h
@@ -248,7 +248,7 @@ void IsTimeLimitExceededPrint(
 	double InTickStartTime,
 	double CurrentTime,
 	double LastTestTime,
-	float InTimeLimit,
+	double InTimeLimit,
 	const TCHAR* InLastTypeOfWorkPerformed = nullptr,
 	UObject* InLastObjectWorkWasPerformedOn = nullptr);
 
