Add missing include for CL 23787796

diff --git a/Engine/Source/Runtime/CoreUObject/Public/UObject/FastReferenceCollector.h b/Engine/Source/Runtime/CoreUObject/Public/UObject/FastReferenceCollector.h
--- a/Engine/Source/Runtime/CoreUObject/Public/UObject/FastReferenceCollector.h
+++ b/Engine/Source/Runtime/CoreUObject/Public/UObject/FastReferenceCollector.h
@@ -15,6 +15,7 @@
 #include "Async/ParallelFor.h"
 #include "UObject/UObjectArray.h"
 #include "UObject/DynamicallyTypedValue.h"
+#include "UObject/GCObject.h"
 
 /** Token stream stack overflow checks are not enabled in test and shipping configs */
 #define UE_ENABLE_TOKENSTREAM_STACKOVERFLOW_CHECKS (!(UE_BUILD_TEST || UE_BUILD_SHIPPING) || 0)
