Logging: Added FAssetLog and converted UE_ASSET_LOG to use structured logging

FAssetLog can be used to log an object path or package path in a structured log statement like UE_SLOG.

#rnx

diff --git a/Engine/Source/Runtime/CoreUObject/Private/UObject/UObjectGlobals.cpp b/Engine/Source/Runtime/CoreUObject/Private/UObject/UObjectGlobals.cpp
--- a/Engine/Source/Runtime/CoreUObject/Private/UObject/UObjectGlobals.cpp
+++ b/Engine/Source/Runtime/CoreUObject/Private/UObject/UObjectGlobals.cpp
@@ -7,6 +7,7 @@
 #include "UObject/UObjectGlobals.h"
 #include "HAL/PlatformFileManager.h"
 #include "HAL/FileManager.h"
+#include "Logging/StructuredLog.h"
 #include "Misc/AsciiSet.h"
 #include "Misc/PackageAccessTrackingOps.h"
 #include "Misc/Paths.h"
@@ -38,6 +39,7 @@
 #include "UObject/UnrealType.h"
 #include "UObject/ObjectRedirector.h"
 #include "UObject/PackageResourceManager.h"
+#include "Serialization/CompactBinaryWriter.h"
 #include "Serialization/DuplicatedObject.h"
 #include "Serialization/DuplicatedDataReader.h"
 #include "Serialization/DuplicatedDataWriter.h"
@@ -5313,6 +5315,119 @@ COREUOBJECT_API UFunction* FindDelegateSignature(FName DelegateSignatureName)
 	return nullptr;
 }
 
+void UE::SerializeForLog(FCbWriter& Writer, const FAssetLog& AssetLog)
+{
+	FString ObjectPath;
+	FString LocalPath;
+
+	if (AssetLog.Path)
+	{
+		ObjectPath = AssetLog.Path;
+		FString PackageName, ObjectName, SubObjectName, Extension;
+		if (FPackageName::TryConvertToMountedPath(ObjectPath, &LocalPath, &PackageName, &ObjectName, &SubObjectName, &Extension))
+		{
+			ObjectPath = PackageName;
+			if (!ObjectName.IsEmpty())
+			{
+				ObjectPath += TEXT('.');
+				ObjectPath += ObjectName;
+			}
+			if (!SubObjectName.IsEmpty())
+			{
+				ObjectPath += SUBOBJECT_DELIMITER;
+				ObjectPath += SubObjectName;
+			}
+
+			if (!Extension.IsEmpty())
+			{
+				LocalPath += Extension;
+			}
+			else if (!FPackageName::DoesPackageExist(PackageName, &LocalPath))
+			{
+				LocalPath.Empty();
+			}
+		}
+	}
+	else if (AssetLog.PackagePath)
+	{
+		ObjectPath = AssetLog.PackagePath->GetPackageName();
+		LocalPath = AssetLog.PackagePath->GetLocalFullPath();
+	}
+	else if (AssetLog.Object)
+	{
+		ObjectPath = AssetLog.Object->GetPathName();
+		if (const UPackage* Package = AssetLog.Object->GetPackage())
+		{
+			LocalPath = Package->GetLoadedPath().GetLocalFullPath();
+		}
+	}
+
+	const auto GetConfigBool = [](const TCHAR* Section, const TCHAR* Key, bool bDefault) -> bool
+	{
+		GConfig->GetBool(Section, Key, bDefault, GEngineIni);
+		return bDefault;
+	};
+
+	if (!LocalPath.IsEmpty())
+	{
+		static bool bShowDiskPath = GetConfigBool(TEXT("Core.System"), TEXT("AssetLogShowsDiskPath"), true);
+		static bool bShowAbsolutePath = GetConfigBool(TEXT("Core.System"), TEXT("AssetLogShowsAbsolutePath"), false);
+		if (bShowAbsolutePath)
+		{
+			LocalPath = IFileManager::Get().ConvertToAbsolutePathForExternalAppForRead(*LocalPath);
+		}
+		FPaths::MakePlatformFilename(LocalPath);
+		if (bShowDiskPath)
+		{
+			ObjectPath = LocalPath;
+		}
+	}
+
+	Writer.BeginObject();
+	Writer.AddString(ANSITEXTVIEW("$type"), ANSITEXTVIEW("Asset"));
+	Writer.AddString(ANSITEXTVIEW("$text"), ObjectPath);
+	if (!LocalPath.IsEmpty())
+	{
+		Writer.AddString(ANSITEXTVIEW("file"), LocalPath);
+	}
+	Writer.EndObject();
+}
+
+void UE::Core::Private::RecordAssetLog(
+	const FName& CategoryName,
+	const ELogVerbosity::Type Verbosity,
+	const FAssetLog& AssetLog,
+	const FString& Message,
+	const ANSICHAR* File,
+	const int32 Line)
+{
+	FCbWriter Writer;
+	Writer.BeginObject();
+	Writer.SetName(ANSITEXTVIEW("Asset"));
+	SerializeForLog(Writer, AssetLog);
+	Writer.AddString(ANSITEXTVIEW("Message"), Message);
+	Writer.EndObject();
+
+	FLogRecord Record;
+	Record.SetFormat(TEXT("[AssetLog] {Asset}: {Message}"));
+	Record.SetFields(Writer.Save().AsObject());
+	Record.SetFile(File);
+	Record.SetLine(Line);
+	Record.SetCategory(CategoryName);
+	Record.SetVerbosity(Verbosity);
+	Record.SetTime(FLogTime::Now());
+
+	switch (Verbosity)
+	{
+	case ELogVerbosity::Error:
+	case ELogVerbosity::Warning:
+	case ELogVerbosity::Display:
+		return GWarn->SerializeRecord(Record);
+	default:
+		return GLog->SerializeRecord(Record);
+	}
+}
+
 /**
  * Takes a path of some sort and attempts to turn it into the asset log's canonical path.
  */
diff --git a/Engine/Source/Runtime/CoreUObject/Public/UObject/UObjectGlobals.h b/Engine/Source/Runtime/CoreUObject/Public/UObject/UObjectGlobals.h
--- a/Engine/Source/Runtime/CoreUObject/Public/UObject/UObjectGlobals.h
+++ b/Engine/Source/Runtime/CoreUObject/Public/UObject/UObjectGlobals.h
@@ -45,6 +45,7 @@
 #include "UObject/UnrealNames.h"
 
 class FArchive;
+class FCbWriter;
 class FLinkerInstancingContext;
 class FObjectPreSaveContext;
 class FOutputDevice;
@@ -2979,6 +2980,37 @@ UE_DEPRECATED(5.0, "This method is no longer in use and will be removed.")
 COREUOBJECT_API TMap<FName, FDynamicClassStaticData>& GetDynamicClassMap();
 PRAGMA_ENABLE_DEPRECATION_WARNINGS
 
+namespace UE
+{
+
+class FAssetLog
+{
+public:
+	inline explicit FAssetLog(const TCHAR* InPath UE_LIFETIMEBOUND) : Path(InPath) {}
+	inline explicit FAssetLog(const FPackagePath& InPath UE_LIFETIMEBOUND) : PackagePath(&InPath) {}
+	inline explicit FAssetLog(const UObject* InObject UE_LIFETIMEBOUND) : Object(InObject) {}
+
+	COREUOBJECT_API friend void SerializeForLog(FCbWriter& Writer, const FAssetLog& AssetLog);
+
+private:
+	const TCHAR* Path = nullptr;
+	const FPackagePath* PackagePath = nullptr;
+	const UObject* Object = nullptr;
+};
+
+} // UE
+
+namespace UE::Core::Private
+{
+COREUOBJECT_API void RecordAssetLog(
+	const FName& CategoryName,
+	ELogVerbosity::Type Verbosity,
+	const FAssetLog& AssetLog,
+	const FString& Message,
+	const ANSICHAR* File,
+	int32 Line);
+} // UE::Core::Private
+
 /**
  * FAssetMsg
  * This struct contains functions for asset-related messaging
@@ -3020,12 +3052,11 @@ struct FAssetMsg
 		{ \
 			UE_LOG_EXPAND_IS_FATAL(Verbosity, PREPROCESSOR_NOTHING, if (!CategoryName.IsSuppressed(ELogVerbosity::Verbosity))) \
 			{ \
-				FString FormatPath = FAssetMsg::FormatPathForAssetLog(Asset);\
-				FMsg::Logf_Internal(__FILE__, __LINE__, CategoryName.GetCategoryName(), ELogVerbosity::Verbosity, TEXT(ASSET_LOG_FORMAT_STRING_ANSI "%s"), *FormatPath, *FString::Printf(Format, ##__VA_ARGS__)); \
+				::UE::Core::Private::RecordAssetLog(CategoryName.GetCategoryName(), ELogVerbosity::Verbosity, ::UE::FAssetLog(Asset), FString::Printf(Format, ##__VA_ARGS__), __FILE__, __LINE__); \
 				UE_LOG_EXPAND_IS_FATAL(Verbosity, \
 					{ \
 						UE_DEBUG_BREAK_AND_PROMPT_FOR_REMOTE(); \
-						FDebug::AssertFailed("", __FILE__, __LINE__, TEXT("%s: %s"), *FormatPath, *FString::Printf(Format, ##__VA_ARGS__)); \
+						FDebug::AssertFailed("", __FILE__, __LINE__, TEXT("%s: %s"), *FAssetMsg::FormatPathForAssetLog(Asset), *FString::Printf(Format, ##__VA_ARGS__)); \
 						CA_ASSUME(false); \
 					}, \
 					PREPROCESSOR_NOTHING \
