Removed HACK_HEADER_GENERATOR

diff --git a/Engine/Source/Developer/MaterialUtilities/Public/MaterialUtilities.h b/Engine/Source/Developer/MaterialUtilities/Public/MaterialUtilities.h
--- a/Engine/Source/Developer/MaterialUtilities/Public/MaterialUtilities.h
+++ b/Engine/Source/Developer/MaterialUtilities/Public/MaterialUtilities.h
@@ -8,7 +8,6 @@
 #include "SceneTypes.h"
 #include "Modules/ModuleInterface.h"
 #include "Engine/TextureStreamingTypes.h"
-#include "UObject/ErrorException.h"
 #include "Engine/Texture.h"
 
 #include "LightMap.h"
diff --git a/Engine/Source/Runtime/Core/Private/GenericPlatform/GenericPlatformMisc.cpp b/Engine/Source/Runtime/Core/Private/GenericPlatform/GenericPlatformMisc.cpp
--- a/Engine/Source/Runtime/Core/Private/GenericPlatform/GenericPlatformMisc.cpp
+++ b/Engine/Source/Runtime/Core/Private/GenericPlatform/GenericPlatformMisc.cpp
@@ -732,13 +732,7 @@ void FGenericPlatformMisc::RaiseException(uint32 ExceptionCode)
 	/** This is the last place to gather memory stats before exception. */
 	FGenericCrashContext::SetMemoryStats(FPlatformMemory::GetStats());
 
-#if HACK_HEADER_GENERATOR && !PLATFORM_EXCEPTIONS_DISABLED
-	// We want Unreal Header Tool to throw an exception but in normal runtime code 
-	// we don't support exception handling
-	throw(ExceptionCode);
-#else	
 	*((uint32*)3) = ExceptionCode;
-#endif
 }
 
 template<typename CharType>
diff --git a/Engine/Source/Runtime/Core/Private/Misc/OutputDeviceHelper.cpp b/Engine/Source/Runtime/Core/Private/Misc/OutputDeviceHelper.cpp
--- a/Engine/Source/Runtime/Core/Private/Misc/OutputDeviceHelper.cpp
+++ b/Engine/Source/Runtime/Core/Private/Misc/OutputDeviceHelper.cpp
@@ -63,10 +63,8 @@ void FOutputDeviceHelper::AppendFormatLogLine(
 	}
 	else if (GPrintLogVerbosity && Verbosity != ELogVerbosity::Log)
 	{
-#if !HACK_HEADER_GENERATOR
 		Format.Append(ToString(Verbosity));
 		Format.Append(WIDETEXT(": "));
-#endif
 	}
 
 	if (Message)
diff --git a/Engine/Source/Runtime/Core/Private/ProfilingDebugging/MallocProfiler.cpp b/Engine/Source/Runtime/Core/Private/ProfilingDebugging/MallocProfiler.cpp
--- a/Engine/Source/Runtime/Core/Private/ProfilingDebugging/MallocProfiler.cpp
+++ b/Engine/Source/Runtime/Core/Private/ProfilingDebugging/MallocProfiler.cpp
@@ -1224,7 +1224,7 @@ FMallocProfilerBufferedFileWriter::~FMallocProfilerBufferedFileWriter()
 
 void FMallocProfilerBufferedFileWriter::Serialize( void* V, int64 Length )
 {
-#if (ALLOW_DEBUG_FILES && !HACK_HEADER_GENERATOR)
+#if ALLOW_DEBUG_FILES
 	// Copy to buffered memory array if GConfig hasn't been set up yet.
 	// This isn't the best solution, but due to complexity of the engine initialization order is the safest.
 	const bool bIsIniReady = GConfig && GConfig->IsReadyForUse();
diff --git a/Engine/Source/Runtime/Core/Public/Misc/Build.h b/Engine/Source/Runtime/Core/Public/Misc/Build.h
--- a/Engine/Source/Runtime/Core/Public/Misc/Build.h
+++ b/Engine/Source/Runtime/Core/Public/Misc/Build.h
@@ -115,18 +115,10 @@
 	#define ENABLE_PGO_PROFILE 0
 #endif
 
-/**
- * Unreal Header Tool requires extra data stored in the structure of a few core files. This enables some ifdef hacks to make this work. 
- * Set via UBT, do not modify directly
- */
-#ifndef HACK_HEADER_GENERATOR
-	#define HACK_HEADER_GENERATOR 0
-#endif
-
 /** Whether we are compiling with automation worker functionality.  Note that automation worker defaults to enabled in
     UE_BUILD_TEST configuration, so that it can be used for performance testing on devices */
 #ifndef WITH_AUTOMATION_WORKER
-	#define WITH_AUTOMATION_WORKER !(UE_BUILD_SHIPPING || HACK_HEADER_GENERATOR)
+	#define WITH_AUTOMATION_WORKER !UE_BUILD_SHIPPING
 #endif
 
 /**
diff --git a/Engine/Source/Runtime/Core/Public/Windows/WindowsPlatform.h b/Engine/Source/Runtime/Core/Public/Windows/WindowsPlatform.h
--- a/Engine/Source/Runtime/Core/Public/Windows/WindowsPlatform.h
+++ b/Engine/Source/Runtime/Core/Public/Windows/WindowsPlatform.h
@@ -127,7 +127,7 @@ typedef FWindowsPlatformTypes FPlatformTypes;
 
 // Intrinsics for 128-bit atomics on Windows platform requires Windows 8 or higher (WINVER>=0x0602)
 // http://msdn.microsoft.com/en-us/library/windows/desktop/hh972640.aspx
-#define PLATFORM_HAS_128BIT_ATOMICS							(!HACK_HEADER_GENERATOR && PLATFORM_64BITS && (WINVER >= 0x602))
+#define PLATFORM_HAS_128BIT_ATOMICS							(PLATFORM_64BITS && (WINVER >= 0x602))
 
 #ifdef CDECL
 #undef CDECL
diff --git a/Engine/Source/Runtime/Core/Public/Windows/WindowsPlatformCompilerSetup.h b/Engine/Source/Runtime/Core/Public/Windows/WindowsPlatformCompilerSetup.h
--- a/Engine/Source/Runtime/Core/Public/Windows/WindowsPlatformCompilerSetup.h
+++ b/Engine/Source/Runtime/Core/Public/Windows/WindowsPlatformCompilerSetup.h
@@ -242,11 +242,6 @@ static_assert(_MSC_VER >= 1929, "Visual Studio 2019 v16.11 or greater is require
 // Mixing MMX/SSE intrinsics will cause this warning, even when it's done correctly.
 //#pragma warning(disable : 4730) //mixing _m64 and floating point expressions may result in incorrect code
 
-// If C++ exception handling is disabled, force guarding to be off.
-#if !defined(_CPPUNWIND) && !defined(__INTELLISENSE__) && !defined(HACK_HEADER_GENERATOR)
-#error "Bad VCC option: C++ exception handling must be enabled" //lint !e309 suppress as lint doesn't have this defined
-#endif
-
 // Make sure characters are signed
 #ifdef _CHAR_UNSIGNED
 #error "Bad VC++ option: Characters must be signed" //lint !e309 suppress as lint doesn't have this defined
diff --git a/Engine/Source/Runtime/CoreUObject/Private/Serialization/AsyncLoading.cpp b/Engine/Source/Runtime/CoreUObject/Private/Serialization/AsyncLoading.cpp
--- a/Engine/Source/Runtime/CoreUObject/Private/Serialization/AsyncLoading.cpp
+++ b/Engine/Source/Runtime/CoreUObject/Private/Serialization/AsyncLoading.cpp
@@ -250,7 +250,6 @@ struct FEDLBootNotificationManager
 		FireCompletedCompiledInImports(AsyncLoadingThread, true);
 		FlushAsyncLoading();
 #endif
-#if !HACK_HEADER_GENERATOR
 		check(!GIsInitialLoad && IsInGameThread());
 		FScopeLock Lock(&EDLBootNotificationManagerLock);
 		for (auto& Pair : PathToState)
@@ -276,7 +275,6 @@ struct FEDLBootNotificationManager
 			}
 			UE_LOG(LogStreaming, Fatal, TEXT("Initial load is complete, but we still have %d imports to fire (listed above)."), PathsToFire.Num());
 		}
-#endif
 		Disable();
 	}
 
diff --git a/Engine/Source/Runtime/CoreUObject/Private/Templates/Casts.cpp b/Engine/Source/Runtime/CoreUObject/Private/Templates/Casts.cpp
--- a/Engine/Source/Runtime/CoreUObject/Private/Templates/Casts.cpp
+++ b/Engine/Source/Runtime/CoreUObject/Private/Templates/Casts.cpp
@@ -10,34 +10,3 @@ void CastLogError(const TCHAR* FromType, const TCHAR* ToType)
 	UE_LOG(LogCasts, Fatal, TEXT("Cast of %s to %s failed"), FromType, ToType);
     for(;;);
 }
-
-
-#if HACK_HEADER_GENERATOR
-
-//////////////////////////////////////////////////////////////////////////
-// ClassCastFlagMap
-
-ClassCastFlagMap::ClassCastFlagMap()
-{
-	// Define macro to be applied to all cast class declarations.
-	#define DECLARE_CAST_BY_FLAG(ClassName) CastFlagMap.Add(FString(#ClassName), CASTCLASS_##ClassName);
-
-	// Now apply the macro to the whole list.
-	DECLARE_ALL_CAST_FLAGS;
-
-	#undef DECLARE_CAST_BY_FLAG
-}
-
-ClassCastFlagMap& ClassCastFlagMap::Get()
-{
-	static ClassCastFlagMap TheInstance;
-	return TheInstance;
-}
-
-EClassCastFlags ClassCastFlagMap::GetCastFlag(const FString& ClassName) const
-{
-	const EClassCastFlags* ValuePtr = CastFlagMap.Find(ClassName);
-	return ValuePtr ? *ValuePtr : CASTCLASS_None;
-}
-
-#endif //HACK_HEADER_GENERATOR
diff --git a/Engine/Source/Runtime/CoreUObject/Private/UObject/Class.cpp b/Engine/Source/Runtime/CoreUObject/Private/UObject/Class.cpp
--- a/Engine/Source/Runtime/CoreUObject/Private/UObject/Class.cpp
+++ b/Engine/Source/Runtime/CoreUObject/Private/UObject/Class.cpp
@@ -14,7 +14,6 @@
 #include "Misc/AutomationTest.h"
 #include "Misc/EnumClassFlags.h"
 #include "Misc/StringBuilder.h"
-#include "UObject/ErrorException.h"
 #include "Modules/ModuleManager.h"
 #include "UObject/UObjectAllocator.h"
 #include "UObject/UObjectHash.h"
@@ -839,11 +838,7 @@ void UStruct::Link(FArchive& Ar, bool bRelinkExistingProperties)
 					if (StructProp != NULL && StructProp->Struct == this)
 					{
 						//we won't support this, too complicated
-					#if HACK_HEADER_GENERATOR
-						FError::Throwf(TEXT("'Struct recursion via arrays is unsupported for properties."));
-					#else
 						UE_LOG(LogClass, Fatal, TEXT("'Struct recursion via arrays is unsupported for properties."));
-					#endif
 					}
 				}
 			}
@@ -2498,9 +2493,6 @@ bool FindConstructorUninitialized(UStruct* BaseClass,uint8* Data,uint8* Defaults
 UScriptStruct::UScriptStruct( EStaticConstructor, int32 InSize, int32 InAlignment, EObjectFlags InFlags )
 	: UStruct( EC_StaticConstructor, InSize, InAlignment, InFlags )
 	, StructFlags(STRUCT_NoFlags)
-#if HACK_HEADER_GENERATOR
-	, StructMacroDeclaredLineNumber(INDEX_NONE)
-#endif
 	, bPrepareCppStructOpsCompleted(false)
 	, CppStructOps(NULL)
 {
@@ -2509,9 +2501,6 @@ UScriptStruct::UScriptStruct( EStaticConstructor, int32 InSize, int32 InAlignmen
 UScriptStruct::UScriptStruct(const FObjectInitializer& ObjectInitializer, UScriptStruct* InSuperStruct, ICppStructOps* InCppStructOps, EStructFlags InStructFlags, SIZE_T ExplicitSize, SIZE_T ExplicitAlignment )
 	: UStruct(ObjectInitializer, InSuperStruct, InCppStructOps ? InCppStructOps->GetSize() : ExplicitSize, InCppStructOps ? InCppStructOps->GetAlignment() : ExplicitAlignment )
 	, StructFlags(EStructFlags(InStructFlags | (InCppStructOps ? STRUCT_Native : STRUCT_NoFlags)))
-#if HACK_HEADER_GENERATOR
-	, StructMacroDeclaredLineNumber(INDEX_NONE)
-#endif
 	, bPrepareCppStructOpsCompleted(false)
 	, CppStructOps(InCppStructOps)
 {
@@ -2521,9 +2510,6 @@ UScriptStruct::UScriptStruct(const FObjectInitializer& ObjectInitializer, UScrip
 UScriptStruct::UScriptStruct(const FObjectInitializer& ObjectInitializer)
 	: UStruct(ObjectInitializer)
 	, StructFlags(STRUCT_NoFlags)
-#if HACK_HEADER_GENERATOR
-	, StructMacroDeclaredLineNumber(INDEX_NONE)
-#endif
 	, bPrepareCppStructOpsCompleted(false)
 	, CppStructOps(NULL)
 {
@@ -2555,13 +2541,6 @@ void UScriptStruct::DeferCppStructOps(FTopLevelAssetPath Target, ICppStructOps*
 	DeferredStructOps.Add(Target, InCppStructOps);
 }
 
-#if HACK_HEADER_GENERATOR
-UScriptStruct::ICppStructOps* UScriptStruct::FindDeferredCppStructOps(FTopLevelAssetPath StructName)
-{
-	return GetDeferredCppStructOps().FindRef(StructName);
-}
-#endif
-
 FTopLevelAssetPath UScriptStruct::GetFlattenedStructPathName() const
 {
 	return FTopLevelAssetPath(GetPackage()->GetFName(), GetFName());
@@ -2592,9 +2571,7 @@ void UScriptStruct::PrepareCppStructOps()
 			bPrepareCppStructOpsCompleted = true;
 			return;
 		}
-#if !HACK_HEADER_GENERATOR
 		StructFlags = EStructFlags(StructFlags | STRUCT_Native);
-#endif
 		
 #if !(UE_BUILD_SHIPPING || UE_BUILD_TEST)
 		// test that the constructor is initializing everything
@@ -3697,14 +3674,12 @@ IMPLEMENT_SIMPLE_AUTOMATION_TEST(FAutomationTestAttemptToFindUninitializedScript
 bool FAutomationTestAttemptToFindUninitializedScriptStructMembers::RunTest(const FString& Parameters)
 {
 	// This test fails when running tests under UHT because there is no TestUninitializedScriptStructMembersTest, so just skip it in that config.
-#if !HACK_HEADER_GENERATOR
 	if (UObjectInitialized())
 	{
 		FStructUtils::AttemptToFindUninitializedScriptStructMembers();
 		return !HasAnyErrors();
 	}
 	else
-#endif
 	{
 		return true;
 	}
@@ -3906,11 +3881,7 @@ IMPLEMENT_SIMPLE_AUTOMATION_TEST(FAutomationTestAttemptToFindShortTypeNamesInMet
 bool FAutomationTestAttemptToFindShortTypeNamesInMetaData::RunTest(const FString& Parameters)
 {
 	// This test is not necessary when running under UHT so just skip it in that config.
-#if !HACK_HEADER_GENERATOR
 	return FStructUtils::AttemptToFindShortTypeNamesInMetaData() == 0;
-#else
-	return true;
-#endif
 }
 
 #endif // WITH_EDITORONLY_DATA
@@ -4497,39 +4468,6 @@ static int32 GValidateReplicatedProperties = 1;
 
 static FAutoConsoleVariable CVarValidateReplicatedPropertyRegistration(TEXT("net.ValidateReplicatedPropertyRegistration"), GValidateReplicatedProperties, TEXT("Warns if replicated properties were not registered in GetLifetimeReplicatedProps."));
 
-#if HACK_HEADER_GENERATOR
-void UClass::SetUpUhtReplicationData()
-{
-	if (!HasAnyClassFlags(CLASS_ReplicationDataIsSetUp) && PropertyLink != NULL)
-	{
-        ClassReps.Empty();
-		if (UClass* SuperClass = GetSuperClass())
-		{
-			SuperClass->SetUpUhtReplicationData();
-			ClassReps = SuperClass->ClassReps;
-			FirstOwnedClassRep = ClassReps.Num();
-		}
-		else
-		{
-			FirstOwnedClassRep = 0;
-		}
-
-		for (TFieldIterator<FProperty> It(this, EFieldIteratorFlags::ExcludeSuper); It; ++It)
-		{
-			if (It->PropertyFlags & CPF_Net)
-			{
-				It->RepIndex = (uint16)ClassReps.Num();
-				new (ClassReps) FRepRecord(*It, 0);
-			}
-		}
-		check(ClassReps.Num() <= 65535);
-
-		ClassFlags |= CLASS_ReplicationDataIsSetUp;
-		ClassReps.Shrink();
-	}
-}
-#endif
-
 void UClass::SetUpRuntimeReplicationData()
 {
 	if (!HasAnyClassFlags(CLASS_ReplicationDataIsSetUp) && PropertyLink != NULL)
@@ -6174,7 +6112,7 @@ bool UClass::IsShortTypeName(FStringView InClassPath)
 }
 
 
-#if WITH_EDITOR || HACK_HEADER_GENERATOR
+#if WITH_EDITOR
 void UClass::GetHideFunctions(TArray<FString>& OutHideFunctions) const
 {
 	static const FName NAME_HideFunctions(TEXT("HideFunctions"));
@@ -6270,7 +6208,7 @@ bool UClass::IsClassGroupName(const TCHAR* InGroupName) const
 	return false;
 }
 
-#endif // WITH_EDITOR || HACK_HEADER_GENERATOR
+#endif // WITH_EDITOR
 
 
 #if WITH_EDITORONLY_DATA
diff --git a/Engine/Source/Runtime/CoreUObject/Private/UObject/Enum.cpp b/Engine/Source/Runtime/CoreUObject/Private/UObject/Enum.cpp
--- a/Engine/Source/Runtime/CoreUObject/Private/UObject/Enum.cpp
+++ b/Engine/Source/Runtime/CoreUObject/Private/UObject/Enum.cpp
@@ -2,7 +2,6 @@
 
 #include "CoreMinimal.h"
 #include "Misc/ConfigCacheIni.h"
-#include "UObject/ErrorException.h"
 #include "UObject/ObjectMacros.h"
 #include "UObject/Class.h"
 #include "UObject/Package.h"
diff --git a/Engine/Source/Runtime/CoreUObject/Private/UObject/ErrorException.cpp b/Engine/Source/Runtime/CoreUObject/Private/UObject/ErrorException.cpp
deleted file mode 100644
diff --git a/Engine/Source/Runtime/CoreUObject/Private/UObject/Field.cpp b/Engine/Source/Runtime/CoreUObject/Private/UObject/Field.cpp
--- a/Engine/Source/Runtime/CoreUObject/Private/UObject/Field.cpp
+++ b/Engine/Source/Runtime/CoreUObject/Private/UObject/Field.cpp
@@ -15,7 +15,6 @@ Field.cpp: Defines FField property system fundamentals
 #include "Misc/OutputDeviceHelper.h"
 #include "Misc/FeedbackContext.h"
 #include "Misc/OutputDeviceConsole.h"
-#include "UObject/ErrorException.h"
 #include "Modules/ModuleManager.h"
 #include "UObject/UObjectAllocator.h"
 #include "UObject/UObjectHash.h"
diff --git a/Engine/Source/Runtime/CoreUObject/Private/UObject/PropertyStruct.cpp b/Engine/Source/Runtime/CoreUObject/Private/UObject/PropertyStruct.cpp
--- a/Engine/Source/Runtime/CoreUObject/Private/UObject/PropertyStruct.cpp
+++ b/Engine/Source/Runtime/CoreUObject/Private/UObject/PropertyStruct.cpp
@@ -246,21 +246,6 @@ void FStructProperty::AddReferencedObjects(FReferenceCollector& Collector)
 	Super::AddReferencedObjects(Collector);
 }
 
-#if HACK_HEADER_GENERATOR
-
-bool FStructProperty::HasNoOpConstructor() const
-{
-	Struct->PrepareCppStructOps();
-	UScriptStruct::ICppStructOps* CppStructOps = Struct->GetCppStructOps();
-	if (CppStructOps && CppStructOps->HasNoopConstructor())
-	{
-		return true;
-	}
-	return false;
-}
-
-#endif
-
 FString FStructProperty::GetCPPType( FString* ExtendedTypeText/*=NULL*/, uint32 CPPExportFlags/*=0*/ ) const
 {
 	return Struct->GetStructCPPName(CPPExportFlags);
diff --git a/Engine/Source/Runtime/CoreUObject/Private/UObject/TopLevelAssetPath.cpp b/Engine/Source/Runtime/CoreUObject/Private/UObject/TopLevelAssetPath.cpp
--- a/Engine/Source/Runtime/CoreUObject/Private/UObject/TopLevelAssetPath.cpp
+++ b/Engine/Source/Runtime/CoreUObject/Private/UObject/TopLevelAssetPath.cpp
@@ -14,11 +14,6 @@
 #include "UObject/UnrealType.h"
 #include "UObject/UObjectThreadContext.h"
 
-// The reason behind HACK_HEADER_GENERATOR is that without it UHT is going to 'see' cppstructops for TopLevelAssetPath
-// and will not generate temp FTopLevelAssetPath struct for codegen purposes where it can access all of its members
-// which are public in the temp struct and private in the actual FTopLevelAssetPath.
-// If UHT compiles with UE_IMPLEMENT_STRUCT("/Script/CoreUObject", TopLevelAssetPath) the generated code will fail to compile trying to access private struct members
-#if !HACK_HEADER_GENERATOR
 template<>
 struct TStructOpsTypeTraits<FTopLevelAssetPath> : public TStructOpsTypeTraitsBase2<FTopLevelAssetPath>
 {
@@ -30,7 +25,6 @@ struct TStructOpsTypeTraits<FTopLevelAssetPath> : public TStructOpsTypeTraitsBas
 	};
 };
 UE_IMPLEMENT_STRUCT("/Script/CoreUObject", TopLevelAssetPath)
-#endif
 
 void FTopLevelAssetPath::AppendString(FStringBuilderBase& Builder) const
 {
diff --git a/Engine/Source/Runtime/CoreUObject/Private/UObject/UObjectArray.cpp b/Engine/Source/Runtime/CoreUObject/Private/UObject/UObjectArray.cpp
--- a/Engine/Source/Runtime/CoreUObject/Private/UObject/UObjectArray.cpp
+++ b/Engine/Source/Runtime/CoreUObject/Private/UObject/UObjectArray.cpp
@@ -81,7 +81,7 @@ FUObjectArray::FUObjectArray()
 : ObjFirstGCIndex(0)
 , ObjLastNonGCIndex(INDEX_NONE)
 , MaxObjectsNotConsideredByGC(0)
-, OpenForDisregardForGC(!HACK_HEADER_GENERATOR)
+, OpenForDisregardForGC(true)
 , PrimarySerialNumber(START_SERIAL_NUMBER)
 {
 	GCoreObjectArrayForDebugVisualizers = &GUObjectArray.ObjObjects;
diff --git a/Engine/Source/Runtime/CoreUObject/Public/CoreUObject.h b/Engine/Source/Runtime/CoreUObject/Public/CoreUObject.h
--- a/Engine/Source/Runtime/CoreUObject/Public/CoreUObject.h
+++ b/Engine/Source/Runtime/CoreUObject/Public/CoreUObject.h
@@ -7,7 +7,6 @@ MONOLITHIC_HEADER_BOILERPLATE()
 
 
 #include "Core.h"
-#include "UObject/ErrorException.h"
 #include "Modules/ModuleManager.h"
 #include "UObject/ScriptInterface.h"
 #include "UObject/Script.h"
diff --git a/Engine/Source/Runtime/CoreUObject/Public/Templates/Casts.h b/Engine/Source/Runtime/CoreUObject/Public/Templates/Casts.h
--- a/Engine/Source/Runtime/CoreUObject/Public/Templates/Casts.h
+++ b/Engine/Source/Runtime/CoreUObject/Public/Templates/Casts.h
@@ -427,24 +427,6 @@ DECLARE_ALL_CAST_FLAGS
 #undef DECLARE_CAST_BY_FLAG_CAST
 #undef DECLARE_CAST_BY_FLAG_FWD
 
-
-#if HACK_HEADER_GENERATOR
-// Singleton class to get the cast flag for a given class name.
-struct COREUOBJECT_API ClassCastFlagMap
-{
-	static ClassCastFlagMap& Get();
-
-	// Get Mapped name -> cast flag. Returns CASTCLASS_None if name is not found.
-	EClassCastFlags GetCastFlag(const FString& ClassName) const;
-
-private:
-	ClassCastFlagMap();
-	TMap<FString, EClassCastFlags> CastFlagMap;
-};
-#endif // HACK_HEADER_GENERATOR
-
-
-
 namespace UECasts_Private
 {
 	template <typename T>
diff --git a/Engine/Source/Runtime/CoreUObject/Public/UObject/Class.h b/Engine/Source/Runtime/CoreUObject/Public/UObject/Class.h
--- a/Engine/Source/Runtime/CoreUObject/Public/UObject/Class.h
+++ b/Engine/Source/Runtime/CoreUObject/Public/UObject/Class.h
@@ -656,11 +656,6 @@ public:
 	}
 #endif // WITH_EDITORONLY_DATA
 
-#if HACK_HEADER_GENERATOR
-	// Required by UHT makefiles for internal data serialization.
-	friend struct FStructArchiveProxy;
-#endif // HACK_HEADER_GENERATOR
-
 	/** Sets the UnresolvedScriptProperties array */
 	void SetUnresolvedScriptProperties(FUnresolvedScriptPropertiesArray& InUnresolvedProperties)
 	{
@@ -1503,16 +1498,6 @@ public:
 public:
 	EStructFlags StructFlags;
 
-
-#if HACK_HEADER_GENERATOR
-	int32 StructMacroDeclaredLineNumber;
-
-	// Required by UHT makefiles for internal data serialization.
-	friend struct FScriptStructArchiveProxy;
-
-	static COREUOBJECT_API ICppStructOps* FindDeferredCppStructOps(FTopLevelAssetPath StructName);
-#endif
-
 protected:
 	/** true if we have performed PrepareCppStructOps **/
 	bool bPrepareCppStructOpsCompleted;
@@ -2656,7 +2641,7 @@ public:
 	/** List of network relevant fields (functions) */
 	TArray<UField*> NetFields;
 
-#if WITH_EDITOR || HACK_HEADER_GENERATOR 
+#if WITH_EDITOR
 	// Editor only properties
 	void GetHideFunctions(TArray<FString>& OutHideFunctions) const;
 	bool IsFunctionHidden(const TCHAR* InFunction) const;
@@ -3330,10 +3315,6 @@ public:
 	 */
 	void SetUpRuntimeReplicationData();
 
-#if HACK_HEADER_GENERATOR
-	void SetUpUhtReplicationData();
-#endif  // HACK_HEADER_GENERATOR
-
 	/**
 	 * Helper function for determining if the given class is compatible with structured archive serialization
 	 */
@@ -3389,11 +3370,6 @@ protected:
 	 * @return		the CDO for this class
 	 **/
 	virtual UObject* CreateDefaultObject();
-
-#if HACK_HEADER_GENERATOR
-	// Required by UHT makefiles for internal data serialization.
-	friend struct FClassArchiveProxy;
-#endif // HACK_HEADER_GENERATOR
 };
 
 // @todo: BP2CPP_remove
diff --git a/Engine/Source/Runtime/CoreUObject/Public/UObject/EnumProperty.h b/Engine/Source/Runtime/CoreUObject/Public/UObject/EnumProperty.h
--- a/Engine/Source/Runtime/CoreUObject/Public/UObject/EnumProperty.h
+++ b/Engine/Source/Runtime/CoreUObject/Public/UObject/EnumProperty.h
@@ -111,9 +111,6 @@ public:
 private:
 	virtual uint32 GetValueTypeHashInternal(const void* Src) const override;
 
-#if HACK_HEADER_GENERATOR
-public:
-#endif
 	FNumericProperty* UnderlyingProp; // The property which represents the underlying type of the enum
 	UEnum* Enum; // The enum represented by this property
 };
diff --git a/Engine/Source/Runtime/CoreUObject/Public/UObject/ErrorException.h b/Engine/Source/Runtime/CoreUObject/Public/UObject/ErrorException.h
--- a/Engine/Source/Runtime/CoreUObject/Public/UObject/ErrorException.h
+++ b/Engine/Source/Runtime/CoreUObject/Public/UObject/ErrorException.h
@@ -2,33 +2,4 @@
 
 #pragma once
 
-#include "CoreMinimal.h"
-#include "Misc/Build.h"
-
-#if HACK_HEADER_GENERATOR
-	#include "Templates/AndOrNot.h"
-	#include "Templates/IsValidVariadicFunctionArg.h"
-
-	/** 
-	 * FError
-	 * Set of functions for error reporting 
-	 **/
-	struct COREUOBJECT_API FError
-	{
-		/**
-		 * Throws a printf-formatted exception as a const TCHAR*.
-		 */
-		template <typename... Types>
-		UE_NORETURN static void VARARGS Throwf(const TCHAR* Fmt, Types... Args)
-		{
-			static_assert(TAnd<TIsValidVariadicFunctionArg<Types>...>::Value, "Invalid argument(s) passed to FError::Throwf");
-
-			ThrowfImpl(Fmt, Args...);
-		}
-
-	private:
-		UE_NORETURN static void VARARGS ThrowfImpl(const TCHAR* Fmt, ...);
-	};
-#endif
-
-
+#pragma message (__FILE__"(5): warning: ErrorException.h has been deprecated")
diff --git a/Engine/Source/Runtime/CoreUObject/Public/UObject/MetaData.h b/Engine/Source/Runtime/CoreUObject/Public/UObject/MetaData.h
--- a/Engine/Source/Runtime/CoreUObject/Public/UObject/MetaData.h
+++ b/Engine/Source/Runtime/CoreUObject/Public/UObject/MetaData.h
@@ -168,11 +168,6 @@ public:
 	// Returns the remapped key name, or NAME_None was not remapped.
 	static FName GetRemappedKeyName(FName OldKey);
 
-#if HACK_HEADER_GENERATOR
-	// Required by UHT makefiles for internal data serialization.
-	friend struct FMetadataArchiveProxy;
-#endif
-
 private:
 	static void InitializeRedirectMap();
 
diff --git a/Engine/Source/Runtime/CoreUObject/Public/UObject/ObjectMacros.h b/Engine/Source/Runtime/CoreUObject/Public/UObject/ObjectMacros.h
--- a/Engine/Source/Runtime/CoreUObject/Public/UObject/ObjectMacros.h
+++ b/Engine/Source/Runtime/CoreUObject/Public/UObject/ObjectMacros.h
@@ -26,12 +26,8 @@ typedef	uint64 ScriptPointerType;
 #endif
 
 #ifndef USE_COMPILED_IN_NATIVES
-#if HACK_HEADER_GENERATOR 
-#define USE_COMPILED_IN_NATIVES	0
-#else
 #define USE_COMPILED_IN_NATIVES	1
 #endif
-#endif
 
 /** Set this to 0 to disable UObject thread safety features */
 #define THREADSAFE_UOBJECTS 1
@@ -41,7 +37,7 @@ typedef	uint64 ScriptPointerType;
 #define USTRUCT_ISCHILDOF_STRUCTARRAY 2 // stores an array of parents per struct and uses this to compare - faster than 1 and thread-safe but can have issues with BP reinstancing and hot reload
 
 // USTRUCT_FAST_ISCHILDOF_IMPL sets which implementation of IsChildOf to use.
-#if UE_EDITOR || HACK_HEADER_GENERATOR
+#if UE_EDITOR
 	// On editor, we use the outerwalk implementation because BP reinstancing and hot reload
 	// mess up the struct array
 	#define USTRUCT_FAST_ISCHILDOF_IMPL USTRUCT_ISCHILDOF_OUTERWALK
diff --git a/Engine/Source/Runtime/CoreUObject/Public/UObject/UObjectBase.h b/Engine/Source/Runtime/CoreUObject/Public/UObject/UObjectBase.h
--- a/Engine/Source/Runtime/CoreUObject/Public/UObject/UObjectBase.h
+++ b/Engine/Source/Runtime/CoreUObject/Public/UObject/UObjectBase.h
@@ -265,11 +265,6 @@ private:
 	/** This is used by the reinstancer to re-class and re-archetype the current instances of a class before recompiling */
 	void SetClass(UClass* NewClass);
 #endif
-
-#if HACK_HEADER_GENERATOR
-	// Required by UHT makefiles for internal data serialization.
-	friend struct FObjectBaseArchiveProxy;
-#endif // HACK_HEADER_GENERATOR
 };
 
 /**
diff --git a/Engine/Source/Runtime/CoreUObject/Public/UObject/UnrealType.h b/Engine/Source/Runtime/CoreUObject/Public/UObject/UnrealType.h
--- a/Engine/Source/Runtime/CoreUObject/Public/UObject/UnrealType.h
+++ b/Engine/Source/Runtime/CoreUObject/Public/UObject/UnrealType.h
@@ -1190,11 +1190,6 @@ public:
 
 	ELifetimeCondition GetBlueprintReplicationCondition() const { return BlueprintReplicationCondition; }
 	void SetBlueprintReplicationCondition(ELifetimeCondition InBlueprintReplicationCondition) { BlueprintReplicationCondition = InBlueprintReplicationCondition; }
-
-#if HACK_HEADER_GENERATOR
-	// Required by UHT makefiles for internal data serialization.
-	friend struct FPropertyArchiveProxy;
-#endif // HACK_HEADER_GENERATOR
 };
 
 
@@ -2559,11 +2554,6 @@ public:
 	}
 
 	uint32 GetValueTypeHashInternal(const void* Src) const override;
-
-#if HACK_HEADER_GENERATOR
-	// Required by UHT makefiles for internal data serialization.
-	friend struct FBoolPropertyArchiveProxy;
-#endif // HACK_HEADER_GENERATOR
 };
 
 /*-----------------------------------------------------------------------------
@@ -5863,18 +5853,6 @@ public:
 
 private:
 	virtual uint32 GetValueTypeHashInternal(const void* Src) const;
-
-public:
-
-#if HACK_HEADER_GENERATOR
-	/**
-	 * Some native structs, like FIntPoint, FIntRect, FVector2D, FVector, FPlane, FRotator, FCylinder have a default constructor that does nothing and require EForceInit
-	 * Since it is name-based, this is not a fast routine intended to be used for header generation only
-	 * 
-	 * @return	true if this struct requires the EForceInit constructor to initialize
-	 */
-	bool HasNoOpConstructor() const;
-#endif
 };
 
 /*-----------------------------------------------------------------------------
