Bulk Data - allow cooked loose file in editor builds

#rnx

diff --git a/Engine/Source/Runtime/CoreUObject/Private/Serialization/AsyncLoading2.cpp b/Engine/Source/Runtime/CoreUObject/Private/Serialization/AsyncLoading2.cpp
--- a/Engine/Source/Runtime/CoreUObject/Private/Serialization/AsyncLoading2.cpp
+++ b/Engine/Source/Runtime/CoreUObject/Private/Serialization/AsyncLoading2.cpp
@@ -1781,41 +1781,40 @@ public:
 		else
 		{
 			const EIoChunkType ChunkType = Meta.HasAnyFlags(BULKDATA_OptionalPayload) ? EIoChunkType::OptionalBulkData : EIoChunkType::BulkData;
-			FIoChunkId ChunkId = CreateIoChunkId(PackageId.Value(), ChunkIndex, ChunkType);
+			BulkData.BulkChunkId = CreateIoChunkId(PackageId.Value(), ChunkIndex, ChunkType);
 
 			if (Meta.HasAnyFlags(BULKDATA_DuplicateNonOptionalPayload))
 			{
 				const FIoChunkId OptionalChunkId = CreateIoChunkId(PackageId.Value(), ChunkIndex, EIoChunkType::OptionalBulkData);
 				if (IoDispatcher.DoesChunkExist(OptionalChunkId))
 				{
-					ChunkId = OptionalChunkId;
+					BulkData.BulkChunkId = OptionalChunkId;
 					Meta.ClearFlags(BULKDATA_DuplicateNonOptionalPayload);
 					Meta.AddFlags(BULKDATA_OptionalPayload);
 					Meta.SetOffset(DuplicateSerialOffset);
 				}
 			}
-			else if (Meta.HasAnyFlags(BULKDATA_MemoryMappedPayload) && Params.bAttemptMemoryMapping)
+			else if (Meta.HasAnyFlags(BULKDATA_MemoryMappedPayload))
 			{
-				const FIoChunkId MappedChunkId = CreateIoChunkId(PackageId.Value(), ChunkIndex, EIoChunkType::MemoryMappedBulkData);
-				TIoStatusOr<FIoMappedRegion> Status = IoDispatcher.OpenMapped(MappedChunkId, FIoReadOptions(Meta.GetOffset(), Meta.GetSize()));
-
-				if (Status.IsOk())
-				{
-					ChunkId = MappedChunkId;
-					FIoMappedRegion Mapping = Status.ConsumeValueOrDie(); 
-					BulkData.DataAllocation.SetMemoryMappedData(&BulkData, Mapping.MappedFileHandle, Mapping.MappedFileRegion);
-				}
-				else
+				BulkData.BulkChunkId = CreateIoChunkId(PackageId.Value(), ChunkIndex, EIoChunkType::MemoryMappedBulkData);
+				if (Params.bAttemptMemoryMapping)
 				{
-					UE_LOG(LogSerialization, Warning, TEXT("Memory map bulk data from chunk '%s', offset '%lld', size '%lld' FAILED"),
-						*LexToString(MappedChunkId), Meta.GetOffset(), Meta.GetSize());
-					
-					BulkData.BulkChunkId = ChunkId;
-					BulkData.ForceBulkDataResident();
+					TIoStatusOr<FIoMappedRegion> Status = IoDispatcher.OpenMapped(BulkData.BulkChunkId, FIoReadOptions(Meta.GetOffset(), Meta.GetSize()));
+
+					if (Status.IsOk())
+					{
+						FIoMappedRegion Mapping = Status.ConsumeValueOrDie(); 
+						BulkData.DataAllocation.SetMemoryMappedData(&BulkData, Mapping.MappedFileHandle, Mapping.MappedFileRegion);
+					}
+					else
+					{
+						UE_LOG(LogSerialization, Warning, TEXT("Memory map bulk data from chunk '%s', offset '%lld', size '%lld' FAILED"),
+							*LexToString(BulkData.BulkChunkId), Meta.GetOffset(), Meta.GetSize());
+
+						BulkData.ForceBulkDataResident();
+					}
 				}
 			}
-
-			BulkData.BulkChunkId = ChunkId;
 		}
 
 		return true;
@@ -8320,4 +8319,4 @@ EAsyncPackageState::Type FAsyncLoadingThread2::ProcessLoadingUntilCompleteFromGa
 IAsyncPackageLoader* MakeAsyncPackageLoader2(FIoDispatcher& InIoDispatcher, IAsyncPackageLoader* InUncookedPackageLoader)
 {
 	return new FAsyncLoadingThread2(InIoDispatcher, InUncookedPackageLoader);
-}
\ No newline at end of file
+}
diff --git a/Engine/Source/Runtime/CoreUObject/Private/UObject/LinkerLoad.cpp b/Engine/Source/Runtime/CoreUObject/Private/UObject/LinkerLoad.cpp
--- a/Engine/Source/Runtime/CoreUObject/Private/UObject/LinkerLoad.cpp
+++ b/Engine/Source/Runtime/CoreUObject/Private/UObject/LinkerLoad.cpp
@@ -6723,7 +6723,8 @@ bool FLinkerLoad::SerializeBulkData(FBulkData& BulkData, const FBulkDataSerializ
 	int64 DuplicateSerialOffset = -1;
 	SerializeBulkMeta(Meta, DuplicateSerialOffset, Params.ElementSize);
 
-	if (IsAllowingLazyLoading())
+	const bool bLazyLoadable = IsAllowingLazyLoading();
+	if (bLazyLoadable)
 	{
 		Meta.AddFlags(BULKDATA_LazyLoadable);
 #if WITH_EDITOR
@@ -6732,10 +6733,6 @@ bool FLinkerLoad::SerializeBulkData(FBulkData& BulkData, const FBulkDataSerializ
 		AttachBulkData(Params.Owner, &BulkData);
 #endif // WITH_EDITOR
 	}
-	else
-	{
-		Meta.ClearFlags(BULKDATA_LazyLoadable);
-	}
 
 	const bool bExternalResource = Meta.HasAnyFlags(BULKDATA_WorkspaceDomainPayload);
 	EPackageSegment Segment = GetBulkDataPackageSegmentFromFlags(Meta.GetFlags(), IsLoadingFromCookedPackage());  
@@ -6753,53 +6750,38 @@ bool FLinkerLoad::SerializeBulkData(FBulkData& BulkData, const FBulkDataSerializ
 		void* Payload = BulkData.ReallocateData(Meta.GetSize());
 		BulkData.SerializeBulkData(*this, Payload, Meta.GetSize(), Meta.GetFlags());
 	}
-#if USE_RUNTIME_BULKDATA
-	else
+	else if (Meta.HasAnyFlags(BULKDATA_PayloadInSeperateFile))
 	{
-		// Streaming cooked bulk data. Since I/O store is the default staged output format
-		// this path is only used when running from loose cooked files.
-		check(Meta.HasAnyFlags(BULKDATA_PayloadInSeperateFile));
-		check(IsLoadingFromCookedPackage());
-		Meta.AddFlags(BULKDATA_LazyLoadable);
-
-		IPackageResourceManager& ResourceMgr = IPackageResourceManager::Get();
-	
-		if (Meta.HasAnyFlags(BULKDATA_DuplicateNonOptionalPayload) &&
-			ResourceMgr.DoesPackageExist(PackagePath, EPackageSegment::BulkDataOptional))
+		// Streaming cooked bulk data / loading from Editor Domain and referencing Workspace domain bulk data
+		if (Meta.HasAnyFlags(BULKDATA_DuplicateNonOptionalPayload))
 		{
-			Segment = EPackageSegment::BulkDataOptional;
-			Meta.ClearFlags(BULKDATA_DuplicateNonOptionalPayload);
-			Meta.AddFlags(BULKDATA_OptionalPayload);
-			Meta.SetOffset(DuplicateSerialOffset);
+			if (IPackageResourceManager::Get().DoesPackageExist(PackagePath, EPackageSegment::BulkDataOptional))
+			{
+				BulkData.BulkChunkId = UE::CreatePackageResourceChunkId(PackagePath.GetPackageFName(), EPackageSegment::BulkDataOptional, bExternalResource);
+				Meta.ClearFlags(BULKDATA_DuplicateNonOptionalPayload);
+				Meta.AddFlags(BULKDATA_OptionalPayload);
+				Meta.SetOffset(DuplicateSerialOffset);
+			}
 		}
-		else if (Meta.HasAnyFlags(BULKDATA_MemoryMappedPayload) && Params.bAttemptMemoryMapping)
+		else if (Meta.HasAnyFlags(BULKDATA_MemoryMappedPayload))
 		{
-			bool bMemoryMapped = false;
-			if (IsAllowingLazyLoading())
+			if (bLazyLoadable && Params.bAttemptMemoryMapping)
 			{
 				TUniquePtr<IMappedFileHandle> MappedFile;
-				MappedFile.Reset(ResourceMgr.OpenMappedHandleToPackage(PackagePath, EPackageSegment::BulkDataMemoryMapped));
+				MappedFile.Reset(IPackageResourceManager::Get().OpenMappedHandleToPackage(PackagePath, EPackageSegment::BulkDataMemoryMapped));
 				IMappedFileRegion* MappedRegion = MappedFile.IsValid() ? MappedFile->MapRegion(Meta.GetOffset(), Meta.GetSize(), true) : nullptr;
 				if (MappedRegion)
 				{
-					Segment = EPackageSegment::BulkDataMemoryMapped;
 					BulkData.DataAllocation.SetMemoryMappedData(&BulkData, MappedFile.Release(), MappedRegion);
-					bMemoryMapped  = true;
 				}
-			}
-
-			if (bMemoryMapped == false)
-			{
-				UE_LOG(LogSerialization, Warning, TEXT("Memory map bulk data '%s' FAILED"), *PackagePath.GetDebugName());
-				
-				BulkData.BulkChunkId = UE::CreatePackageResourceChunkId(PackagePath.GetPackageFName(), EPackageSegment::BulkDataDefault, false);
-				BulkData.ForceBulkDataResident();
+				else
+				{
+					UE_LOG(LogSerialization, Warning, TEXT("Memory map bulk data '%s' FAILED"), *PackagePath.GetDebugName());
+					BulkData.ForceBulkDataResident();
+				}
 			}
 		}
-	
-		BulkData.BulkChunkId = UE::CreatePackageResourceChunkId(PackagePath.GetPackageFName(), Segment, bExternalResource);
 	}
-#else
 	else
 	{
 		// Streaming uncooked bulk data (editor only)
@@ -6812,7 +6794,7 @@ bool FLinkerLoad::SerializeBulkData(FBulkData& BulkData, const FBulkDataSerializ
 			Meta.SetOffset(Meta.GetOffset() + Summary.BulkDataStartOffset);
 		}
 
-		if (IsAllowingLazyLoading() == false)
+		if (bLazyLoadable == false)
 		{
 			FArchive& Ar = *this;
 			FArchive::FScopeSeekTo _(Ar, Meta.GetOffset());
@@ -6820,7 +6802,13 @@ bool FLinkerLoad::SerializeBulkData(FBulkData& BulkData, const FBulkDataSerializ
 			BulkData.SerializeBulkData(Ar, Payload, Meta.GetSize(), Meta.GetFlags());
 		}
 	}
-#endif // USE_RUNTIME_BULKDATA
+
+	if (bLazyLoadable == false)
+	{
+		BulkData.ForceBulkDataResident();
+		Meta.ClearFlags(BULKDATA_LazyLoadable);
+		BulkData.BulkChunkId = FIoChunkId::InvalidChunkId;
+	}
 
 	return true;
 }
