Bulk Data - fix wrong offset being used when making inline bulk data reloadable in the old linker load path

#rnx

diff --git a/Engine/Source/Runtime/CoreUObject/Private/UObject/LinkerLoad.cpp b/Engine/Source/Runtime/CoreUObject/Private/UObject/LinkerLoad.cpp
--- a/Engine/Source/Runtime/CoreUObject/Private/UObject/LinkerLoad.cpp
+++ b/Engine/Source/Runtime/CoreUObject/Private/UObject/LinkerLoad.cpp
@@ -6745,7 +6745,7 @@ bool FLinkerLoad::SerializeBulkData(FBulkData& BulkData, const FBulkDataSerializ
 		{
 			// Cooked packages are split into .uasset/.exp files and the offset needs to be adjusted accordingly.
 			const int64 PkgHeaderSize = IPackageResourceManager::Get().FileSize(PackagePath,  EPackageSegment::Header);
-			Meta.SetOffset(Meta.GetOffset() - PkgHeaderSize);
+			Meta.SetOffset(Tell() - PkgHeaderSize);
 		}
 		void* Payload = BulkData.ReallocateData(Meta.GetSize());
 		BulkData.SerializeBulkData(*this, Payload, Meta.GetSize(), Meta.GetFlags());
