Split GC into full and incremental functions to help both CSV and sampling profilers

#rnx

diff --git a/Engine/Source/Runtime/CoreUObject/Private/UObject/GarbageCollection.cpp b/Engine/Source/Runtime/CoreUObject/Private/UObject/GarbageCollection.cpp
--- a/Engine/Source/Runtime/CoreUObject/Private/UObject/GarbageCollection.cpp
+++ b/Engine/Source/Runtime/CoreUObject/Private/UObject/GarbageCollection.cpp
@@ -4015,26 +4015,51 @@ static void UpdateGCHistory(TConstArrayView<TUniquePtr<FWorkerContext>> Contexts
 }
 
 #if UE_WITH_GC
+
+template<bool bPerformFullPurge>
+FORCENOINLINE void CollectGarbageImpl(EObjectFlags KeepFlags);
+
+FORCENOINLINE static void CollectGarbageIncremental(EObjectFlags KeepFlags)
+{
+	SCOPE_TIME_GUARD(TEXT("Collect Garbage Incremental"));
+	SCOPED_NAMED_EVENT(CollectGarbageIncremental, FColor::Red);
+	CSV_EVENT_GLOBAL(TEXT("GC"));
+	CSV_SCOPED_TIMING_STAT_EXCLUSIVE(GarbageCollection);
+
+	CollectGarbageImpl<false>(KeepFlags);
+}
+
+FORCENOINLINE static void CollectGarbageFull(EObjectFlags KeepFlags)
+{
+	SCOPE_TIME_GUARD(TEXT("Collect Garbage Full"));
+	SCOPED_NAMED_EVENT(CollectGarbageFull, FColor::Red);
+	CSV_EVENT_GLOBAL(TEXT("GC"));
+	CSV_SCOPED_TIMING_STAT_EXCLUSIVE(GarbageCollectionFull);
+
+	CollectGarbageImpl<true>(KeepFlags);
+}
+
+FORCEINLINE void CollectGarbageInternal(EObjectFlags KeepFlags, bool bPerformFullPurge)
+{
+	if (bPerformFullPurge)
+	{
+		CollectGarbageFull(KeepFlags);
+	}
+	else
+	{
+		CollectGarbageIncremental(KeepFlags);
+	}
+}
+
 /** 
  * Deletes all unreferenced objects, keeping objects that have any of the passed in KeepFlags set
  *
  * @param	KeepFlags			objects with those flags will be kept regardless of being referenced or not
  * @param	bPerformFullPurge	if true, perform a full purge after the mark pass
  */
-static void CollectGarbageInternal(EObjectFlags KeepFlags, bool bPerformFullPurge)
+template<bool bPerformFullPurge>
+void CollectGarbageImpl(EObjectFlags KeepFlags)
 {
-	if (GIsInitialLoad)
-	{
-		// During initial load classes may not yet have their GC token streams assembled
-		UE_LOG(LogGarbage, Log, TEXT("Skipping CollectGarbage() call during initial load. It's not safe."));
-		return;
-	}
-
-	SCOPE_TIME_GUARD(TEXT("Collect Garbage"));
-	SCOPED_NAMED_EVENT(CollectGarbageInternal, FColor::Red);
-	CSV_EVENT_GLOBAL(TEXT("GC"));
-	CSV_SCOPED_TIMING_STAT_EXCLUSIVE(GarbageCollection);
-
 	FGCCSyncObject::Get().ResetGCIsWaiting();
 
 #if defined(WITH_CODE_GUARD_HANDLER) && WITH_CODE_GUARD_HANDLER
@@ -4339,6 +4364,13 @@ bool UnhashUnreachableObjects(bool bUseTimeLimit, double TimeLimit)
 
 void CollectGarbage(EObjectFlags KeepFlags, bool bPerformFullPurge)
 {
+	if (GIsInitialLoad)
+	{
+		// During initial load classes may not yet have their GC token streams assembled
+		UE_LOG(LogGarbage, Log, TEXT("Skipping CollectGarbage() call during initial load. It's not safe."));
+		return;
+	}
+
 	// No other thread may be performing UObject operations while we're running
 	AcquireGCLock();
 
@@ -4355,6 +4387,13 @@ void CollectGarbage(EObjectFlags KeepFlags, bool bPerformFullPurge)
 
 bool TryCollectGarbage(EObjectFlags KeepFlags, bool bPerformFullPurge)
 {
+	if (GIsInitialLoad)
+	{
+		// During initial load classes may not yet have their GC token streams assembled
+		UE_LOG(LogGarbage, Log, TEXT("Skipping CollectGarbage() call during initial load. It's not safe."));
+		return false;
+	}
+
 	// No other thread may be performing UObject operations while we're running
 	bool bCanRunGC = FGCCSyncObject::Get().TryGCLock();
 	if (!bCanRunGC)
