Fix issue where the GC reference search was failing to handle weak blueprint ubergraph references properly after it was refactored.

Code by [REDACTED_PERSON]

#rnx

diff --git a/Engine/Source/Runtime/CoreUObject/Private/UObject/ReferenceChainSearch.cpp b/Engine/Source/Runtime/CoreUObject/Private/UObject/ReferenceChainSearch.cpp
--- a/Engine/Source/Runtime/CoreUObject/Private/UObject/ReferenceChainSearch.cpp
+++ b/Engine/Source/Runtime/CoreUObject/Private/UObject/ReferenceChainSearch.cpp
@@ -918,7 +918,7 @@ struct TReferenceSearchHelper
 			UE::GC::FWorkerContext Context;
 			ObjectsToProcess = {Object};
 			Context.SetInitialObjectsUnpadded(ObjectsToProcess);
-			CollectReferences(Processor, Context);
+			CollectReferences<FDirectReferenceCollector>(Processor, Context);
 
 			if (!HandleReferences(Object, Processor.GetReferencedObjects()))
 			{
