Deprecate ISavePackageValidator and replace the feature with list of validation functions for import and exports
Add a default save setting object to which validation function can be added.

diff --git a/Engine/Source/Editor/UnrealEd/Private/CookOnTheFlyServer.cpp b/Engine/Source/Editor/UnrealEd/Private/CookOnTheFlyServer.cpp
--- a/Engine/Source/Editor/UnrealEd/Private/CookOnTheFlyServer.cpp
+++ b/Engine/Source/Editor/UnrealEd/Private/CookOnTheFlyServer.cpp
@@ -20,7 +20,6 @@
 #include "Containers/RingBuffer.h"
 #include "Cooker/AsyncIODelete.h"
 #include "Cooker/CookDirector.h"
-#include "Cooker/CookedSavePackageValidator.h"
 #include "Cooker/CookMPCollector.h"
 #include "Cooker/CookOnTheFlyServerInterface.h"
 #include "Cooker/CookPackageData.h"
@@ -9267,8 +9266,59 @@ UE::Cook::FCookSavePackageContext* UCookOnTheFlyServer::CreateSaveContext(const
 
 	DiffModeHelper->InitializePackageWriter(PackageWriter);
 
-	FCookSavePackageContext* Context = new FCookSavePackageContext(TargetPlatform, PackageWriter, WriterDebugName);
-	Context->SaveContext.SetValidator(MakeUnique<FCookedSavePackageValidator>(TargetPlatform, *this));
+	// Setup save package settings (i.e. validation)
+	FSavePackageSettings SavePackageSettings = FSavePackageSettings::GetDefaultSettings();
+	{
+		// Setup Import Validation for suppressed module native classes
+		TSet<FName> DisabledNativeScriptPackages;
+		for (TSharedRef<IPlugin>& Plugin : IPluginManager::Get().GetEnabledPlugins())
+		{
+			if (!TargetPlatform->IsEnabledForPlugin(Plugin.Get()))
+			{
+				for (const FModuleDescriptor& Module : Plugin->GetDescriptor().Modules)
+				{
+					DisabledNativeScriptPackages.Add(FPackageName::GetModuleScriptPackageName(Module.Name));
+				}
+			}
+		}
+
+		SavePackageSettings.AddExternalImportValidation([SuppressedNativeScriptPackages = MoveTemp(DisabledNativeScriptPackages), TargetPlatform, this](const FImportsValidationContext& ValidationContext)
+			{
+				for (UObject* Object : ValidationContext.Imports)
+				{
+					UClass* Class = Cast<UClass>(Object);
+					if (Class && Class->IsNative() && SuppressedNativeScriptPackages.Contains(Class->GetPackage()->GetFName()))
+					{
+						FInstigator Instigator = GetInstigator(ValidationContext.Package->GetFName());
+						bool bIsError = true;
+						if (Instigator.Category == EInstigator::StartupPackage)
+						{
+							// StartupPackages might be around just because of the editor;
+							// if they're not available on client, ignore them without error
+							bIsError = false;
+						}
+
+						// If you receive this message in a package that you do want to cook, you can remove the object of the
+						// unavailable class by overriding UObject::NeedsLoadForTargetPlatform on that class to return false.
+						if (bIsError)
+						{
+							UE_ASSET_LOG(LogCook, Error, ValidationContext.Package, TEXT("Failed to cook %s for platform %s. It imports class %s, which is in a module that is not available on the platform."),
+								*ValidationContext.Package->GetName(), *TargetPlatform->PlatformName(), *Class->GetPathName());
+							return ESavePackageResult::ValidatorError;
+						}
+						else
+						{
+							UE_ASSET_LOG(LogCook, Display, ValidationContext.Package, TEXT("Skipping package %s for platform %s. It imports class %s, which is in a module that is not available on the platform."),
+								*ValidationContext.Package->GetName(), *TargetPlatform->PlatformName(), *Class->GetPathName());
+							return ESavePackageResult::ValidatorSuppress;
+						}
+					}
+				}
+				return ESavePackageResult::Success;
+			});
+	}
+
+	FCookSavePackageContext* Context = new FCookSavePackageContext(TargetPlatform, PackageWriter, WriterDebugName, MoveTemp(SavePackageSettings));
 	return Context;
 
 }
diff --git a/Engine/Source/Editor/UnrealEd/Private/Cooker/CookTypes.cpp b/Engine/Source/Editor/UnrealEd/Private/Cooker/CookTypes.cpp
--- a/Engine/Source/Editor/UnrealEd/Private/Cooker/CookTypes.cpp
+++ b/Engine/Source/Editor/UnrealEd/Private/Cooker/CookTypes.cpp
@@ -94,8 +94,8 @@ void SetIsSchedulerThread(bool bValue)
 }
 
 FCookSavePackageContext::FCookSavePackageContext(const ITargetPlatform* InTargetPlatform,
-	ICookedPackageWriter* InPackageWriter, FStringView InWriterDebugName)
-	: SaveContext(InTargetPlatform, InPackageWriter)
+	ICookedPackageWriter* InPackageWriter, FStringView InWriterDebugName, FSavePackageSettings InSettings)
+	: SaveContext(InTargetPlatform, InPackageWriter, MoveTemp(InSettings))
 	, WriterDebugName(InWriterDebugName)
 	, PackageWriter(InPackageWriter)
 {
diff --git a/Engine/Source/Editor/UnrealEd/Private/Cooker/CookTypes.h b/Engine/Source/Editor/UnrealEd/Private/Cooker/CookTypes.h
--- a/Engine/Source/Editor/UnrealEd/Private/Cooker/CookTypes.h
+++ b/Engine/Source/Editor/UnrealEd/Private/Cooker/CookTypes.h
@@ -237,7 +237,7 @@ namespace UE::Cook
 	struct FCookSavePackageContext
 	{
 		FCookSavePackageContext(const ITargetPlatform* InTargetPlatform,
-			ICookedPackageWriter* InPackageWriter, FStringView InWriterDebugName);
+			ICookedPackageWriter* InPackageWriter, FStringView InWriterDebugName, FSavePackageSettings InSettings);
 		~FCookSavePackageContext();
 
 		FSavePackageContext SaveContext;
diff --git a/Engine/Source/Editor/UnrealEd/Private/Cooker/CookedSavePackageValidator.cpp b/Engine/Source/Editor/UnrealEd/Private/Cooker/CookedSavePackageValidator.cpp
deleted file mode 100644
diff --git a/Engine/Source/Editor/UnrealEd/Private/Cooker/CookedSavePackageValidator.h b/Engine/Source/Editor/UnrealEd/Private/Cooker/CookedSavePackageValidator.h
deleted file mode 100644
diff --git a/Engine/Source/Editor/UnrealEd/Private/EditorEngine.cpp b/Engine/Source/Editor/UnrealEd/Private/EditorEngine.cpp
--- a/Engine/Source/Editor/UnrealEd/Private/EditorEngine.cpp
+++ b/Engine/Source/Editor/UnrealEd/Private/EditorEngine.cpp
@@ -4581,6 +4581,14 @@ FSavePackageResultStruct UEditorEngine::Save(UPackage* InOuter, UObject* InAsset
 		}
 	}
 
+	// if no save package context was passed in and the default settings were modified, install a context for the save
+	TUniquePtr<FSavePackageContext> UniqueContext;
+	if (SaveArgs.SavePackageContext == nullptr && !FSavePackageSettings::GetDefaultSettings().IsDefault())
+	{
+		UniqueContext = MakeUnique<FSavePackageContext>(nullptr, nullptr, FSavePackageSettings::GetDefaultSettings());
+		SaveArgs.SavePackageContext = UniqueContext.Get();
+	}
+
 	SlowTask.EnterProgressFrame(10);
 
 	UWorld* World = Cast<UWorld>(Asset);
diff --git a/Engine/Source/Runtime/CoreUObject/Private/UObject/SavePackage/SaveContext.h b/Engine/Source/Runtime/CoreUObject/Private/UObject/SavePackage/SaveContext.h
--- a/Engine/Source/Runtime/CoreUObject/Private/UObject/SavePackage/SaveContext.h
+++ b/Engine/Source/Runtime/CoreUObject/Private/UObject/SavePackage/SaveContext.h
@@ -159,6 +159,11 @@ struct FHarvestedRealm
 		return Exports;
 	}
 
+	const TSet<FTaggedExport>& GetExports() const
+	{
+		return Exports;
+	}
+
 	const TSet<UObject*>& GetImports() const
 	{
 		return Imports;
@@ -983,10 +988,34 @@ public:
 		return PackageWriter;
 	}
 
+	PRAGMA_DISABLE_DEPRECATION_WARNINGS
 	ISavePackageValidator* GetPackageValidator() const
 	{
 		return SaveArgs.SavePackageContext ? SaveArgs.SavePackageContext->GetValidator() : nullptr;
 	}
+	PRAGMA_ENABLE_DEPRECATION_WARNINGS
+
+	bool HasExternalImportValidations() const
+	{
+		return SaveArgs.SavePackageContext ? SaveArgs.SavePackageContext->GetExternalImportValidations().Num() > 0 : false;
+	}
+
+	const TArray<TFunction<FSavePackageSettings::ExternalImportValidationFunc>> GetExternalImportValidations() const
+	{
+		check(SaveArgs.SavePackageContext);
+		return SaveArgs.SavePackageContext->GetExternalImportValidations();
+	}
+
+	bool HasExternalExportValidations() const
+	{
+		return SaveArgs.SavePackageContext ? SaveArgs.SavePackageContext->GetExternalExportValidations().Num() > 0 : false;
+	}
+
+	const TArray<TFunction<FSavePackageSettings::ExternalExportValidationFunc>> GetExternalExportValidations() const
+	{
+		check(SaveArgs.SavePackageContext);
+		return SaveArgs.SavePackageContext->GetExternalExportValidations();
+	}
 
 	const FHarvestedRealm& GetHarvestedRealm(ESaveRealm Realm = ESaveRealm::None) const
 	{
@@ -1060,9 +1089,10 @@ private:
 	// The current default harvesting context being queried by the save context
 	ESaveRealm CurrentHarvestingRealm = ESaveRealm::None;
 
-	// Set of harvested content split per harvesting context
+	// List of harvested content split per harvesting context
 	TArray<FHarvestedRealm> HarvestedRealms;
 
+	// List of harvested illegal references
 	TArray<FIllegalReference> HarvestedIllegalReferences;
 
 	// Set of harvested prestream packages, should be deprecated
diff --git a/Engine/Source/Runtime/CoreUObject/Private/UObject/SavePackage/SavePackageUtilities.cpp b/Engine/Source/Runtime/CoreUObject/Private/UObject/SavePackage/SavePackageUtilities.cpp
--- a/Engine/Source/Runtime/CoreUObject/Private/UObject/SavePackage/SavePackageUtilities.cpp
+++ b/Engine/Source/Runtime/CoreUObject/Private/UObject/SavePackage/SavePackageUtilities.cpp
@@ -182,6 +182,12 @@ namespace UE
 	}
 }
 
+ FSavePackageSettings& FSavePackageSettings::GetDefaultSettings()
+{
+	static FSavePackageSettings Default;
+	return Default;
+}
+
 namespace SavePackageUtilities
 {
 const FName NAME_World("World");
diff --git a/Engine/Source/Runtime/CoreUObject/Private/UObject/SavePackage2.cpp b/Engine/Source/Runtime/CoreUObject/Private/UObject/SavePackage2.cpp
--- a/Engine/Source/Runtime/CoreUObject/Private/UObject/SavePackage2.cpp
+++ b/Engine/Source/Runtime/CoreUObject/Private/UObject/SavePackage2.cpp
@@ -529,25 +529,6 @@ ESavePackageResult ValidateExports(FSaveContext& SaveContext)
 		}
 	}
 
-#if WITH_EDITOR
-	if (GOutputCookingWarnings)
-	{
-		// check the name list for UniqueObjectNameForCooking cooking
-		if (SaveContext.NameExists(NAME_UniqueObjectNameForCookingComparisonIndex))
-		{
-			for (const FTaggedExport& Export : SaveContext.GetExports())
-			{
-				FName NameInUse = Export.Obj->GetFName();
-				if (NameInUse.GetComparisonIndex() == NAME_UniqueObjectNameForCookingComparisonIndex)
-				{
-					UObject* Outer = Export.Obj->GetOuter();
-					UE_LOG(LogSavePackage, Warning, TEXT("Saving object into cooked package %s which was created at cook time, Object Name %s, Full Path %s, Class %s, Outer %s, Outer class %s"), SaveContext.GetFilename(), *NameInUse.ToString(), *Export.Obj->GetFullName(), *Export.Obj->GetClass()->GetName(), Outer ? *Outer->GetName() : TEXT("None"), Outer ? *Outer->GetClass()->GetName() : TEXT("None"));
-				}
-			}
-		}
-	}
-#endif
-
 	// If this is a map package, make sure there is a world or level in the export map.
 	if (SaveContext.GetPackage()->ContainsMap())
 	{
@@ -579,7 +560,42 @@ ESavePackageResult ValidateExports(FSaveContext& SaveContext)
 		}
 	}
 
+	// Validate External Export Rules
+	if (SaveContext.HasExternalExportValidations())
+	{
+		TSet<UObject*> Exports;
+		Algo::Transform(SaveContext.GetExports(), Exports, [](const FTaggedExport& InExport) { return InExport.Obj; });
+
+		for (const TFunction<FSavePackageSettings::ExternalExportValidationFunc>& ValidateExport : SaveContext.GetExternalExportValidations())
+		{
+			SaveContext.Result = ValidateExport({ SaveContext.GetPackage(), Exports, SaveContext.IsGenerateSaveError() ? SaveContext.GetError() : nullptr });
+			if (SaveContext.Result != ESavePackageResult::Success)
+			{
+				return SaveContext.Result;
+			}
+		}
+	}
+
 	// Cooking checks
+#if WITH_EDITOR
+	if (GOutputCookingWarnings)
+	{
+		// check the name list for UniqueObjectNameForCooking cooking
+		if (SaveContext.NameExists(NAME_UniqueObjectNameForCookingComparisonIndex))
+		{
+			for (const FTaggedExport& Export : SaveContext.GetExports())
+			{
+				FName NameInUse = Export.Obj->GetFName();
+				if (NameInUse.GetComparisonIndex() == NAME_UniqueObjectNameForCookingComparisonIndex)
+				{
+					UObject* Outer = Export.Obj->GetOuter();
+					UE_LOG(LogSavePackage, Warning, TEXT("Saving object into cooked package %s which was created at cook time, Object Name %s, Full Path %s, Class %s, Outer %s, Outer class %s"), SaveContext.GetFilename(), *NameInUse.ToString(), *Export.Obj->GetFullName(), *Export.Obj->GetClass()->GetName(), Outer ? *Outer->GetName() : TEXT("None"), Outer ? *Outer->GetClass()->GetName() : TEXT("None"));
+				}
+			}
+		}
+	}
+#endif
+
 	if (SaveContext.IsCooking())
 	{
 		// Add the exports for the cook checker
@@ -820,8 +836,21 @@ ESavePackageResult ValidateImports(FSaveContext& SaveContext)
 		return ValidateIllegalReferences(SaveContext, PrivateObjects, ObjectsInOtherMaps);
 	}
 
-	ISavePackageValidator* Validator = SaveContext.GetPackageValidator();
-	if (Validator)
+	// Validate External Import Rules
+	if (SaveContext.HasExternalImportValidations())
+	{
+		for (const TFunction<FSavePackageSettings::ExternalImportValidationFunc>& ValidateImport : SaveContext.GetExternalImportValidations())
+		{
+			SaveContext.Result = ValidateImport({ SaveContext.GetPackage(), SaveContext.GetImports(), SaveContext.IsGenerateSaveError() ? SaveContext.GetError() : nullptr });
+			if (SaveContext.Result != ESavePackageResult::Success)
+			{
+				return SaveContext.Result;
+			}
+		}
+	}
+
+	PRAGMA_DISABLE_DEPRECATION_WARNINGS
+	if (ISavePackageValidator* Validator = SaveContext.GetPackageValidator())
 	{
 		ESavePackageResult ValidatorResult = Validator->ValidateImports(Package, Imports);
 		if (ValidatorResult != ESavePackageResult::Success)
@@ -829,6 +858,7 @@ ESavePackageResult ValidateImports(FSaveContext& SaveContext)
 			return ValidatorResult;
 		}
 	}
+	PRAGMA_ENABLE_DEPRECATION_WARNINGS
 
 	// Cooking checks
 	// Currently do not use the edl checker for the optional context
diff --git a/Engine/Source/Runtime/CoreUObject/Public/UObject/SavePackage.h b/Engine/Source/Runtime/CoreUObject/Public/UObject/SavePackage.h
--- a/Engine/Source/Runtime/CoreUObject/Public/UObject/SavePackage.h
+++ b/Engine/Source/Runtime/CoreUObject/Public/UObject/SavePackage.h
@@ -128,6 +128,7 @@ struct FSavePackageArgs
 };
 
 /** Interface for SavePackage to test for caller-specific errors. */
+class UE_DEPRECATED(5.2, "Use a FSavePackageContext::ExternalValidationFunc if you need to run external validation") ISavePackageValidator;
 class ISavePackageValidator
 {
 public:
@@ -138,12 +139,83 @@ public:
 	virtual ESavePackageResult ValidateImports(const UPackage* Package, const TSet<UObject*>& Imports) = 0;
 };
 
+
+/** Param struct for external import validation functions */
+struct FImportsValidationContext
+{
+	FImportsValidationContext(const UPackage* InPackage, const TSet<UObject*>& InImports, FOutputDevice* InOutputDevice)
+		: Package(InPackage)
+		, Imports(InImports)
+		, OutputDevice(InOutputDevice)
+	{}
+
+	const UPackage* Package;
+	const TSet<UObject*>& Imports;
+	FOutputDevice* OutputDevice;
+};
+
+/** Param struct for external export validation functions */
+struct FExportsValidationContext
+{
+	FExportsValidationContext(const UPackage* InPackage, const TSet<UObject*>& InExports, FOutputDevice* InOutputDevice)
+		: Package(InPackage)
+		, Exports(InExports)
+		, OutputDevice(InOutputDevice)
+	{}
+
+	const UPackage* Package;
+	const TSet<UObject*>& Exports;
+	FOutputDevice* OutputDevice;
+};
+
+/** struct persistent settings used by all save unless overridden. @see FSavePackageContext */
+class FSavePackageSettings
+{
+public:
+	typedef ESavePackageResult ExternalImportValidationFunc(const FImportsValidationContext& InValidationContext);
+	typedef ESavePackageResult ExternalExportValidationFunc(const FExportsValidationContext& InValidationContext);
+	
+	FSavePackageSettings() = default;
+
+	/** Get the default settings by save when none are specified. */
+	COREUOBJECT_API static FSavePackageSettings& GetDefaultSettings();
+
+	COREUOBJECT_API bool IsDefault() const
+	{
+		return ExternalImportValidations.Num() == 0 && ExternalExportValidations.Num() == 0;
+	}
+
+	COREUOBJECT_API const TArray<TFunction<ExternalImportValidationFunc>>& GetExternalImportValidations() const
+	{
+		return ExternalImportValidations;
+	}
+	COREUOBJECT_API const TArray<TFunction<ExternalExportValidationFunc>>& GetExternalExportValidations() const
+	{
+		return ExternalExportValidations;
+	}
+
+	COREUOBJECT_API void AddExternalImportValidation(TFunction<ExternalImportValidationFunc> InValidation)
+	{
+		ExternalImportValidations.Add(MoveTemp(InValidation));
+	}
+	COREUOBJECT_API void AddExternalExportValidation(TFunction<ExternalExportValidationFunc> InValidation)
+	{
+		ExternalExportValidations.Add(MoveTemp(InValidation));
+	}
+
+private:
+	TArray<TFunction<ExternalImportValidationFunc>> ExternalImportValidations;
+	TArray<TFunction<ExternalExportValidationFunc>> ExternalExportValidations;
+};
+
 class FSavePackageContext
 {
 public:
-	FSavePackageContext(const ITargetPlatform* InTargetPlatform, IPackageWriter* InPackageWriter)
+
+	FSavePackageContext(const ITargetPlatform* InTargetPlatform, IPackageWriter* InPackageWriter, FSavePackageSettings InSettings = FSavePackageSettings())
 	: TargetPlatform(InTargetPlatform)
-	, PackageWriter(InPackageWriter) 
+	, PackageWriter(InPackageWriter)
+	, SavePackageSettings(MoveTemp(InSettings))
 	{
 		if (PackageWriter)
 		{
@@ -159,21 +231,36 @@ public:
 
 	COREUOBJECT_API ~FSavePackageContext();
 
+	PRAGMA_DISABLE_DEPRECATION_WARNINGS
 	ISavePackageValidator* GetValidator()
 	{
-		return Validator.Get();
+		return Validator.Get();	
 	}
 	COREUOBJECT_API void SetValidator(TUniquePtr<ISavePackageValidator>&& InValidator)
 	{
 		Validator = MoveTemp(InValidator);
 	}
+	PRAGMA_ENABLE_DEPRECATION_WARNINGS
+
+	COREUOBJECT_API const TArray<TFunction<FSavePackageSettings::ExternalImportValidationFunc>>& GetExternalImportValidations() const
+	{
+		return SavePackageSettings.GetExternalImportValidations();
+	}
+
+	COREUOBJECT_API const TArray<TFunction<FSavePackageSettings::ExternalExportValidationFunc>>& GetExternalExportValidations() const
+	{
+		return SavePackageSettings.GetExternalExportValidations();
+	}
 
 	const ITargetPlatform* const TargetPlatform;
 	IPackageWriter* const PackageWriter;
 	IPackageWriter::FCapabilities PackageWriterCapabilities;
 
 private:
+	PRAGMA_DISABLE_DEPRECATION_WARNINGS
 	TUniquePtr<ISavePackageValidator> Validator;
+	PRAGMA_ENABLE_DEPRECATION_WARNINGS
+	FSavePackageSettings SavePackageSettings;
 public:
 
 	UE_DEPRECATED(5.0, "bForceLegacyOffsets is no longer supported; remove uses of the variable")
