Remove AllowCrossLevel override on TObjectPtr property so that it behaves like raw pointers and has references NULLed when copying + pasting.

diff --git a/Engine/Source/Runtime/CoreUObject/Private/UObject/PropertyObjectPtr.cpp b/Engine/Source/Runtime/CoreUObject/Private/UObject/PropertyObjectPtr.cpp
--- a/Engine/Source/Runtime/CoreUObject/Private/UObject/PropertyObjectPtr.cpp
+++ b/Engine/Source/Runtime/CoreUObject/Private/UObject/PropertyObjectPtr.cpp
@@ -185,11 +185,6 @@ void FObjectPtrProperty::SetObjectPropertyValue_InContainer(void* ContainerAddre
 	}
 }
 
-bool FObjectPtrProperty::AllowCrossLevel() const
-{
-	return true;
-}
-
 bool FObjectPtrProperty::AllowObjectTypeReinterpretationTo(const FObjectPropertyBase* Other) const
 {
 	return Other && Other->IsA<FObjectProperty>();
diff --git a/Engine/Source/Runtime/CoreUObject/Public/UObject/UnrealType.h b/Engine/Source/Runtime/CoreUObject/Public/UObject/UnrealType.h
--- a/Engine/Source/Runtime/CoreUObject/Public/UObject/UnrealType.h
+++ b/Engine/Source/Runtime/CoreUObject/Public/UObject/UnrealType.h
@@ -2990,7 +2990,6 @@ class COREUOBJECT_API FObjectPtrProperty : public FObjectProperty
 	virtual UObject* GetObjectPropertyValue_InContainer(const void* ContainerAddress, int32 ArrayIndex = 0) const override;
 	virtual void SetObjectPropertyValue(void* PropertyValueAddress, UObject* Value) const override;
 	virtual void SetObjectPropertyValue_InContainer(void* ContainerAddress, UObject* Value, int32 ArrayIndex = 0) const override;
-	virtual bool AllowCrossLevel() const override;
 	virtual bool AllowObjectTypeReinterpretationTo(const FObjectPropertyBase* Other) const override;
 private:
 	virtual uint32 GetValueTypeHashInternal(const void* Src) const override;
