changed FObjectPropertyBase::CheckValidObject to not resolve FObjectPtr fields to do validation
this is needed for feature work with validation of accessing asset pointers during serialization

https://p4-swarm.epicgames.net/reviews/23782822

diff --git a/Engine/Source/Runtime/CoreUObject/Private/UObject/PropertyBaseObject.cpp b/Engine/Source/Runtime/CoreUObject/Private/UObject/PropertyBaseObject.cpp
--- a/Engine/Source/Runtime/CoreUObject/Private/UObject/PropertyBaseObject.cpp
+++ b/Engine/Source/Runtime/CoreUObject/Private/UObject/PropertyBaseObject.cpp
@@ -611,69 +611,70 @@ bool FObjectPropertyBase::AllowObjectTypeReinterpretationTo(const FObjectPropert
 
 void FObjectPropertyBase::CheckValidObject(void* ValueAddress, UObject* OldValue) const
 {
-	UObject *Object = GetObjectPropertyValue(ValueAddress);
-	if (Object)
+	const TObjectPtr<UObject> Object = GetObjectPtrPropertyValue(ValueAddress);
+	if (!Object)
 	{
-		//
-		// here we want to make sure the the object value still matches the 
-		// object type expected by the property...
+		return;
+	}
+	//
+	// here we want to make sure the the object value still matches the 
+	// object type expected by the property...
 
-		UClass* ObjectClass = Object->GetClass();
-		UE_CLOG(!ObjectClass, LogProperty, Fatal, TEXT("Object without class referenced by %s, object: 0x%016llx %s"), *GetPathName(), (int64)(PTRINT)Object, *Object->GetPathName());
+	UClass* ObjectClass = Object.GetClass();
+	UE_CLOG(!ObjectClass, LogProperty, Fatal, TEXT("Object without class referenced by %s, object: 0x%016llx %s"), *GetPathName(), (int64)(PTRINT)ValueAddress, *Object.GetPathName());
 
 #if USE_CIRCULAR_DEPENDENCY_LOAD_DEFERRING
-		auto IsDeferringValueLoad = [&]()
-		{
-			FLinkerLoad* PropertyLinker = GetLinker();
-			return ((PropertyLinker == nullptr) || (PropertyLinker->LoadFlags & LOAD_DeferDependencyLoads)) &&
-				(Object->IsA<ULinkerPlaceholderExportObject>() || Object->IsA<ULinkerPlaceholderClass>());
-		};
+	auto IsDeferringValueLoad = [&]()
+	{
+		FLinkerLoad* PropertyLinker = GetLinker();
+		return ((PropertyLinker == nullptr) || (PropertyLinker->LoadFlags & LOAD_DeferDependencyLoads)) &&
+			(ObjectClass->IsChildOf<ULinkerPlaceholderExportObject>() || ObjectClass->IsChildOf<ULinkerPlaceholderClass>());
+	};
 
 #if USE_DEFERRED_DEPENDENCY_CHECK_VERIFICATION_TESTS
-		check( IsDeferringValueLoad() || (!Object->IsA<ULinkerPlaceholderExportObject>() && !Object->IsA<ULinkerPlaceholderClass>()) );
+	check( IsDeferringValueLoad() || (!Object->IsA<ULinkerPlaceholderExportObject>() && !Object->IsA<ULinkerPlaceholderClass>()) );
 #endif // USE_DEFERRED_DEPENDENCY_CHECK_VERIFICATION_TESTS
 
 #else  // USE_CIRCULAR_DEPENDENCY_LOAD_DEFERRING 
-		auto IsDeferringValueLoad = [&]() { return false; };
+	auto IsDeferringValueLoad = [&]() { return false; };
 #endif // USE_CIRCULAR_DEPENDENCY_LOAD_DEFERRING
 
-		if ((PropertyClass != nullptr) && !ObjectClass->IsChildOf(PropertyClass) && !ObjectClass->GetAuthoritativeClass()->IsChildOf(PropertyClass))
-		{
+	if ((PropertyClass != nullptr) && !ObjectClass->IsChildOf(PropertyClass) && !ObjectClass->GetAuthoritativeClass()->IsChildOf(PropertyClass))
+	{
 			
-			// we could be in the middle of replacing references to the 
-			// PropertyClass itself (in the middle of an FArchiveReplaceObjectRef 
-			// pass)... if this is the case, then we might have already replaced 
-			// the object's class, but not the PropertyClass yet (or vise-versa)... 
-			// so we use this to ensure, in that situation, that we don't clear the 
-			// object value (if CLASS_NewerVersionExists is set, then we are likely 
-			// in the middle of an FArchiveReplaceObjectRef pass)
-			bool bIsReplacingClassRefs = PropertyClass && PropertyClass->HasAnyClassFlags(CLASS_NewerVersionExists) != ObjectClass->HasAnyClassFlags(CLASS_NewerVersionExists);
-			if (!bIsReplacingClassRefs && !IsDeferringValueLoad())
+		// we could be in the middle of replacing references to the 
+		// PropertyClass itself (in the middle of an FArchiveReplaceObjectRef 
+		// pass)... if this is the case, then we might have already replaced 
+		// the object's class, but not the PropertyClass yet (or vise-versa)... 
+		// so we use this to ensure, in that situation, that we don't clear the 
+		// object value (if CLASS_NewerVersionExists is set, then we are likely 
+		// in the middle of an FArchiveReplaceObjectRef pass)
+		bool bIsReplacingClassRefs = PropertyClass && PropertyClass->HasAnyClassFlags(CLASS_NewerVersionExists) != ObjectClass->HasAnyClassFlags(CLASS_NewerVersionExists);
+		if (!bIsReplacingClassRefs && !IsDeferringValueLoad())
+		{
+			if (!HasAnyPropertyFlags(CPF_NonNullable))
 			{
-				if (!HasAnyPropertyFlags(CPF_NonNullable))
-				{
-					UE_LOG(LogProperty, Warning,
-						TEXT("Serialized %s for a property of %s. Reference will be nullptred.\n    Property = %s\n    Item = %s"),
-						*Object->GetClass()->GetFullName(),
-						*PropertyClass->GetFullName(),
-						*GetFullName(),
-						*Object->GetFullName()
-					);
-					SetObjectPropertyValue(ValueAddress, nullptr);
-				}
-				else
-				{
-					checkf(OldValue, TEXT("CheckValidObject(\"%s\") trying to assign null object value to non-nullable property \"%s\""), *Object->GetFullName(), *GetFullName());
-					UE_LOG(LogProperty, Warning,
-						TEXT("Serialized %s for a property of %s. Reference will be reverted back to %s.\n    Property = %s\n    Item = %s"),
-						*Object->GetClass()->GetFullName(),
-						*PropertyClass->GetFullName(),
-						*OldValue->GetFullName(),
-						*GetFullName(),
-						*Object->GetFullName()
-					);
-					SetObjectPropertyValue(ValueAddress, OldValue);
-				}
+				UE_LOG(LogProperty, Warning,
+					TEXT("Serialized %s for a property of %s. Reference will be nullptred.\n    Property = %s\n    Item = %s"),
+					*ObjectClass->GetFullName(),
+					*PropertyClass->GetFullName(),
+					*GetFullName(),
+					*Object.GetFullName()
+				);
+				SetObjectPropertyValue(ValueAddress, nullptr);
+			}
+			else
+			{
+				checkf(OldValue, TEXT("CheckValidObject(\"%s\") trying to assign null object value to non-nullable property \"%s\""), *Object.GetFullName(), *GetFullName());
+				UE_LOG(LogProperty, Warning,
+					TEXT("Serialized %s for a property of %s. Reference will be reverted back to %s.\n    Property = %s\n    Item = %s"),
+					*ObjectClass->GetFullName(),
+					*PropertyClass->GetFullName(),
+					OldValue ? *OldValue->GetFullName() : TEXT("None"),
+					*GetFullName(),
+					*Object.GetFullName()
+				);
+				SetObjectPropertyValue(ValueAddress, OldValue);
 			}
 		}
 	}
diff --git a/Engine/Source/Runtime/CoreUObject/Tests/ObjectPropertyTest.cpp b/Engine/Source/Runtime/CoreUObject/Tests/ObjectPropertyTest.cpp
new file mode 100644
--- /dev/null
+++ b/Engine/Source/Runtime/CoreUObject/Tests/ObjectPropertyTest.cpp
@@ -0,0 +1,132 @@
+// Copyright Epic Games, Inc. All Rights Reserved.
+
+#if WITH_LOW_LEVEL_TESTS
+
+#include "TestHarness.h"
+#include "TestMacros/Assertions.h"
+#include "ObjectPtrTestClass.h"
+#include "UObject/Package.h"
+#include "Misc/ScopeExit.h"
+#include "UObject/UnrealType.h"
+#include "UObject/ObjectHandleTracking.h"
+#include "LowLevelTestsRunner/WarnFilterScope.h"
+
+
+TEST_CASE("UE::CoreUObject::FObjectProperty::CheckValidAddress")
+{
+	bool bAllowRead = false;
+#if UE_WITH_OBJECT_HANDLE_TRACKING
+	auto CallbackHandle = UE::CoreUObject::AddObjectHandleReadCallback([&bAllowRead](UObject* ReadObject)
+		{
+			if (!bAllowRead)
+			FAIL("Unexpected read during CheckValidObject");
+		});
+	ON_SCOPE_EXIT
+	{
+		UE::CoreUObject::RemoveObjectHandleReadCallback(CallbackHandle);
+	};
+#endif
+	UClass* Class = UObjectPtrTestClassWithRef::StaticClass();
+	FObjectProperty* Property = CastField<FObjectProperty>(Class->FindPropertyByName(TEXT("ObjectPtr")));
+	REQUIRE(Property != nullptr);
+
+	UPackage* TestPackage = NewObject<UPackage>(nullptr, TEXT("TestPackageName"), RF_Transient);
+	TestPackage->AddToRoot();
+	ON_SCOPE_EXIT
+	{
+		TestPackage->RemoveFromRoot();
+	};
+	UObjectPtrTestClassWithRef* Obj = NewObject<UObjectPtrTestClassWithRef>(TestPackage, TEXT("Object"));
+	UObjectPtrTestClass* Other = NewObject<UObjectPtrTestClass>(Obj, TEXT("Other"));
+
+	//verify nothing happens by default
+	CHECK(!Obj->ObjectPtr);
+	Property->CheckValidObject(&Obj->ObjectPtr, nullptr);
+	CHECK(!Obj->ObjectPtr);
+
+	//valid assignment
+	Obj->ObjectPtr = Other;
+	Property->CheckValidObject(&Obj->ObjectPtr, nullptr);
+	CHECK(Obj->ObjectPtr == Other);
+
+	bAllowRead = true; //has to read the fullpath from the object to make the error message
+	//assign a bad value to the pointer
+	Obj->ObjectPtr = reinterpret_cast<UObjectPtrTestClass*>(Obj);
+	CHECK(Obj->ObjectPtr != nullptr);
+
+	UE::Testing::FWarnFilterScope _([](const TCHAR* Message, ELogVerbosity::Type Verbosity, const FName& Category)
+		{
+			if (Category == TEXT("LogProperty") && FCString::Strstr(Message, TEXT("Reference will be nullptred")) && Verbosity == ELogVerbosity::Type::Warning)
+			{
+				return true;
+			}
+			return false;
+		});
+	Property->CheckValidObject(&Obj->ObjectPtr, Other);
+	CHECK(!Obj->ObjectPtr); //value should be nulled since the type was not compatible	
+}
+
+TEST_CASE("UE::CoreUObject::FObjectProperty::CheckValidAddressNonNullable")
+{
+	bool bAllowRead = false;
+#if UE_WITH_OBJECT_HANDLE_TRACKING
+	auto CallbackHandle = UE::CoreUObject::AddObjectHandleReadCallback([&bAllowRead](UObject* ReadObject)
+		{
+			if (!bAllowRead)
+			{
+				FAIL("Unexpected read during CheckValidObject");
+			}
+		});
+	ON_SCOPE_EXIT
+	{
+		UE::CoreUObject::RemoveObjectHandleReadCallback(CallbackHandle);
+	};
+#endif
+	UClass* Class = UObjectPtrTestClassWithRef::StaticClass();
+	FObjectProperty* Property = CastField<FObjectProperty>(Class->FindPropertyByName(TEXT("ObjectPtrNonNullable")));
+	REQUIRE(Property != nullptr);
+
+	UPackage* TestPackage = NewObject<UPackage>(nullptr, TEXT("TestPackageName"), RF_Transient);
+	TestPackage->AddToRoot();
+	ON_SCOPE_EXIT
+	{
+		TestPackage->RemoveFromRoot();
+	};
+	UObjectPtrTestClassWithRef* Obj = NewObject<UObjectPtrTestClassWithRef>(TestPackage, TEXT("Object"));
+	UObjectPtrTestClass* Other = NewObject<UObjectPtrTestClass>(Obj, TEXT("Other"));
+
+	//property is already null should stay null
+	CHECK(!Obj->ObjectPtrNonNullable);
+	Property->CheckValidObject(&Obj->ObjectPtrNonNullable, Other);
+	CHECK(!Obj->ObjectPtrNonNullable);
+
+	//valid assignment
+	Obj->ObjectPtrNonNullable = Other;
+	Property->CheckValidObject(&Obj->ObjectPtrNonNullable, nullptr);
+	CHECK(Obj->ObjectPtrNonNullable == Other);
+
+	bAllowRead = true; //has to read the fullpath from the object to make the error message
+	//assign a bad value to the pointer
+	Obj->ObjectPtrNonNullable = reinterpret_cast<UObjectPtrTestClass*>(Obj);
+	CHECK(Obj->ObjectPtrNonNullable != nullptr);
+
+	UE::Testing::FWarnFilterScope _([](const TCHAR* Message, ELogVerbosity::Type Verbosity, const FName& Category)
+		{
+			if (Category == TEXT("LogProperty") && FCString::Strstr(Message, TEXT("Reference will be reverted back to")) && Verbosity == ELogVerbosity::Type::Warning)
+			{
+				return true;
+			}
+			return false;
+		});
+	Property->CheckValidObject(&Obj->ObjectPtrNonNullable, Other);
+	CHECK(Obj->ObjectPtrNonNullable == Other); //non nullable properties should be assigned the old value
+
+	//assign a bad value to the pointer
+	Obj->ObjectPtrNonNullable = reinterpret_cast<UObjectPtrTestClass*>(Obj);
+	CHECK(Obj->ObjectPtrNonNullable != nullptr);
+
+	//old value is required for non nullable properties
+	REQUIRE_CHECK(Property->CheckValidObject(&Obj->ObjectPtrNonNullable, nullptr));
+}
+
+#endif
\ No newline at end of file
diff --git a/Engine/Source/Runtime/CoreUObject/Tests/ObjectPtrTestClass.cpp b/Engine/Source/Runtime/CoreUObject/Tests/ObjectPtrTestClass.cpp
--- a/Engine/Source/Runtime/CoreUObject/Tests/ObjectPtrTestClass.cpp
+++ b/Engine/Source/Runtime/CoreUObject/Tests/ObjectPtrTestClass.cpp
@@ -6,6 +6,7 @@
 #include "UObject/Class.h"
 #include "UObject/Package.h"
 #include "UObject/MetaData.h"
+#include "UObject/UnrealType.h"
 
 //can not put the #if inside as the expansion of IMPLEMENT_CORE_INTRINSIC_CLASS fails
 #if WITH_EDITORONLY_DATA
@@ -19,7 +20,45 @@ IMPLEMENT_CORE_INTRINSIC_CLASS(UObjectPtrTestClass, UObject,
 	}
 );
 #else
-IMPLEMENT_CORE_INTRINSIC_CLASS(UObjectPtrTestClass, UObject, { });
+IMPLEMENT_CORE_INTRINSIC_CLASS(UObjectPtrTestClass, UObject, 
+	{ 
+	});
+#endif
+
+
+#if WITH_EDITORONLY_DATA
+IMPLEMENT_CORE_INTRINSIC_CLASS(UObjectPtrTestClassWithRef, UObject,
+	{
+		//add reflection info for the two properties
+		auto Property1 = new FObjectPtrProperty(Class, TEXT("ObjectPtr"), EObjectFlags::RF_NoFlags);
+		Property1->PropertyClass = UObjectPtrTestClass::StaticClass();
+		Class->AddCppProperty(Property1);
+
+		auto Property2 = new FObjectPtrProperty(Class, TEXT("ObjectPtrNonNullable"), EObjectFlags::RF_NoFlags);
+		Property2->PropertyClass = UObjectPtrTestClass::StaticClass();
+		Property2->SetPropertyFlags(EPropertyFlags::CPF_NonNullable); //make non nullable
+		Class->AddCppProperty(Property2);
+
+		auto MetaData = Class->GetOutermost()->GetMetaData();
+		if (MetaData)
+		{
+			MetaData->SetValue(Class, TEXT("LoadBehavior"), TEXT("LazyOnDemand"));
+		}
+	}
+);
+#else
+IMPLEMENT_CORE_INTRINSIC_CLASS(UObjectPtrTestClassWithRef, UObject,
+	{
+		//add reflection info for the two properties
+		auto Property1 = new FObjectPtrProperty(Class, TEXT("ObjectPtr"), EObjectFlags::RF_NoFlags);
+		Property1->PropertyClass = UObjectPtrTestClass::StaticClass();
+		Class->AddCppProperty(Property1);
+
+		auto Property2 = new FObjectPtrProperty(Class, TEXT("ObjectPtrNonNullable"), EObjectFlags::RF_NoFlags);
+		Property2->PropertyClass = UObjectPtrTestClass::StaticClass();
+		Property2->SetPropertyFlags(EPropertyFlags::CPF_NonNullable); //make non nullable
+		Class->AddCppProperty(Property2);
+	});
 #endif
 
 IMPLEMENT_CORE_INTRINSIC_CLASS(UObjectPtrDerrivedTestClass, UObjectPtrTestClass, {});
diff --git a/Engine/Source/Runtime/CoreUObject/Tests/ObjectPtrTestClass.h b/Engine/Source/Runtime/CoreUObject/Tests/ObjectPtrTestClass.h
--- a/Engine/Source/Runtime/CoreUObject/Tests/ObjectPtrTestClass.h
+++ b/Engine/Source/Runtime/CoreUObject/Tests/ObjectPtrTestClass.h
@@ -4,7 +4,7 @@
 #if WITH_LOW_LEVEL_TESTS
 
 #include "UObject/Object.h"
-
+#include "UObject/ObjectPtr.h"
 
 //simple test class for testing TObjectPtr resolve behavior
 class UObjectPtrTestClass : public UObject
@@ -15,6 +15,15 @@ public:
 
 };
 
+//test class with typed reference to another class
+class UObjectPtrTestClassWithRef : public UObject
+{
+	DECLARE_CLASS_INTRINSIC(UObjectPtrTestClassWithRef, UObject, CLASS_MatchedSerializers, TEXT("/Script/CoreUObject"))
+
+public:
+	TObjectPtr<UObjectPtrTestClass> ObjectPtr;
+	TObjectPtr<UObjectPtrTestClass> ObjectPtrNonNullable;
+};
 
 
 //derived test class
