ArchiveStackTrace - fixed bug with missing callstack CRC and added basic low level tests

#rnx

diff --git a/Engine/Source/Runtime/CoreUObject/Private/Serialization/ArchiveStackTrace.cpp b/Engine/Source/Runtime/CoreUObject/Private/Serialization/ArchiveStackTrace.cpp
--- a/Engine/Source/Runtime/CoreUObject/Private/Serialization/ArchiveStackTrace.cpp
+++ b/Engine/Source/Runtime/CoreUObject/Private/Serialization/ArchiveStackTrace.cpp
@@ -207,6 +207,7 @@ FArchiveCallstacks::FArchiveCallstacks(UObject* InAsset)
 	, bCallstacksDirty(true)
 	, StackTraceSize(65535)
 	, LastSerializeCallstack(nullptr)
+	, TotalSize(0)
 {
 	StackTrace = MakeUnique<ANSICHAR[]>(StackTraceSize);
 	StackTrace[0] = 0;
@@ -217,35 +218,41 @@ FName FArchiveCallstacks::GetAssetClass() const
 	return Asset != nullptr ? Asset->GetClass()->GetFName() : NAME_None;
 }
 
-ANSICHAR* FArchiveCallstacks::AddUniqueCallstack(UObject* SerializedObject, FProperty* SerializedProperty, uint32& OutCallstackCRC)
+ANSICHAR* FArchiveCallstacks::AddUniqueCallstack(bool bIsCollectingCallstacks, UObject* SerializedObject, FProperty* SerializedProperty, uint32& OutCallstackCRC)
 {
-	ANSICHAR* OutCallstack = nullptr;
-	OutCallstackCRC = FCrc::StrCrc32(StackTrace.Get());
-
-	if (FCallstackData* ExistingCallstack = UniqueCallstacks.Find(OutCallstackCRC))
+	ANSICHAR* Callstack = nullptr;
+	if (bIsCollectingCallstacks)
 	{
-		OutCallstack = ExistingCallstack->Callstack.Get();
+		OutCallstackCRC = FCrc::StrCrc32(StackTrace.Get());
+
+		if (FCallstackData* ExistingCallstack = UniqueCallstacks.Find(OutCallstackCRC))
+		{
+			Callstack = ExistingCallstack->Callstack.Get();
+		}
+		else
+		{
+			const int32 Len = FCStringAnsi::Strlen(StackTrace.Get()) + 1;
+			TUniquePtr<ANSICHAR[]> NewCallstack = MakeUnique<ANSICHAR[]>(Len);
+			FCStringAnsi::Strcpy(NewCallstack.Get(), Len, StackTrace.Get());
+			FCallstackData& NewEntry = UniqueCallstacks.Add(OutCallstackCRC, FCallstackData(MoveTemp(NewCallstack), SerializedObject, SerializedProperty));
+			Callstack = NewEntry.Callstack.Get();
+		}
 	}
 	else
 	{
-		const int32 Len = FCStringAnsi::Strlen(StackTrace.Get()) + 1;
-		TUniquePtr<ANSICHAR[]> Callstack = MakeUnique<ANSICHAR[]>(Len);
-		FCStringAnsi::Strcpy(Callstack.Get(), Len, StackTrace.Get());
-		FCallstackData& Data = UniqueCallstacks.Add(OutCallstackCRC, FCallstackData(MoveTemp(Callstack), SerializedObject, SerializedProperty));
-
-		OutCallstack = Data.Callstack.Get();
+		OutCallstackCRC = 0;
 	}
-
-	return OutCallstack;
+	return Callstack;
 }
 
 void FArchiveCallstacks::Add(
-	int64 CurrentOffset,
-	int64 CurrentTotalSize,
+	int64 Offset,
+	int64 Length,
 	UObject* SerializedObject,
 	FProperty* SerializedProperty,
 	TArrayView<const FName> DebugDataStack,
-	bool bSaveStackTrace,
+	bool bIsCollectingCallstacks,
+	bool bCollectCurrentCallstack,
 	int32 StackIgnoreCount)
 {
 	if (GIgnoreDiffManager.ShouldBypassDiff())
@@ -253,9 +260,11 @@ void FArchiveCallstacks::Add(
 		return;
 	}
 
-	TotalSize = CurrentTotalSize;
+	const int64 CurrentOffset = Offset;
+	TotalSize = FMath::Max(TotalSize, CurrentOffset + Length); 
 
-	if (bSaveStackTrace && GIgnoreDiffManager.ShouldIgnoreDiff() == false)
+	const bool bShouldCollectCallstack = bIsCollectingCallstacks && bCollectCurrentCallstack && !GIgnoreDiffManager.ShouldIgnoreDiff();
+	if (bShouldCollectCallstack)
 	{
 		StackTrace[0] = '\0';
 		FPlatformStackWalk::StackWalkAndDump(StackTrace.Get(), StackTraceSize, StackIgnoreCount);
@@ -297,22 +306,16 @@ void FArchiveCallstacks::Add(
 		if (CallstackAtOffsetMap.Num() == 0 || CurrentOffset > CallstackAtOffsetMap.Last().Offset)
 		{
 			// New data serialized at the end of archive buffer
-			if (bSaveStackTrace)
-			{
-				LastSerializeCallstack = AddUniqueCallstack(SerializedObject, SerializedProperty, CallstackCRC);
-			}
+			LastSerializeCallstack = AddUniqueCallstack(bIsCollectingCallstacks, SerializedObject, SerializedProperty, CallstackCRC);
 			CallstackAtOffsetMap.Add(FCallstackAtOffset {CurrentOffset, CallstackCRC, GIgnoreDiffManager.ShouldIgnoreDiff()});
 		}
 		else
 		{
 			// This happens usually after Seek() so we need to find the exiting offset or insert a new one
-			int32 CallstackToUpdateIndex = GetCallstackIndexAtOffset(CurrentOffset);
+			const int32 CallstackToUpdateIndex = GetCallstackIndexAtOffset(CurrentOffset);
 			check(CallstackToUpdateIndex != -1);
 			FCallstackAtOffset& CallstackToUpdate = CallstackAtOffsetMap[CallstackToUpdateIndex];
-			if (bSaveStackTrace)
-			{
-				LastSerializeCallstack = AddUniqueCallstack(SerializedObject, SerializedProperty, CallstackCRC);
-			}
+			LastSerializeCallstack = AddUniqueCallstack(bIsCollectingCallstacks, SerializedObject, SerializedProperty, CallstackCRC);
 			if (CallstackToUpdate.Offset == CurrentOffset)
 			{
 				CallstackToUpdate.Callstack = CallstackCRC;
@@ -324,6 +327,7 @@ void FArchiveCallstacks::Add(
 				CallstackAtOffsetMap.Insert(FCallstackAtOffset {CurrentOffset, CallstackCRC, GIgnoreDiffManager.ShouldIgnoreDiff()}, CallstackToUpdateIndex + 1);
 			}
 		}
+		check(CallstackCRC != 0 || !bShouldCollectCallstack);
 	}
 	else if (LastSerializeCallstack)
 	{
@@ -457,7 +461,6 @@ void FArchiveStackTraceWriter::Serialize(void* Data, int64 Length)
 	} BreakAtOffsetSettings;
 
 	const int64 CurrentOffset = DiffMapOffset + Tell();
-	const int64 CurrentTotalSize = DiffMapOffset + TotalSize();
 
 	if (BreakAtOffsetSettings.OffsetToBreakOn >= 0 && BreakAtOffsetSettings.OffsetToBreakOn >= CurrentOffset && BreakAtOffsetSettings.OffsetToBreakOn < CurrentOffset + Length)
 	{
@@ -475,19 +478,22 @@ void FArchiveStackTraceWriter::Serialize(void* Data, int64 Length)
 
 	if (Length > 0)
 	{
-		check(SerializeContext);
-		const bool bSaveStackTrace = DiffMap != nullptr && DiffMap->ContainsOffset(CurrentOffset);
+		UObject* SerializedObject = SerializeContext ? SerializeContext->SerializedObject : nullptr;
 		TArrayView<const FName> DebugStack;
 #if WITH_EDITOR
 		DebugStack = DebugDataStack;
 #endif
+		const bool bIsCollectingCallstacks = DiffMap != nullptr;
+		const bool bCollectCurrentCallstack = bIsCollectingCallstacks && DiffMap->ContainsOffset(CurrentOffset);
+
 		Callstacks.Add(
 			CurrentOffset,
-			CurrentTotalSize,
-			SerializeContext->SerializedObject,
+			Length,
+			SerializedObject,
 			GetSerializedProperty(),
 			DebugStack,
-			bSaveStackTrace,
+			bIsCollectingCallstacks,
+			bCollectCurrentCallstack,
 			StackIgnoreCount);
 	}
 
@@ -665,7 +671,8 @@ void FArchiveStackTraceWriter::Compare(
 	const TCHAR* CallstackCutoffText,
 	const int64 MaxDiffsToLog,
 	int32& InOutDiffsLogged,
-	TMap<FName, FArchiveDiffStats>& OutStats)
+	TMap<FName, FArchiveDiffStats>& OutStats,
+	bool bSuppressLogging)
 {
 #if !NO_LOGGING
 	const TCHAR* const Indent = FDiffFormatHelper::Get().Indent;
@@ -677,7 +684,7 @@ void FArchiveStackTraceWriter::Compare(
 	
 	if (SourceSize != DestSize)
 	{
-		UE_LOG(LogArchiveDiff, Warning, TEXT("%s: Size mismatch: on disk: %lld vs memory: %lld"), AssetFilename, SourceSize, DestSize);
+		UE_CLOG(!bSuppressLogging, LogArchiveDiff, Warning, TEXT("%s: Size mismatch: on disk: %lld vs memory: %lld"), AssetFilename, SourceSize, DestSize);
 		int64 SizeDiff = DestPackage.Size - SourcePackage.Size;
 		OutStats.FindOrAdd(AssetClass).DiffSize += SizeDiff;
 	}
@@ -720,7 +727,7 @@ void FArchiveStackTraceWriter::Compare(
 
 			if (DifferenceCallstackoffsetIndex < 0)
 			{
-				UE_LOG(LogArchiveDiff, Warning, TEXT("%s: Difference at offset %lld (absolute offset: %lld), unknown callstack"), AssetFilename, LocalOffset, DestAbsoluteOffset);
+				UE_CLOG(!bSuppressLogging, LogArchiveDiff, Warning, TEXT("%s: Difference at offset %lld (absolute offset: %lld), unknown callstack"), AssetFilename, LocalOffset, DestAbsoluteOffset);
 				continue;
 			}
 
@@ -802,7 +809,8 @@ void FArchiveStackTraceWriter::Compare(
 					DebugDataStackText = FString::Printf(TEXT("\r\n%s"), FDiffFormatHelper::Get().Indent) + FullStackText.RightChop(DebugDataIndex + 2);
 				}
 #endif
-				UE_LOG(
+				UE_CLOG(
+					!bSuppressLogging,
 					LogArchiveDiff,
 					Warning,
 					TEXT("%s: Difference at offset %lld%s (absolute offset: %lld): byte %d on disk, byte %d in memory, callstack:%s%s%s%s%s"),
@@ -819,7 +827,8 @@ void FArchiveStackTraceWriter::Compare(
 				);
 
 				const int BytesToLog = 128;
-				UE_LOG(
+				UE_CLOG(
+					!bSuppressLogging,
 					LogArchiveDiff,
 					Display,
 					TEXT("%s: Logging %d bytes around absolute offset: %lld (%016X) in the on disk (existing) package, (which corresponds to offset %lld (%016X) in the in-memory package)"),
@@ -832,7 +841,8 @@ void FArchiveStackTraceWriter::Compare(
 				);
 				FCompressionUtil::LogHexDump(SourcePackage.Data, SourcePackage.Size, SourceAbsoluteOffset - BytesToLog / 2, SourceAbsoluteOffset + BytesToLog / 2);
 
-				UE_LOG(
+				UE_CLOG(
+					!bSuppressLogging,
 					LogArchiveDiff,
 					Display,
 					TEXT("%s: Logging %d bytes around absolute offset: %lld (%016X) in the in memory (new) package"),
@@ -869,11 +879,11 @@ void FArchiveStackTraceWriter::Compare(
 	{
 		if (FirstUnreportedDiffIndex != -1)
 		{
-			UE_LOG(LogArchiveDiff, Warning, TEXT("%s: %lld difference(s) not logged (first at absolute offset: %lld)."), AssetFilename, NumDiffsLocal - NumDiffsLoggedLocal, FirstUnreportedDiffIndex);
+			UE_CLOG(!bSuppressLogging, LogArchiveDiff, Warning, TEXT("%s: %lld difference(s) not logged (first at absolute offset: %lld)."), AssetFilename, NumDiffsLocal - NumDiffsLoggedLocal, FirstUnreportedDiffIndex);
 		}
 		else
 		{
-			UE_LOG(LogArchiveDiff, Warning, TEXT("%s: %lld difference(s) not logged."), AssetFilename, NumDiffsLocal - NumDiffsLoggedLocal);
+			UE_CLOG(!bSuppressLogging, LogArchiveDiff, Warning, TEXT("%s: %lld difference(s) not logged."), AssetFilename, NumDiffsLocal - NumDiffsLoggedLocal);
 		}
 	}
 #endif
diff --git a/Engine/Source/Runtime/CoreUObject/Public/Serialization/ArchiveStackTrace.h b/Engine/Source/Runtime/CoreUObject/Public/Serialization/ArchiveStackTrace.h
--- a/Engine/Source/Runtime/CoreUObject/Public/Serialization/ArchiveStackTrace.h
+++ b/Engine/Source/Runtime/CoreUObject/Public/Serialization/ArchiveStackTrace.h
@@ -166,12 +166,13 @@ public:
 
 	/** Capture and append the current callstack. */
 	void Add(
-		int64 CurrentOffset,
-		int64 CurrentTotalSize,
+		int64 Offset,
+		int64 Length,
 		UObject* SerializedObject,
 		FProperty* SerializedProperty,
 		TArrayView<const FName> DebugDataStack,
-		bool bSaveStackTrace,
+		bool bIsCollectingCallstacks,
+		bool bCollectCurrentCallstack,
 		int32 StackIgnoreCount);
 
 	/** Append other callstacks. */
@@ -212,7 +213,7 @@ public:
 
 private:
 	/** Adds a unique callstack to UniqueCallstacks map */
-	ANSICHAR* AddUniqueCallstack(UObject* SerializedObject, FProperty* SerializedProperty, uint32& OutCallstackCRC);
+	ANSICHAR* AddUniqueCallstack(bool bIsCollectingCallstacks, UObject* SerializedObject, FProperty* SerializedProperty, uint32& OutCallstackCRC);
 
 	/** The asset being serialized */
 	UObject* Asset;
@@ -291,7 +292,8 @@ public:
 		const TCHAR* CallstackCutoffText,
 		const int64 MaxDiffsToLog,
 		int32& InOutDiffsLogged,
-		TMap<FName, FArchiveDiffStats>& OutStats);
+		TMap<FName, FArchiveDiffStats>& OutStats,
+		bool bSuppressLogging = false);
 
 	/** Creates map with mismatching callstacks. */
 	static bool GenerateDiffMap(
diff --git a/Engine/Source/Runtime/CoreUObject/Tests/Serialization/ArchiveStackTraceTests.cpp b/Engine/Source/Runtime/CoreUObject/Tests/Serialization/ArchiveStackTraceTests.cpp
new file mode 100644
--- /dev/null
+++ b/Engine/Source/Runtime/CoreUObject/Tests/Serialization/ArchiveStackTraceTests.cpp
@@ -0,0 +1,177 @@
+// Copyright Epic Games, Inc. All Rights Reserved.
+
+#if WITH_LOW_LEVEL_TESTS && WITH_EDITORONLY_DATA
+
+#include "TestHarness.h"
+#include "TestMacros/Assertions.h"
+#include "Memory/SharedBuffer.h"
+#include "Serialization/ArchiveStackTrace.h"
+#include "Serialization/LargeMemoryWriter.h"
+#include "Serialization/MemoryReader.h"
+#include <catch2/generators/catch_generators.hpp>
+
+namespace UE::CoreUObject::Serialization::Tests
+{
+
+struct FBasicDiffState
+{
+	int32 FirstParam = 0;
+	int32 SecondParam = 0;
+	int32 ThirdParam = 0;
+
+	FORCENOINLINE void Serialize(FArchive& Ar)
+	{
+		Ar << FirstParam;
+		Ar << SecondParam;
+		Ar << ThirdParam;
+	}
+};
+
+FArchiveStackTraceWriter::FPackageData ToPackageData(FLargeMemoryWriter& Ar)
+{
+	return FArchiveStackTraceWriter::FPackageData
+	{
+		Ar.GetData(),
+		Ar.TotalSize()
+	};
+}
+
+TEST_CASE("CoreUObject::Serialization::ArchiveStackTrace::Callstacks", "[CoreUObject][Serialization]")
+{
+	const int32 MaxDiffsToLog = 1000;
+	FArchiveCallstacks Callstacks(nullptr);
+
+	FLargeMemoryWriter Memory;
+	FBasicDiffState State = FBasicDiffState {1,2,3};
+	FArchiveStackTraceWriter StackTraceWriter(Memory, Callstacks, 0);
+	State.Serialize(StackTraceWriter);
+
+	CHECK(Callstacks.Num() == 3);
+	CHECK(Callstacks.TotalCapturedSize() == sizeof(FBasicDiffState));
+	CHECK(Callstacks.GetCallstack(0).Offset == 0);
+	CHECK(Callstacks.GetCallstack(0).Callstack == 0);
+	CHECK(Callstacks.GetCallstack(0).bIgnore == false);
+	CHECK(Callstacks.GetCallstack(1).Offset == 4);
+	CHECK(Callstacks.GetCallstack(1).Callstack == 0);
+	CHECK(Callstacks.GetCallstack(1).bIgnore == false);
+	CHECK(Callstacks.GetCallstack(2).Offset == 8);
+	CHECK(Callstacks.GetCallstack(2).Callstack == 0);
+	CHECK(Callstacks.GetCallstack(2).bIgnore == false);
+}
+
+TEST_CASE("CoreUObject::Serialization::ArchiveStackTrace::Basic", "[CoreUObject][Serialization]")
+{
+	const int32 MaxDiffsToLog = 1000;
+
+	SECTION("DiffMap - identical")
+	{
+		FArchiveCallstacks Callstacks(nullptr);
+		FArchiveDiffMap DiffMap;
+
+		FLargeMemoryWriter InitialState;
+		{
+			FBasicDiffState State = FBasicDiffState {1,2,3};
+			State.Serialize(InitialState);
+		}
+
+		FLargeMemoryWriter NewState;
+		{
+			FBasicDiffState State = FBasicDiffState {1,2,3};
+			FArchiveStackTraceWriter StackTraceWriter(NewState, Callstacks);
+			State.Serialize(StackTraceWriter);
+		}
+
+		FArchiveStackTraceWriter::FPackageData InitialData = ToPackageData(InitialState);
+		FArchiveStackTraceWriter::FPackageData NewData = ToPackageData(NewState);
+		const bool bIsIdentical = FArchiveStackTraceWriter::GenerateDiffMap(InitialData, NewData, Callstacks, MaxDiffsToLog, DiffMap);
+		CHECK(bIsIdentical);
+	}
+
+	SECTION("DiffMap - mismatch")
+	{
+		FArchiveCallstacks Callstacks(nullptr);
+		FArchiveDiffMap DiffMap;
+
+		FLargeMemoryWriter InitialState;
+		{
+			FBasicDiffState State = FBasicDiffState {1,2,3};
+			State.Serialize(InitialState);
+		}
+
+		FLargeMemoryWriter NewState;
+		{
+			FBasicDiffState State = FBasicDiffState {1,2,4};
+			FArchiveStackTraceWriter StackTraceWriter(NewState, Callstacks);
+			State.Serialize(StackTraceWriter);
+		}
+
+		FArchiveStackTraceWriter::FPackageData InitialData = ToPackageData(InitialState);
+		FArchiveStackTraceWriter::FPackageData NewData = ToPackageData(NewState);
+		const bool bIsIdentical = FArchiveStackTraceWriter::GenerateDiffMap(InitialData, NewData, Callstacks, MaxDiffsToLog, DiffMap);
+		CHECK(bIsIdentical == false);
+	}
+
+	SECTION("Compare - mismatch")
+	{
+		FArchiveCallstacks Callstacks(nullptr);
+		FArchiveCallstacks CallstacksWithStackTrace(nullptr);
+		FArchiveDiffMap DiffMap;
+
+		FBasicDiffState InitialState = FBasicDiffState {1,2,3};
+		FLargeMemoryWriter InitialMemory;
+		InitialMemory.SetByteSwapping(false);
+
+		InitialState.Serialize(InitialMemory);
+
+		FBasicDiffState NewState = FBasicDiffState {1,22,33};
+		TUniquePtr<FLargeMemoryWriter> NewMemory = MakeUnique<FLargeMemoryWriter>();
+		NewMemory->SetByteSwapping(false);
+
+		{
+			FArchiveStackTraceWriter StackTraceWriter(*NewMemory, Callstacks);
+			NewState.Serialize(StackTraceWriter);
+		}
+
+		// Generate diff map
+		{
+			FArchiveStackTraceWriter::FPackageData InitialData = ToPackageData(InitialMemory);
+			FArchiveStackTraceWriter::FPackageData NewData = ToPackageData(*NewMemory);
+			const bool bIsIdentical = FArchiveStackTraceWriter::GenerateDiffMap(InitialData, NewData, Callstacks, MaxDiffsToLog, DiffMap);
+			CHECK(bIsIdentical == false);
+		}
+		
+		NewMemory = MakeUnique<FLargeMemoryWriter>();
+		NewMemory->SetByteSwapping(false);
+		{
+			FArchiveStackTraceWriter StackTraceWriter(*NewMemory, CallstacksWithStackTrace, &DiffMap);
+			NewState.Serialize(StackTraceWriter);
+		}
+
+		// Compare using callstacks and diffmap
+		{
+			FArchiveStackTraceWriter::FPackageData InitialData = ToPackageData(InitialMemory);
+			FArchiveStackTraceWriter::FPackageData NewData = ToPackageData(*NewMemory);
+			int32 DiffsLogged = 0;
+			TMap<FName, FArchiveDiffStats> DiffStats;
+			const bool bSuppressLogging = true; // Suppress warning logs when running tests on CI
+
+			FArchiveStackTraceWriter::Compare(
+				InitialData,
+				NewData,
+				CallstacksWithStackTrace,
+				DiffMap,
+				TEXT("BasicDiffTest"),
+				nullptr,
+				MaxDiffsToLog,
+				DiffsLogged,
+				DiffStats,
+				bSuppressLogging);
+
+			CHECK(DiffStats[NAME_None].NumDiffs == 2);
+		}
+	}
+}
+
+} // namespace UE::CoreUObject::Serialization::Tests
+
+#endif // WITH_LOW_LEVEL_TESTS && WITH_EDITORONLY_DATA
