Containers: DirectoryTree. An O(log(n)) map from path to ValueType that also supports O(log(n)) queries for the closest parent directory in the tree of a given path.
#rnx

diff --git a/Engine/Source/Runtime/CoreUObject/Internal/Containers/DirectoryTree.h b/Engine/Source/Runtime/CoreUObject/Internal/Containers/DirectoryTree.h
new file mode 100644
--- /dev/null
+++ b/Engine/Source/Runtime/CoreUObject/Internal/Containers/DirectoryTree.h
@@ -0,0 +1,978 @@
+// Copyright Epic Games, Inc. All Rights Reserved.
+
+#pragma once
+
+#include "Containers/Array.h"
+#include "Containers/StringView.h"
+#include "Containers/UnrealString.h"
+#include "HAL/Platform.h"
+#include "Misc/PathViews.h"
+#include "Misc/StringBuilder.h"
+#include "Templates/TypeCompatibleBytes.h"
+#include "Templates/UniquePtr.h"
+
+namespace UE::DirectoryTree
+{
+
+COREUOBJECT_API void FixupPathSeparator(FStringBuilderBase& InOutPath, int32 StartIndex, TCHAR InSeparatorChar);
+COREUOBJECT_API int32 FindInsertionIndex(int32 NumChildNodes, const TUniquePtr<FString[]>& RelPaths,
+	FStringView FirstPathComponent, bool& bOutExists);
+
+}
+
+/**
+ * Container for path -> value that can efficiently report whether a parent directory of a given path exists.
+ * Supports relative and absolute paths.
+ * Supports LongPackageNames and LocalPaths.
+ * 
+ * Note about Value comparisons:
+ * Case-insensitive
+ * / is treated as equal to \
+ * Presence or absence of terminating separator (/) is ignored in the comparison.
+ * Directory elements of . and .. are currently not interpreted and are treated as literal characters.
+ *    Callers should not rely on this behavior as it may be corrected in the future.
+ *    callers should instead conform the paths before calling.
+ * Relative paths and absolute paths are not resolved, and relative paths will never equal absolute paths.
+ *    Callers should not rely on this behavior as it may be corrected in the future;
+ *    callers should instead conform the paths before calling.
+ * 
+ * For functions that find parent paths, parent paths are only discovered if they are conformed to the same format as
+ * the given path: both paths must be either relative or absolute.
+ * 
+ * For functions that return Values by reference or by pointer, that reference or pointer can be invalidated
+ * by any functions that modify the tree, and should be discarded before calling any such functions.
+ */
+template <typename ValueType>
+class COREUOBJECT_API TDirectoryTree
+{
+public:
+	TDirectoryTree();
+
+	/**
+	 * Add a path to the tree if it does not already exist. Construct default Value for it if it did not already exist.
+	 * Return a reference to the added or existing Value. Optionally report whether the path already existed.
+	 * This reference can be invalidated by any operations that modify the tree.
+	 */
+	ValueType& FindOrAdd(FStringView Path, bool* bOutExisted = nullptr);
+	/** Remove a path from the tree and optionally report whether it existed. */
+	void Remove(FStringView Path, bool* bOutExisted = nullptr);
+	/** Free unused slack memory throughout the tree by reallocating containers tightly to their current size. */
+	void Shrink();
+
+	/** Return true if no paths are in the tree. */
+	bool IsEmpty() const;
+	/** Return the number of paths in the tree. */
+	int32 Num() const;
+	/** How much memory is used by *this, not counting sizeof(*this). */
+	SIZE_T GetAllocatedSize() const;
+
+	/** Return whether the given path has been added to the tree. */
+	bool Contains(FStringView Path) const;
+	/** Return a const pointer to the Value set for the given path, or null if it does not exist. */
+	const ValueType* Find(FStringView Path) const;
+	/** Return a pointer to the Value set for the given path, or null if it does not exist. */
+	ValueType* Find(FStringView Path);
+	/** Return whether the given path or any of its parent paths exist in the tree. */
+	bool ContainsPathOrParent(FStringView Path) const;
+	/**
+	 * Return a const pointer to the path's value if it exists, or to its closest parent path's value,
+	 * if any of them exist. Otherwise return nullptr.
+	 */
+	const ValueType* FindClosestValue(FStringView Path) const;
+	/**
+	 * Return a pointer to the path's value if it exists, or to its closest parent path's value,
+	 * if any of them exist. Otherwise return nullptr.
+	 */
+	ValueType* FindClosestValue(FStringView Path);
+	/**
+	 * Return whether the given path or any of its parent paths exist in the tree. If it exists,
+	 * return the discovered path, and optionally return a const pointer to the value of that path.
+	 */
+	bool TryFindClosestPath(FStringView Path, FStringBuilderBase& OutPath, const ValueType** OutValue = nullptr) const;
+	/**
+	 * Return whether the given path or any of its parent paths exist in the tree. If it exists,
+	 * return the discovered path, and optionally return a pointer to the value of that path.
+	 */
+	bool TryFindClosestPath(FStringView Path, FStringBuilderBase& OutPath, ValueType** OutValue);
+	/**
+	 * Return whether the given path or any of its parent paths exist in the tree. If it exists,
+	 * return the discovered path, and optionally return a const pointer to the value of that path.
+	 */
+	bool TryFindClosestPath(FStringView Path, FString& OutPath, const ValueType** OutValue = nullptr) const;
+	/**
+	 * Return whether the given path or any of its parent paths exist in the tree. If it exists,
+	 * return the discovered path, and optionally return a pointer to the value of that path.
+	 */
+	bool TryFindClosestPath(FStringView Path, FString& OutPath, ValueType** OutValue);
+
+private:
+
+	/**
+	 * A tree structure; each node has a sorted array of child paths and a matching array of child nodes.
+	 * Child paths are relative paths and are organized by the FirstComponent of their relative path. If there is only
+	 * a single child path with a given FirstComponent, the entire relative path to	that child is listed as the
+	 * relative path. If there are two or more paths with the same first component, a new child node is created for
+	 * the first component, and the paths are then children of that component.
+	 * Example:
+	 * Root
+	 *		/				(FullPath: /)
+	 *			A			(FullPath: /A
+	 *				X		(FullPath: /A/Y
+	 *				Y/M		(FullPath: /A/Z/N)
+	 *			B/Z			(FullPath: /B/Z)
+	 *				N		(FullPath: /B/Z/N)
+	 *				O		(FullPath: /B/Z/O)
+	 *			C/W/P		(FullPath: /C/W/P)
+	 */
+	struct FTreeNode
+	{
+	public:
+		FTreeNode() = default;
+		FTreeNode(FTreeNode&& Other);
+		FTreeNode& operator=(FTreeNode&& Other);
+		FTreeNode(const FTreeNode& Other) = delete;
+		FTreeNode& operator=(const FTreeNode& Other) = delete;
+
+		/** Remove Value and ChildNodes, return state to default-constructed state. */
+		void Reset();
+
+		int32 GetNumChildNodes() const;
+
+		/**
+		 * Recursively search the node's subtree to find the given relative directory name, adding nodes for the path
+		 * and its parents if required. Return reference to the added or existing node's value.
+		 */
+		ValueType& FindOrAdd(FStringView InRelPath, bool& bOutExisted);
+		/** Remove the Value if it exists in the tree. */
+		void Remove(FStringView InRelPath, bool& bOutExisted);
+		/** Return pointer to the Value stored in RelPath, if RelPath exists in the tree. */
+		ValueType* Find(FStringView InRelPath);
+
+		/**
+		 * Recursively search the node's subtree to find the given RelPath. Return whether the path or any parent is
+		 * found, and if any is found and Outpath is present, append relative path of the discovered path into OutPath.
+		 */
+		ValueType* TryFindClosestPath(FStringView RelPath, FStringBuilderBase* OutPath, TCHAR SeparatorChar);
+
+		/** Are no paths contained within the node. */
+		bool IsEmpty() const;
+		/** How much memory is used by *this, not counting sizeof(*this). */
+		SIZE_T GetAllocatedSize() const;
+
+		/** Reduce memory used in buffers. */
+		void Shrink();
+
+		/**
+		 * Report whether the node has a value, which is equivalent to the node's path existing in the DirectoryTree.
+		 * Nodes might exist in the tree without their path existing, if they are parent paths that have not been
+		 * added.
+		 */
+		bool HasValue() const;
+		/** Get a reference to the node's Value. Invalid to call if !HasValue. */
+		ValueType& GetValue();
+		/** Set HasValue=true, and move InValue into the node's value, after destructing any existing old value. */
+		void SetValue(ValueType&& InValue);
+		/** Set HasValue=true, and default-construct the node's value, after destructing any existing old value. */
+		void SetDefaultValue();
+		/** Set HasValue=false, and destruct any existing old value. */
+		void RemoveValue();
+
+	private:
+		const static uint32 NumFlagBits = 1;
+		const static uint32 FlagsShift = (8 * sizeof(uint32) - NumFlagBits);
+		const static uint32 FlagsMask = 0x1 << FlagsShift;
+		void SetNumChildNodes(int32 InNumChildNodes);
+
+		/**
+		 * Search the sorted ChildNode RelPaths for the given FirstPathComponent, which must be only a single
+		 * path component, and return the index either of the existing path or where the path should be inserted
+		 * if new. Also report whether it exists.
+		 */
+		int32 FindInsertionIndex(FStringView FirstPathComponent, bool& bOutExists) const;
+		/** Insert the given RelPath and ChildNode at the given index; must be the index from FindInsertionIndex. */
+		FTreeNode& InsertChildNode(int32 InsertionIndex, FString&& RelPath, FTreeNode&& ChildNode);
+		/** Remove the RelPath and ChildNode from the given index. */
+		void RemoveChildNodeAt(int32 RemoveIndex);
+
+		/** Merge the node with its direct child if possible, and if so adjust the input RelPath to match. */
+		static void ConditionalCompactNode(FString& RelPath, FTreeNode& ChildNode);
+		/** Reallocate the NumChildNodes arrays to match the given new capacity. */
+		void Realloc(int32 NewCapacity);
+
+	private:
+		TTypeCompatibleBytes<ValueType> Value;
+		TUniquePtr<FString[]> RelPaths;
+		TUniquePtr<FTreeNode[]> ChildNodes;
+		uint32 NumChildNodesAndFlags = 0;
+		int32 CapacityChildNodes = 0;
+	};
+
+private:
+	ValueType* TryFindClosestPathInternal(FStringView Path, FStringBuilderBase* OutPath);
+
+private:
+	FTreeNode Root;
+	int32 NumPaths = 0;
+	TCHAR PathSeparator = '/';
+	bool bPathSeparatorInitialized = false;
+};
+
+
+template <typename ValueType>
+inline TDirectoryTree<ValueType>::TDirectoryTree()
+{
+	check(Root.IsEmpty());
+}
+
+template <typename ValueType>
+inline ValueType& TDirectoryTree<ValueType>::FindOrAdd(FStringView Path, bool* bOutExisted)
+{
+	if (Path.IsEmpty())
+	{
+		if (bOutExisted)
+		{
+			*bOutExisted = Root.HasValue();
+		}
+		if (!Root.HasValue())
+		{
+			Root.SetDefaultValue();
+			++NumPaths;
+		}
+		return Root.GetValue();
+	}
+
+	if (!bPathSeparatorInitialized)
+	{
+		int32 UnusedIndex;
+		if (Path.FindChar('/', UnusedIndex))
+		{
+			PathSeparator = '/';
+			bPathSeparatorInitialized = true;
+		}
+		else if (Path.FindChar('\\', UnusedIndex))
+		{
+			PathSeparator = '\\';
+			bPathSeparatorInitialized = true;
+		}
+	}
+
+	bool bExisted;
+	ValueType& Result = Root.FindOrAdd(Path, bExisted);
+	if (!bExisted)
+	{
+		++NumPaths;
+	}
+	return Result;
+}
+
+template <typename ValueType>
+inline void TDirectoryTree<ValueType>::Remove(FStringView Path, bool* bOutExisted)
+{
+	bool bExisted;
+	if (Path.IsEmpty())
+	{
+		bExisted = Root.HasValue();
+		if (bExisted)
+		{
+			Root.RemoveValue();
+		}
+	}
+	else
+	{
+		Root.Remove(Path, bExisted);
+	}
+	if (bExisted)
+	{
+		check(NumPaths > 0);
+		--NumPaths;
+	}
+	if (bOutExisted)
+	{
+		*bOutExisted = bExisted;
+	}
+}
+
+template <typename ValueType>
+inline void TDirectoryTree<ValueType>::Shrink()
+{
+	Root.Shrink();
+}
+
+template <typename ValueType>
+inline bool TDirectoryTree<ValueType>::IsEmpty() const
+{
+	return Root.IsEmpty();
+}
+
+template <typename ValueType>
+inline int32 TDirectoryTree<ValueType>::Num() const
+{
+	return NumPaths;
+}
+
+template <typename ValueType>
+inline SIZE_T TDirectoryTree<ValueType>::GetAllocatedSize() const
+{
+	return Root.GetAllocatedSize();
+}
+
+template <typename ValueType>
+inline bool TDirectoryTree<ValueType>::Contains(FStringView Path) const
+{
+	return Find(Path) != nullptr;
+}
+
+template <typename ValueType>
+inline const ValueType* TDirectoryTree<ValueType>::Find(FStringView Path) const
+{
+	return const_cast<TDirectoryTree*>(this)->Find(Path);
+}
+
+template <typename ValueType>
+inline ValueType* TDirectoryTree<ValueType>::Find(FStringView Path)
+{
+	if (Path.IsEmpty())
+	{
+		return Root.HasValue() ? &Root.GetValue() : nullptr;
+	}
+	return Root.Find(Path);
+}
+
+template <typename ValueType>
+inline bool TDirectoryTree<ValueType>::ContainsPathOrParent(FStringView Path) const
+{
+	return FindClosestValue(Path) != nullptr;
+}
+
+template <typename ValueType>
+inline const ValueType* TDirectoryTree<ValueType>::FindClosestValue(FStringView Path) const
+{
+	return const_cast<TDirectoryTree*>(this)->TryFindClosestPathInternal(Path, nullptr);
+}
+
+template <typename ValueType>
+inline ValueType* TDirectoryTree<ValueType>::FindClosestValue(FStringView Path)
+{
+	return TryFindClosestPathInternal(Path, nullptr);
+}
+
+template <typename ValueType>
+inline bool TDirectoryTree<ValueType>::TryFindClosestPath(FStringView Path, FString& OutPath, const ValueType** OutValue) const
+{
+	ValueType* ResultValue;
+	bool bResult = const_cast<TDirectoryTree*>(this)->TryFindClosestPath(Path, OutPath, &ResultValue);
+	if (bResult)
+	{
+		if (OutValue)
+		{
+			*OutValue = ResultValue;
+		}
+	}
+	return bResult;
+}
+
+template <typename ValueType>
+inline bool TDirectoryTree<ValueType>::TryFindClosestPath(FStringView Path, FString& OutPath, ValueType** OutValue)
+{
+	TStringBuilder<1024> Builder;
+	OutPath.Reset(); // Reset the path even if we fail, to match the behavior when taking a StringBuilderBase
+	ValueType* Result = TryFindClosestPathInternal(Path, &Builder);
+	if (Result)
+	{
+		OutPath.Append(Builder);
+		if (OutValue)
+		{
+			*OutValue = Result;
+		}
+		return true;
+	}
+	else
+	{
+		return false;
+	}
+}
+
+template <typename ValueType>
+inline bool TDirectoryTree<ValueType>::TryFindClosestPath(FStringView Path, FStringBuilderBase& OutPath, const ValueType** OutValue) const
+{
+	ValueType* ResultValue;
+	bool bResult = const_cast<TDirectoryTree*>(this)->TryFindClosestPath(Path, OutPath, &ResultValue);
+	if (bResult)
+	{
+		if (OutValue)
+		{
+			*OutValue = ResultValue;
+		}
+	}
+	return bResult;
+}
+
+template <typename ValueType>
+inline bool TDirectoryTree<ValueType>::TryFindClosestPath(FStringView Path, FStringBuilderBase& OutPath, ValueType** OutValue)
+{
+	ValueType* Result = TryFindClosestPathInternal(Path, &OutPath);
+	if (Result)
+	{
+		if (OutValue)
+		{
+			*OutValue = Result;
+		}
+		return true;
+	}
+	else
+	{
+		return false;
+	}
+}
+
+template <typename ValueType>
+inline ValueType* TDirectoryTree<ValueType>::TryFindClosestPathInternal(FStringView Path,
+	FStringBuilderBase* OutPath)
+{
+	if (OutPath)
+	{
+		// We build the path as we go along, before we know whether we will return true,
+		// so we have to reset it now
+		OutPath->Reset();
+	}
+	if (!Path.IsEmpty())
+	{
+		ValueType* Result = Root.TryFindClosestPath(Path, OutPath, PathSeparator);
+		if (Result)
+		{
+			return Result;
+		}
+	}
+
+	return Root.HasValue() ? &Root.GetValue() : nullptr;
+}
+
+template <typename ValueType>
+TDirectoryTree<ValueType>::FTreeNode::FTreeNode(FTreeNode&& Other)
+{
+	*this = MoveTemp(Other);
+}
+
+template <typename ValueType>
+typename TDirectoryTree<ValueType>::FTreeNode& TDirectoryTree<ValueType>::FTreeNode::operator=(FTreeNode&& Other)
+{
+	if (Other.HasValue())
+	{
+		SetValue(MoveTemp(Other.GetValue()));
+		Other.RemoveValue();
+	}
+	else
+	{
+		RemoveValue();
+	}
+	RelPaths = MoveTemp(Other.RelPaths);
+	ChildNodes = MoveTemp(Other.ChildNodes);
+	SetNumChildNodes(Other.GetNumChildNodes());
+	CapacityChildNodes = MoveTemp(Other.CapacityChildNodes);
+
+	Other.SetNumChildNodes(0);
+	Other.CapacityChildNodes = 0;
+	return *this;
+}
+
+template <typename ValueType>
+void TDirectoryTree<ValueType>::FTreeNode::Reset()
+{
+	RemoveValue();
+	RelPaths.Reset();
+	ChildNodes.Reset();
+	NumChildNodesAndFlags = 0;
+	CapacityChildNodes = 0;
+}
+
+template <typename ValueType>
+inline int32 TDirectoryTree<ValueType>::FTreeNode::GetNumChildNodes() const
+{
+	return (int32)(NumChildNodesAndFlags & ~FlagsMask);
+}
+
+template <typename ValueType>
+inline void TDirectoryTree<ValueType>::FTreeNode::SetNumChildNodes(int32 InNumChildNodes)
+{
+	check(0 <= InNumChildNodes && (InNumChildNodes & FlagsMask) == 0);
+	NumChildNodesAndFlags =
+		InNumChildNodes |
+		(NumChildNodesAndFlags & FlagsMask);
+}
+
+template <typename ValueType>
+inline bool TDirectoryTree<ValueType>::FTreeNode::HasValue() const
+{
+	return (NumChildNodesAndFlags & FlagsMask) != 0;
+}
+
+template <typename ValueType>
+inline ValueType& TDirectoryTree<ValueType>::FTreeNode::GetValue()
+{
+	return *reinterpret_cast<ValueType*>(&Value);
+}
+
+template <typename ValueType>
+inline void TDirectoryTree<ValueType>::FTreeNode::SetValue(ValueType&& InValue)
+{
+	RemoveValue();
+
+	uint32 FlagsValue = 0x1;
+	NumChildNodesAndFlags =
+		(NumChildNodesAndFlags & ~FlagsMask) |
+		(FlagsValue << FlagsShift);
+	new (&Value) ValueType(MoveTemp(InValue));
+}
+
+template <typename ValueType>
+inline void TDirectoryTree<ValueType>::FTreeNode::SetDefaultValue()
+{
+	RemoveValue();
+
+	uint32 FlagsValue = 0x1;
+	NumChildNodesAndFlags =
+		(NumChildNodesAndFlags & ~FlagsMask) |
+		(FlagsValue << FlagsShift);
+	new (&Value) ValueType();
+}
+
+template <typename ValueType>
+inline void TDirectoryTree<ValueType>::FTreeNode::RemoveValue()
+{
+	if (HasValue())
+	{
+		uint32 FlagsValue = 0x0;
+		NumChildNodesAndFlags =
+			(NumChildNodesAndFlags & ~FlagsMask) |
+			(FlagsValue << FlagsShift);
+
+		// We need a typedef here because VC won't compile the destructor call below if ValueType itself has a member called ValueType
+		typedef ValueType FDirectoryTreeDestructValueType;
+		((ValueType*)&Value)->FDirectoryTreeDestructValueType::~FDirectoryTreeDestructValueType();
+	}
+}
+
+template <typename ValueType>
+inline ValueType& TDirectoryTree<ValueType>::FTreeNode::FindOrAdd(FStringView InRelPath, bool& bOutExisted)
+{
+	FStringView FirstComponent;
+	FStringView RemainingPath;
+	FPathViews::SplitFirstComponent(InRelPath, FirstComponent, RemainingPath);
+	check(!FirstComponent.IsEmpty());
+	bool bExists;
+	int32 InsertionIndex = FindInsertionIndex(FirstComponent, bExists);
+	if (!bExists)
+	{
+		bOutExisted = false;
+		FTreeNode& ChildNode = InsertChildNode(InsertionIndex, FString(InRelPath), FTreeNode());
+		ChildNode.SetDefaultValue();
+		return ChildNode.GetValue();
+	}
+
+
+	FTreeNode& ChildNode = ChildNodes[InsertionIndex];
+	FString& ChildRelPath = RelPaths[InsertionIndex];
+
+	FStringView ExistingFirstComponent;
+	FStringView ExistingRemainingPath;
+	FPathViews::SplitFirstComponent(ChildRelPath, ExistingFirstComponent, ExistingRemainingPath);
+	check(!ExistingFirstComponent.IsEmpty());
+	int32 NumFirstComponents = 1;
+	for (int32 RunawayLoop = 0; RunawayLoop <= InRelPath.Len(); ++RunawayLoop)
+	{
+		if (ExistingRemainingPath.IsEmpty())
+		{
+			// We've reached the end of the existing path
+			if (!RemainingPath.IsEmpty())
+			{
+				// We have not reached the end of the input path, so it is a child of the existing path
+				return ChildNode.FindOrAdd(RemainingPath, bOutExisted);
+			}
+			else
+			{
+				// The input path matches the existing path
+				bOutExisted = ChildNode.HasValue();
+				if (!ChildNode.HasValue())
+				{
+					ChildNode.SetDefaultValue();
+				}
+				return ChildNode.GetValue();
+			}
+		}
+		else if (RemainingPath.IsEmpty())
+		{
+			// We've reached the end of the input path, so it is a parent of the existing path
+			bOutExisted = false;
+
+			// Create a new ChildNode and move the existing ChildNode into a child of it
+			// We can modify RelPath in place without messing up the sort order because the nodes are sorted by
+			// FirstComponent only and the FirstComponent is not changing
+			FTreeNode OldTreeNode = MoveTemp(ChildNode);
+			FString OldRelPath(ExistingRemainingPath);
+			ChildNode.Reset();
+			ChildNode.SetDefaultValue();
+			ChildRelPath = InRelPath;
+			ChildNode.InsertChildNode(0, MoveTemp(OldRelPath), MoveTemp(OldTreeNode));
+			return ChildNode.GetValue();
+		}
+		else
+		{
+			// Both existing and remaining have more directory components
+			FStringView NextFirstComponent;
+			FStringView NextRemainingPath;
+			FPathViews::SplitFirstComponent(RemainingPath, NextFirstComponent, NextRemainingPath);
+			check(!NextFirstComponent.IsEmpty());
+			FStringView NextExistingFirstComponent;
+			FStringView NextExistingRemainingPath;
+			FPathViews::SplitFirstComponent(ExistingRemainingPath, NextExistingFirstComponent, NextExistingRemainingPath);
+			check(!NextExistingFirstComponent.IsEmpty());
+			if (NextFirstComponent == NextExistingFirstComponent)
+			{
+				// Next component is also a match, go to the next loop iteration to handle the new remainingpaths
+				RemainingPath = NextRemainingPath;
+				ExistingRemainingPath = NextExistingRemainingPath;
+				++NumFirstComponents;
+				continue;
+			}
+			else
+			{
+				// Existing and remaining firstcomponent differ, so they are both child paths of a mutual parent path
+				// Reconstruct the CommonParentPath
+				TStringBuilder<1024> CommonParentPath;
+				FStringView ParentRemaining = ChildRelPath;
+				for (int32 FirstComponentIndex = 0; FirstComponentIndex < NumFirstComponents; ++FirstComponentIndex)
+				{
+					FStringView ParentFirstComponent;
+					FStringView NextParentRemaining;
+					FPathViews::SplitFirstComponent(ParentRemaining, ParentFirstComponent, NextParentRemaining);
+					FPathViews::Append(CommonParentPath, ParentFirstComponent);
+					ParentRemaining = NextParentRemaining;
+				}
+
+				// Create a new ChildNode and move the existing ChildNode into a child of it
+				// We can modify RelPath in place without messing up the sort order because the nodes are sorted by
+				// FirstComponent only and the FirstComponent is not changing
+				FTreeNode OldTreeNode = MoveTemp(ChildNode);
+				FString OldRelPath(ExistingRemainingPath);
+				ChildNode.Reset();
+				ChildRelPath = CommonParentPath;
+				ChildNode.InsertChildNode(0, MoveTemp(OldRelPath), MoveTemp(OldTreeNode));
+
+				// The input path is now a child of the modified ChildNode
+				// It's impossible for FindOrAdd to fail.
+				// If it ever becomes possible for it to fail, we will have to remove the child we added if it fails.
+				return ChildNode.FindOrAdd(RemainingPath, bOutExisted);
+			}
+		}
+	}
+	checkf(false, TEXT("Infinite loop trying to split path %.*s into components."), InRelPath.Len(), InRelPath.GetData());
+	return GetValue();
+}
+
+template <typename ValueType>
+inline void TDirectoryTree<ValueType>::FTreeNode::Remove(FStringView InRelPath, bool& bOutExisted)
+{
+	check(!InRelPath.IsEmpty());
+
+	FStringView FirstComponent;
+	FStringView RemainingPath;
+	FPathViews::SplitFirstComponent(InRelPath, FirstComponent, RemainingPath);
+	bool bExists;
+	int32 InsertionIndex = FindInsertionIndex(FirstComponent, bExists);
+	if (!bExists)
+	{
+		bOutExisted = false;
+		return;
+	}
+
+	FTreeNode& ChildNode = ChildNodes[InsertionIndex];
+	FString& ChildRelPath = RelPaths[InsertionIndex];
+
+	if (!FPathViews::TryMakeChildPathRelativeTo(InRelPath, ChildRelPath, RemainingPath))
+	{
+		bOutExisted = false;
+		return;
+	}
+
+	if (!RemainingPath.IsEmpty())
+	{
+		// The input path is a child of the existing path
+		ChildNode.Remove(RemainingPath, bOutExisted);
+		if (bOutExisted)
+		{
+			// If the remove was successful, the childnode must have had at least two paths, because
+			// otherwise we would have previously Compacted it to combine its RelPath with RemainingPath.
+			// In case it only had two paths and now has one path, try to compact it.
+			ConditionalCompactNode(ChildRelPath, ChildNode);
+		}
+	}
+	else
+	{
+		// The input path matches the existing path
+		bOutExisted = ChildNode.HasValue();
+		if (ChildNode.GetNumChildNodes() != 0)
+		{
+			if (ChildNode.HasValue())
+			{
+				ChildNode.RemoveValue();
+				ConditionalCompactNode(ChildRelPath, ChildNode);
+			}
+		}
+		else
+		{
+			RemoveChildNodeAt(InsertionIndex);
+		}
+	}
+}
+
+template <typename ValueType>
+inline ValueType* TDirectoryTree<ValueType>::FTreeNode::Find(FStringView InRelPath)
+{
+	check(!InRelPath.IsEmpty());
+
+	FStringView FirstComponent;
+	FStringView RemainingPath;
+	FPathViews::SplitFirstComponent(InRelPath, FirstComponent, RemainingPath);
+	bool bExists;
+	int32 InsertionIndex = FindInsertionIndex(FirstComponent, bExists);
+	if (!bExists)
+	{
+		return nullptr;
+	}
+
+	FTreeNode& ChildNode = ChildNodes[InsertionIndex];
+	FString& ChildRelPath = RelPaths[InsertionIndex];
+
+	if (!FPathViews::TryMakeChildPathRelativeTo(InRelPath, ChildRelPath, RemainingPath))
+	{
+		return nullptr;
+	}
+
+	if (!RemainingPath.IsEmpty())
+	{
+		// The input path is a child of the existing path
+		return ChildNode.Find(RemainingPath);
+	}
+	else
+	{
+		// The input path matches the existing path
+		return ChildNode.HasValue() ? &ChildNode.GetValue() : nullptr;
+	}
+}
+
+template <typename ValueType>
+inline ValueType* TDirectoryTree<ValueType>::FTreeNode::TryFindClosestPath(FStringView InRelPath,
+	FStringBuilderBase* OutPath, TCHAR InPathSeparator)
+{
+	check(!InRelPath.IsEmpty());
+
+	FStringView FirstComponent;
+	FStringView RemainingPath;
+	FPathViews::SplitFirstComponent(InRelPath, FirstComponent, RemainingPath);
+	bool bExists;
+	int32 InsertionIndex = FindInsertionIndex(FirstComponent, bExists);
+	if (!bExists)
+	{
+		return nullptr;
+	}
+
+	FTreeNode& ChildNode = ChildNodes[InsertionIndex];
+	const FString& ChildRelPath = RelPaths[InsertionIndex];
+
+	if (!FPathViews::TryMakeChildPathRelativeTo(InRelPath, ChildRelPath, RemainingPath))
+	{
+		return nullptr;
+	}
+
+	if (!RemainingPath.IsEmpty())
+	{
+		// The input path is a child of the existing path
+		int32 SavedOutPathLen = INDEX_NONE;
+		if (OutPath)
+		{
+			SavedOutPathLen = OutPath->Len();
+			FPathViews::Append(*OutPath, ChildRelPath);
+			UE::DirectoryTree::FixupPathSeparator(*OutPath, SavedOutPathLen, InPathSeparator);
+		}
+		ValueType* Result = ChildNode.TryFindClosestPath(RemainingPath, OutPath, InPathSeparator);
+		if (Result)
+		{
+			return Result;
+		}
+
+		if (ChildNode.HasValue())
+		{
+			return &ChildNode.GetValue();
+		}
+		if (OutPath)
+		{
+			OutPath->RemoveSuffix(OutPath->Len() - SavedOutPathLen);
+		}
+		return nullptr;
+	}
+	else
+	{
+		// The input path matches the existing path
+		if (!ChildNode.HasValue())
+		{
+			return nullptr;
+		}
+		if (OutPath)
+		{
+			int32 SavedOutPathLen = OutPath->Len();
+			FPathViews::Append(*OutPath, ChildRelPath);
+			UE::DirectoryTree::FixupPathSeparator(*OutPath, SavedOutPathLen, InPathSeparator);
+		}
+		return &ChildNode.GetValue();
+	}
+}
+
+template <typename ValueType>
+inline bool TDirectoryTree<ValueType>::FTreeNode::IsEmpty() const
+{
+	return !HasValue() && GetNumChildNodes() == 0;
+}
+
+template <typename ValueType>
+inline SIZE_T TDirectoryTree<ValueType>::FTreeNode::GetAllocatedSize() const
+{
+	SIZE_T Size = 0;
+	Size = CapacityChildNodes * (sizeof(FTreeNode) + sizeof(FString));
+	int32 NumChildNodes = GetNumChildNodes();
+	for (const FString& RelPath : TConstArrayView<FString>(RelPaths.Get(), NumChildNodes))
+	{
+		Size += RelPath.GetAllocatedSize();
+	}
+	for (const FTreeNode& ChildNode : TConstArrayView<FTreeNode>(ChildNodes.Get(), NumChildNodes))
+	{
+		Size += ChildNode.GetAllocatedSize();
+	}
+	return Size;
+}
+
+template <typename ValueType>
+inline void TDirectoryTree<ValueType>::FTreeNode::Shrink()
+{
+	int32 NumChildNodes = GetNumChildNodes();
+	Realloc(NumChildNodes);
+
+	for (FString& RelPath : TArrayView<FString>(RelPaths.Get(), NumChildNodes))
+	{
+		RelPath.Shrink();
+	}
+	for (FTreeNode& ChildNode : TArrayView<FTreeNode>(ChildNodes.Get(), NumChildNodes))
+	{
+		ChildNode.Shrink();
+	}
+}
+
+template <typename ValueType>
+inline int32 TDirectoryTree<ValueType>::FTreeNode::FindInsertionIndex(FStringView FirstPathComponent, bool& bOutExists) const
+{
+	return UE::DirectoryTree::FindInsertionIndex(GetNumChildNodes(), RelPaths, FirstPathComponent, bOutExists);
+}
+
+template <typename ValueType>
+inline typename TDirectoryTree<ValueType>::FTreeNode& TDirectoryTree<ValueType>::FTreeNode::InsertChildNode(int32 InsertionIndex, FString&& RelPath,
+	FTreeNode&& ChildNode)
+{
+	int32 NumChildNodes = GetNumChildNodes();
+	check(0 <= InsertionIndex && InsertionIndex <= NumChildNodes);
+	checkf((((uint32)(NumChildNodes + 1)) & FlagsMask) == 0, TEXT("Overflow"));
+
+	if (NumChildNodes == CapacityChildNodes)
+	{
+		int32 NewCapacity;
+		if (CapacityChildNodes == 0)
+		{
+			NewCapacity = 1;
+		}
+		else if (CapacityChildNodes < 8)
+		{
+			NewCapacity = CapacityChildNodes * 2;
+		}
+		else
+		{
+			NewCapacity = (CapacityChildNodes * 3) / 2;
+			checkf(NewCapacity > CapacityChildNodes, TEXT("Overflow"));
+		}
+		Realloc(NewCapacity);
+		check(NumChildNodes < CapacityChildNodes);
+	}
+
+	for (int32 ShiftIndex = NumChildNodes; ShiftIndex > InsertionIndex; --ShiftIndex)
+	{
+		RelPaths[ShiftIndex] = MoveTemp(RelPaths[ShiftIndex - 1]);
+		ChildNodes[ShiftIndex] = MoveTemp(ChildNodes[ShiftIndex - 1]);
+	}
+	RelPaths[InsertionIndex] = MoveTemp(RelPath);
+	ChildNodes[InsertionIndex] = MoveTemp(ChildNode);
+	SetNumChildNodes(NumChildNodes + 1);
+	return ChildNodes[InsertionIndex];
+}
+
+template <typename ValueType>
+inline void TDirectoryTree<ValueType>::FTreeNode::RemoveChildNodeAt(int32 RemoveIndex)
+{
+	int32 NumChildNodes = GetNumChildNodes();
+	check(0 <= RemoveIndex && RemoveIndex < NumChildNodes);
+	for (int32 ShiftIndex = RemoveIndex; ShiftIndex < NumChildNodes - 1; ++ShiftIndex)
+	{
+		RelPaths[ShiftIndex] = MoveTemp(RelPaths[ShiftIndex + 1]);
+		ChildNodes[ShiftIndex] = MoveTemp(ChildNodes[ShiftIndex + 1]);
+	}
+	RelPaths[NumChildNodes - 1].Empty();
+	ChildNodes[NumChildNodes - 1].Reset();
+	SetNumChildNodes(NumChildNodes - 1);
+}
+
+template <typename ValueType>
+inline void TDirectoryTree<ValueType>::FTreeNode::ConditionalCompactNode(FString& RelPath, FTreeNode& ChildNode)
+{
+	if (ChildNode.HasValue())
+	{
+		return;
+	}
+
+	int32 NumChildNodes = ChildNode.GetNumChildNodes();
+	checkf(NumChildNodes > 0, TEXT("Invalid to call ConditionalCompactNode with an empty ChildNode."));
+	if (NumChildNodes > 1)
+	{
+		return;
+	}
+	TStringBuilder<1024> NewRelPath;
+	FPathViews::AppendPath(NewRelPath, RelPath);
+	FPathViews::AppendPath(NewRelPath, ChildNode.RelPaths[0]);
+	RelPath = NewRelPath;
+	FTreeNode OldGrandChild = MoveTemp(ChildNode.ChildNodes[0]);
+	ChildNode = MoveTemp(OldGrandChild);
+}
+
+template <typename ValueType>
+inline void TDirectoryTree<ValueType>::FTreeNode::Realloc(int32 NewCapacity)
+{
+	if (CapacityChildNodes == NewCapacity)
+	{
+		return;
+	}
+	int32 NumChildNodes = GetNumChildNodes();
+	check(NumChildNodes <= NewCapacity);
+	if (NewCapacity > 0)
+	{
+		FString* NewRelPaths = new FString[NewCapacity];
+		FTreeNode* NewChildNodes = new FTreeNode[NewCapacity];
+		for (int32 Index = 0; Index < NumChildNodes; ++Index)
+		{
+			NewRelPaths[Index] = MoveTemp(RelPaths[Index]);
+			NewChildNodes[Index] = MoveTemp(ChildNodes[Index]);
+		}
+		ChildNodes.Reset(NewChildNodes);
+		RelPaths.Reset(NewRelPaths);
+	}
+	else
+	{
+		ChildNodes.Reset();
+		RelPaths.Reset();
+	}
+	CapacityChildNodes = NewCapacity;
+}
diff --git a/Engine/Source/Runtime/CoreUObject/Private/Misc/DirectoryTree.cpp b/Engine/Source/Runtime/CoreUObject/Private/Misc/DirectoryTree.cpp
new file mode 100644
--- /dev/null
+++ b/Engine/Source/Runtime/CoreUObject/Private/Misc/DirectoryTree.cpp
@@ -0,0 +1,467 @@
+// Copyright Epic Games, Inc. All Rights Reserved.
+
+#include "Containers/DirectoryTree.h"
+
+#include "Algo/BinarySearch.h"
+#include "Misc/AutomationTest.h"
+
+namespace UE::DirectoryTree
+{
+
+void FixupPathSeparator(FStringBuilderBase& InOutPath, int32 StartIndex, TCHAR InPathSeparator)
+{
+	if (InPathSeparator == '/')
+	{
+		return;
+	}
+	int32 SeparatorIndex;
+	while (InOutPath.ToView().RightChop(StartIndex).FindChar('/', SeparatorIndex))
+	{
+		StartIndex += SeparatorIndex;
+		InOutPath.GetData()[StartIndex] = InPathSeparator;
+	}
+}
+
+int32 FindInsertionIndex(int32 NumChildNodes, const TUniquePtr<FString[]>& RelPaths, FStringView FirstPathComponent, bool& bOutExists)
+{
+	TConstArrayView<FString> RelPathsRange(RelPaths.Get(), NumChildNodes);
+	int32 Index = Algo::LowerBound(RelPathsRange, FirstPathComponent,
+		[](const FString& ChildRelPath, FStringView FirstPathComponent)
+		{
+			return FPathViews::Less(ChildRelPath, FirstPathComponent);
+		}
+	);
+	bOutExists = Index < NumChildNodes && FPathViews::IsParentPathOf(FirstPathComponent, RelPaths[Index]);
+	return Index;
+}
+
+} // namespace UE::DirectoryTree
+
+#if WITH_DEV_AUTOMATION_TESTS
+
+IMPLEMENT_SIMPLE_AUTOMATION_TEST(FDirectoryTreeTest, "System.Core.DirectoryTree", EAutomationTestFlags::EditorContext | EAutomationTestFlags::EngineFilter)
+bool FDirectoryTreeTest::RunTest(const FString& Parameters)
+{
+	constexpr int32 NumPathTypes = 5;
+	constexpr int32 NumPaths = 7;
+	FStringView PathsByTypeAndIndex[NumPathTypes][NumPaths] =
+	{
+		{
+			TEXTVIEW("/Game/Dir2"),
+			TEXTVIEW("/Game/Path1"),
+			TEXTVIEW("/Game/Dir2/Path2"),
+			TEXTVIEW("/Plugin1/Path1"),
+			TEXTVIEW("/Plugin1/Dir2/Path2"),
+			TEXTVIEW("/Engine/Path1"),
+			TEXTVIEW("/Plugin2/Path1"),
+		},
+		{
+			TEXTVIEW("d:\\root\\Project\\Content\\Dir2"),
+			TEXTVIEW("d:\\root\\Project\\Content\\Path1.uasset"),
+			TEXTVIEW("d:\\root\\Project\\Content\\Dir2\\Path2.uasset"),
+			TEXTVIEW("d:\\root\\Project\\Plugins\\Plugin1\\Content\\Path1.uasset"),
+			TEXTVIEW("d:\\root\\Project\\Plugins\\Plugin1\\Content\\Dir2\\Path2.uasset"),
+			TEXTVIEW("d:\\root\\Engine\\Content\\Path1.uasset"),
+			TEXTVIEW("e:\\root\\Project\\Plugins\\Plugin2\\Content\\Path1.uasset"),
+		},
+		{
+			TEXTVIEW("d:/root/Project/Content/Dir2"),
+			TEXTVIEW("d:/root/Project/Content/Path1.uasset"),
+			TEXTVIEW("d:/root/Project/Content/Dir2/Path2.uasset"),
+			TEXTVIEW("d:/root/Project/Plugins/Plugin1/Content/Path1.uasset"),
+			TEXTVIEW("d:/root/Project/Plugins/Plugin1/Content/Dir2/Path2.uasset"),
+			TEXTVIEW("d:/root/Engine/Content/Path1.uasset"),
+			TEXTVIEW("e:/root/Project/Plugins/Plugin2/Content/Path1.uasset"),
+		},
+		{
+			TEXTVIEW("..\\..\\..\\Project\\Content\\Dir2"),
+			TEXTVIEW("..\\..\\..\\Project\\Content\\Path1.uasset"),
+			TEXTVIEW("..\\..\\..\\Project\\Content\\Dir2\\Path2.uasset"),
+			TEXTVIEW("..\\..\\..\\Project\\Plugins\\Plugin1\\Content\\Path1.uasset"),
+			TEXTVIEW("..\\..\\..\\Project\\Plugins\\Plugin1\\Content\\Dir2\\Path2.uasset"),
+			TEXTVIEW("..\\..\\..\\Engine\\Content\\Path1.uasset"),
+			TEXTVIEW("e:\\root\\Project\\Plugins\\Plugin2\\Content\\Path1.uasset"),
+		},
+		{
+			TEXTVIEW("../../../Project/Content/Dir2"),
+			TEXTVIEW("../../../Project/Content/Path1.uasset"),
+			TEXTVIEW("../../../Project/Content/Dir2/Path2.uasset"),
+			TEXTVIEW("../../../Project/Plugins/Plugin1/Content/Path1.uasset"),
+			TEXTVIEW("../../../Project/Plugins/Plugin1/Content/Dir2/Path2.uasset"),
+			TEXTVIEW("../../../Engine/Content/Path1.uasset"),
+			TEXTVIEW("e:/root/Project/Plugins/Plugin2/Content/Path1.uasset"),
+		},
+	};
+	FStringView Path0SubPath[NumPathTypes] =
+	{
+		TEXTVIEW("/Game/Dir2/Sub"),
+		TEXTVIEW("d:\\root\\Project\\Content\\Dir2\\Sub"),
+		TEXTVIEW("d:/root/Project/Content/Dir2/Sub"),
+		TEXTVIEW("..\\..\\..\\Project\\Content\\Dir2"),
+		TEXTVIEW("../../../Project/Content/Dir2/Sub"),
+	};
+	int32 ValueByIndex[NumPaths] = { 1,2,3,4,5,6,7 };
+	FStringView NonPathsByTypeAndIndex0[] =
+		{
+			TEXTVIEW(""),
+			TEXTVIEW("/"),
+			TEXTVIEW("/Game"),
+			TEXTVIEW("/Game/"),
+			TEXTVIEW("/Plugin1"),
+			TEXTVIEW("/Plugin1/"),
+			TEXTVIEW("/Plugin1/Dir2"),
+			TEXTVIEW("/Plugin1/Dir2/"),
+			TEXTVIEW("/Engine"),
+			TEXTVIEW("/Engine/"),
+		};
+	FStringView NonPathsByTypeAndIndex1[] =
+		{
+			TEXTVIEW(""),
+			TEXTVIEW("d:\\"),
+			TEXTVIEW("d:\\root1"),
+			TEXTVIEW("d:\\root1\\"),
+			TEXTVIEW("d:\\root1\\Project"),
+			TEXTVIEW("d:\\root1\\Project\\"),
+			TEXTVIEW("d:\\root1\\Project\\Content"),
+			TEXTVIEW("d:\\root1\\Project\\Plugins\\"),
+			TEXTVIEW("d:\\root1\\Project\\Plugins\\Content"),
+			TEXTVIEW("d:\\root1\\Project\\Plugins\\Content\\"),
+			TEXTVIEW("d:\\root1\\Project\\Plugins\\Content\\Plugin1"),
+			TEXTVIEW("d:\\root1\\Project\\Plugins\\Content\\Plugin1\\"),
+			TEXTVIEW("d:\\root1\\Project\\Plugins\\Content\\Plugin1\\Dir2"),
+			TEXTVIEW("d:\\root1\\Project\\Plugins\\Content\\Plugin1\\Dir2\\"),
+			TEXTVIEW("d:\\root1\\Engine"),
+			TEXTVIEW("d:\\root1\\Engine\\"),
+			TEXTVIEW("d:\\root1\\Engine\\Content"),
+			TEXTVIEW("d:\\root1\\Engine\\Content\\"),
+		};
+	FStringView NonPathsByTypeAndIndex2[] =
+		{
+			TEXTVIEW(""),
+			TEXTVIEW("d:/"),
+			TEXTVIEW("d:/root1"),
+			TEXTVIEW("d:/root1/Project"),
+			TEXTVIEW("d:/root1/Project/Content"),
+			TEXTVIEW("d:/root1/Project/Plugins/Content"),
+			TEXTVIEW("d:/root1/Project/Plugins/Content/Plugin1"),
+			TEXTVIEW("d:/root1/Project/Plugins/Content/Plugin1/Dir2"),
+			TEXTVIEW("d:/root1/Engine"),
+			TEXTVIEW("d:/root1/Engine/Content"),
+		};
+	FStringView NonPathsByTypeAndIndex3[] =
+		{
+			TEXTVIEW(""),
+			TEXTVIEW("\\"),
+			TEXTVIEW(".."),
+			TEXTVIEW("..\\"),
+			TEXTVIEW("..\\.."),
+			TEXTVIEW("..\\..\\"),
+			TEXTVIEW("..\\..\\.."),
+			TEXTVIEW("..\\..\\..\\"),
+			TEXTVIEW("..\\..\\..\\Project"),
+			TEXTVIEW("..\\..\\..\\Project\\Content"),
+			TEXTVIEW("..\\..\\..\\Project\\Plugins\\Content"),
+			TEXTVIEW("..\\..\\..\\Project\\Plugins\\Content\\Plugin1"),
+			TEXTVIEW("..\\..\\..\\Project\\Plugins\\Content\\Plugin1\\Dir2"),
+			TEXTVIEW("..\\..\\..\\Engine"),
+			TEXTVIEW("..\\..\\..\\Engine\\Content"),
+		};
+	FStringView NonPathsByTypeAndIndex4[] =
+		{
+			TEXTVIEW(""),
+			TEXTVIEW("/"),
+			TEXTVIEW(".."),
+			TEXTVIEW("../.."),
+			TEXTVIEW("../../.."),
+			TEXTVIEW("../../../Project"),
+			TEXTVIEW("../../../Project/Content"),
+			TEXTVIEW("../../../Project/Plugins/Content"),
+			TEXTVIEW("../../../Project/Plugins/Content/Plugin1"),
+			TEXTVIEW("../../../Project/Plugins/Content/Plugin1/Dir2"),
+			TEXTVIEW("../../../Engine"),
+			TEXTVIEW("../../../Engine/Content"),
+		};
+	TConstArrayView<FStringView> NonPathsByTypeAndIndex[NumPathTypes] =
+	{
+		NonPathsByTypeAndIndex0, NonPathsByTypeAndIndex1, NonPathsByTypeAndIndex2, NonPathsByTypeAndIndex3,
+		NonPathsByTypeAndIndex4
+	};
+
+	constexpr int32 NumPermutations = 2;
+	int32 Permutations[NumPermutations][NumPaths] =
+	{
+		{ 0, 1, 2, 3, 4, 5, 6 },
+		{ 6, 5, 4, 3, 2, 1, 0 },
+	};
+
+	auto GetDirTreeTestName = [](FStringView InTestName, int32 PathType, int32 Permutation, int32 EditPermutationIndex,
+		int32 OtherPermutationIndex)
+	{
+		return FString::Printf(TEXT("%.*s(%d, %d, %d, %d)"),
+			InTestName.Len(), InTestName.GetData(),
+			PathType, Permutation, EditPermutationIndex, OtherPermutationIndex);
+	};
+
+	for (int32 PathType = 0; PathType < NumPathTypes; ++PathType)
+	{
+		int32 NumNonPaths = NonPathsByTypeAndIndex[PathType].Num();
+		for (int32 Permutation = 0; Permutation < NumPermutations; ++Permutation)
+		{
+			TDirectoryTree<int32> Tree;
+
+			// Add all the Paths in the given order and make Contains assertions after each addition
+			for (int32 AddPathPermutationIndex = 0; AddPathPermutationIndex < NumPaths; ++AddPathPermutationIndex)
+			{
+				int32 AddPathIndex = Permutations[Permutation][AddPathPermutationIndex];
+				FStringView AddPath = PathsByTypeAndIndex[PathType][AddPathIndex];
+
+				// Add the path
+				Tree.FindOrAdd(AddPath) = ValueByIndex[AddPathIndex];
+
+				if (Tree.Num() != AddPathPermutationIndex+1)
+				{
+					AddError(GetDirTreeTestName(TEXTVIEW("TreeNum has expected value"), PathType, Permutation,
+						AddPathPermutationIndex, 0));
+				}
+
+				// Assert all paths up to and including this one are included
+				for (int32 OtherPermutationIndex = 0; OtherPermutationIndex <= AddPathPermutationIndex;
+					++OtherPermutationIndex)
+				{
+					int32 OtherIndex = Permutations[Permutation][OtherPermutationIndex];
+					FStringView OtherPath = PathsByTypeAndIndex[PathType][OtherIndex];
+					int32 OtherValue = ValueByIndex[OtherIndex];
+					int32* ExistingValue = Tree.Find(OtherPath);
+					if (!ExistingValue)
+					{
+						AddError(GetDirTreeTestName(TEXTVIEW("OtherInList"), PathType, Permutation,
+							AddPathPermutationIndex, OtherPermutationIndex));
+					}
+					else
+					{
+						if (*ExistingValue != OtherValue)
+						{
+							AddError(GetDirTreeTestName(TEXTVIEW("OtherInListMatchesValue"), PathType, Permutation,
+								AddPathPermutationIndex, OtherPermutationIndex));
+						}
+						if (!Tree.ContainsPathOrParent(OtherPath))
+						{
+							AddError(GetDirTreeTestName(TEXTVIEW("ContainsPathOrParentOtherInList"), PathType, Permutation,
+								AddPathPermutationIndex, OtherPermutationIndex));
+						}
+						else
+						{
+							FString ClosestPath;
+							if (!Tree.TryFindClosestPath(OtherPath, ClosestPath))
+							{
+								AddError(GetDirTreeTestName(TEXTVIEW("TryFindClosestPathOtherInListSucceeds"), PathType, Permutation,
+									AddPathPermutationIndex, OtherPermutationIndex));
+							}
+							else if (ClosestPath != OtherPath)
+							{
+								AddError(GetDirTreeTestName(TEXTVIEW("TryFindClosestPathOtherInListMatches"), PathType, Permutation,
+									AddPathPermutationIndex, OtherPermutationIndex));
+							}
+						}
+					}
+				}
+
+				// Assert all paths not yet added are not included
+				for (int32 OtherPermutationIndex = AddPathPermutationIndex + 1; OtherPermutationIndex < NumPaths;
+					++OtherPermutationIndex)
+				{
+					int32 OtherIndex = Permutations[Permutation][OtherPermutationIndex];
+					FStringView OtherPath = PathsByTypeAndIndex[PathType][OtherIndex];
+					if (Tree.Contains(OtherPath))
+					{
+						AddError(GetDirTreeTestName(TEXTVIEW("OtherNotInList"), PathType, Permutation,
+							AddPathPermutationIndex, OtherPermutationIndex));
+					}
+					// Call ContainsPathOrParent to test whether it crashes, but not yet implemented that we verify 
+					// what its return value should be.
+					(void)Tree.ContainsPathOrParent(OtherPath);
+				}
+
+				// Assert all non paths are not included
+				for (int32 NonPathIndex = 0; NonPathIndex < NumNonPaths; ++NonPathIndex)
+				{
+					FStringView NonPath = NonPathsByTypeAndIndex[PathType][NonPathIndex];
+					if (Tree.Contains(NonPath))
+					{
+						AddError(GetDirTreeTestName(TEXTVIEW("NonPathNotInList"), PathType, Permutation,
+							AddPathPermutationIndex, NonPathIndex));
+					}
+					// Call ContainsPathOrParent to test whether it crashes, but not yet implemented that we verify 
+					// what its return value should be.
+					(void)Tree.ContainsPathOrParent(NonPath);
+				}
+			}
+
+			// Verify that the SubPath is present
+			FString ExistingSubParentPath;
+			int32* ExistingSubParentValue;
+			if (!Tree.TryFindClosestPath(Path0SubPath[PathType], ExistingSubParentPath, &ExistingSubParentValue))
+			{
+				AddError(GetDirTreeTestName(TEXTVIEW("SubPathInTree"), PathType, Permutation, 0, 0));
+			}
+			else if (ExistingSubParentPath != PathsByTypeAndIndex[PathType][0] || *ExistingSubParentValue != ValueByIndex[0])
+			{
+				AddError(GetDirTreeTestName(TEXTVIEW("SubPathInTreeMatches"), PathType, Permutation, 0, 0));
+			}
+
+			// Remove all the Paths (in specified order) and make Contains assertions after each removal
+			// Currently we only test removal in FIFO order; bugs that are specific to a removal-order
+			// should be dependent only on the final added state and should not be dependent on the earlier add-order.
+			for (int32 RemovePathPermutationIndex = 0; RemovePathPermutationIndex < NumPaths;
+				++RemovePathPermutationIndex)
+			{
+				int32 RemovePathIndex = Permutations[Permutation][RemovePathPermutationIndex];
+				FStringView RemovePath = PathsByTypeAndIndex[PathType][RemovePathIndex];
+
+				// Remove the path
+				bool bExisted;
+				Tree.Remove(RemovePath, &bExisted);
+				if (!bExisted)
+				{
+					AddError(GetDirTreeTestName(TEXTVIEW("RemoveFoundSomethingToRemove"), PathType, Permutation,
+						RemovePathPermutationIndex, 0));
+				}
+				if (Tree.Num() != NumPaths - (RemovePathPermutationIndex+1))
+				{
+					AddError(GetDirTreeTestName(TEXTVIEW("TreeNum has expected value"), PathType, Permutation,
+						RemovePathPermutationIndex, 0));
+				}
+
+
+				// Assert all paths not yet removed are still included
+				for (int32 OtherPermutationIndex = RemovePathPermutationIndex+1; OtherPermutationIndex < NumPaths;
+					++OtherPermutationIndex)
+				{
+					int32 OtherIndex = Permutations[Permutation][OtherPermutationIndex];
+					FStringView OtherPath = PathsByTypeAndIndex[PathType][OtherIndex];
+					int32 OtherValue = ValueByIndex[OtherIndex];
+					int32* ExistingValue = Tree.Find(OtherPath);
+					if (!ExistingValue)
+					{
+						AddError(GetDirTreeTestName(TEXTVIEW("OtherInListAfterRemoval"), PathType, Permutation,
+							RemovePathPermutationIndex, OtherPermutationIndex));
+					}
+					else
+					{
+						if (*ExistingValue != OtherValue)
+						{
+							AddError(GetDirTreeTestName(TEXTVIEW("OtherInListAfterRemovalMatches"), PathType, Permutation,
+								RemovePathPermutationIndex, OtherPermutationIndex));
+						}
+						if (!Tree.ContainsPathOrParent(OtherPath))
+						{
+							AddError(GetDirTreeTestName(TEXTVIEW("OtherContainsPathOrParentAfterRemoval"), PathType, Permutation,
+								RemovePathPermutationIndex, OtherPermutationIndex));
+						}
+					}
+				}
+
+				// Assert all paths up to and including this one have been removed
+				for (int32 OtherPermutationIndex = 0; OtherPermutationIndex <= RemovePathPermutationIndex;
+					++OtherPermutationIndex)
+				{
+					int32 OtherIndex = Permutations[Permutation][OtherPermutationIndex];
+					FStringView OtherPath = PathsByTypeAndIndex[PathType][OtherIndex];
+					if (Tree.Contains(OtherPath))
+					{
+						AddError(GetDirTreeTestName(TEXT("OtherNotInListAfterRemoval"), PathType, Permutation,
+							RemovePathPermutationIndex, OtherPermutationIndex));
+					}
+					// Call ContainsPathOrParent to test whether it crashes, but not yet implemented that we verify 
+					// what its return value should be.
+					(void)Tree.ContainsPathOrParent(OtherPath);
+				}
+
+				// Assert all non paths are not still included
+				for (int32 NonPathIndex = 0; NonPathIndex < NumNonPaths; ++NonPathIndex)
+				{
+					FStringView NonPath = NonPathsByTypeAndIndex[PathType][NonPathIndex];
+					if (Tree.Contains(NonPath))
+					{
+						AddError(GetDirTreeTestName(TEXTVIEW("NonPathNotInListAfterRemoval"), PathType, Permutation,
+							RemovePathPermutationIndex, NonPathIndex));
+					}
+					// Call ContainsPathOrParent to test whether it crashes, but not yet implemented that we verify 
+					// what its return value should be.
+					(void)Tree.ContainsPathOrParent(NonPath);
+				}
+			}
+			if (!Tree.IsEmpty())
+			{
+				AddError(GetDirTreeTestName(TEXTVIEW("TreeEmptyAfterRemoval"), PathType, Permutation, 0, 0));
+			}
+		}
+	}
+
+	// Testing some pathtype-independent scenarios
+	{
+		TDirectoryTree<int32> Tree;
+		Tree.FindOrAdd(TEXTVIEW("/Root/Path1")) = 1;
+		Tree.FindOrAdd(TEXTVIEW("/Root/Path2")) = 2;
+		int32* FoundRoot = Tree.FindClosestValue(TEXTVIEW("/Root"));
+		int32* FoundPath1 = Tree.FindClosestValue(TEXTVIEW("/Root/Path1"));
+		int32* FoundPath1Sub = Tree.FindClosestValue(TEXTVIEW("/Root/Path1/Sub"));
+		int32* FoundPath2 = Tree.FindClosestValue(TEXTVIEW("/Root/Path2"));
+		int32* FoundPath2Sub = Tree.FindClosestValue(TEXTVIEW("/Root/Path2/Sub"));
+		TestTrue(TEXT("TwoPaths Root does not exist"), FoundRoot == nullptr);
+		TestTrue(TEXT("TwoPaths Path1 Value matches"), FoundPath1 && *FoundPath1 == 1);
+		TestTrue(TEXT("TwoPaths Path1Sub Value matches"), FoundPath1Sub &&* FoundPath1Sub == 1);
+		TestTrue(TEXT("TwoPaths Path2 Value matches"), FoundPath2 && *FoundPath2 == 2);
+		TestTrue(TEXT("TwoPaths Path2Sub Value matches"), FoundPath2Sub && *FoundPath2Sub == 2);
+	}
+	{
+		TDirectoryTree<int32> Tree;
+		Tree.FindOrAdd(TEXTVIEW("/Root/Path1/A/B/C")) = 1;
+		Tree.FindOrAdd(TEXTVIEW("/Root/Path2/A/B/C")) = 2;
+		int32* FoundRoot = Tree.FindClosestValue(TEXTVIEW("/Root"));
+		int32* FoundPath1 = Tree.FindClosestValue(TEXTVIEW("/Root/Path1/A/B/C"));
+		int32* FoundPath1Sub = Tree.FindClosestValue(TEXTVIEW("/Root/Path1/A/B/C/Sub"));
+		int32* FoundPath1Parent = Tree.FindClosestValue(TEXTVIEW("/Root/Path1/A"));
+		int32* FoundPath2 = Tree.FindClosestValue(TEXTVIEW("/Root/Path2/A/B/C"));
+		int32* FoundPath2Sub = Tree.FindClosestValue(TEXTVIEW("/Root/Path2/A/B/C/Sub"));
+		int32* FoundPath2Parent = Tree.FindClosestValue(TEXTVIEW("/Root/Path2/A"));
+		TestTrue(TEXT("TwoPathsLong Root does not exist"), FoundRoot == nullptr);
+		TestTrue(TEXT("TwoPathsLong Path1 Value matches"), FoundPath1 && *FoundPath1 == 1);
+		TestTrue(TEXT("TwoPathsLong Path1Sub Value matches"), FoundPath1Sub && *FoundPath1Sub == 1);
+		TestTrue(TEXT("TwoPathsLong Path1 Parent does not exist"), FoundPath1Parent == nullptr);
+		TestTrue(TEXT("TwoPathsLong Path2 Value matches"), FoundPath2 && *FoundPath2 == 2);
+		TestTrue(TEXT("TwoPathsLong Path2Sub Value matches"), FoundPath2Sub && *FoundPath2Sub == 2);
+		TestTrue(TEXT("TwoPathsLong Path2 Parent does not exist"), FoundPath2Parent == nullptr);
+	}
+
+	struct FMoveConstructOnly
+	{
+		FMoveConstructOnly() { Value = 437; }
+		FMoveConstructOnly(FMoveConstructOnly&& Other) { Value = Other.Value; Other.Value = -1; }
+		FMoveConstructOnly(const FMoveConstructOnly& Other) = delete;
+		FMoveConstructOnly& operator=(FMoveConstructOnly&& Other) = delete;
+		FMoveConstructOnly& operator=(const FMoveConstructOnly& Other) = delete;
+		
+		int32 Value;
+	};
+
+	{
+		TDirectoryTree<FMoveConstructOnly> Tree;
+		Tree.FindOrAdd(TEXTVIEW("/Root/PathM")).Value = 1;
+		Tree.FindOrAdd(TEXTVIEW("/Root/PathP")).Value = 2;
+		Tree.FindOrAdd(TEXTVIEW("/Root/PathA")).Value = 3;
+		Tree.FindOrAdd(TEXTVIEW("/Root/PathZ"));
+
+		FMoveConstructOnly* Value = Tree.Find(TEXTVIEW("/Root/PathA"));
+		TestTrue(TEXT("MoveConstructOnlyValueA correct"), Value && Value->Value == 3);
+		Value = Tree.Find(TEXTVIEW("/Root/PathM"));
+		TestTrue(TEXT("MoveConstructOnlyValueM correct"), Value && Value->Value == 1);
+		Value = Tree.Find(TEXTVIEW("/Root/PathP"));
+		TestTrue(TEXT("MoveConstructOnlyValueP correct"), Value && Value->Value == 2);
+		Value = Tree.Find(TEXTVIEW("/Root/PathZ"));
+		TestTrue(TEXT("MoveConstructOnlyValueZ correct"), Value && Value->Value == 437);
+	}
+	return true;
+}
+
+#endif
