Updated GetWrappedObjectPropertyValue_InContainer to allow getting any part of an entire UObject* wrapper array in one call.

diff --git a/Engine/Source/Runtime/CoreUObject/Private/UObject/PropertyClassPtr.cpp b/Engine/Source/Runtime/CoreUObject/Private/UObject/PropertyClassPtr.cpp
--- a/Engine/Source/Runtime/CoreUObject/Private/UObject/PropertyClassPtr.cpp
+++ b/Engine/Source/Runtime/CoreUObject/Private/UObject/PropertyClassPtr.cpp
@@ -53,7 +53,9 @@ UObject* FClassPtrProperty::GetObjectPropertyValue(const void* PropertyValueAddr
 
 UObject* FClassPtrProperty::GetObjectPropertyValue_InContainer(const void* ContainerAddress, int32 ArrayIndex) const
 {
-	return GetWrappedObjectPropertyValue_InContainer<FObjectPtr>(ContainerAddress, ArrayIndex);
+	UObject* Result = nullptr;
+	GetWrappedUObjectPtrValues_InContainer<FObjectPtr>(&Result, ContainerAddress, ArrayIndex, 1);
+	return Result;
 }
 
 void FClassPtrProperty::SetObjectPropertyValue(void* PropertyValueAddress, UObject* Value) const
diff --git a/Engine/Source/Runtime/CoreUObject/Private/UObject/PropertyLazyObjectPtr.cpp b/Engine/Source/Runtime/CoreUObject/Private/UObject/PropertyLazyObjectPtr.cpp
--- a/Engine/Source/Runtime/CoreUObject/Private/UObject/PropertyLazyObjectPtr.cpp
+++ b/Engine/Source/Runtime/CoreUObject/Private/UObject/PropertyLazyObjectPtr.cpp
@@ -117,7 +117,9 @@ UObject* FLazyObjectProperty::GetObjectPropertyValue(const void* PropertyValueAd
 
 UObject* FLazyObjectProperty::GetObjectPropertyValue_InContainer(const void* ContainerAddress, int32 ArrayIndex) const
 {
-	return GetWrappedObjectPropertyValue_InContainer<FLazyObjectPtr>(ContainerAddress, ArrayIndex);
+	UObject* Result = nullptr;
+	GetWrappedUObjectPtrValues_InContainer<FLazyObjectPtr>(&Result, ContainerAddress, ArrayIndex, 1);
+	return Result;
 }
 
 void FLazyObjectProperty::SetObjectPropertyValue(void* PropertyValueAddress, UObject* Value) const
diff --git a/Engine/Source/Runtime/CoreUObject/Private/UObject/PropertyObject.cpp b/Engine/Source/Runtime/CoreUObject/Private/UObject/PropertyObject.cpp
--- a/Engine/Source/Runtime/CoreUObject/Private/UObject/PropertyObject.cpp
+++ b/Engine/Source/Runtime/CoreUObject/Private/UObject/PropertyObject.cpp
@@ -192,7 +192,7 @@ UObject* FObjectProperty::GetObjectPropertyValue(const void* PropertyValueAddres
 
 UObject* FObjectProperty::GetObjectPropertyValue_InContainer(const void* ContainerAddress, int32 ArrayIndex) const
 {
-	// Dummy struct to satisfy GetWrappedObjectPropertyValue_InContainer API requirements (Get() function)
+	// Dummy struct to satisfy GetWrappedUObjectPtrValues_InContainer API requirements (Get() function)
 	struct FWrappedObjectPtr
 	{
 		UObject* Obj = nullptr;
@@ -201,7 +201,9 @@ UObject* FObjectProperty::GetObjectPropertyValue_InContainer(const void* Contain
 			return Obj;
 		}
 	};
-	return GetWrappedObjectPropertyValue_InContainer<FWrappedObjectPtr>(ContainerAddress, ArrayIndex);
+	UObject* Result = nullptr;
+	GetWrappedUObjectPtrValues_InContainer<FWrappedObjectPtr>(&Result, ContainerAddress, ArrayIndex, 1);
+	return Result;
 }
 
 void FObjectProperty::SetObjectPropertyValue(void* PropertyValueAddress, UObject* Value) const
diff --git a/Engine/Source/Runtime/CoreUObject/Private/UObject/PropertyObjectPtr.cpp b/Engine/Source/Runtime/CoreUObject/Private/UObject/PropertyObjectPtr.cpp
--- a/Engine/Source/Runtime/CoreUObject/Private/UObject/PropertyObjectPtr.cpp
+++ b/Engine/Source/Runtime/CoreUObject/Private/UObject/PropertyObjectPtr.cpp
@@ -158,7 +158,9 @@ UObject* FObjectPtrProperty::GetObjectPropertyValue(const void* PropertyValueAdd
 
 UObject* FObjectPtrProperty::GetObjectPropertyValue_InContainer(const void* ContainerAddress, int32 ArrayIndex) const
 {
-	return GetWrappedObjectPropertyValue_InContainer<FObjectPtr>(ContainerAddress, ArrayIndex);
+	UObject* Result = nullptr;
+	GetWrappedUObjectPtrValues_InContainer<FObjectPtr>(&Result, ContainerAddress, ArrayIndex, 1);
+	return Result;
 }
 
 void FObjectPtrProperty::SetObjectPropertyValue(void* PropertyValueAddress, UObject* Value) const
diff --git a/Engine/Source/Runtime/CoreUObject/Private/UObject/PropertySoftObjectPtr.cpp b/Engine/Source/Runtime/CoreUObject/Private/UObject/PropertySoftObjectPtr.cpp
--- a/Engine/Source/Runtime/CoreUObject/Private/UObject/PropertySoftObjectPtr.cpp
+++ b/Engine/Source/Runtime/CoreUObject/Private/UObject/PropertySoftObjectPtr.cpp
@@ -243,7 +243,9 @@ UObject* FSoftObjectProperty::GetObjectPropertyValue(const void* PropertyValueAd
 
 UObject* FSoftObjectProperty::GetObjectPropertyValue_InContainer(const void* ContainerAddress, int32 ArrayIndex) const
 {
-	return GetWrappedObjectPropertyValue_InContainer<FSoftObjectPtr>(ContainerAddress, ArrayIndex);
+	UObject* Result = nullptr;
+	GetWrappedUObjectPtrValues_InContainer<FSoftObjectPtr>(&Result, ContainerAddress, ArrayIndex, 1);
+	return Result;
 }
 
 void FSoftObjectProperty::SetObjectPropertyValue(void* PropertyValueAddress, UObject* Value) const
diff --git a/Engine/Source/Runtime/CoreUObject/Private/UObject/PropertyWeakObjectPtr.cpp b/Engine/Source/Runtime/CoreUObject/Private/UObject/PropertyWeakObjectPtr.cpp
--- a/Engine/Source/Runtime/CoreUObject/Private/UObject/PropertyWeakObjectPtr.cpp
+++ b/Engine/Source/Runtime/CoreUObject/Private/UObject/PropertyWeakObjectPtr.cpp
@@ -98,7 +98,9 @@ UObject* FWeakObjectProperty::GetObjectPropertyValue(const void* PropertyValueAd
 
 UObject* FWeakObjectProperty::GetObjectPropertyValue_InContainer(const void* ContainerAddress, int32 ArrayIndex) const
 {
-	return GetWrappedObjectPropertyValue_InContainer<FWeakObjectPtr>(ContainerAddress, ArrayIndex);
+	UObject* Result = nullptr;
+	GetWrappedUObjectPtrValues_InContainer<FWeakObjectPtr>(&Result, ContainerAddress, ArrayIndex, 1);
+	return Result;
 }
 
 void FWeakObjectProperty::SetObjectPropertyValue(void* PropertyValueAddress, UObject* Value) const
diff --git a/Engine/Source/Runtime/CoreUObject/Public/UObject/UnrealType.h b/Engine/Source/Runtime/CoreUObject/Public/UObject/UnrealType.h
--- a/Engine/Source/Runtime/CoreUObject/Public/UObject/UnrealType.h
+++ b/Engine/Source/Runtime/CoreUObject/Public/UObject/UnrealType.h
@@ -2745,31 +2745,52 @@ protected:
 
 	/* Helper functions for UObject property types that wrap the object pointer in a smart pointer */
 	template <typename T>
-	UObject* GetWrappedObjectPropertyValue_InContainer(const void* ContainerAddress, int32 ArrayIndex) const
+	void GetWrappedUObjectPtrValues_InContainer(UObject** OutObjects, const void* ContainerAddress, int32 ArrayIndex, int32 ArrayCount) const
 	{
+		// Ensure required range is valid
+		checkf(ArrayIndex >= 0 && ArrayCount >= 0 && ArrayIndex <= ArrayDim && ArrayCount <= ArrayDim && ArrayIndex <= ArrayDim - ArrayCount, TEXT("ArrayIndex (%d) and ArrayCount (%d) is invalid for an array of size %d"), ArrayIndex, ArrayCount, ArrayDim);
+
 		if (!HasGetter())
 		{
 			// Fast path - direct memory access
-			return GetObjectPropertyValue(ContainerPtrToValuePtr<void>(ContainerAddress, ArrayIndex));
+			const uint8* ObjAddress = (const uint8*)ContainerPtrToValuePtr<void>(ContainerAddress, ArrayIndex);
+
+			if (ArrayCount == 1)
+			{
+				*OutObjects = GetObjectPropertyValue(ObjAddress);
+			}
+			else
+			{
+				int32 LocalElementSize = ElementSize;
+				for (int32 OutIndex = 0; OutIndex != ArrayCount; ++OutIndex)
+				{
+					OutObjects[OutIndex] = GetObjectPropertyValue(ObjAddress + OutIndex * LocalElementSize);
+				}
+			}
 		}
 		else
 		{
-			T Value;
-			if (ArrayDim == 1)
+			if (ArrayCount == 1)
 			{
 				// Slower but no mallocs. We can copy the value directly to the resulting param
+				T Value;
 				GetValue_InContainer(ContainerAddress, &Value);
+				*OutObjects = Value.Get();
 			}
 			else
 			{
 				// Malloc a temp value that is the size of the array. Getter will then copy the entire array to the temp value
 				T* ValueArray = (T*)AllocateAndInitializeValue();
 				FProperty::GetValue_InContainer(ContainerAddress, ValueArray);
-				// Grab the item we care about and free the temp array
-				Value = ValueArray[ArrayIndex];
+
+				// Grab the items we care about and free the temp array
+				int32 LocalElementSize = ElementSize;
+				for (int32 OutIndex = 0; OutIndex != ArrayCount; ++OutIndex)
+				{
+					OutObjects[OutIndex] = ValueArray[ArrayIndex + OutIndex].Get();
+				}
 				DestroyAndFreeValue(ValueArray);
 			}
-			return Value.Get();
 		}
 	}
 	template <typename T>
