Fixed memory corruption caused by a TWeakObjectPtr property with a getter used in a blueprint.

diff --git a/Engine/Source/Runtime/CoreUObject/Private/UObject/ScriptCore.cpp b/Engine/Source/Runtime/CoreUObject/Private/UObject/ScriptCore.cpp
--- a/Engine/Source/Runtime/CoreUObject/Private/UObject/ScriptCore.cpp
+++ b/Engine/Source/Runtime/CoreUObject/Private/UObject/ScriptCore.cpp
@@ -501,7 +501,18 @@ void FFrame::StepExplicitProperty(void*const Result, FProperty* Property)
 		MostRecentPropertyContainer = Locals;
 		if (Property->HasGetter())
 		{
-			Property->GetValue_InContainer(MostRecentPropertyContainer, Result);
+			// Temporary hack to fix TWeakObjectPtrs with getters being copied wrongly to the VM -
+			// hacked here because we cannot change UnrealType.h
+			// Object properties are copied to the VM by UObject*, except for FSoftObjectProperty (see overrides for FSoftObjectProperty::CopyCompleteValueToScriptVM).
+			FObjectPropertyBase* ObjProperty = CastField<FObjectPropertyBase>(Property);
+			if (ObjProperty && !ObjProperty->IsA<FSoftObjectProperty>())
+			{
+				*(UObject**)Result = ObjProperty->GetObjectPropertyValue_InContainer(MostRecentPropertyContainer);
+			}
+			else
+			{
+				Property->GetValue_InContainer(MostRecentPropertyContainer, Result);
+			}
 		}
 		else
 		{
@@ -2197,7 +2208,18 @@ DEFINE_VM_FUNCTION(UObject::execLocalVariable)
 		{
 			if (VarProperty->HasGetter())
 			{
-				VarProperty->GetValue_InContainer(Stack.MostRecentPropertyContainer, RESULT_PARAM);
+				// Temporary hack to fix TWeakObjectPtrs with getters being copied wrongly to the VM -
+				// hacked here because we cannot change UnrealType.h
+				// Object properties are copied to the VM by UObject*, except for FSoftObjectProperty (see overrides for FSoftObjectProperty::CopyCompleteValueToScriptVM).
+				FObjectPropertyBase* ObjProperty = CastField<FObjectPropertyBase>(VarProperty);
+				if (ObjProperty && !ObjProperty->IsA<FSoftObjectProperty>())
+				{
+					*(UObject**)RESULT_PARAM = ObjProperty->GetObjectPropertyValue_InContainer(Stack.MostRecentPropertyContainer);
+				}
+				else
+				{
+					ObjProperty->GetValue_InContainer(Stack.MostRecentPropertyContainer, RESULT_PARAM);
+				}
 			}
 			else
 			{
@@ -2229,7 +2251,18 @@ DEFINE_VM_FUNCTION(UObject::execInstanceVariable)
 		{
 			if (VarProperty->HasGetter())
 			{
-				VarProperty->GetValue_InContainer(Stack.MostRecentPropertyContainer, RESULT_PARAM);
+				// Temporary hack to fix TWeakObjectPtrs with getters being copied wrongly to the VM -
+				// hacked here because we cannot change UnrealType.h
+				// Object properties are copied to the VM by UObject*, except for FSoftObjectProperty (see overrides for FSoftObjectProperty::CopyCompleteValueToScriptVM).
+				FObjectPropertyBase* ObjProperty = CastField<FObjectPropertyBase>(VarProperty);
+				if (ObjProperty && !ObjProperty->IsA<FSoftObjectProperty>())
+				{
+					*(UObject**)RESULT_PARAM = ObjProperty->GetObjectPropertyValue_InContainer(Stack.MostRecentPropertyContainer);
+				}
+				else
+				{
+					VarProperty->GetValue_InContainer(Stack.MostRecentPropertyContainer, RESULT_PARAM);
+				}
 			}
 			else
 			{
@@ -2265,7 +2298,18 @@ DEFINE_VM_FUNCTION(UObject::execClassSparseDataVariable)
 		{
 			if (VarProperty->HasGetter())
 			{
-				VarProperty->GetValue_InContainer(Stack.MostRecentPropertyContainer, RESULT_PARAM);
+				// Temporary hack to fix TWeakObjectPtrs with getters being copied wrongly to the VM -
+				// hacked here because we cannot change UnrealType.h
+				// Object properties are copied to the VM by UObject*, except for FSoftObjectProperty (see overrides for FSoftObjectProperty::CopyCompleteValueToScriptVM).
+				FObjectPropertyBase* ObjProperty = CastField<FObjectPropertyBase>(VarProperty);
+				if (ObjProperty && !ObjProperty->IsA<FSoftObjectProperty>())
+				{
+					*(UObject**)RESULT_PARAM = ObjProperty->GetObjectPropertyValue_InContainer(Stack.MostRecentPropertyContainer);
+				}
+				else
+				{
+					VarProperty->GetValue_InContainer(Stack.MostRecentPropertyContainer, RESULT_PARAM);
+				}
 			}
 			else
 			{
@@ -2309,7 +2353,18 @@ DEFINE_VM_FUNCTION(UObject::execDefaultVariable)
 			{
 				if (VarProperty->HasGetter())
 				{
-					VarProperty->GetValue_InContainer(Stack.MostRecentPropertyContainer, RESULT_PARAM);
+					// Temporary hack to fix TWeakObjectPtrs with getters being copied wrongly to the VM -
+					// hacked here because we cannot change UnrealType.h
+					// Object properties are copied to the VM by UObject*, except for FSoftObjectProperty (see overrides for FSoftObjectProperty::CopyCompleteValueToScriptVM).
+					FObjectPropertyBase* ObjProperty = CastField<FObjectPropertyBase>(VarProperty);
+					if (ObjProperty && !ObjProperty->IsA<FSoftObjectProperty>())
+					{
+						*(UObject**)RESULT_PARAM = ObjProperty->GetObjectPropertyValue_InContainer(Stack.MostRecentPropertyContainer);
+					}
+					else
+					{
+						VarProperty->GetValue_InContainer(Stack.MostRecentPropertyContainer, RESULT_PARAM);
+					}
 				}
 				else
 				{
@@ -3193,7 +3248,18 @@ DEFINE_VM_FUNCTION(UObject::execStructMemberContext)
 		{
 			if (StructProperty->HasGetter())
 			{
-				StructProperty->GetValue_InContainer(Stack.MostRecentPropertyContainer, RESULT_PARAM);
+				// Temporary hack to fix TWeakObjectPtrs with getters being copied wrongly to the VM -
+				// hacked here because we cannot change UnrealType.h
+				// Object properties are copied to the VM by UObject*, except for FSoftObjectProperty (see overrides for FSoftObjectProperty::CopyCompleteValueToScriptVM).
+				FObjectPropertyBase* ObjProperty = CastField<FObjectPropertyBase>(StructProperty);
+				if (ObjProperty && !ObjProperty->IsA<FSoftObjectProperty>())
+				{
+					*(UObject**)RESULT_PARAM = ObjProperty->GetObjectPropertyValue_InContainer(Stack.MostRecentPropertyContainer);
+				}
+				else
+				{
+					StructProperty->GetValue_InContainer(Stack.MostRecentPropertyContainer, RESULT_PARAM);
+				}
 			}
 			else
 			{
