Save Package: if we have at least one export that is outered to an object in the import's package, consider the reference as allowed.

diff --git a/Engine/Source/Runtime/CoreUObject/Private/UObject/SavePackage2.cpp b/Engine/Source/Runtime/CoreUObject/Private/UObject/SavePackage2.cpp
--- a/Engine/Source/Runtime/CoreUObject/Private/UObject/SavePackage2.cpp
+++ b/Engine/Source/Runtime/CoreUObject/Private/UObject/SavePackage2.cpp
@@ -719,6 +719,17 @@ ESavePackageResult ValidateImports(FSaveContext& SaveContext)
 
 	auto IsMapReferenceAllowed = [&SaveContext](UObject* InImport) -> bool
 	{
+		// If we have at least one export that is outered to an object in the import's package, consider the reference as allowed.
+		// Ideally, we would like to only allow imports from exports that shares an outer in the same package, but this will be ok
+		// for now.
+		for (const FTaggedExport& Export : SaveContext.GetExports())
+		{
+			if (Export.Obj->GetOuter()->GetPackage() == InImport->GetPackage())
+			{
+				return true;
+			}
+		}
+
 		if (!InImport->HasAnyFlags(RF_Public))
 		{
 			return false;
