DirectoryTree: Remove COREUBOBJECT_API, because inline functions cannot be marked as dllexport/dllimport.
#rnx

diff --git a/Engine/Source/Runtime/CoreUObject/Internal/Containers/DirectoryTree.h b/Engine/Source/Runtime/CoreUObject/Internal/Containers/DirectoryTree.h
--- a/Engine/Source/Runtime/CoreUObject/Internal/Containers/DirectoryTree.h
+++ b/Engine/Source/Runtime/CoreUObject/Internal/Containers/DirectoryTree.h
@@ -43,7 +43,7 @@ COREUOBJECT_API int32 FindInsertionIndex(int32 NumChildNodes, const TUniquePtr<F
  * by any functions that modify the tree, and should be discarded before calling any such functions.
  */
 template <typename ValueType>
-class COREUOBJECT_API TDirectoryTree
+class TDirectoryTree
 {
 public:
 	TDirectoryTree();
