Typo fixes.

diff --git a/Engine/Source/Runtime/CoreUObject/Private/UObject/Property.cpp b/Engine/Source/Runtime/CoreUObject/Private/UObject/Property.cpp
--- a/Engine/Source/Runtime/CoreUObject/Private/UObject/Property.cpp
+++ b/Engine/Source/Runtime/CoreUObject/Private/UObject/Property.cpp
@@ -1397,30 +1397,30 @@ void* FProperty::GetValueAddressAtIndex_Direct(const FProperty* Inner, void* InV
 	return (uint8*)InValueAddress + ElementSize * Index;
 }
 
-void FProperty::SetSingleValue_InContainer(void* InContainer, const void* InValue, int32 ArrayIndex) const
+void FProperty::SetSingleValue_InContainer(void* OutContainer, const void* InValue, int32 ArrayIndex) const
 {
 	checkf(ArrayIndex <= ArrayDim, TEXT("ArrayIndex (%d) must be less than the property %s array size (%d)"), ArrayIndex, *GetFullName(), ArrayDim);
 	if (!HasSetter())
 	{
 		// Fast path - direct memory access
-		CopySingleValue(ContainerVoidPtrToValuePtrInternal((void*)InContainer, ArrayIndex), InValue);
+		CopySingleValue(ContainerVoidPtrToValuePtrInternal((void*)OutContainer, ArrayIndex), InValue);
 	}
 	else
 	{
 		if (ArrayDim == 1)
 		{
 			// Slower but no mallocs. We can copy the value directly to the resulting param
-			CallSetter(InContainer, InValue);
+			CallSetter(OutContainer, InValue);
 		}
 		else
 		{
 			// Malloc a temp value that is the size of the array. We will then copy the entire array to the temp value
 			uint8* ValueArray = (uint8*)AllocateAndInitializeValue();
-			GetValue_InContainer(InContainer, ValueArray);
+			GetValue_InContainer(OutContainer, ValueArray);
 			// Replace the value at the specified index in the temp array with the InValue
 			CopySingleValue(ValueArray + ArrayIndex * ElementSize, InValue);
 			// Now call a setter to replace the entire array and then destroy the temp value
-			CallSetter(InContainer, ValueArray);
+			CallSetter(OutContainer, ValueArray);
 			DestroyAndFreeValue(ValueArray);
 		}
 	}
@@ -1453,21 +1453,21 @@ void FProperty::GetSingleValue_InContainer(const void* InContainer, void* OutVal
 	}
 }
 
-void FProperty::PerformOperationWithSetter(void* InContainer, void* DirectPropertyAddress, TFunctionRef<void(void*)> DirectValueAccessFunc) const
+void FProperty::PerformOperationWithSetter(void* OutContainer, void* DirectPropertyAddress, TFunctionRef<void(void*)> DirectValueAccessFunc) const
 {
-	if (InContainer && HasSetterOrGetter()) // If there's a getter we need to allocate a temp value even if there's no setter
+	if (OutContainer && HasSetterOrGetter()) // If there's a getter we need to allocate a temp value even if there's no setter
 	{
 		// When modifying container or struct properties that have a setter or getter function we first allocate a temp value
 		// that we can operate on directly (add new elements or modify existing ones)
 		void* LocalValuePtr = AllocateAndInitializeValue();
 		// Copy the value to the allocated local (using a getter if present)
-		GetValue_InContainer(InContainer, LocalValuePtr);
+		GetValue_InContainer(OutContainer, LocalValuePtr);
 
 		// Perform operation on the temp value
 		DirectValueAccessFunc(LocalValuePtr);
 
 		// Assign the temp value back to the property using a setter function
-		SetValue_InContainer(InContainer, LocalValuePtr);
+		SetValue_InContainer(OutContainer, LocalValuePtr);
 		// Destroy and free the temp value
 		DestroyAndFreeValue(LocalValuePtr);
 	}
@@ -1476,22 +1476,22 @@ void FProperty::PerformOperationWithSetter(void* InContainer, void* DirectProper
 		// When there's no setter or getter present it's ok to perform the operation directly on the container / struct memory
 		if (!DirectPropertyAddress)
 		{
-			checkf(InContainer, TEXT("Container pointr must be valid if DirectPropertyAddress is not valid"));
-			DirectPropertyAddress = PointerToValuePtr(InContainer, EPropertyPointerType::Container);
+			checkf(OutContainer, TEXT("Container pointr must be valid if DirectPropertyAddress is not valid"));
+			DirectPropertyAddress = PointerToValuePtr(OutContainer, EPropertyPointerType::Container);
 		}
 		DirectValueAccessFunc(DirectPropertyAddress);
 	}
 }
 
-void FProperty::PerformOperationWithGetter(void* InContainer, const void* DirectPropertyAddress, TFunctionRef<void(const void*)> DirectValueAccessFunc) const
+void FProperty::PerformOperationWithGetter(void* OutContainer, const void* DirectPropertyAddress, TFunctionRef<void(const void*)> DirectValueAccessFunc) const
 {
-	if (InContainer && HasGetter())
+	if (OutContainer && HasGetter())
 	{
 		// When modifying container or struct properties that have a getter function we first allocate a temp value
 		// that we can operate on directly (add new elements or modify existing ones)
 		void* LocalValuePtr = AllocateAndInitializeValue();
 		// Copy the value to the allocated local using a getter
-		GetValue_InContainer(InContainer, LocalValuePtr);
+		GetValue_InContainer(OutContainer, LocalValuePtr);
 
 		// Perform read-only operation on the temp value
 		DirectValueAccessFunc(LocalValuePtr);
@@ -1503,8 +1503,8 @@ void FProperty::PerformOperationWithGetter(void* InContainer, const void* Direct
 	{
 		if (!DirectPropertyAddress)
 		{
-			checkf(InContainer, TEXT("Container pointr must be valid if DirectPropertyAddress is not valid"));
-			DirectPropertyAddress = PointerToValuePtr(InContainer, EPropertyPointerType::Container);
+			checkf(OutContainer, TEXT("Container pointr must be valid if DirectPropertyAddress is not valid"));
+			DirectPropertyAddress = PointerToValuePtr(OutContainer, EPropertyPointerType::Container);
 		}
 		DirectValueAccessFunc(DirectPropertyAddress);
 	}
diff --git a/Engine/Source/Runtime/CoreUObject/Private/UObject/ScriptCore.cpp b/Engine/Source/Runtime/CoreUObject/Private/UObject/ScriptCore.cpp
--- a/Engine/Source/Runtime/CoreUObject/Private/UObject/ScriptCore.cpp
+++ b/Engine/Source/Runtime/CoreUObject/Private/UObject/ScriptCore.cpp
@@ -3341,8 +3341,8 @@ public:
 				if(NULL != Stack.MostRecentPropertyAddress)
 				{
 					check(Property->IsInContainer(SignatureFunction->ParmsSize));
-					uint8* ConstRefCopyParamAdress = Property->ContainerPtrToValuePtr<uint8>(Parameters);
-					Property->CopyCompleteValueToScriptVM(ConstRefCopyParamAdress, Stack.MostRecentPropertyAddress);
+					uint8* ConstRefCopyParamAddress = Property->ContainerPtrToValuePtr<uint8>(Parameters);
+					Property->CopyCompleteValueToScriptVM(ConstRefCopyParamAddress, Stack.MostRecentPropertyAddress);
 				}
 			}
 			else
diff --git a/Engine/Source/Runtime/CoreUObject/Public/UObject/UnrealType.h b/Engine/Source/Runtime/CoreUObject/Public/UObject/UnrealType.h
--- a/Engine/Source/Runtime/CoreUObject/Public/UObject/UnrealType.h
+++ b/Engine/Source/Runtime/CoreUObject/Public/UObject/UnrealType.h
@@ -571,15 +571,15 @@ public:
 		return ImportText_Internal(Buffer, PropertyPtr, EPropertyPointerType::Direct, OwnerObject, PortFlags, ErrorText);
 	}
 
-	FORCEINLINE void SetValue_InContainer(void* InContainer, const void* InValue) const
+	FORCEINLINE void SetValue_InContainer(void* OutContainer, const void* InValue) const
 	{
 		if (!HasSetter())
 		{
-			CopyCompleteValue(ContainerVoidPtrToValuePtrInternal(InContainer, 0), InValue);
+			CopyCompleteValue(ContainerVoidPtrToValuePtrInternal(OutContainer, 0), InValue);
 		}
 		else
 		{
-			CallSetter(InContainer, InValue);
+			CallSetter(OutContainer, InValue);
 		}
 	}
 	FORCEINLINE void GetValue_InContainer(void const* InContainer, void* OutValue) const
@@ -596,11 +596,11 @@ public:
 
 	/**
 	* Copies a single value to the property even if the property represents a static array of values
-	* @param InContainer Instance owner of the property
+	* @param OutContainer Instance owner of the property
 	* @param InValue Pointer to the memory that the value will be copied from. Must be at least ElementSize big
 	* @param ArrayIndex Index into the static array to copy the value from. If the property is not a static array it should be 0
 	*/
-	void SetSingleValue_InContainer(void* InContainer, const void* InValue, int32 ArrayIndex) const;
+	void SetSingleValue_InContainer(void* OutContainer, const void* InValue, int32 ArrayIndex) const;
 
 	/**
 	* Copies a single value to OutValue even if the property represents a static array of values
@@ -618,19 +618,19 @@ public:
 
 	/**
 	 * Helper function for setting container / struct property value and performing operation directly on the value memory
-	 * @param InContainer Pointer to the container that owns the property. Can be null but then setters and getters will not be used.
-	 * @param DirectPropertyAddress Direct property value address. Can be null only if InContainer is a valid pointer.
-	 * @param DirectValueAccessFunc Function that manipulates directly on property value address. The value address can be different than the passed in DirectPropertyAddress if setters and getters are present and InContainer pointer is valid.
+	 * @param OutContainer Pointer to the container that owns the property. Can be null but then setters and getters will not be used.
+	 * @param DirectPropertyAddress Direct property value address. Can be null only if OutContainer is a valid pointer.
+	 * @param DirectValueAccessFunc Function that manipulates directly on property value address. The value address can be different than the passed in DirectPropertyAddress if setters and getters are present and OutContainer pointer is valid.
 	 */
-	void PerformOperationWithSetter(void* InContainer, void* DirectPropertyAddress, TFunctionRef<void(void*)> DirectValueAccessFunc) const;
+	void PerformOperationWithSetter(void* OutContainer, void* DirectPropertyAddress, TFunctionRef<void(void*)> DirectValueAccessFunc) const;
 
 	/**
 	 * Helper function for getting container / struct property value and performing operation directly on the value memory
-	 * @param InContainer Pointer to the container that owns the property. Can be null but then setters and getters will not be used.
-	 * @param DirectPropertyAddress Direct property value address. Can be null only if InContainer is a valid pointer.
-	 * @param DirectValueAccessFunc Function that manipulates directly on property value address. The value address can be different than the passed in DirectPropertyAddress if setters and getters are present and InContainer pointer is valid.
+	 * @param OutContainer Pointer to the container that owns the property. Can be null but then setters and getters will not be used.
+	 * @param DirectPropertyAddress Direct property value address. Can be null only if OutContainer is a valid pointer.
+	 * @param DirectValueAccessFunc Function that manipulates directly on property value address. The value address can be different than the passed in DirectPropertyAddress if setters and getters are present and OutContainer pointer is valid.
 	 */
-	void PerformOperationWithGetter(void* InContainer, const void* DirectPropertyAddress, TFunctionRef<void(const void*)> DirectValueAccessFunc) const;
+	void PerformOperationWithGetter(void* OutContainer, const void* DirectPropertyAddress, TFunctionRef<void(const void*)> DirectValueAccessFunc) const;
 
 	/** 
 	 * Gets value address at given index inside of a static array or container
@@ -1472,9 +1472,9 @@ public:
 		*GetPropertyValuePtr_InContainer(A, ArrayIndex) = Value;
 	}
 
-	FORCEINLINE void SetValue_InContainer(void* InContainer, const TCppType& InValue) const
+	FORCEINLINE void SetValue_InContainer(void* OutContainer, const TCppType& InValue) const
 	{
-		TInPropertyBaseClass::SetValue_InContainer(InContainer, &InValue);
+		TInPropertyBaseClass::SetValue_InContainer(OutContainer, &InValue);
 	}
 
 	FORCEINLINE void GetValue_InContainer(void const* InContainer, TCppType* OutValue) const
