Updated SetWrappedObjectPropertyValue_InContainer to allow setting any part of an entire UObject* wrapper array in one call.

diff --git a/Engine/Source/Runtime/CoreUObject/Private/UObject/PropertyClassPtr.cpp b/Engine/Source/Runtime/CoreUObject/Private/UObject/PropertyClassPtr.cpp
--- a/Engine/Source/Runtime/CoreUObject/Private/UObject/PropertyClassPtr.cpp
+++ b/Engine/Source/Runtime/CoreUObject/Private/UObject/PropertyClassPtr.cpp
@@ -74,7 +74,7 @@ void FClassPtrProperty::SetObjectPropertyValue_InContainer(void* ContainerAddres
 {
 	if (Value || !HasAnyPropertyFlags(CPF_NonNullable))
 	{
-		SetWrappedObjectPropertyValue_InContainer<FObjectPtr>(ContainerAddress, Value, ArrayIndex);
+		SetWrappedUObjectPtrValues_InContainer<FObjectPtr>(ContainerAddress, &Value, ArrayIndex, 1);
 	}
 	else
 	{
diff --git a/Engine/Source/Runtime/CoreUObject/Private/UObject/PropertyLazyObjectPtr.cpp b/Engine/Source/Runtime/CoreUObject/Private/UObject/PropertyLazyObjectPtr.cpp
--- a/Engine/Source/Runtime/CoreUObject/Private/UObject/PropertyLazyObjectPtr.cpp
+++ b/Engine/Source/Runtime/CoreUObject/Private/UObject/PropertyLazyObjectPtr.cpp
@@ -138,7 +138,7 @@ void FLazyObjectProperty::SetObjectPropertyValue_InContainer(void* ContainerAddr
 {
 	if (Value || !HasAnyPropertyFlags(CPF_NonNullable))
 	{
-		SetWrappedObjectPropertyValue_InContainer<FLazyObjectPtr>(ContainerAddress, Value, ArrayIndex);
+		SetWrappedUObjectPtrValues_InContainer<FLazyObjectPtr>(ContainerAddress, &Value, ArrayIndex, 1);
 	}
 	else
 	{
diff --git a/Engine/Source/Runtime/CoreUObject/Private/UObject/PropertyObject.cpp b/Engine/Source/Runtime/CoreUObject/Private/UObject/PropertyObject.cpp
--- a/Engine/Source/Runtime/CoreUObject/Private/UObject/PropertyObject.cpp
+++ b/Engine/Source/Runtime/CoreUObject/Private/UObject/PropertyObject.cpp
@@ -220,7 +220,7 @@ void FObjectProperty::SetObjectPropertyValue(void* PropertyValueAddress, UObject
 
 void FObjectProperty::SetObjectPropertyValue_InContainer(void* ContainerAddress, UObject* Value, int32 ArrayIndex) const
 {
-	// Dummy struct to satisfy SetWrappedObjectPropertyValue_InContainer API requirements
+	// Dummy struct to satisfy SetWrappedUObjectPtrValues_InContainer API requirements
 	struct FWrappedObjectPtr
 	{
 		FWrappedObjectPtr() = default;
@@ -232,7 +232,7 @@ void FObjectProperty::SetObjectPropertyValue_InContainer(void* ContainerAddress,
 	};
 	if (Value || !HasAnyPropertyFlags(CPF_NonNullable))
 	{
-		SetWrappedObjectPropertyValue_InContainer<FWrappedObjectPtr>(ContainerAddress, Value, ArrayIndex);
+		SetWrappedUObjectPtrValues_InContainer<FWrappedObjectPtr>(ContainerAddress, &Value, ArrayIndex, 1);
 	}
 	else
 	{
diff --git a/Engine/Source/Runtime/CoreUObject/Private/UObject/PropertyObjectPtr.cpp b/Engine/Source/Runtime/CoreUObject/Private/UObject/PropertyObjectPtr.cpp
--- a/Engine/Source/Runtime/CoreUObject/Private/UObject/PropertyObjectPtr.cpp
+++ b/Engine/Source/Runtime/CoreUObject/Private/UObject/PropertyObjectPtr.cpp
@@ -179,7 +179,7 @@ void FObjectPtrProperty::SetObjectPropertyValue_InContainer(void* ContainerAddre
 {
 	if (Value || !HasAnyPropertyFlags(CPF_NonNullable))
 	{
-		SetWrappedObjectPropertyValue_InContainer<FObjectPtr>(ContainerAddress, Value, ArrayIndex);
+		SetWrappedUObjectPtrValues_InContainer<FObjectPtr>(ContainerAddress, &Value, ArrayIndex, 1);
 	}
 	else
 	{
diff --git a/Engine/Source/Runtime/CoreUObject/Private/UObject/PropertySoftObjectPtr.cpp b/Engine/Source/Runtime/CoreUObject/Private/UObject/PropertySoftObjectPtr.cpp
--- a/Engine/Source/Runtime/CoreUObject/Private/UObject/PropertySoftObjectPtr.cpp
+++ b/Engine/Source/Runtime/CoreUObject/Private/UObject/PropertySoftObjectPtr.cpp
@@ -255,7 +255,7 @@ void FSoftObjectProperty::SetObjectPropertyValue(void* PropertyValueAddress, UOb
 
 void FSoftObjectProperty::SetObjectPropertyValue_InContainer(void* ContainerAddress, UObject* Value, int32 ArrayIndex) const
 {
-	SetWrappedObjectPropertyValue_InContainer<FSoftObjectPtr>(ContainerAddress, Value, ArrayIndex);
+	SetWrappedUObjectPtrValues_InContainer<FSoftObjectPtr>(ContainerAddress, &Value, ArrayIndex, 1);
 }
 
 bool FSoftObjectProperty::AllowCrossLevel() const
diff --git a/Engine/Source/Runtime/CoreUObject/Private/UObject/PropertyWeakObjectPtr.cpp b/Engine/Source/Runtime/CoreUObject/Private/UObject/PropertyWeakObjectPtr.cpp
--- a/Engine/Source/Runtime/CoreUObject/Private/UObject/PropertyWeakObjectPtr.cpp
+++ b/Engine/Source/Runtime/CoreUObject/Private/UObject/PropertyWeakObjectPtr.cpp
@@ -110,7 +110,7 @@ void FWeakObjectProperty::SetObjectPropertyValue(void* PropertyValueAddress, UOb
 
 void FWeakObjectProperty::SetObjectPropertyValue_InContainer(void* ContainerAddress, UObject* Value, int32 ArrayIndex) const
 {
-	SetWrappedObjectPropertyValue_InContainer<FWeakObjectPtr>(ContainerAddress, Value, ArrayIndex);
+	SetWrappedUObjectPtrValues_InContainer<FWeakObjectPtr>(ContainerAddress, &Value, ArrayIndex, 1);
 }
 
 uint32 FWeakObjectProperty::GetValueTypeHashInternal(const void* Src) const
diff --git a/Engine/Source/Runtime/CoreUObject/Public/UObject/UnrealType.h b/Engine/Source/Runtime/CoreUObject/Public/UObject/UnrealType.h
--- a/Engine/Source/Runtime/CoreUObject/Public/UObject/UnrealType.h
+++ b/Engine/Source/Runtime/CoreUObject/Public/UObject/UnrealType.h
@@ -2794,19 +2794,35 @@ protected:
 		}
 	}
 	template <typename T>
-	void SetWrappedObjectPropertyValue_InContainer(void* ContainerAddress, UObject* InValue, int32 ArrayIndex) const
+	void SetWrappedUObjectPtrValues_InContainer(void* ContainerAddress, UObject** InValues, int32 ArrayIndex, int32 ArrayCount) const
 	{
+		// Ensure required range is valid
+		checkf(ArrayIndex >= 0 && ArrayCount >= 0 && ArrayIndex <= ArrayDim && ArrayCount <= ArrayDim && ArrayIndex <= ArrayDim - ArrayCount, TEXT("ArrayIndex (%d) and ArrayCount (%d) is invalid for an array of size %d"), ArrayIndex, ArrayCount, ArrayDim);
+
 		if (!HasSetter())
 		{
 			// Fast path - direct memory access
-			SetObjectPropertyValue(ContainerPtrToValuePtr<void>(ContainerAddress, ArrayIndex), InValue);
+			uint8* ObjAddress = (uint8*)ContainerPtrToValuePtr<void>(ContainerAddress, ArrayIndex);
+
+			if (ArrayCount == 1)
+			{
+				SetObjectPropertyValue(ObjAddress, *InValues);
+			}
+			else
+			{
+				int32 LocalElementSize = ElementSize;
+				for (int32 OutIndex = 0; OutIndex != ArrayCount; ++OutIndex)
+				{
+					SetObjectPropertyValue(ObjAddress + OutIndex * LocalElementSize, InValues[OutIndex]);
+				}
+			}
 		}
 		else
-		{			
-			if (ArrayDim == 1)
+		{
+			if (ArrayCount == 1)
 			{
 				// Slower but no mallocs. We can copy a local wrapped value directly to the resulting param
-				T WrappedValue(InValue);
+				T WrappedValue(*InValues);
 				SetValue_InContainer(ContainerAddress, &WrappedValue);
 			}
 			else
@@ -2814,8 +2830,14 @@ protected:
 				// Malloc a temp value that is the size of the array. Getter will then copy the entire array to the temp value
 				T* ValueArray = (T*)AllocateAndInitializeValue();
 				FProperty::GetValue_InContainer(ContainerAddress, ValueArray);
-				// Replace the item we care about
-				ValueArray[ArrayIndex] = InValue;
+
+				// Replace the items we care about
+				int32 LocalElementSize = ElementSize;
+				for (int32 OutIndex = 0; OutIndex != ArrayCount; ++OutIndex)
+				{
+					ValueArray[ArrayIndex + OutIndex] = InValues[OutIndex];
+				}
+
 				// Now copy the entire array back to the property using a setter
 				SetValue_InContainer(ContainerAddress, ValueArray);
 				DestroyAndFreeValue(ValueArray);
