Proper implementation of the hack fix for the VM crash when using TWeakObjectPtr properties with getters in blueprints.

diff --git a/Engine/Source/Runtime/CoreUObject/Private/UObject/Property.cpp b/Engine/Source/Runtime/CoreUObject/Private/UObject/Property.cpp
--- a/Engine/Source/Runtime/CoreUObject/Private/UObject/Property.cpp
+++ b/Engine/Source/Runtime/CoreUObject/Private/UObject/Property.cpp
@@ -945,6 +945,32 @@ void FProperty::CopyCompleteValueFromScriptVM( void* Dest, void const* Src ) con
 	CopyCompleteValue(Dest, Src);
 }
 
+void FProperty::CopyCompleteValueToScriptVM_InContainer( void* OutValue, void const* InContainer ) const
+{
+	if (HasGetter())
+	{
+		CallGetter(InContainer, OutValue);
+	}
+	else
+	{
+		const void* InObj = ContainerPtrToValuePtr<uint8>(InContainer);
+		CopyCompleteValue(OutValue, InObj);
+	}
+}
+
+void FProperty::CopyCompleteValueFromScriptVM_InContainer( void* OutContainer, void const* InValue ) const
+{
+	if (HasSetter())
+	{
+		CallSetter(OutContainer, InValue);
+	}
+	else
+	{
+		void* OutObj = ContainerPtrToValuePtr<uint8>(OutContainer);
+		CopyCompleteValue(OutObj, InValue);
+	}
+}
+
 void FProperty::ClearValueInternal( void* Data ) const
 {
 	checkf(0, TEXT("%s failed to handle ClearValueInternal, but it was not CPF_NoDestructor | CPF_ZeroConstructor"), *GetFullName());
diff --git a/Engine/Source/Runtime/CoreUObject/Private/UObject/PropertyClassPtr.cpp b/Engine/Source/Runtime/CoreUObject/Private/UObject/PropertyClassPtr.cpp
--- a/Engine/Source/Runtime/CoreUObject/Private/UObject/PropertyClassPtr.cpp
+++ b/Engine/Source/Runtime/CoreUObject/Private/UObject/PropertyClassPtr.cpp
@@ -87,3 +87,12 @@ uint32 FClassPtrProperty::GetValueTypeHashInternal(const void* Src) const
 	return GetTypeHash((FObjectPtr&)GetPropertyValue(Src));
 }
 
+void FClassPtrProperty::CopyCompleteValueToScriptVM_InContainer(void* OutValue, void const* InContainer) const
+{
+	GetWrappedUObjectPtrValues_InContainer<FObjectPtr>((UObject**)OutValue, InContainer, 0, ArrayDim);
+}
+
+void FClassPtrProperty::CopyCompleteValueFromScriptVM_InContainer(void* OutContainer, void const* InValue) const
+{
+	SetWrappedUObjectPtrValues_InContainer<FObjectPtr>(OutContainer, (UObject**)InValue, 0, ArrayDim);
+}
diff --git a/Engine/Source/Runtime/CoreUObject/Private/UObject/PropertyLazyObjectPtr.cpp b/Engine/Source/Runtime/CoreUObject/Private/UObject/PropertyLazyObjectPtr.cpp
--- a/Engine/Source/Runtime/CoreUObject/Private/UObject/PropertyLazyObjectPtr.cpp
+++ b/Engine/Source/Runtime/CoreUObject/Private/UObject/PropertyLazyObjectPtr.cpp
@@ -156,3 +156,12 @@ uint32 FLazyObjectProperty::GetValueTypeHashInternal(const void* Src) const
 	return GetTypeHash(GetPropertyValue(Src));
 }
 
+void FLazyObjectProperty::CopyCompleteValueToScriptVM_InContainer(void* OutValue, void const* InContainer) const
+{
+	GetWrappedUObjectPtrValues_InContainer<FLazyObjectPtr>((UObject**)OutValue, InContainer, 0, ArrayDim);
+}
+
+void FLazyObjectProperty::CopyCompleteValueFromScriptVM_InContainer(void* OutContainer, void const* InValue) const
+{
+	SetWrappedUObjectPtrValues_InContainer<FLazyObjectPtr>(OutContainer, (UObject**)InValue, 0, ArrayDim);
+}
diff --git a/Engine/Source/Runtime/CoreUObject/Private/UObject/PropertyObject.cpp b/Engine/Source/Runtime/CoreUObject/Private/UObject/PropertyObject.cpp
--- a/Engine/Source/Runtime/CoreUObject/Private/UObject/PropertyObject.cpp
+++ b/Engine/Source/Runtime/CoreUObject/Private/UObject/PropertyObject.cpp
@@ -10,6 +10,27 @@
 #include "UObject/LinkerPlaceholderClass.h"
 #include "UObject/PropertyHelper.h"
 
+namespace UE::Core::Private
+{
+	// Dummy struct to satisfy [Get/Set]WrappedUObjectPtrValues_InContainer API requirements (Get() function)
+	struct FWrappedObjectPtr
+	{
+		UObject* Obj = nullptr;
+
+		FWrappedObjectPtr() = default;
+
+		FWrappedObjectPtr(UObject* InObj)
+			: Obj(InObj)
+		{
+		}
+
+		UObject* Get() const
+		{
+			return Obj;
+		}
+	};
+}
+
 /*-----------------------------------------------------------------------------
 	FObjectProperty.
 -----------------------------------------------------------------------------*/
@@ -192,17 +213,8 @@ UObject* FObjectProperty::GetObjectPropertyValue(const void* PropertyValueAddres
 
 UObject* FObjectProperty::GetObjectPropertyValue_InContainer(const void* ContainerAddress, int32 ArrayIndex) const
 {
-	// Dummy struct to satisfy GetWrappedUObjectPtrValues_InContainer API requirements (Get() function)
-	struct FWrappedObjectPtr
-	{
-		UObject* Obj = nullptr;
-		UObject* Get() const
-		{
-			return Obj;
-		}
-	};
 	UObject* Result = nullptr;
-	GetWrappedUObjectPtrValues_InContainer<FWrappedObjectPtr>(&Result, ContainerAddress, ArrayIndex, 1);
+	GetWrappedUObjectPtrValues_InContainer<UE::Core::Private::FWrappedObjectPtr>(&Result, ContainerAddress, ArrayIndex, 1);
 	return Result;
 }
 
@@ -220,22 +232,22 @@ void FObjectProperty::SetObjectPropertyValue(void* PropertyValueAddress, UObject
 
 void FObjectProperty::SetObjectPropertyValue_InContainer(void* ContainerAddress, UObject* Value, int32 ArrayIndex) const
 {
-	// Dummy struct to satisfy SetWrappedUObjectPtrValues_InContainer API requirements
-	struct FWrappedObjectPtr
-	{
-		FWrappedObjectPtr() = default;
-		FWrappedObjectPtr(UObject* InObj)
-			: Obj(InObj)
-		{
-		}
-		UObject* Obj = nullptr;
-	};
 	if (Value || !HasAnyPropertyFlags(CPF_NonNullable))
 	{
-		SetWrappedUObjectPtrValues_InContainer<FWrappedObjectPtr>(ContainerAddress, &Value, ArrayIndex, 1);
+		SetWrappedUObjectPtrValues_InContainer<UE::Core::Private::FWrappedObjectPtr>(ContainerAddress, &Value, ArrayIndex, 1);
 	}
 	else
 	{
 		UE_LOG(LogProperty, Verbose /*Warning*/, TEXT("Trying to assign null object value to non-nullable \"%s\""), *GetFullName());
 	}
 }
+
+void FObjectProperty::CopyCompleteValueToScriptVM_InContainer(void* OutValue, void const* InContainer) const
+{
+	GetWrappedUObjectPtrValues_InContainer<UE::Core::Private::FWrappedObjectPtr>((UObject**)OutValue, InContainer, 0, ArrayDim);
+}
+
+void FObjectProperty::CopyCompleteValueFromScriptVM_InContainer(void* OutContainer, void const* InValue) const
+{
+	SetWrappedUObjectPtrValues_InContainer<UE::Core::Private::FWrappedObjectPtr>(OutContainer, (UObject**)InValue, 0, ArrayDim);
+}
diff --git a/Engine/Source/Runtime/CoreUObject/Private/UObject/PropertyObjectPtr.cpp b/Engine/Source/Runtime/CoreUObject/Private/UObject/PropertyObjectPtr.cpp
--- a/Engine/Source/Runtime/CoreUObject/Private/UObject/PropertyObjectPtr.cpp
+++ b/Engine/Source/Runtime/CoreUObject/Private/UObject/PropertyObjectPtr.cpp
@@ -197,3 +197,12 @@ uint32 FObjectPtrProperty::GetValueTypeHashInternal(const void* Src) const
 	return GetTypeHash((FObjectPtr&)GetPropertyValue(Src));
 }
 
+void FObjectPtrProperty::CopyCompleteValueToScriptVM_InContainer(void* OutValue, void const* InContainer) const
+{
+	GetWrappedUObjectPtrValues_InContainer<FObjectPtr>((UObject**)OutValue, InContainer, 0, ArrayDim);
+}
+
+void FObjectPtrProperty::CopyCompleteValueFromScriptVM_InContainer(void* OutContainer, void const* InValue) const
+{
+	SetWrappedUObjectPtrValues_InContainer<FObjectPtr>(OutContainer, (UObject**)InValue, 0, ArrayDim);
+}
diff --git a/Engine/Source/Runtime/CoreUObject/Private/UObject/PropertyWeakObjectPtr.cpp b/Engine/Source/Runtime/CoreUObject/Private/UObject/PropertyWeakObjectPtr.cpp
--- a/Engine/Source/Runtime/CoreUObject/Private/UObject/PropertyWeakObjectPtr.cpp
+++ b/Engine/Source/Runtime/CoreUObject/Private/UObject/PropertyWeakObjectPtr.cpp
@@ -117,3 +117,13 @@ uint32 FWeakObjectProperty::GetValueTypeHashInternal(const void* Src) const
 {
 	return GetTypeHash(*(FWeakObjectPtr*)Src);
 }
+
+void FWeakObjectProperty::CopyCompleteValueToScriptVM_InContainer(void* OutValue, void const* InContainer) const
+{
+	GetWrappedUObjectPtrValues_InContainer<FWeakObjectPtr>((UObject**)OutValue, InContainer, 0, ArrayDim);
+}
+
+void FWeakObjectProperty::CopyCompleteValueFromScriptVM_InContainer(void* OutContainer, void const* InValue) const
+{
+	SetWrappedUObjectPtrValues_InContainer<FWeakObjectPtr>(OutContainer, (UObject**)InValue, 0, ArrayDim);
+}
diff --git a/Engine/Source/Runtime/CoreUObject/Private/UObject/ScriptCore.cpp b/Engine/Source/Runtime/CoreUObject/Private/UObject/ScriptCore.cpp
--- a/Engine/Source/Runtime/CoreUObject/Private/UObject/ScriptCore.cpp
+++ b/Engine/Source/Runtime/CoreUObject/Private/UObject/ScriptCore.cpp
@@ -499,25 +499,7 @@ void FFrame::StepExplicitProperty(void*const Result, FProperty* Property)
 	{
 		MostRecentPropertyAddress = Property->ContainerPtrToValuePtr<uint8>(Locals);
 		MostRecentPropertyContainer = Locals;
-		if (Property->HasGetter())
-		{
-			// Temporary hack to fix TWeakObjectPtrs with getters being copied wrongly to the VM -
-			// hacked here because we cannot change UnrealType.h
-			// Object properties are copied to the VM by UObject*, except for FSoftObjectProperty (see overrides for FSoftObjectProperty::CopyCompleteValueToScriptVM).
-			FObjectPropertyBase* ObjProperty = CastField<FObjectPropertyBase>(Property);
-			if (ObjProperty && !ObjProperty->IsA<FSoftObjectProperty>())
-			{
-				*(UObject**)Result = ObjProperty->GetObjectPropertyValue_InContainer(MostRecentPropertyContainer);
-			}
-			else
-			{
-				Property->GetValue_InContainer(MostRecentPropertyContainer, Result);
-			}
-		}
-		else
-		{
-			Property->CopyCompleteValueToScriptVM(Result, MostRecentPropertyAddress);
-		}
+		Property->CopyCompleteValueToScriptVM_InContainer(Result, MostRecentPropertyContainer);
 	}
 }
 
@@ -2206,25 +2188,7 @@ DEFINE_VM_FUNCTION(UObject::execLocalVariable)
 
 		if (RESULT_PARAM)
 		{
-			if (VarProperty->HasGetter())
-			{
-				// Temporary hack to fix TWeakObjectPtrs with getters being copied wrongly to the VM -
-				// hacked here because we cannot change UnrealType.h
-				// Object properties are copied to the VM by UObject*, except for FSoftObjectProperty (see overrides for FSoftObjectProperty::CopyCompleteValueToScriptVM).
-				FObjectPropertyBase* ObjProperty = CastField<FObjectPropertyBase>(VarProperty);
-				if (ObjProperty && !ObjProperty->IsA<FSoftObjectProperty>())
-				{
-					*(UObject**)RESULT_PARAM = ObjProperty->GetObjectPropertyValue_InContainer(Stack.MostRecentPropertyContainer);
-				}
-				else
-				{
-					ObjProperty->GetValue_InContainer(Stack.MostRecentPropertyContainer, RESULT_PARAM);
-				}
-			}
-			else
-			{
-				VarProperty->CopyCompleteValueToScriptVM(RESULT_PARAM, Stack.MostRecentPropertyAddress);
-			}
+			VarProperty->CopyCompleteValueToScriptVM_InContainer(RESULT_PARAM, Stack.MostRecentPropertyContainer);
 		}
 	}
 }
@@ -2249,25 +2213,7 @@ DEFINE_VM_FUNCTION(UObject::execInstanceVariable)
 		Stack.MostRecentPropertyContainer = (uint8*)P_THIS;
 		if (RESULT_PARAM)
 		{
-			if (VarProperty->HasGetter())
-			{
-				// Temporary hack to fix TWeakObjectPtrs with getters being copied wrongly to the VM -
-				// hacked here because we cannot change UnrealType.h
-				// Object properties are copied to the VM by UObject*, except for FSoftObjectProperty (see overrides for FSoftObjectProperty::CopyCompleteValueToScriptVM).
-				FObjectPropertyBase* ObjProperty = CastField<FObjectPropertyBase>(VarProperty);
-				if (ObjProperty && !ObjProperty->IsA<FSoftObjectProperty>())
-				{
-					*(UObject**)RESULT_PARAM = ObjProperty->GetObjectPropertyValue_InContainer(Stack.MostRecentPropertyContainer);
-				}
-				else
-				{
-					VarProperty->GetValue_InContainer(Stack.MostRecentPropertyContainer, RESULT_PARAM);
-				}
-			}
-			else
-			{
-				VarProperty->CopyCompleteValueToScriptVM(RESULT_PARAM, Stack.MostRecentPropertyAddress);
-			}
+			VarProperty->CopyCompleteValueToScriptVM_InContainer(RESULT_PARAM, Stack.MostRecentPropertyContainer);
 		}
 	}
 
@@ -2296,25 +2242,7 @@ DEFINE_VM_FUNCTION(UObject::execClassSparseDataVariable)
 
 		if (RESULT_PARAM)
 		{
-			if (VarProperty->HasGetter())
-			{
-				// Temporary hack to fix TWeakObjectPtrs with getters being copied wrongly to the VM -
-				// hacked here because we cannot change UnrealType.h
-				// Object properties are copied to the VM by UObject*, except for FSoftObjectProperty (see overrides for FSoftObjectProperty::CopyCompleteValueToScriptVM).
-				FObjectPropertyBase* ObjProperty = CastField<FObjectPropertyBase>(VarProperty);
-				if (ObjProperty && !ObjProperty->IsA<FSoftObjectProperty>())
-				{
-					*(UObject**)RESULT_PARAM = ObjProperty->GetObjectPropertyValue_InContainer(Stack.MostRecentPropertyContainer);
-				}
-				else
-				{
-					VarProperty->GetValue_InContainer(Stack.MostRecentPropertyContainer, RESULT_PARAM);
-				}
-			}
-			else
-			{
-				VarProperty->CopyCompleteValueToScriptVM(RESULT_PARAM, Stack.MostRecentPropertyAddress);
-			}
+			VarProperty->CopyCompleteValueToScriptVM_InContainer(RESULT_PARAM, Stack.MostRecentPropertyContainer);
 		}
 	}
 }
@@ -2351,25 +2279,7 @@ DEFINE_VM_FUNCTION(UObject::execDefaultVariable)
 
 			if(RESULT_PARAM)
 			{
-				if (VarProperty->HasGetter())
-				{
-					// Temporary hack to fix TWeakObjectPtrs with getters being copied wrongly to the VM -
-					// hacked here because we cannot change UnrealType.h
-					// Object properties are copied to the VM by UObject*, except for FSoftObjectProperty (see overrides for FSoftObjectProperty::CopyCompleteValueToScriptVM).
-					FObjectPropertyBase* ObjProperty = CastField<FObjectPropertyBase>(VarProperty);
-					if (ObjProperty && !ObjProperty->IsA<FSoftObjectProperty>())
-					{
-						*(UObject**)RESULT_PARAM = ObjProperty->GetObjectPropertyValue_InContainer(Stack.MostRecentPropertyContainer);
-					}
-					else
-					{
-						VarProperty->GetValue_InContainer(Stack.MostRecentPropertyContainer, RESULT_PARAM);
-					}
-				}
-				else
-				{
-					VarProperty->CopyCompleteValueToScriptVM(RESULT_PARAM, Stack.MostRecentPropertyAddress);
-				}
+				VarProperty->CopyCompleteValueToScriptVM_InContainer(RESULT_PARAM, Stack.MostRecentPropertyContainer);
 			}
 		}
 		else
@@ -3246,25 +3156,7 @@ DEFINE_VM_FUNCTION(UObject::execStructMemberContext)
 		// Handle variable reads
 		if (RESULT_PARAM)
 		{
-			if (StructProperty->HasGetter())
-			{
-				// Temporary hack to fix TWeakObjectPtrs with getters being copied wrongly to the VM -
-				// hacked here because we cannot change UnrealType.h
-				// Object properties are copied to the VM by UObject*, except for FSoftObjectProperty (see overrides for FSoftObjectProperty::CopyCompleteValueToScriptVM).
-				FObjectPropertyBase* ObjProperty = CastField<FObjectPropertyBase>(StructProperty);
-				if (ObjProperty && !ObjProperty->IsA<FSoftObjectProperty>())
-				{
-					*(UObject**)RESULT_PARAM = ObjProperty->GetObjectPropertyValue_InContainer(Stack.MostRecentPropertyContainer);
-				}
-				else
-				{
-					StructProperty->GetValue_InContainer(Stack.MostRecentPropertyContainer, RESULT_PARAM);
-				}
-			}
-			else
-			{
-				StructProperty->CopyCompleteValueToScriptVM(RESULT_PARAM, Stack.MostRecentPropertyAddress);
-			}
+			StructProperty->CopyCompleteValueToScriptVM_InContainer(RESULT_PARAM, Stack.MostRecentPropertyContainer);
 		}
 	}
 	else
diff --git a/Engine/Source/Runtime/CoreUObject/Public/UObject/UnrealType.h b/Engine/Source/Runtime/CoreUObject/Public/UObject/UnrealType.h
--- a/Engine/Source/Runtime/CoreUObject/Public/UObject/UnrealType.h
+++ b/Engine/Source/Runtime/CoreUObject/Public/UObject/UnrealType.h
@@ -889,6 +889,12 @@ public:
 	 */
 	virtual void CopyCompleteValueToScriptVM( void* Dest, void const* Src ) const;
 
+	/**
+	 * Equivalent to the above functions, but using the container and aware of getters/setters when the container has them.
+	 */
+	virtual void CopyCompleteValueToScriptVM_InContainer( void* OutValue, void const* InContainer ) const;
+	virtual void CopyCompleteValueFromScriptVM_InContainer( void* OutContainer, void const* InValue ) const;
+
 	/**
 	 * Copy the value for a single element of this property. From the script VM.
 	 * 
@@ -2970,6 +2976,8 @@ PRAGMA_ENABLE_DEPRECATION_WARNINGS
 private:
 	virtual uint32 GetValueTypeHashInternal(const void* Src) const override;
 public:
+	virtual void CopyCompleteValueToScriptVM_InContainer( void* OutValue, void const* InContainer ) const override;
+	virtual void CopyCompleteValueFromScriptVM_InContainer( void* OutContainer, void const* InValue ) const override;
 	// End of FProperty interface
 
 	// FObjectPropertyBase interface
@@ -3017,6 +3025,8 @@ class COREUOBJECT_API FObjectPtrProperty : public FObjectProperty
 	virtual bool SameType(const FProperty* Other) const override;
 	virtual bool Identical(const void* A, const void* B, uint32 PortFlags) const override;
 	virtual void SerializeItem(FStructuredArchive::FSlot Slot, void* Value, void const* Defaults) const override;
+	virtual void CopyCompleteValueToScriptVM_InContainer( void* OutValue, void const* InContainer ) const override;
+	virtual void CopyCompleteValueFromScriptVM_InContainer( void* OutContainer, void const* InValue ) const override;
 	// End of FProperty interface
 
 	// Helper method for sharing code with FClassPtrProperty even though one doesn't inherit from the other
@@ -3089,6 +3099,8 @@ protected:
 private:
 	virtual uint32 GetValueTypeHashInternal(const void* Src) const override;
 public:
+	virtual void CopyCompleteValueToScriptVM_InContainer( void* OutValue, void const* InContainer ) const override;
+	virtual void CopyCompleteValueFromScriptVM_InContainer( void* OutContainer, void const* InValue ) const override;
 	// End of FProperty interface
 
 	// FObjectProperty interface
@@ -3147,6 +3159,8 @@ PRAGMA_ENABLE_DEPRECATION_WARNINGS
 	virtual bool Identical( const void* A, const void* B, uint32 PortFlags ) const override;
 	virtual void SerializeItem( FStructuredArchive::FSlot Slot, void* Value, void const* Defaults ) const override;
 	virtual void EmitReferenceInfo(UE::GC::FTokenStreamBuilder& TokenStream, int32 BaseOffset, TArray<const FStructProperty*>& EncounteredStructProps, FGCStackSizeHelper& StackSizeHelper) override;
+	virtual void CopyCompleteValueToScriptVM_InContainer( void* OutValue, void const* InContainer ) const override;
+	virtual void CopyCompleteValueFromScriptVM_InContainer( void* OutContainer, void const* InValue ) const override;
 	// End of FProperty interface
 
 	// FObjectProperty interface
@@ -3366,6 +3380,8 @@ class COREUOBJECT_API FClassPtrProperty : public FClassProperty
 	virtual bool SameType(const FProperty* Other) const override;
 	virtual bool Identical(const void* A, const void* B, uint32 PortFlags) const override;
 	virtual void SerializeItem(FStructuredArchive::FSlot Slot, void* Value, void const* Defaults) const override;
+	virtual void CopyCompleteValueToScriptVM_InContainer( void* OutValue, void const* InContainer ) const override;
+	virtual void CopyCompleteValueFromScriptVM_InContainer( void* OutContainer, void const* InValue ) const override;
 	// End of FProperty interface
 
 	// FObjectProperty interface
