PackageName: Optimize TryConvertFilenameToLongPackageName family of functions using an O(log(n)) DirectoryTree instead of O(n) linear search. 51s to 2s for a 2.5h cook.
#rnx

diff --git a/Engine/Source/Runtime/CoreUObject/Private/Misc/PackageName.cpp b/Engine/Source/Runtime/CoreUObject/Private/Misc/PackageName.cpp
--- a/Engine/Source/Runtime/CoreUObject/Private/Misc/PackageName.cpp
+++ b/Engine/Source/Runtime/CoreUObject/Private/Misc/PackageName.cpp
@@ -7,6 +7,7 @@
 #include "Misc/PackageName.h"
 
 #include "Algo/Find.h"
+#include "Containers/DirectoryTree.h"
 #include "Containers/StringView.h"
 #include "GenericPlatform/GenericPlatformFile.h"
 #include "HAL/CriticalSection.h"
@@ -27,6 +28,7 @@
 #include "Modules/ModuleManager.h"
 #include "ProfilingDebugging/LoadTimeTracker.h"
 #include "Stats/Stats.h"
+#include "Templates/UniquePtr.h"
 #include "UObject/Linker.h"
 #include "UObject/Package.h"
 #include "UObject/PackageFileSummary.h"
@@ -152,27 +154,26 @@ bool FPackageName::TryConvertGameRelativePackagePathToLocalPath(FStringView Rela
 }
 
 
-struct FPathPair
+struct FMountPoint
 {
-	// The virtual path (e.g., "/Engine/")
+	/** The LongPackageName path(e.g., "/Engine/") */
 	const FString RootPath;
 
-	// The physical relative path (e.g., "../../../Engine/Content/")
-	const FString ContentPath;
+	/** The LocalPath relative path(e.g., "../../../Engine/Content/") */
+	const FString ContentPathRelative;
 
-	// The physical absolute path
-	const FString AbsolutePath;
+	/** The LocalPath absolute path */
+	const FString ContentPathAbsolute;
 
-	bool operator ==(const FPathPair& Other) const
-	{
-		return RootPath == Other.RootPath && ContentPath == Other.ContentPath;
-	}
+	/** Whether the mountpoint was created by InsertMountPoint and can be removed by RemoveMountPoint. */
+	const bool bRemovable = false;
 
-	// Construct a path pair
-	FPathPair(const FString& InRootPath, const FString& InContentPath, const FString& InAbsolutePath)
-		: RootPath(InRootPath)
-		, ContentPath(InContentPath)
-		, AbsolutePath(InAbsolutePath)
+	FMountPoint(FString&& InRootPath, FString&& InContentPathRelative,
+		FString&& InContentPathAbsolute, bool bInRemovable = false)
+		: RootPath(MoveTemp(InRootPath))
+		, ContentPathRelative(MoveTemp(InContentPathRelative))
+		, ContentPathAbsolute(MoveTemp(InContentPathAbsolute))
+		, bRemovable(bInRemovable)
 	{
 	}
 };
@@ -187,7 +188,6 @@ struct FLongPackagePathsSingleton
 	FString ScriptRootPath;
 	FString MemoryRootPath;
 	FString TempRootPath;
-	TArray<FString> MountPointRootPaths;
 
 	FString VerseSubPath;
 
@@ -204,9 +204,9 @@ struct FLongPackagePathsSingleton
 	FString GameScriptPathRebased;
 	FString GameSavedPathRebased;
 
-	//@TODO: Can probably consolidate these into a single array, if it weren't for EngineContentPathShort
-	TArray<FPathPair> ContentRootToPath;
-	TArray<FPathPair> ContentPathToRoot;
+	TDirectoryTree<FMountPoint*> RootPathTree;
+	TDirectoryTree<FMountPoint*> ContentPathTree;
+	TArray<TUniquePtr<FMountPoint>> MountPoints;
 
 	// Cached values
 	TArray<FString> ValidLongPackageRoots[2];
@@ -234,10 +234,19 @@ struct FLongPackagePathsSingleton
 
 	void UpdateValidLongPackageRoots()
 	{
+		// List Writable roots (roots into which packages can be saved or resaved) in ValidLongPackageRoots[0]
 		ValidLongPackageRoots[0].Empty();
 		ValidLongPackageRoots[0].Add(EngineRootPath);
 		ValidLongPackageRoots[0].Add(GameRootPath);
-		ValidLongPackageRoots[0] += MountPointRootPaths;
+		for (TUniquePtr<FMountPoint>& MountPoint : MountPoints)
+		{
+			if (MountPoint->bRemovable)
+			{
+				ValidLongPackageRoots[0].Add(MountPoint->RootPath);
+			}
+		}
+
+		// List all roots in ValidLongPackageRoots[1] (Writable roots plus ReadOnly roots)
 		ValidLongPackageRoots[1] = ValidLongPackageRoots[0];
 		ValidLongPackageRoots[1].Add(ConfigRootPath);
 		ValidLongPackageRoots[1].Add(ScriptRootPath);
@@ -245,8 +254,11 @@ struct FLongPackagePathsSingleton
 		ValidLongPackageRoots[1].Add(TempRootPath);
 	}
 
-	// Given a content path ensure it is consistent, specifically with FileManager relative paths 
-	static FString ProcessContentMountPoint(const FString& ContentPath)
+	/**
+	 * Given a content path return the consistently-formatted RelativePath and AbsolutePath, suitable for lookup based only on text
+	 * The RelativePath is consistent with FileManager relative paths.
+	 */
+	static FMountPoint ConstructMountPoint(const FString& RootPath, const FString& ContentPath, bool bRemovable)
 	{
 		FString MountPath = ContentPath;
 
@@ -264,74 +276,126 @@ struct FLongPackagePathsSingleton
 		}
 
 		// Convert to a relative path using the FileManager
-		return IFileManager::Get().ConvertToRelativePath(*MountPath);
-	}
+		FString RelativeContentPath = IFileManager::Get().ConvertToRelativePath(*MountPath);
 
-
-	// This will insert a mount point at the head of the search chain (so it can overlap an existing mount point and win)
-	void InsertMountPoint(const FString& RootPath, const FString& ContentPath)
-	{	
-		// Make sure the content path is stored as a relative path, consistent with the other paths we have
-		FString RelativeContentPath = ProcessContentMountPoint(ContentPath);
-
-		// Make sure the path ends in a trailing path separator.  We are expecting that in the InternalFilenameToLongPackageName code.
-		if( !RelativeContentPath.EndsWith( TEXT( "/" ), ESearchCase::CaseSensitive ) )
+		// Make sure the path ends in a trailing path separator for consistency.
+		if (!RelativeContentPath.EndsWith(TEXT("/"), ESearchCase::CaseSensitive))
 		{
-			RelativeContentPath += TEXT( "/" );
+			RelativeContentPath += TEXT("/");
 		}
 
-		TStringBuilder<256> AbsolutePath;
-		FPathViews::ToAbsolutePath(RelativeContentPath, AbsolutePath);
-		FPathPair Pair(RootPath, RelativeContentPath, AbsolutePath.ToString());
-		
+		TStringBuilder<256> AbsolutePathBuilder;
+		FPathViews::ToAbsolutePath(RelativeContentPath, AbsolutePathBuilder);
+
+		return FMountPoint(FString(RootPath), MoveTemp(RelativeContentPath), AbsolutePathBuilder.ToString(),
+			bRemovable);
+	}
+
+	void InsertMountPoint(const FString& RootPath, const FString& ContentPath)
+	{	
+		TUniquePtr<FMountPoint> MountPointOwner(MakeUnique<FMountPoint>(
+			ConstructMountPoint(RootPath, ContentPath, true /* bRemovable */)));
+		FMountPoint* MountPoint = MountPointOwner.Get();
+		const FString& ContentPathRelative = MountPoint->ContentPathRelative;
 		{
 			FWriteScopeLock ScopeLock(MountLock);
-			ContentRootToPath.Insert(Pair, 0);
-			ContentPathToRoot.Insert(Pair, 0);
-			MountPointRootPaths.Add(RootPath);
+			MountPoints.Insert(MoveTemp(MountPointOwner), 0);
+
+			RootPathTree.FindOrAdd(MountPoint->RootPath) = MountPoint;
+			ContentPathTree.FindOrAdd(MountPoint->ContentPathRelative) = MountPoint;
+			ContentPathTree.FindOrAdd(MountPoint->ContentPathAbsolute) = MountPoint;
 
 			UpdateValidLongPackageRoots();
 		}		
 
-		UE_LOG(LogPackageName, Display, TEXT("FPackageName: Mount point added: '%s' mounted to '%s'"), *RelativeContentPath, *RootPath);
+		UE_LOG(LogPackageName, Display, TEXT("FPackageName: Mount point added: '%s' mounted to '%s'"),
+			*ContentPathRelative, *RootPath);
 		
 		// Let subscribers know that a new content path was mounted
-		FPackageName::OnContentPathMounted().Broadcast( RootPath, RelativeContentPath);
+		FPackageName::OnContentPathMounted().Broadcast( RootPath, ContentPathRelative);
 	}
 
 	// This will remove a previously inserted mount point
 	void RemoveMountPoint(const FString& RootPath, const FString& ContentPath)
 	{
-		// Make sure the content path is stored as a relative path, consistent with the other paths we have
-		FString RelativeContentPath = ProcessContentMountPoint(ContentPath);
-
-		// Make sure the path ends in a trailing path separator.  We are expecting that in the InternalFilenameToLongPackageName code.
-		if (!RelativeContentPath.EndsWith(TEXT("/"), ESearchCase::CaseSensitive))
-		{
-			RelativeContentPath += TEXT("/");
-		}
+		RemoveMountPoint(ConstructMountPoint(RootPath, ContentPath, true /* bRemovable */));
+	}
 
-		bool bFirePathDismountedDelegate = false;
+	void RemoveMountPoint(const FMountPoint& MountPoint)
+	{
+		bool bRemoved = false;
 		{
 			FWriteScopeLock ScopeLock(MountLock);
-			if ( MountPointRootPaths.Remove(RootPath) > 0 )
+			for (int32 MountPointIndex = 0; MountPointIndex < MountPoints.Num(); ++MountPointIndex)
 			{
-				// Absolute path doesn't need to be computed here as we are removing mount point
-				FPathPair Pair(RootPath, RelativeContentPath, FString());
-				ContentRootToPath.Remove(Pair);
-				ContentPathToRoot.Remove(Pair);
-				MountPointRootPaths.Remove(RootPath);
-
-				// Let subscribers know that a new content path was unmounted
-				bFirePathDismountedDelegate = true;
+				FMountPoint& ExistingMount = *MountPoints[MountPointIndex];
+				if (ExistingMount.RootPath == MountPoint.RootPath &&
+					ExistingMount.ContentPathRelative == MountPoint.ContentPathRelative &&
+					ExistingMount.bRemovable)
+				{
+					bool bRootPathExisted;
+					RootPathTree.Remove(ExistingMount.RootPath, &bRootPathExisted);
+					if (bRootPathExisted)
+					{
+						// See if there is a lower priority path that replaces it
+						for (TUniquePtr<FMountPoint>& OtherMount : TArrayView<TUniquePtr<FMountPoint>>(MountPoints).LeftChop(MountPointIndex + 1))
+						{
+							if (OtherMount->RootPath == MountPoint.RootPath)
+							{
+								RootPathTree.FindOrAdd(OtherMount->RootPath) = OtherMount.Get();
+								break;
+							}
+						}
+					}
+					bool bContentPathRelativeExisted;
+					ContentPathTree.Remove(ExistingMount.ContentPathRelative, &bContentPathRelativeExisted);
+					if (bContentPathRelativeExisted)
+					{
+						// See if there is a lower priority path that replaces it
+						for (TUniquePtr<FMountPoint>& OtherMount :
+							TArrayView<TUniquePtr<FMountPoint>>(MountPoints).LeftChop(MountPointIndex+1))
+						{
+							if (OtherMount->ContentPathRelative == MountPoint.ContentPathRelative)
+							{
+								ContentPathTree.FindOrAdd(OtherMount->ContentPathRelative) = OtherMount.Get();
+								break;
+							}
+						}
+					}
+					bool bContentPathAbsoluteExisted;
+					ContentPathTree.Remove(ExistingMount.ContentPathAbsolute, &bContentPathAbsoluteExisted);
+					if (bContentPathAbsoluteExisted)
+					{
+						// See if there is a lower priority path that replaces it
+						for (TUniquePtr<FMountPoint>& OtherMount :
+							TArrayView<TUniquePtr<FMountPoint>>(MountPoints).LeftChop(MountPointIndex + 1))
+						{
+							if (OtherMount->ContentPathAbsolute == MountPoint.ContentPathAbsolute)
+							{
+								ContentPathTree.FindOrAdd(OtherMount->ContentPathAbsolute) = OtherMount.Get();
+								break;
+							}
+						}
+					}
 
+					MountPoints.RemoveAt(MountPointIndex);
+					bRemoved = true;
+				}
+				else
+				{
+					++MountPointIndex;
+				}
+			}
+			if (bRemoved)
+			{
 				UpdateValidLongPackageRoots();
-			}			
+			}
 		}
 
-		if (bFirePathDismountedDelegate)
+		// Let subscribers know that a new content path was unmounted
+		if (bRemoved)
 		{
-			FPackageName::OnContentPathDismounted().Broadcast(RootPath, RelativeContentPath);
+			FPackageName::OnContentPathDismounted().Broadcast(MountPoint.RootPath, MountPoint.ContentPathRelative);
 		}
 	}
 
@@ -392,46 +456,59 @@ private:
 		GameScriptPathRebased  = RebasedGameDir / TEXT("Script/");
 		GameSavedPathRebased   = RebasedGameDir / TEXT("Saved/");
 
-		ContentPathToRoot.Empty(13);
-		TStringBuilder<256> AbsolutePath;
-		auto AddPathPair = [&AbsolutePath](TArray<FPathPair>& Paths, const FString& RootPath, const FString& ContentPath)
+		auto LocalInsertMountPoint = [this](const FString& RootPath, const FString& ContentPathRelative)
 		{
-			AbsolutePath.Reset();
-			FPathViews::ToAbsolutePath(ContentPath, AbsolutePath);
-			Paths.Emplace(RootPath, ContentPath, AbsolutePath.ToString());
+			// Don't call ConstructMountPoint, because that calls FPaths::ConvertRelativePathToFull followed by
+			// IFileManager::Get().ConvertToRelativePath, and that combination will convert ../../../Engine/Content
+			// into ../../Content, which we don't want here; we want to enter in the full ../../.. path for some entries.
+			TStringBuilder<256> AbsolutePathBuilder;
+			FPathViews::ToAbsolutePath(ContentPathRelative, AbsolutePathBuilder);
+			FMountPoint& MountPoint = *MountPoints.Add_GetRef(MakeUnique<FMountPoint>(
+				FString(RootPath), FString(ContentPathRelative), AbsolutePathBuilder.ToString()));
+
+			// We are appending to the MountPoint list, and will have only a few duplicates. If any duplicates do occur,
+			// keep the MountPoint that was appended earlier as the registered mountpoint.
+			FMountPoint*& ExistingRootPath = RootPathTree.FindOrAdd(MountPoint.RootPath);
+			if (!ExistingRootPath)
+			{
+				ExistingRootPath = &MountPoint;
+			}
+			FMountPoint*& ExistingContentPathRelative = ContentPathTree.FindOrAdd(MountPoint.ContentPathRelative);
+			if (!ExistingContentPathRelative)
+			{
+				ExistingContentPathRelative = &MountPoint;
+			}
+			FMountPoint*& ExistingContentPathAbsolute = ContentPathTree.FindOrAdd(MountPoint.ContentPathAbsolute);
+			if (!ExistingContentPathAbsolute)
+			{
+				ExistingContentPathAbsolute = &MountPoint;
+			}
 		};
 		
-		AddPathPair(ContentPathToRoot, EngineRootPath, EngineContentPath);
+		LocalInsertMountPoint(EngineRootPath, EngineContentPath);
+		LocalInsertMountPoint(EngineRootPath, EngineShadersPath);
+		LocalInsertMountPoint(GameRootPath,   GameContentPath);
+		LocalInsertMountPoint(ScriptRootPath, GameScriptPath);
+		LocalInsertMountPoint(TempRootPath,   GameSavedPath);
+		LocalInsertMountPoint(ConfigRootPath, GameConfigPath);
+
+		// Add other LocalPaths that have different sets of .. but are the same location on disk
+		LocalInsertMountPoint(EngineRootPath, EngineShadersPathShort);
 		if (FPaths::IsSamePath(GameContentPath, ContentPathShort))
 		{
-			AddPathPair(ContentPathToRoot, GameRootPath, ContentPathShort);
+			// ../../Content points to the the Engine directory, the same as ../../../Engine/Content
+			// But if the /Game is pointing to ../../../Engine/Content as well, then map ../../Content to
+			// /Game instead of /Engine
+			LocalInsertMountPoint(GameRootPath, ContentPathShort);
 		}
 		else
 		{
-			AddPathPair(ContentPathToRoot, EngineRootPath, ContentPathShort);
-		}
-		AddPathPair(ContentPathToRoot, EngineRootPath, EngineShadersPath);
-		AddPathPair(ContentPathToRoot, EngineRootPath, EngineShadersPathShort);
-		AddPathPair(ContentPathToRoot, GameRootPath,   GameContentPath);
-		AddPathPair(ContentPathToRoot, ScriptRootPath, GameScriptPath);
-		AddPathPair(ContentPathToRoot, TempRootPath,   GameSavedPath);
-		AddPathPair(ContentPathToRoot, ConfigRootPath, GameConfigPath);
-		AddPathPair(ContentPathToRoot, GameRootPath,   GameContentPathRebased);
-		AddPathPair(ContentPathToRoot, ScriptRootPath, GameScriptPathRebased);
-		AddPathPair(ContentPathToRoot, TempRootPath,   GameSavedPathRebased);
-		AddPathPair(ContentPathToRoot, ConfigRootPath, GameConfigPathRebased);
-
-		ContentRootToPath.Empty(11);
-		AddPathPair(ContentRootToPath, EngineRootPath, EngineContentPath);
-		AddPathPair(ContentRootToPath, EngineRootPath, EngineShadersPath);
-		AddPathPair(ContentRootToPath, GameRootPath,   GameContentPath);
-		AddPathPair(ContentRootToPath, ScriptRootPath, GameScriptPath);
-		AddPathPair(ContentRootToPath, TempRootPath,   GameSavedPath);
-		AddPathPair(ContentRootToPath, ConfigRootPath, GameConfigPath);
-		AddPathPair(ContentRootToPath, GameRootPath,   GameContentPathRebased);
-		AddPathPair(ContentRootToPath, ScriptRootPath, GameScriptPathRebased);
-		AddPathPair(ContentRootToPath, TempRootPath,   GameSavedPathRebased);
-		AddPathPair(ContentRootToPath, ConfigRootPath, GameConfigPathRebased);
+			LocalInsertMountPoint(EngineRootPath, ContentPathShort);
+		}
+		LocalInsertMountPoint(GameRootPath,   GameContentPathRebased);
+		LocalInsertMountPoint(ScriptRootPath, GameScriptPathRebased);
+		LocalInsertMountPoint(TempRootPath,   GameSavedPathRebased);
+		LocalInsertMountPoint(ConfigRootPath, GameConfigPathRebased);
 
 		UpdateValidLongPackageRoots();
 	}
@@ -442,17 +519,19 @@ private:
 		UE_LOG(LogPackageName, Log, TEXT("Valid mount points:"));
 
 		FReadScopeLock ScopeLock(MountLock);
-		for (const auto& Pair : ContentRootToPath)
+		for (const TUniquePtr<FMountPoint>& MountPoint : MountPoints)
 		{
-			UE_LOG(LogPackageName, Log, TEXT("	'%s' -> '%s'"), *Pair.RootPath, *Pair.ContentPath);
+			UE_LOG(LogPackageName, Log, TEXT("	'%s' -> '%s'"), *MountPoint->RootPath, *MountPoint->ContentPathRelative);
 		}
 
 		UE_LOG(LogPackageName, Log, TEXT("Removable mount points:"));
-		for (const auto& Root : MountPointRootPaths)
+		for (const TUniquePtr<FMountPoint>& MountPoint : MountPoints)
 		{
-			UE_LOG(LogPackageName, Log, TEXT("	'%s'"), *Root);
+			if (MountPoint->bRemovable)
+			{
+				UE_LOG(LogPackageName, Log, TEXT("	'%s'"), *MountPoint->RootPath);
+			}
 		}
-		
 	}
 
 	void ExecInsertMountPoint(const TArray<FString>& Args)
@@ -508,49 +587,63 @@ private:
 			{
 				FReadScopeLock ScopeLock(MountLock);
 				
-				RootPath = Path;
-				auto Elements = Get().ContentRootToPath.FilterByPredicate([&RootPath](const FPathPair& elem){ return elem.RootPath.Equals(RootPath); });
-				if (Elements.Num() == 0)
+				TArray<FString> ContentPathRelatives;
+				for (TUniquePtr<FMountPoint>& ExistingMount : Get().MountPoints)
+				{
+					if (ExistingMount->RootPath == RootPath)
+					{
+						ContentPathRelatives.Add(ExistingMount->ContentPathRelative);
+					}
+				}
+				if (ContentPathRelatives.Num() == 0)
 				{
 					UE_LOG(LogPackageName, Error, TEXT("PackageName.UnregisterMountPoint: Root path '%s' is not mounted!"), *RootPath);
 					return;
 				}
 
-				if (Elements.Num() > 1)
+				if (ContentPathRelatives.Num() > 1)
 				{
 					UE_LOG(LogPackageName, Error, TEXT("PackageName.UnregisterMountPoint: Root path '%s' is mounted to multiple content paths, specify content path to unmount explicitly!"), *RootPath);
-					for (const FPathPair& elem : Elements)
+					for (const FString& ContentPathRelative : ContentPathRelatives)
 					{
-						UE_LOG(LogPackageName, Error, TEXT("- %s"), *elem.ContentPath);
+						UE_LOG(LogPackageName, Error, TEXT("- %s"), *ContentPathRelative);
 					}
 					return;
 				}
 
-				ContentPath = Elements[0].ContentPath;
+				ContentPath = ContentPathRelatives[0];
 			}
 			else
 			{
 				FReadScopeLock ScopeLock(MountLock);
 			
 				ContentPath = Path;
-				auto Elements = Get().ContentPathToRoot.FilterByPredicate([&ContentPath](const FPathPair& elem){ return elem.ContentPath.Equals(ContentPath); });
-				if (Elements.Num() == 0)
+				TArray<FString> RootPaths;
+				for (TUniquePtr<FMountPoint>& ExistingMount : Get().MountPoints)
+				{
+					if (ExistingMount->ContentPathRelative == ContentPath)
+					{
+						RootPaths.Add(ExistingMount->RootPath);
+					}
+				}
+
+				if (RootPaths.Num() == 0)
 				{
 					UE_LOG(LogPackageName, Error, TEXT("PackageName.UnregisterMountPoint: Content path '%s' is not mounted!"), *ContentPath);
 					return;
 				}
 
-				if (Elements.Num() > 1)
+				if (RootPaths.Num() > 1)
 				{
 					UE_LOG(LogPackageName, Error, TEXT("PackageName.UnregisterMountPoint: Content path '%s' is mounted to multiple root paths, specify root path to unmount explicitly!"), *ContentPath);
-					for (const FPathPair& elem : Elements)
+					for (const FString& MountRootPath: RootPaths)
 					{
-						UE_LOG(LogPackageName, Error, TEXT("- %s"), *elem.RootPath);
+						UE_LOG(LogPackageName, Error, TEXT("- %s"), *MountRootPath);
 					}
 					return;
 				}
 
-				RootPath = Elements[0].RootPath;
+				RootPath = RootPaths[0];
 			}
 			
 			Get().RemoveMountPoint(RootPath, ContentPath);
@@ -587,17 +680,13 @@ void FPackageName::InternalFilenameToLongPackageName(FStringView InFilename, FSt
 	bool bIsValidLongPackageName = false;
 	{
 		FReadScopeLock ScopeLock(Paths.MountLock);
-		for (const auto& Pair : Paths.ContentRootToPath)
+		const FMountPoint* const* MountPointPtr = Paths.RootPathTree.FindClosestValue(Filename);
+		if (MountPointPtr)
 		{
-			if (FPathViews::IsParentPathOf(Pair.RootPath, Filename))
-			{
-				bIsValidLongPackageName = true;
-				break;
-			}
+			bIsValidLongPackageName = true;
 		}
 	}
 
-	FStringView Result;
 	if (!bIsValidLongPackageName)
 	{
 		Filename = IFileManager::Get().ConvertToRelativePath(*Filename);
@@ -620,6 +709,7 @@ void FPackageName::InternalFilenameToLongPackageName(FStringView InFilename, FSt
 
 	// Do not strip the extension from paths with an empty pathname; if we do the caller can't tell the difference
 	// between /Path/.ext and /Path/
+	FStringView Result;
 	if (!SplitOutName.IsEmpty() || SplitOutExtension.IsEmpty())
 	{
 		Result = FStringView(SplitOutPath.GetData(), UE_PTRDIFF_TO_INT32(SplitOutName.GetData() + SplitOutName.Len() - SplitOutPath.GetData()));
@@ -631,21 +721,32 @@ void FPackageName::InternalFilenameToLongPackageName(FStringView InFilename, FSt
 			"Accepting ObjectPaths is deprecated behavior and will be removed in a future release; TryConvertFilenameToLongPackageName will fail on ObjectPaths."), InFilename.Len(), InFilename.GetData());
 	}
 
+	auto TryGetLocalPathResult = [&Paths, &OutPackageName](FStringView SearchName)
 	{
 		FReadScopeLock ScopeLock(Paths.MountLock);
-		for (const auto& Pair : Paths.ContentPathToRoot)
+		const FMountPoint* const* MountPointPtr = Paths.ContentPathTree.FindClosestValue(SearchName);
+		if (!MountPointPtr)
 		{
-			FStringView RelPath;
-			if (FPathViews::TryMakeChildPathRelativeTo(Result, Pair.ContentPath, RelPath))
-			{
-				OutPackageName << Pair.RootPath << RelPath;
-				return;
-			}
+			return false;
 		}
+		const FMountPoint* MountPoint = *MountPointPtr;
+		OutPackageName << MountPoint->RootPath;
+		FStringView RelPath;
+		if (!FPathViews::TryMakeChildPathRelativeTo(SearchName, MountPoint->ContentPathRelative, RelPath))
+		{
+			verify(FPathViews::TryMakeChildPathRelativeTo(SearchName, MountPoint->ContentPathAbsolute, RelPath));
+		}
+		FPathViews::AppendPath(OutPackageName, RelPath);
+		return true;
+	};
+
+	if (TryGetLocalPathResult(Result))
+	{
+		return;
 	}
 
-	// if we get here, we haven't converted to a package name, and it may be because the path was absolute. in which case we should 
-	// check the absolute representation in ContentPathToRoot
+	// if we get here, we haven't converted to a package name, and it may be because the path was an unnormalized absolute ContentPath.
+	// In that case, normalize it and check the ContentPathTree again
 	if (!bIsValidLongPackageName)
 	{
 		// reset to the incoming string
@@ -654,27 +755,14 @@ void FPackageName::InternalFilenameToLongPackageName(FStringView InFilename, FSt
 		{
 			FPaths::NormalizeFilename(Filename);
 			Result = FPathViews::GetBaseFilenameWithPath(Filename);
+			if (TryGetLocalPathResult(Result))
 			{
-				FReadScopeLock ScopeLock(Paths.MountLock);
-				for (const auto& Pair : Paths.ContentPathToRoot)
-				{
-					FStringView RelPath;
-					// Compare against AbsolutePath, but also try ContentPath in case ContentPath is absolute
-					if (FPathViews::TryMakeChildPathRelativeTo(Result, Pair.AbsolutePath, RelPath))
-					{
-						OutPackageName << Pair.RootPath << RelPath;
-						return;
-					}
-					if (FPathViews::TryMakeChildPathRelativeTo(Result, Pair.ContentPath, RelPath))
-					{
-						OutPackageName << Pair.RootPath << RelPath;
-						return;
-					}
-				}
+				return;
 			}
 		}
 	}
 
+	// Either the input string was already a LongPackageName, or we did not find a mapping for it. Return the last attempted Result
 	OutPackageName << Result;
 }
 
@@ -751,10 +839,10 @@ FString FPackageName::FilenameToLongPackageName(const FString& InFilename)
 		{
 			const auto& Paths = FLongPackagePathsSingleton::Get();
 			FReadScopeLock ScopeLock(Paths.MountLock);
-			for (const FPathPair& Pair : Paths.ContentPathToRoot)
+			for (const TUniquePtr<FMountPoint>& MountPoint : Paths.MountPoints)
 			{
-				ContentRoots << TEXT("\n\t\t") << Pair.ContentPath;
-				ContentRoots << TEXT("\n\t\t\t") << Pair.AbsolutePath;
+				ContentRoots << TEXT("\n\t\t") << MountPoint->ContentPathRelative;
+				ContentRoots << TEXT("\n\t\t\t") << MountPoint->ContentPathAbsolute;
 			}
 		}
 
@@ -778,18 +866,22 @@ bool FPackageName::TryConvertLongPackageNameToFilename(const FString& InLongPack
 {
 	const auto& Paths = FLongPackagePathsSingleton::Get();
 	FReadScopeLock ScopeLock(Paths.MountLock);
-	for (const auto& Pair : Paths.ContentRootToPath)
+	const FMountPoint* const* MountPointPtr = Paths.RootPathTree.FindClosestValue(InLongPackageName);
+	if (!MountPointPtr)
 	{
-		FStringView RelPath;
-		if (FPathViews::TryMakeChildPathRelativeTo(InLongPackageName, Pair.RootPath, RelPath))
-		{
-			OutFilename = Pair.ContentPath + RelPath + InExtension;
-			return true;
-		}
+		// Note that we return false for the root folder "/".
+		return false;
 	}
 
-	// This is not a long package name or the root folder is not handled in the above cases
-	return false;
+	const FMountPoint* MountPoint = *MountPointPtr;
+	FStringView RelPath;
+	verify(FPathViews::TryMakeChildPathRelativeTo(InLongPackageName, MountPoint->RootPath, RelPath));
+	TStringBuilder<256> Builder;
+	Builder << MountPoint->ContentPathRelative;
+	FPathViews::AppendPath(Builder, RelPath);
+	Builder << InExtension;
+	OutFilename = Builder.ToView();
+	return true;
 }
 
 FString FPackageName::LongPackageNameToFilename(const FString& InLongPackageName, const FString& InExtension)
@@ -1197,16 +1289,8 @@ bool FPackageName::IsValidPath(const FString& InPath)
 {
 	const FLongPackagePathsSingleton& Paths = FLongPackagePathsSingleton::Get();
 	FReadScopeLock ScopeLock(Paths.MountLock);
-	for (const FPathPair& Pair : Paths.ContentRootToPath)
-	{
-		if (FPathViews::IsParentPathOf(Pair.RootPath, InPath))
-		{
-			return true;
-		}
-	}
-
-	// The root folder is not handled in the above cases
-	return false;
+	// Note that we return false for the root folder "/".
+	return Paths.RootPathTree.FindClosestValue(InPath) != nullptr;
 }
 
 void FPackageName::RegisterMountPoint(const FString& RootPath, const FString& ContentPath)
@@ -1248,15 +1332,13 @@ FString FPackageName::GetContentPathForPackageRoot(FStringView InMountPoint)
 	FLongPackagePathsSingleton& Paths = FLongPackagePathsSingleton::Get();
 
 	FReadScopeLock ScopeLock(Paths.MountLock);
-	for (const FPathPair& RootToPathEntry : Paths.ContentRootToPath)
+	const FMountPoint* const* MountPointPtr = Paths.RootPathTree.FindClosestValue(InMountPoint);
+	if (!MountPointPtr)
 	{
-		if (RootToPathEntry.RootPath == InMountPoint)
-		{
-			return RootToPathEntry.AbsolutePath;
-		}
+		return FString();
 	}
-
-	return FString();
+	const FMountPoint* MountPoint = *MountPointPtr;
+	return MountPoint->ContentPathAbsolute;
 }
 
 bool FPackageName::TryConvertToMountedPathComponents(FStringView InFilePathOrPackageName,
@@ -1413,48 +1495,55 @@ bool FPackageName::TryGetMountPointForPath(FStringView InFilePathOrPackageName,
 
 	TStringBuilder<256> PossibleAbsFilePath;
 	FPathViews::ToAbsolutePath(InFilePathOrPackageName, PossibleAbsFilePath);
+	FStringView RelPath;
 	const FLongPackagePathsSingleton& Paths = FLongPackagePathsSingleton::Get();
 	FReadScopeLock ScopeLock(Paths.MountLock);
-	for (const auto& Pair : Paths.ContentRootToPath)
-	{
-		FStringView RelPath;
-		if (FPathViews::TryMakeChildPathRelativeTo(InFilePathOrPackageName, Pair.RootPath, RelPath))
+	const FMountPoint* const* MountPointPtr = Paths.RootPathTree.FindClosestValue(InFilePathOrPackageName);
+	if (MountPointPtr)
+	{
+		const FMountPoint* MountPoint = *MountPointPtr;
+		verify(FPathViews::TryMakeChildPathRelativeTo(InFilePathOrPackageName, MountPoint->RootPath, RelPath));
+		OutMountPointPackageName << MountPoint->RootPath;
+		OutMountPointFilePath << MountPoint->ContentPathRelative;
+		OutRelPath << RelPath;
+		if (OutFlexNameType)
 		{
-			OutMountPointPackageName << Pair.RootPath;
-			OutMountPointFilePath << Pair.ContentPath;
-			OutRelPath << RelPath;
-			if (OutFlexNameType)
+			if (Algo::Find(RelPath, TEXT('.')))
 			{
-				if (Algo::Find(RelPath, TEXT('.')))
-				{
-					*OutFlexNameType = EFlexNameType::ObjectPath;
-				}
-				else
-				{
-					*OutFlexNameType = EFlexNameType::PackageName;
-				}
+				*OutFlexNameType = EFlexNameType::ObjectPath;
 			}
-			if (OutFailureReason)
+			else
 			{
-				*OutFailureReason = EErrorCode::PackageNameUnknown;
+				*OutFlexNameType = EFlexNameType::PackageName;
 			}
-			return true;
 		}
-		else if (FPathViews::TryMakeChildPathRelativeTo(PossibleAbsFilePath, Pair.AbsolutePath, RelPath))
+		if (OutFailureReason)
 		{
-			OutMountPointPackageName << Pair.RootPath;
-			OutMountPointFilePath << Pair.ContentPath;
-			OutRelPath << RelPath;
-			if (OutFlexNameType)
-			{
-				*OutFlexNameType = EFlexNameType::LocalPath;
-			}
-			if (OutFailureReason)
-			{
-				*OutFailureReason = EErrorCode::PackageNameUnknown;
-			}
-			return true;
+			*OutFailureReason = EErrorCode::PackageNameUnknown;
+		}
+		return true;
+	}
+
+	MountPointPtr = Paths.ContentPathTree.FindClosestValue(PossibleAbsFilePath);
+	if (MountPointPtr)
+	{
+		const FMountPoint* MountPoint = *MountPointPtr;
+		if (!FPathViews::TryMakeChildPathRelativeTo(PossibleAbsFilePath, MountPoint->ContentPathAbsolute, RelPath))
+		{
+			verify(FPathViews::TryMakeChildPathRelativeTo(PossibleAbsFilePath, MountPoint->ContentPathRelative, RelPath));
 		}
+		OutMountPointPackageName << MountPoint->RootPath;
+		OutMountPointFilePath << MountPoint->ContentPathRelative;
+		OutRelPath << RelPath;
+		if (OutFlexNameType)
+		{
+			*OutFlexNameType = EFlexNameType::LocalPath;
+		}
+		if (OutFailureReason)
+		{
+			*OutFailureReason = EErrorCode::PackageNameUnknown;
+		}
+		return true;
 	}
 	if (OutFlexNameType)
 	{
@@ -1462,7 +1551,6 @@ bool FPackageName::TryGetMountPointForPath(FStringView InFilePathOrPackageName,
 	}
 	if (OutFailureReason)
 	{
-		FStringView RelPath;
 		if (FPathViews::TryMakeChildPathRelativeTo(InFilePathOrPackageName, Paths.MemoryRootPath, RelPath))
 		{
 			*OutFailureReason = EErrorCode::PackageNamePathIsMemoryOnly;
