Prevent re-dirtying packages that have been marked clean by an undo/redo operation

#rnx

diff --git a/Engine/Source/Editor/UnrealEd/Private/EditorTransaction.cpp b/Engine/Source/Editor/UnrealEd/Private/EditorTransaction.cpp
--- a/Engine/Source/Editor/UnrealEd/Private/EditorTransaction.cpp
+++ b/Engine/Source/Editor/UnrealEd/Private/EditorTransaction.cpp
@@ -810,7 +810,7 @@ void FTransaction::Apply()
 	UE::Transaction::DiffUtil::FDiffableObjectArchetypeCache ArchetypeCache;
 
 	// Update the package dirty states
-	// We do this prior to applying any object updates as we want to respect if an undo operation causes an object to dirty its own package
+	// We do this prior to applying any object updates since a package cannot be re-dirtied during an undo/redo operation
 	if (bFlip)
 	{
 		for (TTuple<UE::Transaction::FPersistentObjectRef, FPackageRecord>& PackageRecordPair : PackageRecordMap)
diff --git a/Engine/Source/Runtime/CoreUObject/Private/UObject/UObjectBaseUtility.cpp b/Engine/Source/Runtime/CoreUObject/Private/UObject/UObjectBaseUtility.cpp
--- a/Engine/Source/Runtime/CoreUObject/Private/UObject/UObjectBaseUtility.cpp
+++ b/Engine/Source/Runtime/CoreUObject/Private/UObject/UObjectBaseUtility.cpp
@@ -313,12 +313,13 @@ bool UObjectBaseUtility::MarkPackageDirty() const
 
 		if( Package != NULL	)
 		{
-			// It is against policy to dirty a map or package during load in the Editor, to enforce this policy
-			// we explicitly disable the ability to dirty a package or map during load.  Commandlets can still
+			// It is against policy to dirty a map or package during load/undo/redo in the Editor, to enforce this policy
+			// we explicitly disable the ability to dirty a package or map during load/undo/redo.  Commandlets can still
 			// set the dirty state on load.
 			if( IsRunningCommandlet() || 
 				(GIsEditor && !GIsEditorLoadingPackage && !GIsCookerLoadingPackage && !GIsPlayInEditorWorld && !IsInAsyncLoadingThread() && !IsReloadActive()
 #if WITH_EDITORONLY_DATA
+				&& !GIsTransacting
 				&& !Package->bIsCookedForEditor // Cooked packages can't be modified nor marked as dirty
 #endif
 				))
