[Backout] - CL23859876
Original CL Desc
-----------------------------------------------------------------
PackageName: Optimize TryConvertFilenameToLongPackageName family of functions using an O(log(n)) DirectoryTree instead of O(n) linear search. 51s to 2s for a 2.5h cook.
#rnx

diff --git a/Engine/Source/Runtime/CoreUObject/Private/Misc/PackageName.cpp b/Engine/Source/Runtime/CoreUObject/Private/Misc/PackageName.cpp
--- a/Engine/Source/Runtime/CoreUObject/Private/Misc/PackageName.cpp
+++ b/Engine/Source/Runtime/CoreUObject/Private/Misc/PackageName.cpp
@@ -7,7 +7,6 @@
 #include "Misc/PackageName.h"
 
 #include "Algo/Find.h"
-#include "Containers/DirectoryTree.h"
 #include "Containers/StringView.h"
 #include "GenericPlatform/GenericPlatformFile.h"
 #include "HAL/CriticalSection.h"
@@ -28,7 +27,6 @@
 #include "Modules/ModuleManager.h"
 #include "ProfilingDebugging/LoadTimeTracker.h"
 #include "Stats/Stats.h"
-#include "Templates/UniquePtr.h"
 #include "UObject/Linker.h"
 #include "UObject/Package.h"
 #include "UObject/PackageFileSummary.h"
@@ -154,26 +152,27 @@ bool FPackageName::TryConvertGameRelativePackagePathToLocalPath(FStringView Rela
 }
 
 
-struct FMountPoint
+struct FPathPair
 {
-	/** The LongPackageName path(e.g., "/Engine/") */
+	// The virtual path (e.g., "/Engine/")
 	const FString RootPath;
 
-	/** The LocalPath relative path(e.g., "../../../Engine/Content/") */
-	const FString ContentPathRelative;
+	// The physical relative path (e.g., "../../../Engine/Content/")
+	const FString ContentPath;
 
-	/** The LocalPath absolute path */
-	const FString ContentPathAbsolute;
+	// The physical absolute path
+	const FString AbsolutePath;
 
-	/** Whether the mountpoint was created by InsertMountPoint and can be removed by RemoveMountPoint. */
-	const bool bRemovable = false;
+	bool operator ==(const FPathPair& Other) const
+	{
+		return RootPath == Other.RootPath && ContentPath == Other.ContentPath;
+	}
 
-	FMountPoint(FString&& InRootPath, FString&& InContentPathRelative,
-		FString&& InContentPathAbsolute, bool bInRemovable = false)
-		: RootPath(MoveTemp(InRootPath))
-		, ContentPathRelative(MoveTemp(InContentPathRelative))
-		, ContentPathAbsolute(MoveTemp(InContentPathAbsolute))
-		, bRemovable(bInRemovable)
+	// Construct a path pair
+	FPathPair(const FString& InRootPath, const FString& InContentPath, const FString& InAbsolutePath)
+		: RootPath(InRootPath)
+		, ContentPath(InContentPath)
+		, AbsolutePath(InAbsolutePath)
 	{
 	}
 };
@@ -188,6 +187,7 @@ struct FLongPackagePathsSingleton
 	FString ScriptRootPath;
 	FString MemoryRootPath;
 	FString TempRootPath;
+	TArray<FString> MountPointRootPaths;
 
 	FString VerseSubPath;
 
@@ -204,9 +204,9 @@ struct FLongPackagePathsSingleton
 	FString GameScriptPathRebased;
 	FString GameSavedPathRebased;
 
-	TDirectoryTree<FMountPoint*> RootPathTree;
-	TDirectoryTree<FMountPoint*> ContentPathTree;
-	TArray<TUniquePtr<FMountPoint>> MountPoints;
+	//@TODO: Can probably consolidate these into a single array, if it weren't for EngineContentPathShort
+	TArray<FPathPair> ContentRootToPath;
+	TArray<FPathPair> ContentPathToRoot;
 
 	// Cached values
 	TArray<FString> ValidLongPackageRoots[2];
@@ -234,19 +234,10 @@ struct FLongPackagePathsSingleton
 
 	void UpdateValidLongPackageRoots()
 	{
-		// List Writable roots (roots into which packages can be saved or resaved) in ValidLongPackageRoots[0]
 		ValidLongPackageRoots[0].Empty();
 		ValidLongPackageRoots[0].Add(EngineRootPath);
 		ValidLongPackageRoots[0].Add(GameRootPath);
-		for (TUniquePtr<FMountPoint>& MountPoint : MountPoints)
-		{
-			if (MountPoint->bRemovable)
-			{
-				ValidLongPackageRoots[0].Add(MountPoint->RootPath);
-			}
-		}
-
-		// List all roots in ValidLongPackageRoots[1] (Writable roots plus ReadOnly roots)
+		ValidLongPackageRoots[0] += MountPointRootPaths;
 		ValidLongPackageRoots[1] = ValidLongPackageRoots[0];
 		ValidLongPackageRoots[1].Add(ConfigRootPath);
 		ValidLongPackageRoots[1].Add(ScriptRootPath);
@@ -254,11 +245,8 @@ struct FLongPackagePathsSingleton
 		ValidLongPackageRoots[1].Add(TempRootPath);
 	}
 
-	/**
-	 * Given a content path return the consistently-formatted RelativePath and AbsolutePath, suitable for lookup based only on text
-	 * The RelativePath is consistent with FileManager relative paths.
-	 */
-	static FMountPoint ConstructMountPoint(const FString& RootPath, const FString& ContentPath, bool bRemovable)
+	// Given a content path ensure it is consistent, specifically with FileManager relative paths 
+	static FString ProcessContentMountPoint(const FString& ContentPath)
 	{
 		FString MountPath = ContentPath;
 
@@ -276,126 +264,74 @@ struct FLongPackagePathsSingleton
 		}
 
 		// Convert to a relative path using the FileManager
-		FString RelativeContentPath = IFileManager::Get().ConvertToRelativePath(*MountPath);
-
-		// Make sure the path ends in a trailing path separator for consistency.
-		if (!RelativeContentPath.EndsWith(TEXT("/"), ESearchCase::CaseSensitive))
-		{
-			RelativeContentPath += TEXT("/");
-		}
-
-		TStringBuilder<256> AbsolutePathBuilder;
-		FPathViews::ToAbsolutePath(RelativeContentPath, AbsolutePathBuilder);
-
-		return FMountPoint(FString(RootPath), MoveTemp(RelativeContentPath), AbsolutePathBuilder.ToString(),
-			bRemovable);
+		return IFileManager::Get().ConvertToRelativePath(*MountPath);
 	}
 
+
+	// This will insert a mount point at the head of the search chain (so it can overlap an existing mount point and win)
 	void InsertMountPoint(const FString& RootPath, const FString& ContentPath)
 	{	
-		TUniquePtr<FMountPoint> MountPointOwner(MakeUnique<FMountPoint>(
-			ConstructMountPoint(RootPath, ContentPath, true /* bRemovable */)));
-		FMountPoint* MountPoint = MountPointOwner.Get();
-		const FString& ContentPathRelative = MountPoint->ContentPathRelative;
+		// Make sure the content path is stored as a relative path, consistent with the other paths we have
+		FString RelativeContentPath = ProcessContentMountPoint(ContentPath);
+
+		// Make sure the path ends in a trailing path separator.  We are expecting that in the InternalFilenameToLongPackageName code.
+		if( !RelativeContentPath.EndsWith( TEXT( "/" ), ESearchCase::CaseSensitive ) )
 		{
-			FWriteScopeLock ScopeLock(MountLock);
-			MountPoints.Insert(MoveTemp(MountPointOwner), 0);
+			RelativeContentPath += TEXT( "/" );
+		}
 
-			RootPathTree.FindOrAdd(MountPoint->RootPath) = MountPoint;
-			ContentPathTree.FindOrAdd(MountPoint->ContentPathRelative) = MountPoint;
-			ContentPathTree.FindOrAdd(MountPoint->ContentPathAbsolute) = MountPoint;
+		TStringBuilder<256> AbsolutePath;
+		FPathViews::ToAbsolutePath(RelativeContentPath, AbsolutePath);
+		FPathPair Pair(RootPath, RelativeContentPath, AbsolutePath.ToString());
+		
+		{
+			FWriteScopeLock ScopeLock(MountLock);
+			ContentRootToPath.Insert(Pair, 0);
+			ContentPathToRoot.Insert(Pair, 0);
+			MountPointRootPaths.Add(RootPath);
 
 			UpdateValidLongPackageRoots();
 		}		
 
-		UE_LOG(LogPackageName, Display, TEXT("FPackageName: Mount point added: '%s' mounted to '%s'"),
-			*ContentPathRelative, *RootPath);
+		UE_LOG(LogPackageName, Display, TEXT("FPackageName: Mount point added: '%s' mounted to '%s'"), *RelativeContentPath, *RootPath);
 		
 		// Let subscribers know that a new content path was mounted
-		FPackageName::OnContentPathMounted().Broadcast( RootPath, ContentPathRelative);
+		FPackageName::OnContentPathMounted().Broadcast( RootPath, RelativeContentPath);
 	}
 
 	// This will remove a previously inserted mount point
 	void RemoveMountPoint(const FString& RootPath, const FString& ContentPath)
 	{
-		RemoveMountPoint(ConstructMountPoint(RootPath, ContentPath, true /* bRemovable */));
-	}
+		// Make sure the content path is stored as a relative path, consistent with the other paths we have
+		FString RelativeContentPath = ProcessContentMountPoint(ContentPath);
 
-	void RemoveMountPoint(const FMountPoint& MountPoint)
-	{
-		bool bRemoved = false;
+		// Make sure the path ends in a trailing path separator.  We are expecting that in the InternalFilenameToLongPackageName code.
+		if (!RelativeContentPath.EndsWith(TEXT("/"), ESearchCase::CaseSensitive))
+		{
+			RelativeContentPath += TEXT("/");
+		}
+
+		bool bFirePathDismountedDelegate = false;
 		{
 			FWriteScopeLock ScopeLock(MountLock);
-			for (int32 MountPointIndex = 0; MountPointIndex < MountPoints.Num(); ++MountPointIndex)
+			if ( MountPointRootPaths.Remove(RootPath) > 0 )
 			{
-				FMountPoint& ExistingMount = *MountPoints[MountPointIndex];
-				if (ExistingMount.RootPath == MountPoint.RootPath &&
-					ExistingMount.ContentPathRelative == MountPoint.ContentPathRelative &&
-					ExistingMount.bRemovable)
-				{
-					bool bRootPathExisted;
-					RootPathTree.Remove(ExistingMount.RootPath, &bRootPathExisted);
-					if (bRootPathExisted)
-					{
-						// See if there is a lower priority path that replaces it
-						for (TUniquePtr<FMountPoint>& OtherMount : TArrayView<TUniquePtr<FMountPoint>>(MountPoints).LeftChop(MountPointIndex + 1))
-						{
-							if (OtherMount->RootPath == MountPoint.RootPath)
-							{
-								RootPathTree.FindOrAdd(OtherMount->RootPath) = OtherMount.Get();
-								break;
-							}
-						}
-					}
-					bool bContentPathRelativeExisted;
-					ContentPathTree.Remove(ExistingMount.ContentPathRelative, &bContentPathRelativeExisted);
-					if (bContentPathRelativeExisted)
-					{
-						// See if there is a lower priority path that replaces it
-						for (TUniquePtr<FMountPoint>& OtherMount :
-							TArrayView<TUniquePtr<FMountPoint>>(MountPoints).LeftChop(MountPointIndex+1))
-						{
-							if (OtherMount->ContentPathRelative == MountPoint.ContentPathRelative)
-							{
-								ContentPathTree.FindOrAdd(OtherMount->ContentPathRelative) = OtherMount.Get();
-								break;
-							}
-						}
-					}
-					bool bContentPathAbsoluteExisted;
-					ContentPathTree.Remove(ExistingMount.ContentPathAbsolute, &bContentPathAbsoluteExisted);
-					if (bContentPathAbsoluteExisted)
-					{
-						// See if there is a lower priority path that replaces it
-						for (TUniquePtr<FMountPoint>& OtherMount :
-							TArrayView<TUniquePtr<FMountPoint>>(MountPoints).LeftChop(MountPointIndex + 1))
-						{
-							if (OtherMount->ContentPathAbsolute == MountPoint.ContentPathAbsolute)
-							{
-								ContentPathTree.FindOrAdd(OtherMount->ContentPathAbsolute) = OtherMount.Get();
-								break;
-							}
-						}
-					}
+				// Absolute path doesn't need to be computed here as we are removing mount point
+				FPathPair Pair(RootPath, RelativeContentPath, FString());
+				ContentRootToPath.Remove(Pair);
+				ContentPathToRoot.Remove(Pair);
+				MountPointRootPaths.Remove(RootPath);
+
+				// Let subscribers know that a new content path was unmounted
+				bFirePathDismountedDelegate = true;
 
-					MountPoints.RemoveAt(MountPointIndex);
-					bRemoved = true;
-				}
-				else
-				{
-					++MountPointIndex;
-				}
-			}
-			if (bRemoved)
-			{
 				UpdateValidLongPackageRoots();
-			}
+			}			
 		}
 
-		// Let subscribers know that a new content path was unmounted
-		if (bRemoved)
+		if (bFirePathDismountedDelegate)
 		{
-			FPackageName::OnContentPathDismounted().Broadcast(MountPoint.RootPath, MountPoint.ContentPathRelative);
+			FPackageName::OnContentPathDismounted().Broadcast(RootPath, RelativeContentPath);
 		}
 	}
 
@@ -456,59 +392,44 @@ private:
 		GameScriptPathRebased  = RebasedGameDir / TEXT("Script/");
 		GameSavedPathRebased   = RebasedGameDir / TEXT("Saved/");
 
-		auto LocalInsertMountPoint = [this](const FString& RootPath, const FString& ContentPathRelative)
+		ContentPathToRoot.Empty(13);
+		TStringBuilder<256> AbsolutePath;
+		auto AddPathPair = [&AbsolutePath](TArray<FPathPair>& Paths, const FString& RootPath, const FString& ContentPath)
 		{
-			// Don't call ConstructMountPoint, because that calls FPaths::ConvertRelativePathToFull followed by
-			// IFileManager::Get().ConvertToRelativePath, and that combination will convert ../../../Engine/Content
-			// into ../../Content, which we don't want here; we want to enter in the full ../../.. path for some entries.
-			TStringBuilder<256> AbsolutePathBuilder;
-			FPathViews::ToAbsolutePath(ContentPathRelative, AbsolutePathBuilder);
-			FMountPoint& MountPoint = *MountPoints.Add_GetRef(MakeUnique<FMountPoint>(
-				FString(RootPath), FString(ContentPathRelative), AbsolutePathBuilder.ToString()));
-
-			// We are appending to the MountPoint list, and will have only a few duplicates. If any duplicates do occur,
-			// keep the MountPoint that was appended earlier as the registered mountpoint.
-			FMountPoint*& ExistingRootPath = RootPathTree.FindOrAdd(MountPoint.RootPath);
-			if (!ExistingRootPath)
-			{
-				ExistingRootPath = &MountPoint;
-			}
-			FMountPoint*& ExistingContentPathRelative = ContentPathTree.FindOrAdd(MountPoint.ContentPathRelative);
-			if (!ExistingContentPathRelative)
-			{
-				ExistingContentPathRelative = &MountPoint;
-			}
-			FMountPoint*& ExistingContentPathAbsolute = ContentPathTree.FindOrAdd(MountPoint.ContentPathAbsolute);
-			if (!ExistingContentPathAbsolute)
-			{
-				ExistingContentPathAbsolute = &MountPoint;
-			}
+			AbsolutePath.Reset();
+			FPathViews::ToAbsolutePath(ContentPath, AbsolutePath);
+			Paths.Emplace(RootPath, ContentPath, AbsolutePath.ToString());
 		};
 		
-		LocalInsertMountPoint(EngineRootPath, EngineContentPath);
-		LocalInsertMountPoint(EngineRootPath, EngineShadersPath);
-		LocalInsertMountPoint(GameRootPath,   GameContentPath);
-		LocalInsertMountPoint(ScriptRootPath, GameScriptPath);
-		LocalInsertMountPoint(TempRootPath,   GameSavedPath);
-		LocalInsertMountPoint(ConfigRootPath, GameConfigPath);
-
-		// Add other LocalPaths that have different sets of .. but are the same location on disk
-		LocalInsertMountPoint(EngineRootPath, EngineShadersPathShort);
+		AddPathPair(ContentPathToRoot, EngineRootPath, EngineContentPath);
 		if (FPaths::IsSamePath(GameContentPath, ContentPathShort))
 		{
-			// ../../Content points to the the Engine directory, the same as ../../../Engine/Content
-			// But if the /Game is pointing to ../../../Engine/Content as well, then map ../../Content to
-			// /Game instead of /Engine
-			LocalInsertMountPoint(GameRootPath, ContentPathShort);
+			AddPathPair(ContentPathToRoot, GameRootPath, ContentPathShort);
 		}
 		else
 		{
-			LocalInsertMountPoint(EngineRootPath, ContentPathShort);
-		}
-		LocalInsertMountPoint(GameRootPath,   GameContentPathRebased);
-		LocalInsertMountPoint(ScriptRootPath, GameScriptPathRebased);
-		LocalInsertMountPoint(TempRootPath,   GameSavedPathRebased);
-		LocalInsertMountPoint(ConfigRootPath, GameConfigPathRebased);
+			AddPathPair(ContentPathToRoot, EngineRootPath, ContentPathShort);
+		}
+		AddPathPair(ContentPathToRoot, EngineRootPath, EngineShadersPath);
+		AddPathPair(ContentPathToRoot, EngineRootPath, EngineShadersPathShort);
+		AddPathPair(ContentPathToRoot, GameRootPath,   GameContentPath);
+		AddPathPair(ContentPathToRoot, ScriptRootPath, GameScriptPath);
+		AddPathPair(ContentPathToRoot, TempRootPath,   GameSavedPath);
+		AddPathPair(ContentPathToRoot, GameRootPath,   GameContentPathRebased);
+		AddPathPair(ContentPathToRoot, ScriptRootPath, GameScriptPathRebased);
+		AddPathPair(ContentPathToRoot, TempRootPath,   GameSavedPathRebased);
+		AddPathPair(ContentPathToRoot, ConfigRootPath, GameConfigPath);
+
+		ContentRootToPath.Empty(11);
+		AddPathPair(ContentRootToPath, EngineRootPath, EngineContentPath);
+		AddPathPair(ContentRootToPath, EngineRootPath, EngineShadersPath);
+		AddPathPair(ContentRootToPath, GameRootPath,   GameContentPath);
+		AddPathPair(ContentRootToPath, ScriptRootPath, GameScriptPath);
+		AddPathPair(ContentRootToPath, TempRootPath,   GameSavedPath);
+		AddPathPair(ContentRootToPath, GameRootPath,   GameContentPathRebased);
+		AddPathPair(ContentRootToPath, ScriptRootPath, GameScriptPathRebased);
+		AddPathPair(ContentRootToPath, TempRootPath,   GameSavedPathRebased);
+		AddPathPair(ContentRootToPath, ConfigRootPath, GameConfigPathRebased);
 
 		UpdateValidLongPackageRoots();
 	}
@@ -519,19 +440,17 @@ private:
 		UE_LOG(LogPackageName, Log, TEXT("Valid mount points:"));
 
 		FReadScopeLock ScopeLock(MountLock);
-		for (const TUniquePtr<FMountPoint>& MountPoint : MountPoints)
+		for (const auto& Pair : ContentRootToPath)
 		{
-			UE_LOG(LogPackageName, Log, TEXT("	'%s' -> '%s'"), *MountPoint->RootPath, *MountPoint->ContentPathRelative);
+			UE_LOG(LogPackageName, Log, TEXT("	'%s' -> '%s'"), *Pair.RootPath, *Pair.ContentPath);
 		}
 
 		UE_LOG(LogPackageName, Log, TEXT("Removable mount points:"));
-		for (const TUniquePtr<FMountPoint>& MountPoint : MountPoints)
+		for (const auto& Root : MountPointRootPaths)
 		{
-			if (MountPoint->bRemovable)
-			{
-				UE_LOG(LogPackageName, Log, TEXT("	'%s'"), *MountPoint->RootPath);
-			}
+			UE_LOG(LogPackageName, Log, TEXT("	'%s'"), *Root);
 		}
+		
 	}
 
 	void ExecInsertMountPoint(const TArray<FString>& Args)
@@ -587,63 +506,49 @@ private:
 			{
 				FReadScopeLock ScopeLock(MountLock);
 				
-				TArray<FString> ContentPathRelatives;
-				for (TUniquePtr<FMountPoint>& ExistingMount : Get().MountPoints)
-				{
-					if (ExistingMount->RootPath == RootPath)
-					{
-						ContentPathRelatives.Add(ExistingMount->ContentPathRelative);
-					}
-				}
-				if (ContentPathRelatives.Num() == 0)
+				RootPath = Path;
+				auto Elements = Get().ContentRootToPath.FilterByPredicate([&RootPath](const FPathPair& elem){ return elem.RootPath.Equals(RootPath); });
+				if (Elements.Num() == 0)
 				{
 					UE_LOG(LogPackageName, Error, TEXT("PackageName.UnregisterMountPoint: Root path '%s' is not mounted!"), *RootPath);
 					return;
 				}
 
-				if (ContentPathRelatives.Num() > 1)
+				if (Elements.Num() > 1)
 				{
 					UE_LOG(LogPackageName, Error, TEXT("PackageName.UnregisterMountPoint: Root path '%s' is mounted to multiple content paths, specify content path to unmount explicitly!"), *RootPath);
-					for (const FString& ContentPathRelative : ContentPathRelatives)
+					for (const FPathPair& elem : Elements)
 					{
-						UE_LOG(LogPackageName, Error, TEXT("- %s"), *ContentPathRelative);
+						UE_LOG(LogPackageName, Error, TEXT("- %s"), *elem.ContentPath);
 					}
 					return;
 				}
 
-				ContentPath = ContentPathRelatives[0];
+				ContentPath = Elements[0].ContentPath;
 			}
 			else
 			{
 				FReadScopeLock ScopeLock(MountLock);
 			
 				ContentPath = Path;
-				TArray<FString> RootPaths;
-				for (TUniquePtr<FMountPoint>& ExistingMount : Get().MountPoints)
-				{
-					if (ExistingMount->ContentPathRelative == ContentPath)
-					{
-						RootPaths.Add(ExistingMount->RootPath);
-					}
-				}
-
-				if (RootPaths.Num() == 0)
+				auto Elements = Get().ContentPathToRoot.FilterByPredicate([&ContentPath](const FPathPair& elem){ return elem.ContentPath.Equals(ContentPath); });
+				if (Elements.Num() == 0)
 				{
 					UE_LOG(LogPackageName, Error, TEXT("PackageName.UnregisterMountPoint: Content path '%s' is not mounted!"), *ContentPath);
 					return;
 				}
 
-				if (RootPaths.Num() > 1)
+				if (Elements.Num() > 1)
 				{
 					UE_LOG(LogPackageName, Error, TEXT("PackageName.UnregisterMountPoint: Content path '%s' is mounted to multiple root paths, specify root path to unmount explicitly!"), *ContentPath);
-					for (const FString& MountRootPath: RootPaths)
+					for (const FPathPair& elem : Elements)
 					{
-						UE_LOG(LogPackageName, Error, TEXT("- %s"), *MountRootPath);
+						UE_LOG(LogPackageName, Error, TEXT("- %s"), *elem.RootPath);
 					}
 					return;
 				}
 
-				RootPath = RootPaths[0];
+				RootPath = Elements[0].RootPath;
 			}
 			
 			Get().RemoveMountPoint(RootPath, ContentPath);
@@ -680,13 +585,17 @@ void FPackageName::InternalFilenameToLongPackageName(FStringView InFilename, FSt
 	bool bIsValidLongPackageName = false;
 	{
 		FReadScopeLock ScopeLock(Paths.MountLock);
-		const FMountPoint* const* MountPointPtr = Paths.RootPathTree.FindClosestValue(Filename);
-		if (MountPointPtr)
+		for (const auto& Pair : Paths.ContentRootToPath)
 		{
-			bIsValidLongPackageName = true;
+			if (FPathViews::IsParentPathOf(Pair.RootPath, Filename))
+			{
+				bIsValidLongPackageName = true;
+				break;
+			}
 		}
 	}
 
+	FStringView Result;
 	if (!bIsValidLongPackageName)
 	{
 		Filename = IFileManager::Get().ConvertToRelativePath(*Filename);
@@ -709,7 +618,6 @@ void FPackageName::InternalFilenameToLongPackageName(FStringView InFilename, FSt
 
 	// Do not strip the extension from paths with an empty pathname; if we do the caller can't tell the difference
 	// between /Path/.ext and /Path/
-	FStringView Result;
 	if (!SplitOutName.IsEmpty() || SplitOutExtension.IsEmpty())
 	{
 		Result = FStringView(SplitOutPath.GetData(), UE_PTRDIFF_TO_INT32(SplitOutName.GetData() + SplitOutName.Len() - SplitOutPath.GetData()));
@@ -721,32 +629,21 @@ void FPackageName::InternalFilenameToLongPackageName(FStringView InFilename, FSt
 			"Accepting ObjectPaths is deprecated behavior and will be removed in a future release; TryConvertFilenameToLongPackageName will fail on ObjectPaths."), InFilename.Len(), InFilename.GetData());
 	}
 
-	auto TryGetLocalPathResult = [&Paths, &OutPackageName](FStringView SearchName)
 	{
 		FReadScopeLock ScopeLock(Paths.MountLock);
-		const FMountPoint* const* MountPointPtr = Paths.ContentPathTree.FindClosestValue(SearchName);
-		if (!MountPointPtr)
+		for (const auto& Pair : Paths.ContentPathToRoot)
 		{
-			return false;
-		}
-		const FMountPoint* MountPoint = *MountPointPtr;
-		OutPackageName << MountPoint->RootPath;
-		FStringView RelPath;
-		if (!FPathViews::TryMakeChildPathRelativeTo(SearchName, MountPoint->ContentPathRelative, RelPath))
-		{
-			verify(FPathViews::TryMakeChildPathRelativeTo(SearchName, MountPoint->ContentPathAbsolute, RelPath));
+			FStringView RelPath;
+			if (FPathViews::TryMakeChildPathRelativeTo(Result, Pair.ContentPath, RelPath))
+			{
+				OutPackageName << Pair.RootPath << RelPath;
+				return;
+			}
 		}
-		FPathViews::AppendPath(OutPackageName, RelPath);
-		return true;
-	};
-
-	if (TryGetLocalPathResult(Result))
-	{
-		return;
 	}
 
-	// if we get here, we haven't converted to a package name, and it may be because the path was an unnormalized absolute ContentPath.
-	// In that case, normalize it and check the ContentPathTree again
+	// if we get here, we haven't converted to a package name, and it may be because the path was absolute. in which case we should 
+	// check the absolute representation in ContentPathToRoot
 	if (!bIsValidLongPackageName)
 	{
 		// reset to the incoming string
@@ -755,14 +652,27 @@ void FPackageName::InternalFilenameToLongPackageName(FStringView InFilename, FSt
 		{
 			FPaths::NormalizeFilename(Filename);
 			Result = FPathViews::GetBaseFilenameWithPath(Filename);
-			if (TryGetLocalPathResult(Result))
 			{
-				return;
+				FReadScopeLock ScopeLock(Paths.MountLock);
+				for (const auto& Pair : Paths.ContentPathToRoot)
+				{
+					FStringView RelPath;
+					// Compare against AbsolutePath, but also try ContentPath in case ContentPath is absolute
+					if (FPathViews::TryMakeChildPathRelativeTo(Result, Pair.AbsolutePath, RelPath))
+					{
+						OutPackageName << Pair.RootPath << RelPath;
+						return;
+					}
+					if (FPathViews::TryMakeChildPathRelativeTo(Result, Pair.ContentPath, RelPath))
+					{
+						OutPackageName << Pair.RootPath << RelPath;
+						return;
+					}
+				}
 			}
 		}
 	}
 
-	// Either the input string was already a LongPackageName, or we did not find a mapping for it. Return the last attempted Result
 	OutPackageName << Result;
 }
 
@@ -839,10 +749,10 @@ FString FPackageName::FilenameToLongPackageName(const FString& InFilename)
 		{
 			const auto& Paths = FLongPackagePathsSingleton::Get();
 			FReadScopeLock ScopeLock(Paths.MountLock);
-			for (const TUniquePtr<FMountPoint>& MountPoint : Paths.MountPoints)
+			for (const FPathPair& Pair : Paths.ContentPathToRoot)
 			{
-				ContentRoots << TEXT("\n\t\t") << MountPoint->ContentPathRelative;
-				ContentRoots << TEXT("\n\t\t\t") << MountPoint->ContentPathAbsolute;
+				ContentRoots << TEXT("\n\t\t") << Pair.ContentPath;
+				ContentRoots << TEXT("\n\t\t\t") << Pair.AbsolutePath;
 			}
 		}
 
@@ -866,22 +776,18 @@ bool FPackageName::TryConvertLongPackageNameToFilename(const FString& InLongPack
 {
 	const auto& Paths = FLongPackagePathsSingleton::Get();
 	FReadScopeLock ScopeLock(Paths.MountLock);
-	const FMountPoint* const* MountPointPtr = Paths.RootPathTree.FindClosestValue(InLongPackageName);
-	if (!MountPointPtr)
+	for (const auto& Pair : Paths.ContentRootToPath)
 	{
-		// Note that we return false for the root folder "/".
-		return false;
+		FStringView RelPath;
+		if (FPathViews::TryMakeChildPathRelativeTo(InLongPackageName, Pair.RootPath, RelPath))
+		{
+			OutFilename = Pair.ContentPath + RelPath + InExtension;
+			return true;
+		}
 	}
 
-	const FMountPoint* MountPoint = *MountPointPtr;
-	FStringView RelPath;
-	verify(FPathViews::TryMakeChildPathRelativeTo(InLongPackageName, MountPoint->RootPath, RelPath));
-	TStringBuilder<256> Builder;
-	Builder << MountPoint->ContentPathRelative;
-	FPathViews::AppendPath(Builder, RelPath);
-	Builder << InExtension;
-	OutFilename = Builder.ToView();
-	return true;
+	// This is not a long package name or the root folder is not handled in the above cases
+	return false;
 }
 
 FString FPackageName::LongPackageNameToFilename(const FString& InLongPackageName, const FString& InExtension)
@@ -1289,8 +1195,16 @@ bool FPackageName::IsValidPath(const FString& InPath)
 {
 	const FLongPackagePathsSingleton& Paths = FLongPackagePathsSingleton::Get();
 	FReadScopeLock ScopeLock(Paths.MountLock);
-	// Note that we return false for the root folder "/".
-	return Paths.RootPathTree.FindClosestValue(InPath) != nullptr;
+	for (const FPathPair& Pair : Paths.ContentRootToPath)
+	{
+		if (FPathViews::IsParentPathOf(Pair.RootPath, InPath))
+		{
+			return true;
+		}
+	}
+
+	// The root folder is not handled in the above cases
+	return false;
 }
 
 void FPackageName::RegisterMountPoint(const FString& RootPath, const FString& ContentPath)
@@ -1332,13 +1246,15 @@ FString FPackageName::GetContentPathForPackageRoot(FStringView InMountPoint)
 	FLongPackagePathsSingleton& Paths = FLongPackagePathsSingleton::Get();
 
 	FReadScopeLock ScopeLock(Paths.MountLock);
-	const FMountPoint* const* MountPointPtr = Paths.RootPathTree.FindClosestValue(InMountPoint);
-	if (!MountPointPtr)
+	for (const FPathPair& RootToPathEntry : Paths.ContentRootToPath)
 	{
-		return FString();
+		if (RootToPathEntry.RootPath == InMountPoint)
+		{
+			return RootToPathEntry.AbsolutePath;
+		}
 	}
-	const FMountPoint* MountPoint = *MountPointPtr;
-	return MountPoint->ContentPathAbsolute;
+
+	return FString();
 }
 
 bool FPackageName::TryConvertToMountedPathComponents(FStringView InFilePathOrPackageName,
@@ -1495,55 +1411,48 @@ bool FPackageName::TryGetMountPointForPath(FStringView InFilePathOrPackageName,
 
 	TStringBuilder<256> PossibleAbsFilePath;
 	FPathViews::ToAbsolutePath(InFilePathOrPackageName, PossibleAbsFilePath);
-	FStringView RelPath;
 	const FLongPackagePathsSingleton& Paths = FLongPackagePathsSingleton::Get();
 	FReadScopeLock ScopeLock(Paths.MountLock);
-	const FMountPoint* const* MountPointPtr = Paths.RootPathTree.FindClosestValue(InFilePathOrPackageName);
-	if (MountPointPtr)
-	{
-		const FMountPoint* MountPoint = *MountPointPtr;
-		verify(FPathViews::TryMakeChildPathRelativeTo(InFilePathOrPackageName, MountPoint->RootPath, RelPath));
-		OutMountPointPackageName << MountPoint->RootPath;
-		OutMountPointFilePath << MountPoint->ContentPathRelative;
-		OutRelPath << RelPath;
-		if (OutFlexNameType)
+	for (const auto& Pair : Paths.ContentRootToPath)
+	{
+		FStringView RelPath;
+		if (FPathViews::TryMakeChildPathRelativeTo(InFilePathOrPackageName, Pair.RootPath, RelPath))
 		{
-			if (Algo::Find(RelPath, TEXT('.')))
+			OutMountPointPackageName << Pair.RootPath;
+			OutMountPointFilePath << Pair.ContentPath;
+			OutRelPath << RelPath;
+			if (OutFlexNameType)
 			{
-				*OutFlexNameType = EFlexNameType::ObjectPath;
+				if (Algo::Find(RelPath, TEXT('.')))
+				{
+					*OutFlexNameType = EFlexNameType::ObjectPath;
+				}
+				else
+				{
+					*OutFlexNameType = EFlexNameType::PackageName;
+				}
 			}
-			else
+			if (OutFailureReason)
 			{
-				*OutFlexNameType = EFlexNameType::PackageName;
+				*OutFailureReason = EErrorCode::PackageNameUnknown;
 			}
+			return true;
 		}
-		if (OutFailureReason)
-		{
-			*OutFailureReason = EErrorCode::PackageNameUnknown;
-		}
-		return true;
-	}
-
-	MountPointPtr = Paths.ContentPathTree.FindClosestValue(PossibleAbsFilePath);
-	if (MountPointPtr)
-	{
-		const FMountPoint* MountPoint = *MountPointPtr;
-		if (!FPathViews::TryMakeChildPathRelativeTo(PossibleAbsFilePath, MountPoint->ContentPathAbsolute, RelPath))
-		{
-			verify(FPathViews::TryMakeChildPathRelativeTo(PossibleAbsFilePath, MountPoint->ContentPathRelative, RelPath));
-		}
-		OutMountPointPackageName << MountPoint->RootPath;
-		OutMountPointFilePath << MountPoint->ContentPathRelative;
-		OutRelPath << RelPath;
-		if (OutFlexNameType)
-		{
-			*OutFlexNameType = EFlexNameType::LocalPath;
-		}
-		if (OutFailureReason)
+		else if (FPathViews::TryMakeChildPathRelativeTo(PossibleAbsFilePath, Pair.AbsolutePath, RelPath))
 		{
-			*OutFailureReason = EErrorCode::PackageNameUnknown;
+			OutMountPointPackageName << Pair.RootPath;
+			OutMountPointFilePath << Pair.ContentPath;
+			OutRelPath << RelPath;
+			if (OutFlexNameType)
+			{
+				*OutFlexNameType = EFlexNameType::LocalPath;
+			}
+			if (OutFailureReason)
+			{
+				*OutFailureReason = EErrorCode::PackageNameUnknown;
+			}
+			return true;
 		}
-		return true;
 	}
 	if (OutFlexNameType)
 	{
@@ -1551,6 +1460,7 @@ bool FPackageName::TryGetMountPointForPath(FStringView InFilePathOrPackageName,
 	}
 	if (OutFailureReason)
 	{
+		FStringView RelPath;
 		if (FPathViews::TryMakeChildPathRelativeTo(InFilePathOrPackageName, Paths.MemoryRootPath, RelPath))
 		{
 			*OutFailureReason = EErrorCode::PackageNamePathIsMemoryOnly;
