Add a way for OFPA packages to relinquish the lock on their package file without fully tearing down their FLinkerLoad.

#rnx

- The goal of FLinkerLoad::DetachLoader is to allow the caller to relinquish the file lock that we have on a package but without deleting the FLinkerLoad that the package owns.
- First it will detach any bulkdata/editorbulkdata in the package from the file on disk, meaning that they will immediately load their payloads into memory making them safe to use in the future, then the internal loader of the FLinkerLoad is deleted meaning it can no longer be used for accessing the package file. At this point we can no longer trust the package file on disk.
- Once this is called any future attempts to use the FLinkerLoad as an archive will result in an assert. This should be a problem as the only thing that should be using it as an archive would be bulkdata, but as that has been detached it should be a none issue. The asserts are mostly a safety precaution that we do not expect to see fire, but if they do it will make the problem much easier to track down.
- This is not a long term solution and the method has a warning in its documentation to try and prevent people from using it where possible.

diff --git a/Engine/Source/Runtime/CoreUObject/Private/UObject/LinkerLoad.cpp b/Engine/Source/Runtime/CoreUObject/Private/UObject/LinkerLoad.cpp
--- a/Engine/Source/Runtime/CoreUObject/Private/UObject/LinkerLoad.cpp
+++ b/Engine/Source/Runtime/CoreUObject/Private/UObject/LinkerLoad.cpp
@@ -943,6 +943,9 @@ FLinkerLoad::FLinkerLoad(UPackage* InParent, const FPackagePath& InPackagePath,
 , bLockoutLegacyOperations(false)
 , bIsAsyncLoader(false)
 , bIsDestroyingLoader(false)
+#if WITH_EDITOR
+, bDetachedLoader(false)
+#endif // WITH_EDITOR
 , StructuredArchive(nullptr)
 , StructuredArchiveFormatter(nullptr)
 , PackagePath(InPackagePath)
@@ -5755,6 +5758,19 @@ void FLinkerLoad::DestroyLoader()
 	bIsDestroyingLoader = false;
 }
 
+void FLinkerLoad::DetachLoader()
+{
+#if WITH_EDITOR
+	DetachAllBulkData(true);
+#endif // WITH_EDITOR
+
+	DestroyLoader();
+
+#if WITH_EDITOR
+	bDetachedLoader = true;
+#endif // WITH_EDITOR
+}
+
 void FLinkerLoad::DetachExports()
 {
 	// Detach all objects linked with this linker.
diff --git a/Engine/Source/Runtime/CoreUObject/Public/UObject/LinkerLoad.h b/Engine/Source/Runtime/CoreUObject/Public/UObject/LinkerLoad.h
--- a/Engine/Source/Runtime/CoreUObject/Public/UObject/LinkerLoad.h
+++ b/Engine/Source/Runtime/CoreUObject/Public/UObject/LinkerLoad.h
@@ -204,6 +204,10 @@ private:
 
 	/** True if the linker is currently deleting loader */
 	bool					bIsDestroyingLoader;
+#if WITH_EDITOR
+	/** Tracks if DetachLoader has been called or not */
+	bool					bDetachedLoader;
+#endif // WITH_EDITOR
 
 	/** Structured archive interface. Wraps underlying loader to provide contextual metadata to the values being written
 	 *  which ultimately allows text based serialization of the data
@@ -262,6 +266,14 @@ public:
 
 	void DestroyLoader();
 
+	/**
+	 * Detaches all bulkdata currently attached to the FLinkerLoad followed by destroying the internal loader.
+	 * This is a fairly dangerous method to call as it will leave the FLinkerLoad in a state where using it as
+	 * a FArchive will cause an assert/crash. It is intended as a short/medium term fix to an internal engine
+	 * issue and will be deprecated as soon as possible. 
+	 */
+	COREUOBJECT_API void DetachLoader();
+
 	FORCEINLINE bool IsDestroyingLoader() const
 	{
 		return bIsDestroyingLoader;
@@ -868,27 +880,6 @@ private:
 
 	void DetachExport( int32 i );
 
-	// FArchive interface.
-	/**
-	 * Hint the archive that the region starting at passed in offset and spanning the passed in size
-	 * is going to be read soon and should be precached.
-	 *
-	 * The function returns whether the precache operation has completed or not which is an important
-	 * hint for code knowing that it deals with potential async I/O. The archive is free to either not 
-	 * implement this function or only partially precache so it is required that given sufficient time
-	 * the function will return true. Archives not based on async I/O should always return true.
-	 *
-	 * This function will not change the current archive position.
-	 *
-	 * @param	PrecacheOffset	Offset at which to begin precaching.
-	 * @param	PrecacheSize	Number of bytes to precache
-	 * @return	false if precache operation is still pending, true otherwise
-	 */
-	FORCEINLINE virtual bool Precache(int64 PrecacheOffset, int64 PrecacheSize) override
-	{
-		return Loader->Precache(PrecacheOffset, PrecacheSize);
-	}
-
 #if WITH_EDITOR
 	/**
 	 * Attaches/ associates the passed in bulk data object with the linker.
@@ -938,16 +929,55 @@ public:
 
 private:
 
+	// FArchive interface.
+	/**
+	 * Hint the archive that the region starting at passed in offset and spanning the passed in size
+	 * is going to be read soon and should be precached.
+	 *
+	 * The function returns whether the precache operation has completed or not which is an important
+	 * hint for code knowing that it deals with potential async I/O. The archive is free to either not
+	 * implement this function or only partially precache so it is required that given sufficient time
+	 * the function will return true. Archives not based on async I/O should always return true.
+	 *
+	 * This function will not change the current archive position.
+	 *
+	 * @param	PrecacheOffset	Offset at which to begin precaching.
+	 * @param	PrecacheSize	Number of bytes to precache
+	 * @return	false if precache operation is still pending, true otherwise
+	 */
+	FORCEINLINE virtual bool Precache(int64 PrecacheOffset, int64 PrecacheSize) override
+	{
+#if WITH_EDITOR
+		checkf(!bDetachedLoader, TEXT("Attempting to call ::Precache on a FLinkerLoad that has previously called ::DetachLoader"));
+#endif // WITH_EDITOR
+
+		return Loader->Precache(PrecacheOffset, PrecacheSize);
+	}
+
 	FORCEINLINE virtual void Seek(int64 InPos) override
 	{
+#if WITH_EDITOR
+		checkf(!bDetachedLoader, TEXT("Attempting to call ::Seek on a FLinkerLoad that has previously called ::DetachLoader"));
+#endif // WITH_EDITOR
+
 		Loader->Seek(InPos);
 	}
+
 	FORCEINLINE virtual int64 Tell() override
 	{
+#if WITH_EDITOR
+		checkf(!bDetachedLoader, TEXT("Attempting to call ::Tell on a FLinkerLoad that has previously called ::DetachLoader"));
+#endif // WITH_EDITOR
+
 		return Loader->Tell();
 	}
+
 	FORCEINLINE virtual int64 TotalSize() override
 	{
+#if WITH_EDITOR
+		checkf(!bDetachedLoader, TEXT("Attempting to call ::TotalSize on a FLinkerLoad that has previously called ::DetachLoader"));
+#endif // WITH_EDITOR
+
 		return Loader->TotalSize();
 	}
 
@@ -955,6 +985,10 @@ private:
 	using FLinker::Serialize;
 	FORCEINLINE virtual void Serialize(void* V, int64 Length) override
 	{
+#if WITH_EDITOR
+		checkf(!bDetachedLoader, TEXT("Attempting to call ::Serialize on a FLinkerLoad that has previously called ::DetachLoader"));
+#endif // WITH_EDITOR
+
 		checkSlow(FPlatformTLS::GetCurrentThreadId() == OwnerThread);
 #if WITH_EDITOR
 		Loader->SetSerializedProperty(GetSerializedProperty());
