Removed UE::Core::Private::FWrappedObjectPtr, since it's unnecessary now that we have FObjectPtr.  It was added when FObjectProperty was based on UObject*, not FObjectPtr.

diff --git a/Engine/Source/Runtime/CoreUObject/Private/UObject/PropertyObject.cpp b/Engine/Source/Runtime/CoreUObject/Private/UObject/PropertyObject.cpp
--- a/Engine/Source/Runtime/CoreUObject/Private/UObject/PropertyObject.cpp
+++ b/Engine/Source/Runtime/CoreUObject/Private/UObject/PropertyObject.cpp
@@ -10,27 +10,6 @@
 #include "UObject/LinkerPlaceholderClass.h"
 #include "UObject/PropertyHelper.h"
 
-namespace UE::Core::Private
-{
-	// Dummy struct to satisfy [Get/Set]WrappedUObjectPtrValues_InContainer API requirements (Get() function)
-	struct FWrappedObjectPtr
-	{
-		UObject* Obj = nullptr;
-
-		FWrappedObjectPtr() = default;
-
-		FWrappedObjectPtr(UObject* InObj)
-			: Obj(InObj)
-		{
-		}
-
-		UObject* Get() const
-		{
-			return Obj;
-		}
-	};
-}
-
 /*-----------------------------------------------------------------------------
 	FObjectProperty.
 -----------------------------------------------------------------------------*/
@@ -214,7 +193,7 @@ UObject* FObjectProperty::GetObjectPropertyValue(const void* PropertyValueAddres
 UObject* FObjectProperty::GetObjectPropertyValue_InContainer(const void* ContainerAddress, int32 ArrayIndex) const
 {
 	UObject* Result = nullptr;
-	GetWrappedUObjectPtrValues_InContainer<UE::Core::Private::FWrappedObjectPtr>(&Result, ContainerAddress, ArrayIndex, 1);
+	GetWrappedUObjectPtrValues_InContainer<FObjectPtr>(&Result, ContainerAddress, ArrayIndex, 1);
 	return Result;
 }
 
@@ -234,7 +213,7 @@ void FObjectProperty::SetObjectPropertyValue_InContainer(void* ContainerAddress,
 {
 	if (Value || !HasAnyPropertyFlags(CPF_NonNullable))
 	{
-		SetWrappedUObjectPtrValues_InContainer<UE::Core::Private::FWrappedObjectPtr>(ContainerAddress, &Value, ArrayIndex, 1);
+		SetWrappedUObjectPtrValues_InContainer<FObjectPtr>(ContainerAddress, &Value, ArrayIndex, 1);
 	}
 	else
 	{
@@ -244,10 +223,10 @@ void FObjectProperty::SetObjectPropertyValue_InContainer(void* ContainerAddress,
 
 void FObjectProperty::CopyCompleteValueToScriptVM_InContainer(void* OutValue, void const* InContainer) const
 {
-	GetWrappedUObjectPtrValues_InContainer<UE::Core::Private::FWrappedObjectPtr>((UObject**)OutValue, InContainer, 0, ArrayDim);
+	GetWrappedUObjectPtrValues_InContainer<FObjectPtr>((UObject**)OutValue, InContainer, 0, ArrayDim);
 }
 
 void FObjectProperty::CopyCompleteValueFromScriptVM_InContainer(void* OutContainer, void const* InValue) const
 {
-	SetWrappedUObjectPtrValues_InContainer<UE::Core::Private::FWrappedObjectPtr>(OutContainer, (UObject**)InValue, 0, ArrayDim);
+	SetWrappedUObjectPtrValues_InContainer<FObjectPtr>(OutContainer, (UObject**)InValue, 0, ArrayDim);
 }
