Making GetValueAddressAtIndex_Direct functions return null if the requested index is out of bounds

diff --git a/Engine/Source/Runtime/CoreUObject/Private/UObject/PropertyArray.cpp b/Engine/Source/Runtime/CoreUObject/Private/UObject/PropertyArray.cpp
--- a/Engine/Source/Runtime/CoreUObject/Private/UObject/PropertyArray.cpp
+++ b/Engine/Source/Runtime/CoreUObject/Private/UObject/PropertyArray.cpp
@@ -887,7 +887,12 @@ void* FArrayProperty::GetValueAddressAtIndex_Direct(const FProperty* InInner, vo
 {
 	FScriptArrayHelper ArrayHelper(this, InValueAddress);
 	checkf(Inner == InInner, TEXT("Passed in inner must be identical to the array property inner property"));
-	checkf(Index < ArrayHelper.Num() && Index >= 0, TEXT("Array index (%d) out of range"), Index);
-
-	return ArrayHelper.GetRawPtr(Index);
+	if (Index < ArrayHelper.Num() && Index >= 0)
+	{
+		return ArrayHelper.GetRawPtr(Index);
+	}
+	else
+	{
+		return nullptr;
+	}
 }
diff --git a/Engine/Source/Runtime/CoreUObject/Private/UObject/PropertyMap.cpp b/Engine/Source/Runtime/CoreUObject/Private/UObject/PropertyMap.cpp
--- a/Engine/Source/Runtime/CoreUObject/Private/UObject/PropertyMap.cpp
+++ b/Engine/Source/Runtime/CoreUObject/Private/UObject/PropertyMap.cpp
@@ -1333,6 +1333,5 @@ void* FMapProperty::GetValueAddressAtIndex_Direct(const FProperty* Inner, void*
 			Num--;
 		}
 	}
-	checkf(false, TEXT("Map element index (%d) out of range"), Index);
 	return nullptr;
 }
diff --git a/Engine/Source/Runtime/CoreUObject/Private/UObject/PropertySet.cpp b/Engine/Source/Runtime/CoreUObject/Private/UObject/PropertySet.cpp
--- a/Engine/Source/Runtime/CoreUObject/Private/UObject/PropertySet.cpp
+++ b/Engine/Source/Runtime/CoreUObject/Private/UObject/PropertySet.cpp
@@ -1105,6 +1105,5 @@ void* FSetProperty::GetValueAddressAtIndex_Direct(const FProperty* Inner, void*
 			Num--;
 		}
 	}
-	checkf(false, TEXT("Set element index (%d) out of range"), Index);
 	return nullptr;
 }
