DirectoryTree: Give CA_ASSUME hint to suppress spurious static analysis warning.
#rnx

diff --git a/Engine/Source/Runtime/CoreUObject/Internal/Containers/DirectoryTree.h b/Engine/Source/Runtime/CoreUObject/Internal/Containers/DirectoryTree.h
--- a/Engine/Source/Runtime/CoreUObject/Internal/Containers/DirectoryTree.h
+++ b/Engine/Source/Runtime/CoreUObject/Internal/Containers/DirectoryTree.h
@@ -963,6 +963,8 @@ inline void TDirectoryTree<ValueType>::FTreeNode::Realloc(int32 NewCapacity)
 		FTreeNode* NewChildNodes = new FTreeNode[NewCapacity];
 		for (int32 Index = 0; Index < NumChildNodes; ++Index)
 		{
+			// Index < NumChildNodes <= NewCapacity. Some analyzers miss check(NumChildNodes <= NewCapacity) and need a direct hint.
+			CA_ASSUME(Index < NewCapacity);
 			NewRelPaths[Index] = MoveTemp(RelPaths[Index]);
 			NewChildNodes[Index] = MoveTemp(ChildNodes[Index]);
 		}
