fix for ObjectHandleTracking initialization. order of global variable initialization was problematic

move globals into a struct added a singleton for the struct to do lazy initialization

https://p4-swarm.epicgames.net/reviews/23877016

diff --git a/Engine/Source/Runtime/CoreUObject/Private/UObject/ObjectHandleTracking.cpp b/Engine/Source/Runtime/CoreUObject/Private/UObject/ObjectHandleTracking.cpp
--- a/Engine/Source/Runtime/CoreUObject/Private/UObject/ObjectHandleTracking.cpp
+++ b/Engine/Source/Runtime/CoreUObject/Private/UObject/ObjectHandleTracking.cpp
@@ -9,143 +9,213 @@ namespace UE::CoreUObject
 {
 	namespace Private
 	{
-		int32 NextHandleId = 0;
-		TArray<TTuple<int32, FObjectHandleReadFunc>> ReadHandleCallbacks;
-		TArray<TTuple<int32, FObjectHandleClassResolvedFunc>> ClassResolvedCallbacks;
-		TArray<TTuple<int32, FObjectHandleReferenceResolvedFunc>> HandleResolvedCallbacks;
-		TArray<TTuple<int32, FObjectHandleReferenceLoadedFunc>> HandleLoadedCallbacks;
-		
-		//using a single lock as currently there is not any contention. could add separate locks for each list. 
-		FRWLock HandleLock;
-
-		void OnHandleRead(UObject* Object)
+		struct ObjectHandleCallbacks
 		{
-			FReadScopeLock _(HandleLock);
-			for (auto&& Pair : ReadHandleCallbacks)
+			static ObjectHandleCallbacks& Get()
+			{
+				static ObjectHandleCallbacks Callbacks;
+				return Callbacks;
+			}
+
+			void OnHandleRead(UObject* Object)
+			{
+				FReadScopeLock _(HandleLock);
+				for (auto&& Pair : ReadHandleCallbacks)
+				{
+					Pair.Value(Object);
+				}
+			}
+
+			void OnClassReferenceResolved(const FObjectRef& ObjectRef, UPackage* Package, UClass* Class)
+			{
+				FReadScopeLock _(HandleLock);
+				for (auto&& Pair : ClassResolvedCallbacks)
+				{
+					Pair.Value(ObjectRef, Package, Class);
+				}
+			}
+
+			void OnReferenceResolved(const FObjectRef& ObjectRef, UPackage* Package, UObject* Object)
+			{
+				FReadScopeLock _(HandleLock);
+				for (auto&& Pair : HandleResolvedCallbacks)
+				{
+					Pair.Value(ObjectRef, Package, Object);
+				}
+			}
+
+			void OnReferenceLoaded(const FObjectRef& ObjectRef, UPackage* Package, UObject* Object)
+			{
+				FReadScopeLock _(HandleLock);
+				for (auto&& Pair : HandleLoadedCallbacks)
+				{
+					Pair.Value(ObjectRef, Package, Object);
+				}
+			}
+
+			FObjectHandleTrackingCallbackId AddObjectHandleReadCallback(FObjectHandleReadFunc Func)
+			{
+				FWriteScopeLock _(HandleLock);
+				NextHandleId++;
+				ReadHandleCallbacks.Add({ NextHandleId, Func });
+				return { NextHandleId };
+			}
+
+			void RemoveObjectHandleReadCallback(FObjectHandleTrackingCallbackId Handle)
+			{
+				FWriteScopeLock _(HandleLock);
+				for (int32 i = ReadHandleCallbacks.Num() - 1; i >= 0; --i)
+				{
+					auto& Pair = ReadHandleCallbacks[i];
+					if (Pair.Key == Handle.Id)
+					{
+						ReadHandleCallbacks.RemoveAt(i);
+					}
+				}
+			}
+
+			FObjectHandleTrackingCallbackId AddObjectHandleClassResolvedCallback(FObjectHandleClassResolvedFunc Func)
+			{
+				FWriteScopeLock _(HandleLock);
+				NextHandleId++;
+				ClassResolvedCallbacks.Add({ NextHandleId, Func });
+				return { NextHandleId };
+			}
+
+			void RemoveObjectHandleClassResolvedCallback(FObjectHandleTrackingCallbackId Handle)
+			{
+				FWriteScopeLock _(HandleLock);
+				for (int32 i = ClassResolvedCallbacks.Num() - 1; i >= 0; --i)
+				{
+					auto& Pair = ClassResolvedCallbacks[i];
+					if (Pair.Key == Handle.Id)
+					{
+						ClassResolvedCallbacks.RemoveAt(i);
+					}
+				}
+			}
+
+			FObjectHandleTrackingCallbackId AddObjectHandleReferenceResolvedCallback(FObjectHandleReferenceResolvedFunc Func)
+			{
+				FWriteScopeLock _(HandleLock);
+				NextHandleId++;
+				HandleResolvedCallbacks.Add({ NextHandleId, Func });
+				return { NextHandleId };
+			}
+
+			void RemoveObjectHandleReferenceResolvedCallback(FObjectHandleTrackingCallbackId Handle)
+			{
+				FWriteScopeLock _(HandleLock);
+				for (int32 i = HandleResolvedCallbacks.Num() - 1; i >= 0; --i)
+				{
+					auto& Pair = HandleResolvedCallbacks[i];
+					if (Pair.Key == Handle.Id)
+					{
+						HandleResolvedCallbacks.RemoveAt(i);
+					}
+				}
+			}
+
+			FObjectHandleTrackingCallbackId AddObjectHandleReferenceLoadedCallback(FObjectHandleReferenceLoadedFunc Func)
+			{
+				FWriteScopeLock _(HandleLock);
+				NextHandleId++;
+				HandleLoadedCallbacks.Add({ NextHandleId, Func });
+				return { NextHandleId };
+			}
+
+			void RemoveObjectHandleReferenceLoadedCallback(FObjectHandleTrackingCallbackId Handle)
 			{
-				Pair.Value(Object);
+				FWriteScopeLock _(HandleLock);
+				for (int32 i = HandleLoadedCallbacks.Num() - 1; i >= 0; --i)
+				{
+					auto& Pair = HandleLoadedCallbacks[i];
+					if (Pair.Key == Handle.Id)
+					{
+						HandleLoadedCallbacks.RemoveAt(i);
+					}
+				}
 			}
+
+		private:
+			int32 NextHandleId = 0;
+			TArray<TTuple<int32, FObjectHandleReadFunc>> ReadHandleCallbacks;
+			TArray<TTuple<int32, FObjectHandleClassResolvedFunc>> ClassResolvedCallbacks;
+			TArray<TTuple<int32, FObjectHandleReferenceResolvedFunc>> HandleResolvedCallbacks;
+			TArray<TTuple<int32, FObjectHandleReferenceLoadedFunc>> HandleLoadedCallbacks;
+
+			//using a single lock as currently there is not any contention. could add separate locks for each list. 
+			FRWLock HandleLock;
+		};
+
+		void OnHandleRead(UObject* Object)
+		{
+			ObjectHandleCallbacks::Get().OnHandleRead(Object);
 		}
 
 		void OnClassReferenceResolved(const FObjectRef& ObjectRef, UPackage* Package, UClass* Class)
 		{
-			FReadScopeLock _(HandleLock);
-			for (auto&& Pair : ClassResolvedCallbacks)
-			{
-				Pair.Value(ObjectRef, Package, Class);
-			}
+			ObjectHandleCallbacks::Get().OnClassReferenceResolved(ObjectRef, Package, Class);
 		}
 
 		void OnReferenceResolved(const FObjectRef& ObjectRef, UPackage* Package, UObject* Object)
 		{
-			FReadScopeLock _(HandleLock);
-			for (auto&& Pair : HandleResolvedCallbacks)
-			{
-				Pair.Value(ObjectRef, Package, Object);
-			}
+			ObjectHandleCallbacks::Get().OnReferenceResolved(ObjectRef, Package, Object);
 		}
 
 		void OnReferenceLoaded(const FObjectRef& ObjectRef, UPackage* Package, UObject* Object)
 		{
-			FReadScopeLock _(HandleLock);
-			for (auto&& Pair : HandleLoadedCallbacks)
-			{
-				Pair.Value(ObjectRef, Package, Object);
-			}
+			ObjectHandleCallbacks::Get().OnReferenceLoaded(ObjectRef, Package, Object);
 		}
+
 	}
 
 	FObjectHandleTrackingCallbackId AddObjectHandleReadCallback(FObjectHandleReadFunc Func)
 	{
 		using namespace UE::CoreUObject::Private;
-		FWriteScopeLock _(HandleLock);
-		NextHandleId++;
-		ReadHandleCallbacks.Add({ NextHandleId, Func });
-		return { NextHandleId };
+		return ObjectHandleCallbacks::Get().AddObjectHandleReadCallback(Func);
 	}
 
 	void RemoveObjectHandleReadCallback(FObjectHandleTrackingCallbackId Handle)
 	{
 		using namespace UE::CoreUObject::Private;
-		FWriteScopeLock _(Private::HandleLock);
-		for (int32 i = ReadHandleCallbacks.Num() - 1; i >= 0; --i)
-		{
-			auto& Pair = ReadHandleCallbacks[i];
-			if (Pair.Key == Handle.Id)
-			{
-				ReadHandleCallbacks.RemoveAt(i);
-			}
-		}
-		
+		ObjectHandleCallbacks::Get().RemoveObjectHandleReadCallback(Handle);
 	}
 
 	FObjectHandleTrackingCallbackId AddObjectHandleClassResolvedCallback(FObjectHandleClassResolvedFunc Func)
 	{
 		using namespace UE::CoreUObject::Private;
-		FWriteScopeLock _(HandleLock);
-		NextHandleId++;
-		ClassResolvedCallbacks.Add({ NextHandleId, Func });
-		return { NextHandleId };
+		return ObjectHandleCallbacks::Get().AddObjectHandleClassResolvedCallback(Func);
 	}
 
 	void RemoveObjectHandleClassResolvedCallback(FObjectHandleTrackingCallbackId Handle)
 	{
 		using namespace UE::CoreUObject::Private;
-		FWriteScopeLock _(Private::HandleLock);
-		for (int32 i = ClassResolvedCallbacks.Num() - 1; i >= 0; --i)
-		{
-			auto& Pair = ClassResolvedCallbacks[i];
-			if (Pair.Key == Handle.Id)
-			{
-				ClassResolvedCallbacks.RemoveAt(i);
-			}
-		}
+		ObjectHandleCallbacks::Get().RemoveObjectHandleClassResolvedCallback(Handle);
 	}
 
 	FObjectHandleTrackingCallbackId AddObjectHandleReferenceResolvedCallback(FObjectHandleReferenceResolvedFunc Func)
 	{
 		using namespace UE::CoreUObject::Private;
-		FWriteScopeLock _(HandleLock);
-		NextHandleId++;
-		HandleResolvedCallbacks.Add({ NextHandleId, Func });
-		return { NextHandleId };
+		return ObjectHandleCallbacks::Get().AddObjectHandleReferenceResolvedCallback(Func);
 	}
 
 	void RemoveObjectHandleReferenceResolvedCallback(FObjectHandleTrackingCallbackId Handle)
 	{
 		using namespace UE::CoreUObject::Private;
-		FWriteScopeLock _(Private::HandleLock);
-		for (int32 i = HandleResolvedCallbacks.Num() - 1; i >= 0; --i)
-		{
-			auto& Pair = HandleResolvedCallbacks[i];
-			if (Pair.Key == Handle.Id)
-			{
-				HandleResolvedCallbacks.RemoveAt(i);
-			}
-		}
+		ObjectHandleCallbacks::Get().RemoveObjectHandleReferenceResolvedCallback(Handle);
 	}
 
 	FObjectHandleTrackingCallbackId AddObjectHandleReferenceLoadedCallback(FObjectHandleReferenceLoadedFunc Func)
 	{
 		using namespace UE::CoreUObject::Private;
-		FWriteScopeLock _(HandleLock);
-		NextHandleId++;
-		HandleLoadedCallbacks.Add({ NextHandleId, Func });
-		return { NextHandleId };
+		return ObjectHandleCallbacks::Get().AddObjectHandleReferenceLoadedCallback(Func);
 	}
 
 	void RemoveObjectHandleReferenceLoadedCallback(FObjectHandleTrackingCallbackId Handle)
 	{
 		using namespace UE::CoreUObject::Private;
-		FWriteScopeLock _(Private::HandleLock);
-		for (int32 i = HandleLoadedCallbacks.Num() - 1; i >= 0; --i)
-		{
-			auto& Pair = HandleLoadedCallbacks[i];
-			if (Pair.Key == Handle.Id)
-			{
-				HandleLoadedCallbacks.RemoveAt(i);
-			}
-		}
+		ObjectHandleCallbacks::Get().RemoveObjectHandleReferenceLoadedCallback(Handle);
 	}
 }
 #endif // UE_WITH_OBJECT_HANDLE_TRACKING
