Fix TSAN warning by letting it know our intention is really to do a relaxed read

#rnx

diff --git a/Engine/Source/Runtime/CoreUObject/Private/UObject/GarbageCollection.cpp b/Engine/Source/Runtime/CoreUObject/Private/UObject/GarbageCollection.cpp
--- a/Engine/Source/Runtime/CoreUObject/Private/UObject/GarbageCollection.cpp
+++ b/Engine/Source/Runtime/CoreUObject/Private/UObject/GarbageCollection.cpp
@@ -2263,9 +2263,9 @@ constexpr FORCEINLINE EKillable MayKill(EGCTokenType TokenType, bool bAllowKill)
 FORCEINLINE static bool ClearUnreachableInterlocked(int32& Flags)
 {
 	static constexpr int32 FlagToClear = int32(EInternalObjectFlags::Unreachable);
-	if (Flags & FlagToClear)
+	if (FPlatformAtomics::AtomicRead_Relaxed(&Flags) & FlagToClear)
 	{
-		int Old = FPlatformAtomics::InterlockedAnd(&Flags, ~FlagToClear);
+		int32 Old = FPlatformAtomics::InterlockedAnd(&Flags, ~FlagToClear);
 		return Old & FlagToClear;
 	}
 
