Fix TSAN warning during object creation on ALT because GUndo is not thread-safe

#rnx

diff --git a/Engine/Source/Runtime/CoreUObject/Private/UObject/UObjectGlobals.cpp b/Engine/Source/Runtime/CoreUObject/Private/UObject/UObjectGlobals.cpp
--- a/Engine/Source/Runtime/CoreUObject/Private/UObject/UObjectGlobals.cpp
+++ b/Engine/Source/Runtime/CoreUObject/Private/UObject/UObjectGlobals.cpp
@@ -4324,11 +4324,14 @@ UObject* StaticConstructObject_Internal(const FStaticConstructObjectParameters&
 		(*InClass->ClassConstructor)(FObjectInitializer(Result, Params));
 	}
 	
-	if (GIsEditor && GUndo && 
-		(InFlags & RF_Transactional) && !(InFlags & RF_NeedLoad) && 
-		!InClass->IsChildOf(UField::StaticClass()) &&
+	if (GIsEditor && 
 		// Do not consider object creation in transaction if the object is marked as async or in being async loaded 
-		!Result->HasAnyInternalFlags(EInternalObjectFlags::Async|EInternalObjectFlags::AsyncLoading))
+		!Result->HasAnyInternalFlags(EInternalObjectFlags::Async | EInternalObjectFlags::AsyncLoading) &&
+		// Read GUndo only if not having Async flags set to avoid making TSAN unhappy that we're trying to read an unsynchronized global
+		GUndo &&
+		(InFlags & RF_Transactional) && !(InFlags & RF_NeedLoad) && 
+		!InClass->IsChildOf(UField::StaticClass())
+		)
 	{
 		// Set RF_PendingKill and update the undo buffer so an undo operation will set RF_PendingKill on the newly constructed object.
 		Result->MarkAsGarbage();
