Flipped the internal condition of [Get|Set]WrappedUObjectPtrValues_InContainer to make an upcoming change easier to diff.

diff --git a/Engine/Source/Runtime/CoreUObject/Public/UObject/UnrealType.h b/Engine/Source/Runtime/CoreUObject/Public/UObject/UnrealType.h
--- a/Engine/Source/Runtime/CoreUObject/Public/UObject/UnrealType.h
+++ b/Engine/Source/Runtime/CoreUObject/Public/UObject/UnrealType.h
@@ -2756,25 +2756,7 @@ protected:
 		// Ensure required range is valid
 		checkf(ArrayIndex >= 0 && ArrayCount >= 0 && ArrayIndex <= ArrayDim && ArrayCount <= ArrayDim && ArrayIndex <= ArrayDim - ArrayCount, TEXT("ArrayIndex (%d) and ArrayCount (%d) is invalid for an array of size %d"), ArrayIndex, ArrayCount, ArrayDim);
 
-		if (!HasGetter())
-		{
-			// Fast path - direct memory access
-			const uint8* ObjAddress = (const uint8*)ContainerPtrToValuePtr<void>(ContainerAddress, ArrayIndex);
-
-			if (ArrayCount == 1)
-			{
-				*OutObjects = GetObjectPropertyValue(ObjAddress);
-			}
-			else
-			{
-				int32 LocalElementSize = ElementSize;
-				for (int32 OutIndex = 0; OutIndex != ArrayCount; ++OutIndex)
-				{
-					OutObjects[OutIndex] = GetObjectPropertyValue(ObjAddress + OutIndex * LocalElementSize);
-				}
-			}
-		}
-		else
+		if (HasGetter())
 		{
 			if (ArrayCount == 1)
 			{
@@ -2798,32 +2780,32 @@ protected:
 				DestroyAndFreeValue(ValueArray);
 			}
 		}
-	}
-	template <typename T>
-	void SetWrappedUObjectPtrValues_InContainer(void* ContainerAddress, UObject** InValues, int32 ArrayIndex, int32 ArrayCount) const
-	{
-		// Ensure required range is valid
-		checkf(ArrayIndex >= 0 && ArrayCount >= 0 && ArrayIndex <= ArrayDim && ArrayCount <= ArrayDim && ArrayIndex <= ArrayDim - ArrayCount, TEXT("ArrayIndex (%d) and ArrayCount (%d) is invalid for an array of size %d"), ArrayIndex, ArrayCount, ArrayDim);
-
-		if (!HasSetter())
+		else
 		{
 			// Fast path - direct memory access
-			uint8* ObjAddress = (uint8*)ContainerPtrToValuePtr<void>(ContainerAddress, ArrayIndex);
+			const uint8* ObjAddress = (const uint8*)ContainerPtrToValuePtr<void>(ContainerAddress, ArrayIndex);
 
 			if (ArrayCount == 1)
 			{
-				SetObjectPropertyValue(ObjAddress, *InValues);
+				*OutObjects = GetObjectPropertyValue(ObjAddress);
 			}
 			else
 			{
 				int32 LocalElementSize = ElementSize;
 				for (int32 OutIndex = 0; OutIndex != ArrayCount; ++OutIndex)
 				{
-					SetObjectPropertyValue(ObjAddress + OutIndex * LocalElementSize, InValues[OutIndex]);
+					OutObjects[OutIndex] = GetObjectPropertyValue(ObjAddress + OutIndex * LocalElementSize);
 				}
 			}
 		}
-		else
+	}
+	template <typename T>
+	void SetWrappedUObjectPtrValues_InContainer(void* ContainerAddress, UObject** InValues, int32 ArrayIndex, int32 ArrayCount) const
+	{
+		// Ensure required range is valid
+		checkf(ArrayIndex >= 0 && ArrayCount >= 0 && ArrayIndex <= ArrayDim && ArrayCount <= ArrayDim && ArrayIndex <= ArrayDim - ArrayCount, TEXT("ArrayIndex (%d) and ArrayCount (%d) is invalid for an array of size %d"), ArrayIndex, ArrayCount, ArrayDim);
+
+		if (HasSetter())
 		{
 			if (ArrayCount == 1)
 			{
@@ -2849,6 +2831,24 @@ protected:
 				DestroyAndFreeValue(ValueArray);
 			}
 		}
+		else
+		{
+			// Fast path - direct memory access
+			uint8* ObjAddress = (uint8*)ContainerPtrToValuePtr<void>(ContainerAddress, ArrayIndex);
+
+			if (ArrayCount == 1)
+			{
+				SetObjectPropertyValue(ObjAddress, *InValues);
+			}
+			else
+			{
+				int32 LocalElementSize = ElementSize;
+				for (int32 OutIndex = 0; OutIndex != ArrayCount; ++OutIndex)
+				{
+					SetObjectPropertyValue(ObjAddress + OutIndex * LocalElementSize, InValues[OutIndex]);
+				}
+			}
+		}
 	}
 };
 
