Changed [Get|Set]WrappedUObjectPtrValues to take a flag which says whether the pointer is 'in container' or not.

diff --git a/Engine/Source/Runtime/CoreUObject/Private/UObject/PropertyClassPtr.cpp b/Engine/Source/Runtime/CoreUObject/Private/UObject/PropertyClassPtr.cpp
--- a/Engine/Source/Runtime/CoreUObject/Private/UObject/PropertyClassPtr.cpp
+++ b/Engine/Source/Runtime/CoreUObject/Private/UObject/PropertyClassPtr.cpp
@@ -54,7 +54,7 @@ UObject* FClassPtrProperty::GetObjectPropertyValue(const void* PropertyValueAddr
 UObject* FClassPtrProperty::GetObjectPropertyValue_InContainer(const void* ContainerAddress, int32 ArrayIndex) const
 {
 	UObject* Result = nullptr;
-	GetWrappedUObjectPtrValues_InContainer<FObjectPtr>(&Result, ContainerAddress, ArrayIndex, 1);
+	GetWrappedUObjectPtrValues<FObjectPtr>(&Result, ContainerAddress, EPropertyMemoryAccess::InContainer, ArrayIndex, 1);
 	return Result;
 }
 
@@ -74,7 +74,7 @@ void FClassPtrProperty::SetObjectPropertyValue_InContainer(void* ContainerAddres
 {
 	if (Value || !HasAnyPropertyFlags(CPF_NonNullable))
 	{
-		SetWrappedUObjectPtrValues_InContainer<FObjectPtr>(ContainerAddress, &Value, ArrayIndex, 1);
+		SetWrappedUObjectPtrValues<FObjectPtr>(ContainerAddress, EPropertyMemoryAccess::InContainer, &Value, ArrayIndex, 1);
 	}
 	else
 	{
@@ -89,10 +89,10 @@ uint32 FClassPtrProperty::GetValueTypeHashInternal(const void* Src) const
 
 void FClassPtrProperty::CopyCompleteValueToScriptVM_InContainer(void* OutValue, void const* InContainer) const
 {
-	GetWrappedUObjectPtrValues_InContainer<FObjectPtr>((UObject**)OutValue, InContainer, 0, ArrayDim);
+	GetWrappedUObjectPtrValues<FObjectPtr>((UObject**)OutValue, InContainer, EPropertyMemoryAccess::InContainer, 0, ArrayDim);
 }
 
 void FClassPtrProperty::CopyCompleteValueFromScriptVM_InContainer(void* OutContainer, void const* InValue) const
 {
-	SetWrappedUObjectPtrValues_InContainer<FObjectPtr>(OutContainer, (UObject**)InValue, 0, ArrayDim);
+	SetWrappedUObjectPtrValues<FObjectPtr>(OutContainer, EPropertyMemoryAccess::InContainer, (UObject**)InValue, 0, ArrayDim);
 }
diff --git a/Engine/Source/Runtime/CoreUObject/Private/UObject/PropertyLazyObjectPtr.cpp b/Engine/Source/Runtime/CoreUObject/Private/UObject/PropertyLazyObjectPtr.cpp
--- a/Engine/Source/Runtime/CoreUObject/Private/UObject/PropertyLazyObjectPtr.cpp
+++ b/Engine/Source/Runtime/CoreUObject/Private/UObject/PropertyLazyObjectPtr.cpp
@@ -118,7 +118,7 @@ UObject* FLazyObjectProperty::GetObjectPropertyValue(const void* PropertyValueAd
 UObject* FLazyObjectProperty::GetObjectPropertyValue_InContainer(const void* ContainerAddress, int32 ArrayIndex) const
 {
 	UObject* Result = nullptr;
-	GetWrappedUObjectPtrValues_InContainer<FLazyObjectPtr>(&Result, ContainerAddress, ArrayIndex, 1);
+	GetWrappedUObjectPtrValues<FLazyObjectPtr>(&Result, ContainerAddress, EPropertyMemoryAccess::InContainer, ArrayIndex, 1);
 	return Result;
 }
 
@@ -138,7 +138,7 @@ void FLazyObjectProperty::SetObjectPropertyValue_InContainer(void* ContainerAddr
 {
 	if (Value || !HasAnyPropertyFlags(CPF_NonNullable))
 	{
-		SetWrappedUObjectPtrValues_InContainer<FLazyObjectPtr>(ContainerAddress, &Value, ArrayIndex, 1);
+		SetWrappedUObjectPtrValues<FLazyObjectPtr>(ContainerAddress, EPropertyMemoryAccess::InContainer, &Value, ArrayIndex, 1);
 	}
 	else
 	{
@@ -158,10 +158,10 @@ uint32 FLazyObjectProperty::GetValueTypeHashInternal(const void* Src) const
 
 void FLazyObjectProperty::CopyCompleteValueToScriptVM_InContainer(void* OutValue, void const* InContainer) const
 {
-	GetWrappedUObjectPtrValues_InContainer<FLazyObjectPtr>((UObject**)OutValue, InContainer, 0, ArrayDim);
+	GetWrappedUObjectPtrValues<FLazyObjectPtr>((UObject**)OutValue, InContainer, EPropertyMemoryAccess::InContainer, 0, ArrayDim);
 }
 
 void FLazyObjectProperty::CopyCompleteValueFromScriptVM_InContainer(void* OutContainer, void const* InValue) const
 {
-	SetWrappedUObjectPtrValues_InContainer<FLazyObjectPtr>(OutContainer, (UObject**)InValue, 0, ArrayDim);
+	SetWrappedUObjectPtrValues<FLazyObjectPtr>(OutContainer, EPropertyMemoryAccess::InContainer, (UObject**)InValue, 0, ArrayDim);
 }
diff --git a/Engine/Source/Runtime/CoreUObject/Private/UObject/PropertyObject.cpp b/Engine/Source/Runtime/CoreUObject/Private/UObject/PropertyObject.cpp
--- a/Engine/Source/Runtime/CoreUObject/Private/UObject/PropertyObject.cpp
+++ b/Engine/Source/Runtime/CoreUObject/Private/UObject/PropertyObject.cpp
@@ -193,7 +193,7 @@ UObject* FObjectProperty::GetObjectPropertyValue(const void* PropertyValueAddres
 UObject* FObjectProperty::GetObjectPropertyValue_InContainer(const void* ContainerAddress, int32 ArrayIndex) const
 {
 	UObject* Result = nullptr;
-	GetWrappedUObjectPtrValues_InContainer<FObjectPtr>(&Result, ContainerAddress, ArrayIndex, 1);
+	GetWrappedUObjectPtrValues<FObjectPtr>(&Result, ContainerAddress, EPropertyMemoryAccess::InContainer, ArrayIndex, 1);
 	return Result;
 }
 
@@ -213,7 +213,7 @@ void FObjectProperty::SetObjectPropertyValue_InContainer(void* ContainerAddress,
 {
 	if (Value || !HasAnyPropertyFlags(CPF_NonNullable))
 	{
-		SetWrappedUObjectPtrValues_InContainer<FObjectPtr>(ContainerAddress, &Value, ArrayIndex, 1);
+		SetWrappedUObjectPtrValues<FObjectPtr>(ContainerAddress, EPropertyMemoryAccess::InContainer, &Value, ArrayIndex, 1);
 	}
 	else
 	{
@@ -223,10 +223,10 @@ void FObjectProperty::SetObjectPropertyValue_InContainer(void* ContainerAddress,
 
 void FObjectProperty::CopyCompleteValueToScriptVM_InContainer(void* OutValue, void const* InContainer) const
 {
-	GetWrappedUObjectPtrValues_InContainer<FObjectPtr>((UObject**)OutValue, InContainer, 0, ArrayDim);
+	GetWrappedUObjectPtrValues<FObjectPtr>((UObject**)OutValue, InContainer, EPropertyMemoryAccess::InContainer, 0, ArrayDim);
 }
 
 void FObjectProperty::CopyCompleteValueFromScriptVM_InContainer(void* OutContainer, void const* InValue) const
 {
-	SetWrappedUObjectPtrValues_InContainer<FObjectPtr>(OutContainer, (UObject**)InValue, 0, ArrayDim);
+	SetWrappedUObjectPtrValues<FObjectPtr>(OutContainer, EPropertyMemoryAccess::InContainer, (UObject**)InValue, 0, ArrayDim);
 }
diff --git a/Engine/Source/Runtime/CoreUObject/Private/UObject/PropertyObjectPtr.cpp b/Engine/Source/Runtime/CoreUObject/Private/UObject/PropertyObjectPtr.cpp
--- a/Engine/Source/Runtime/CoreUObject/Private/UObject/PropertyObjectPtr.cpp
+++ b/Engine/Source/Runtime/CoreUObject/Private/UObject/PropertyObjectPtr.cpp
@@ -159,7 +159,7 @@ UObject* FObjectPtrProperty::GetObjectPropertyValue(const void* PropertyValueAdd
 UObject* FObjectPtrProperty::GetObjectPropertyValue_InContainer(const void* ContainerAddress, int32 ArrayIndex) const
 {
 	UObject* Result = nullptr;
-	GetWrappedUObjectPtrValues_InContainer<FObjectPtr>(&Result, ContainerAddress, ArrayIndex, 1);
+	GetWrappedUObjectPtrValues<FObjectPtr>(&Result, ContainerAddress, EPropertyMemoryAccess::InContainer, ArrayIndex, 1);
 	return Result;
 }
 
@@ -179,7 +179,7 @@ void FObjectPtrProperty::SetObjectPropertyValue_InContainer(void* ContainerAddre
 {
 	if (Value || !HasAnyPropertyFlags(CPF_NonNullable))
 	{
-		SetWrappedUObjectPtrValues_InContainer<FObjectPtr>(ContainerAddress, &Value, ArrayIndex, 1);
+		SetWrappedUObjectPtrValues<FObjectPtr>(ContainerAddress, EPropertyMemoryAccess::InContainer, &Value, ArrayIndex, 1);
 	}
 	else
 	{
@@ -199,10 +199,10 @@ uint32 FObjectPtrProperty::GetValueTypeHashInternal(const void* Src) const
 
 void FObjectPtrProperty::CopyCompleteValueToScriptVM_InContainer(void* OutValue, void const* InContainer) const
 {
-	GetWrappedUObjectPtrValues_InContainer<FObjectPtr>((UObject**)OutValue, InContainer, 0, ArrayDim);
+	GetWrappedUObjectPtrValues<FObjectPtr>((UObject**)OutValue, InContainer, EPropertyMemoryAccess::InContainer, 0, ArrayDim);
 }
 
 void FObjectPtrProperty::CopyCompleteValueFromScriptVM_InContainer(void* OutContainer, void const* InValue) const
 {
-	SetWrappedUObjectPtrValues_InContainer<FObjectPtr>(OutContainer, (UObject**)InValue, 0, ArrayDim);
+	SetWrappedUObjectPtrValues<FObjectPtr>(OutContainer, EPropertyMemoryAccess::InContainer, (UObject**)InValue, 0, ArrayDim);
 }
diff --git a/Engine/Source/Runtime/CoreUObject/Private/UObject/PropertySoftObjectPtr.cpp b/Engine/Source/Runtime/CoreUObject/Private/UObject/PropertySoftObjectPtr.cpp
--- a/Engine/Source/Runtime/CoreUObject/Private/UObject/PropertySoftObjectPtr.cpp
+++ b/Engine/Source/Runtime/CoreUObject/Private/UObject/PropertySoftObjectPtr.cpp
@@ -244,7 +244,7 @@ UObject* FSoftObjectProperty::GetObjectPropertyValue(const void* PropertyValueAd
 UObject* FSoftObjectProperty::GetObjectPropertyValue_InContainer(const void* ContainerAddress, int32 ArrayIndex) const
 {
 	UObject* Result = nullptr;
-	GetWrappedUObjectPtrValues_InContainer<FSoftObjectPtr>(&Result, ContainerAddress, ArrayIndex, 1);
+	GetWrappedUObjectPtrValues<FSoftObjectPtr>(&Result, ContainerAddress, EPropertyMemoryAccess::InContainer, ArrayIndex, 1);
 	return Result;
 }
 
@@ -255,7 +255,7 @@ void FSoftObjectProperty::SetObjectPropertyValue(void* PropertyValueAddress, UOb
 
 void FSoftObjectProperty::SetObjectPropertyValue_InContainer(void* ContainerAddress, UObject* Value, int32 ArrayIndex) const
 {
-	SetWrappedUObjectPtrValues_InContainer<FSoftObjectPtr>(ContainerAddress, &Value, ArrayIndex, 1);
+	SetWrappedUObjectPtrValues<FSoftObjectPtr>(ContainerAddress, EPropertyMemoryAccess::InContainer, &Value, ArrayIndex, 1);
 }
 
 bool FSoftObjectProperty::AllowCrossLevel() const
diff --git a/Engine/Source/Runtime/CoreUObject/Private/UObject/PropertyWeakObjectPtr.cpp b/Engine/Source/Runtime/CoreUObject/Private/UObject/PropertyWeakObjectPtr.cpp
--- a/Engine/Source/Runtime/CoreUObject/Private/UObject/PropertyWeakObjectPtr.cpp
+++ b/Engine/Source/Runtime/CoreUObject/Private/UObject/PropertyWeakObjectPtr.cpp
@@ -99,7 +99,7 @@ UObject* FWeakObjectProperty::GetObjectPropertyValue(const void* PropertyValueAd
 UObject* FWeakObjectProperty::GetObjectPropertyValue_InContainer(const void* ContainerAddress, int32 ArrayIndex) const
 {
 	UObject* Result = nullptr;
-	GetWrappedUObjectPtrValues_InContainer<FWeakObjectPtr>(&Result, ContainerAddress, ArrayIndex, 1);
+	GetWrappedUObjectPtrValues<FWeakObjectPtr>(&Result, ContainerAddress, EPropertyMemoryAccess::InContainer, ArrayIndex, 1);
 	return Result;
 }
 
@@ -110,7 +110,7 @@ void FWeakObjectProperty::SetObjectPropertyValue(void* PropertyValueAddress, UOb
 
 void FWeakObjectProperty::SetObjectPropertyValue_InContainer(void* ContainerAddress, UObject* Value, int32 ArrayIndex) const
 {
-	SetWrappedUObjectPtrValues_InContainer<FWeakObjectPtr>(ContainerAddress, &Value, ArrayIndex, 1);
+	SetWrappedUObjectPtrValues<FWeakObjectPtr>(ContainerAddress, EPropertyMemoryAccess::InContainer, &Value, ArrayIndex, 1);
 }
 
 uint32 FWeakObjectProperty::GetValueTypeHashInternal(const void* Src) const
@@ -120,10 +120,10 @@ uint32 FWeakObjectProperty::GetValueTypeHashInternal(const void* Src) const
 
 void FWeakObjectProperty::CopyCompleteValueToScriptVM_InContainer(void* OutValue, void const* InContainer) const
 {
-	GetWrappedUObjectPtrValues_InContainer<FWeakObjectPtr>((UObject**)OutValue, InContainer, 0, ArrayDim);
+	GetWrappedUObjectPtrValues<FWeakObjectPtr>((UObject**)OutValue, InContainer, EPropertyMemoryAccess::InContainer, 0, ArrayDim);
 }
 
 void FWeakObjectProperty::CopyCompleteValueFromScriptVM_InContainer(void* OutContainer, void const* InValue) const
 {
-	SetWrappedUObjectPtrValues_InContainer<FWeakObjectPtr>(OutContainer, (UObject**)InValue, 0, ArrayDim);
+	SetWrappedUObjectPtrValues<FWeakObjectPtr>(OutContainer, EPropertyMemoryAccess::InContainer, (UObject**)InValue, 0, ArrayDim);
 }
diff --git a/Engine/Source/Runtime/CoreUObject/Public/UObject/UnrealType.h b/Engine/Source/Runtime/CoreUObject/Public/UObject/UnrealType.h
--- a/Engine/Source/Runtime/CoreUObject/Public/UObject/UnrealType.h
+++ b/Engine/Source/Runtime/CoreUObject/Public/UObject/UnrealType.h
@@ -129,6 +129,16 @@ enum class EPropertyObjectReferenceType : uint32
 };
 ENUM_CLASS_FLAGS(EPropertyObjectReferenceType);
 
+enum class EPropertyMemoryAccess : uint8
+{
+	// Direct memory access - the associated pointer points to the memory at the reflected item.
+	Direct,
+
+	// Container access - the associated pointer points to the outer of the reflected item.
+	// Access via containers will use getter and setters, if present.
+	InContainer
+};
+
 namespace UEProperty_Private { class FProperty_DoNotUse; }
 
 /**
@@ -2751,102 +2761,108 @@ protected:
 
 	/* Helper functions for UObject property types that wrap the object pointer in a smart pointer */
 	template <typename T>
-	void GetWrappedUObjectPtrValues_InContainer(UObject** OutObjects, const void* ContainerAddress, int32 ArrayIndex, int32 ArrayCount) const
+	void GetWrappedUObjectPtrValues(UObject** OutObjects, const void* SrcAddress, EPropertyMemoryAccess SrcAccess, int32 ArrayIndex, int32 ArrayCount) const
 	{
 		// Ensure required range is valid
 		checkf(ArrayIndex >= 0 && ArrayCount >= 0 && ArrayIndex <= ArrayDim && ArrayCount <= ArrayDim && ArrayIndex <= ArrayDim - ArrayCount, TEXT("ArrayIndex (%d) and ArrayCount (%d) is invalid for an array of size %d"), ArrayIndex, ArrayCount, ArrayDim);
 
-		if (HasGetter())
+		if (SrcAccess == EPropertyMemoryAccess::InContainer)
 		{
-			if (ArrayCount == 1)
-			{
-				// Slower but no mallocs. We can copy the value directly to the resulting param
-				T Value;
-				GetValue_InContainer(ContainerAddress, &Value);
-				*OutObjects = Value.Get();
-			}
-			else
+			if (HasGetter())
 			{
-				// Malloc a temp value that is the size of the array. Getter will then copy the entire array to the temp value
-				T* ValueArray = (T*)AllocateAndInitializeValue();
-				FProperty::GetValue_InContainer(ContainerAddress, ValueArray);
-
-				// Grab the items we care about and free the temp array
-				int32 LocalElementSize = ElementSize;
-				for (int32 OutIndex = 0; OutIndex != ArrayCount; ++OutIndex)
+				if (ArrayCount == 1)
 				{
-					OutObjects[OutIndex] = ValueArray[ArrayIndex + OutIndex].Get();
+					// Slower but no mallocs. We can copy the value directly to the resulting param
+					T Value;
+					GetValue_InContainer(SrcAddress, &Value);
+					*OutObjects = Value.Get();
 				}
-				DestroyAndFreeValue(ValueArray);
+				else
+				{
+					// Malloc a temp value that is the size of the array. Getter will then copy the entire array to the temp value
+					T* ValueArray = (T*)AllocateAndInitializeValue();
+					FProperty::GetValue_InContainer(SrcAddress, ValueArray);
+
+					// Grab the items we care about and free the temp array
+					int32 LocalElementSize = ElementSize;
+					for (int32 OutIndex = 0; OutIndex != ArrayCount; ++OutIndex)
+					{
+						OutObjects[OutIndex] = ValueArray[ArrayIndex + OutIndex].Get();
+					}
+					DestroyAndFreeValue(ValueArray);
+				}
+
+				return;
 			}
+
+			SrcAddress = ContainerPtrToValuePtr<void>(SrcAddress, ArrayIndex);
+		}
+
+		// Fast path - direct memory access
+		if (ArrayCount == 1)
+		{
+			*OutObjects = GetObjectPropertyValue(SrcAddress);
 		}
 		else
 		{
-			// Fast path - direct memory access
-			const uint8* ObjAddress = (const uint8*)ContainerPtrToValuePtr<void>(ContainerAddress, ArrayIndex);
-
-			if (ArrayCount == 1)
-			{
-				*OutObjects = GetObjectPropertyValue(ObjAddress);
-			}
-			else
+			int32 LocalElementSize = ElementSize;
+			for (int32 OutIndex = 0; OutIndex != ArrayCount; ++OutIndex)
 			{
-				int32 LocalElementSize = ElementSize;
-				for (int32 OutIndex = 0; OutIndex != ArrayCount; ++OutIndex)
-				{
-					OutObjects[OutIndex] = GetObjectPropertyValue(ObjAddress + OutIndex * LocalElementSize);
-				}
+				OutObjects[OutIndex] = GetObjectPropertyValue((const uint8*)SrcAddress + OutIndex * LocalElementSize);
 			}
 		}
 	}
 	template <typename T>
-	void SetWrappedUObjectPtrValues_InContainer(void* ContainerAddress, UObject** InValues, int32 ArrayIndex, int32 ArrayCount) const
+	void SetWrappedUObjectPtrValues(void* DestAddress, EPropertyMemoryAccess DestAccess, UObject** InValues, int32 ArrayIndex, int32 ArrayCount) const
 	{
 		// Ensure required range is valid
 		checkf(ArrayIndex >= 0 && ArrayCount >= 0 && ArrayIndex <= ArrayDim && ArrayCount <= ArrayDim && ArrayIndex <= ArrayDim - ArrayCount, TEXT("ArrayIndex (%d) and ArrayCount (%d) is invalid for an array of size %d"), ArrayIndex, ArrayCount, ArrayDim);
 
-		if (HasSetter())
+		if (DestAccess == EPropertyMemoryAccess::InContainer)
 		{
-			if (ArrayCount == 1)
-			{
-				// Slower but no mallocs. We can copy a local wrapped value directly to the resulting param
-				T WrappedValue(*InValues);
-				SetValue_InContainer(ContainerAddress, &WrappedValue);
-			}
-			else
+			if (HasSetter())
 			{
-				// Malloc a temp value that is the size of the array. Getter will then copy the entire array to the temp value
-				T* ValueArray = (T*)AllocateAndInitializeValue();
-				FProperty::GetValue_InContainer(ContainerAddress, ValueArray);
-
-				// Replace the items we care about
-				int32 LocalElementSize = ElementSize;
-				for (int32 OutIndex = 0; OutIndex != ArrayCount; ++OutIndex)
+				if (ArrayCount == 1)
 				{
-					ValueArray[ArrayIndex + OutIndex] = InValues[OutIndex];
+					// Slower but no mallocs. We can copy a local wrapped value directly to the resulting param
+					T WrappedValue(*InValues);
+					SetValue_InContainer(DestAddress, &WrappedValue);
 				}
+				else
+				{
+					// Malloc a temp value that is the size of the array. Getter will then copy the entire array to the temp value
+					T* ValueArray = (T*)AllocateAndInitializeValue();
+					FProperty::GetValue_InContainer(DestAddress, ValueArray);
+
+					// Replace the items we care about
+					int32 LocalElementSize = ElementSize;
+					for (int32 OutIndex = 0; OutIndex != ArrayCount; ++OutIndex)
+					{
+						ValueArray[ArrayIndex + OutIndex] = InValues[OutIndex];
+					}
 
-				// Now copy the entire array back to the property using a setter
-				SetValue_InContainer(ContainerAddress, ValueArray);
-				DestroyAndFreeValue(ValueArray);
+					// Now copy the entire array back to the property using a setter
+					SetValue_InContainer(DestAddress, ValueArray);
+					DestroyAndFreeValue(ValueArray);
+				}
+
+				return;
 			}
+
+			DestAddress = ContainerPtrToValuePtr<void>(DestAddress, ArrayIndex);
+		}
+
+		// Fast path - direct memory access
+		if (ArrayCount == 1)
+		{
+			SetObjectPropertyValue(DestAddress, *InValues);
 		}
 		else
 		{
-			// Fast path - direct memory access
-			uint8* ObjAddress = (uint8*)ContainerPtrToValuePtr<void>(ContainerAddress, ArrayIndex);
-
-			if (ArrayCount == 1)
+			int32 LocalElementSize = ElementSize;
+			for (int32 OutIndex = 0; OutIndex != ArrayCount; ++OutIndex)
 			{
-				SetObjectPropertyValue(ObjAddress, *InValues);
-			}
-			else
-			{
-				int32 LocalElementSize = ElementSize;
-				for (int32 OutIndex = 0; OutIndex != ArrayCount; ++OutIndex)
-				{
-					SetObjectPropertyValue(ObjAddress + OutIndex * LocalElementSize, InValues[OutIndex]);
-				}
+				SetObjectPropertyValue((uint8*)DestAddress + OutIndex * LocalElementSize, InValues[OutIndex]);
 			}
 		}
 	}
