Bulk Data - fix bug when serializing text based asset

#rnx

diff --git a/Engine/Source/Runtime/CoreUObject/Private/UObject/LinkerLoad.cpp b/Engine/Source/Runtime/CoreUObject/Private/UObject/LinkerLoad.cpp
--- a/Engine/Source/Runtime/CoreUObject/Private/UObject/LinkerLoad.cpp
+++ b/Engine/Source/Runtime/CoreUObject/Private/UObject/LinkerLoad.cpp
@@ -6728,7 +6728,7 @@ bool FLinkerLoad::SerializeBulkData(FBulkData& BulkData, const FBulkDataSerializ
 {
 	using namespace UE::BulkData::Private;
 
-	if (ShouldSkipBulkData())
+	if (ShouldSkipBulkData() || IsTextFormat())
 	{
 		return false;
 	}
