Added an overload of FTextProperty::Identical_Implementation that lets you control the fallback lexical compare mode

This lets you use it at runtime to do comparison in a way that matches the editor

diff --git a/Engine/Source/Runtime/CoreUObject/Private/UObject/TextProperty.cpp b/Engine/Source/Runtime/CoreUObject/Private/UObject/TextProperty.cpp
--- a/Engine/Source/Runtime/CoreUObject/Private/UObject/TextProperty.cpp
+++ b/Engine/Source/Runtime/CoreUObject/Private/UObject/TextProperty.cpp
@@ -49,6 +49,12 @@ EConvertFromTypeResult FTextProperty::ConvertFromType(const FPropertyTag& Tag, F
 }
 
 bool FTextProperty::Identical_Implementation(const FText& ValueA, const FText& ValueB, uint32 PortFlags)
+{
+	// We compare the display strings in editor (as we author in the native language)
+	return Identical_Implementation(ValueA, ValueB, PortFlags, GIsEditor ? EIdenticalLexicalCompareMethod::DisplayString : EIdenticalLexicalCompareMethod::None);
+}
+
+bool FTextProperty::Identical_Implementation(const FText& ValueA, const FText& ValueB, uint32 PortFlags, EIdenticalLexicalCompareMethod LexicalCompareMethod)
 {
 	// A culture variant text is never equal to a culture invariant text
 	// A transient text is never equal to a non-transient text
@@ -85,13 +91,19 @@ bool FTextProperty::Identical_Implementation(const FText& ValueA, const FText& V
 		return true;
 	}
 
-	// We compare the display strings in editor (as we author in the native language)
+	// We compare the display strings if asked
 	// We compare the display string for culture invariant and transient texts as they don't have an identity
-	if (GIsEditor || ValueA.IsCultureInvariant() || ValueA.IsTransient())
+	if (LexicalCompareMethod == EIdenticalLexicalCompareMethod::DisplayString || ValueA.IsCultureInvariant() || ValueA.IsTransient())
 	{
 		return FTextInspector::GetDisplayString(ValueA).Equals(FTextInspector::GetDisplayString(ValueB), ESearchCase::CaseSensitive);
 	}
 	
+	// We compare the source strings if asked
+	if (LexicalCompareMethod == EIdenticalLexicalCompareMethod::SourceString)
+	{
+		return FTextInspector::GetSourceString(ValueA)->Equals(*FTextInspector::GetSourceString(ValueB), ESearchCase::CaseSensitive);
+	}
+
 	// If we got this far then the texts don't share the same pointer, which means that they can't share the same identity
 	return false;
 }
diff --git a/Engine/Source/Runtime/CoreUObject/Public/UObject/TextProperty.h b/Engine/Source/Runtime/CoreUObject/Public/UObject/TextProperty.h
--- a/Engine/Source/Runtime/CoreUObject/Public/UObject/TextProperty.h
+++ b/Engine/Source/Runtime/CoreUObject/Public/UObject/TextProperty.h
@@ -67,7 +67,14 @@ public:
 	virtual FString GetCPPTypeForwardDeclaration() const override;
 	// End of FProperty interface
 
+	enum class EIdenticalLexicalCompareMethod : uint8
+	{
+		None,
+		SourceString,
+		DisplayString
+	};
 	static bool Identical_Implementation(const FText& A, const FText& B, uint32 PortFlags);
+	static bool Identical_Implementation(const FText& A, const FText& B, uint32 PortFlags, EIdenticalLexicalCompareMethod LexicalCompareMethod);
 };
 
 #if UE_ENABLE_INCLUDE_ORDER_DEPRECATED_IN_5_2
