Replaced TModels traits class with TModels_V.

diff --git a/Engine/Plugins/Experimental/AnimNext/Source/AnimNextInterface/Public/AnimNextInterfaceContext.h b/Engine/Plugins/Experimental/AnimNext/Source/AnimNextInterface/Public/AnimNextInterfaceContext.h
--- a/Engine/Plugins/Experimental/AnimNext/Source/AnimNextInterface/Public/AnimNextInterfaceContext.h
+++ b/Engine/Plugins/Experimental/AnimNext/Source/AnimNextInterface/Public/AnimNextInterfaceContext.h
@@ -172,7 +172,7 @@ public:
 			}
 
 			// deal with containers
-			if constexpr (TModels<Private::CSizedContainerWithAccessibleDataAsRawPtr, ValueType>::Value)
+			if constexpr (TModels_V<Private::CSizedContainerWithAccessibleDataAsRawPtr, ValueType>::Value)
 			{
 				MutableValueType* ParamData = const_cast<MutableValueType>(static_cast<ValueType*>(Value->GetData()));
 				ParamHandle = EnumHasAnyFlags(Flags, FParam::EFlags::Value) 
@@ -299,7 +299,7 @@ public:
 	void SetResult(const ValueType& InValue) const
 	{
 		// Support containers with Num() and GetData()
-		if constexpr (TModels<Private::CSizedContainerWithAccessibleDataAsRawPtr, ValueType>::Value)
+		if constexpr (TModels_V<Private::CSizedContainerWithAccessibleDataAsRawPtr, ValueType>::Value)
 		{
 			checkSlow(Private::CheckParam<ValueType>(*Result));
 			checkSlow(InValue.Num() == Result->NumElements);
diff --git a/Engine/Plugins/Experimental/AnimNext/Source/AnimNextInterface/Public/AnimNextInterfaceParam.h b/Engine/Plugins/Experimental/AnimNext/Source/AnimNextInterface/Public/AnimNextInterfaceParam.h
--- a/Engine/Plugins/Experimental/AnimNext/Source/AnimNextInterface/Public/AnimNextInterfaceParam.h
+++ b/Engine/Plugins/Experimental/AnimNext/Source/AnimNextInterface/Public/AnimNextInterfaceParam.h
@@ -390,7 +390,7 @@ struct TParamTypeHelper
 {
 	FORCEINLINE static const UScriptStruct* GetStruct()
 	{
-		if constexpr (TModels<CHasStaticStruct, ValueType>::Value)
+		if constexpr (TModels_V<CHasStaticStruct, ValueType>)
 		{
 			return TBaseStructure<ValueType>::Get();
 		}
diff --git a/Engine/Plugins/Experimental/Web/WebAPI/Source/WebAPI/Public/WebAPIHttpRequest.h b/Engine/Plugins/Experimental/Web/WebAPI/Source/WebAPI/Public/WebAPIHttpRequest.h
--- a/Engine/Plugins/Experimental/Web/WebAPI/Source/WebAPI/Public/WebAPIHttpRequest.h
+++ b/Engine/Plugins/Experimental/Web/WebAPI/Source/WebAPI/Public/WebAPIHttpRequest.h
@@ -1,4 +1,4 @@
-﻿// Copyright Epic Games, Inc. All Rights Reserved.
+// Copyright Epic Games, Inc. All Rights Reserved.
 
 #pragma once
 
@@ -237,7 +237,7 @@ TRequest<PayloadType>& TRequest<PayloadType>::BindHeaderReceivedCallback(TFuncti
 template <typename PayloadType>
 TRequest<PayloadType>& TRequest<PayloadType>::SetPayloadData(const PayloadType& InPayloadData)
 {
-	static_assert(TModels<CStaticStructProvider, PayloadType>::Value, TEXT("PayloadType must be a UStruct"));
+	static_assert(TModels_V<CStaticStructProvider, PayloadType>, TEXT("PayloadType must be a UStruct"));
 
 	FString PayloadAsString = PayloadToString_Internal(InPayloadData).GetValue();
 	PayloadAsString.Append(TEXT("]")).InsertAt(0, TEXT("["));
@@ -262,7 +262,7 @@ bool TRequest<PayloadType>::HasPayload() const
 template <typename PayloadType>
 TValueOrError<FString, FString> TRequest<PayloadType>::PayloadToString_Internal(const PayloadType& InContent) const
 {
-	static_assert(TModels<CStaticStructProvider, PayloadType>::Value, TEXT("PayloadType must be a UStruct"));
+	static_assert(TModels_V<CStaticStructProvider, PayloadType>, TEXT("PayloadType must be a UStruct"));
 
 	FString StringifiedContent;
 	if (!FJsonObjectConverter::UStructToJsonObjectString(InContent, StringifiedContent))
diff --git a/Engine/Plugins/Experimental/Web/WebAPI/Source/WebAPI/Public/WebAPISubsystem.h b/Engine/Plugins/Experimental/Web/WebAPI/Source/WebAPI/Public/WebAPISubsystem.h
--- a/Engine/Plugins/Experimental/Web/WebAPI/Source/WebAPI/Public/WebAPISubsystem.h
+++ b/Engine/Plugins/Experimental/Web/WebAPI/Source/WebAPI/Public/WebAPISubsystem.h
@@ -113,13 +113,13 @@ private:
 template <typename OperationType>
 TObjectPtr<OperationType> UWebAPISubsystem::MakeOperation(const UWebAPIDeveloperSettings* InSettings)
 {
-	static_assert(TModels<CStaticClassProvider, OperationType>::Value, "OperationType should provide a StaticClass.");
+	static_assert(TModels_V<CStaticClassProvider, OperationType>, "OperationType should provide a StaticClass.");
 	return Cast<OperationType>(MakeOperation(InSettings, OperationType::StaticClass()));
 }
 
 template <typename OperationType>
 void UWebAPISubsystem::ReleaseOperation(const TObjectPtr<OperationType>& InOperation)
 {
-	static_assert(TModels<CStaticClassProvider, OperationType>::Value, "OperationType should provide a StaticClass.");
+	static_assert(TModels_V<CStaticClassProvider, OperationType>, "OperationType should provide a StaticClass.");
 	ReleaseOperation(OperationType::StaticClass(), InOperation);
 }
diff --git a/Engine/Plugins/Experimental/Web/WebAPI/Source/WebAPIEditor/Public/WebAPIJsonUtilities.h b/Engine/Plugins/Experimental/Web/WebAPI/Source/WebAPIEditor/Public/WebAPIJsonUtilities.h
--- a/Engine/Plugins/Experimental/Web/WebAPI/Source/WebAPIEditor/Public/WebAPIJsonUtilities.h
+++ b/Engine/Plugins/Experimental/Web/WebAPI/Source/WebAPIEditor/Public/WebAPIJsonUtilities.h
@@ -67,15 +67,15 @@ namespace UE::Json
 
 		/** Numeric value types. */
 		template <typename ValueType>
-		using TIsNumeric = TAnd<TModels<Concepts::CNumerical, std::decay_t<ValueType>>>;
+		constexpr bool TIsNumeric_V = TModels_V<Concepts::CNumerical, std::decay_t<ValueType>>;
 
 		/** ValueType has FromJson. */
 		template <typename ValueType>
-		using THasFromJson = TModels<Concepts::CFromJsonable, std::decay_t<ValueType>>;
+		constexpr bool THasFromJson_V = TModels_V<Concepts::CFromJsonable, std::decay_t<ValueType>>;
 
 		/** ValueType is derived from TMap. */
 		template <typename ValueType>
-		using TIsDerivedFromMap = TModels<Concepts::CMap, std::decay_t<ValueType>>;
+		constexpr bool TIsDerivedFromMap_V = TModels_V<Concepts::CMap, std::decay_t<ValueType>>;
 	}
 
 	/** Contains either an object constructed in place, or a reference to an object declared elsewhere. */
@@ -262,7 +262,7 @@ namespace UE::Json
 
 	// Numeric
 	template <typename ValueType>
-	constexpr typename TEnableIf<TypeTraits::TIsNumeric<ValueType>::Value, void>::Type
+	constexpr typename TEnableIf<TypeTraits::TIsNumeric_V<ValueType>, void>::Type
 	As(const TSharedPtr<FJsonValue>& InJsonValue, ValueType& OutValue)
 	{
 		OutValue = InJsonValue->AsNumber();
@@ -286,11 +286,11 @@ namespace UE::Json
 
 	// Numeric
 	template <typename ValueType>
-	constexpr typename TEnableIf<TypeTraits::TIsNumeric<ValueType>::Value, bool>::Type
+	constexpr typename TEnableIf<TypeTraits::TIsNumeric_V<ValueType>, bool>::Type
 	TryGet(const TSharedPtr<FJsonValue>& InJsonValue, ValueType& OutValue);
 
 	template <typename ValueType>
-	constexpr typename TEnableIf<TypeTraits::TIsNumeric<ValueType>::Value, bool>::Type
+	constexpr typename TEnableIf<TypeTraits::TIsNumeric_V<ValueType>, bool>::Type
 	TryGetField(const TSharedPtr<FJsonObject>& InJsonObject, const FString& InFieldName, ValueType& OutValue);
 
 	// Enum
@@ -349,16 +349,14 @@ namespace UE::Json
 	// Object (with FromJson)
 	template <typename ValueType>
 	typename TEnableIf<
-		TAnd<
-			TNot<TypeTraits::TIsDerivedFromMap<ValueType>>,
-			TypeTraits::THasFromJson<ValueType>>::Value, bool>::Type
+		!TypeTraits::TIsDerivedFromMap_V<ValueType> &&
+		TypeTraits::THasFromJson_V<ValueType>, bool>::Type
 	TryGet(const TSharedPtr<FJsonValue>& InJsonValue, ValueType& OutValue);
 
 	template <typename ValueType>
 	typename TEnableIf<
-		TAnd<
-			TNot<TypeTraits::TIsDerivedFromMap<ValueType>>,
-			TypeTraits::THasFromJson<ValueType>>::Value, bool>::Type
+		!TypeTraits::TIsDerivedFromMap_V<ValueType> &&
+		TypeTraits::THasFromJson_V<ValueType>, bool>::Type
 	TryGetField(const TSharedPtr<FJsonObject>& InJsonObject, const FString& InFieldName, ValueType& OutValue);
 
 	// Variant
@@ -393,16 +391,14 @@ namespace UE::Json
 	// Object (with FromJson)
 	template <typename ValueType>
 	typename TEnableIf<
-		TAnd<
-			TNot<TypeTraits::TIsDerivedFromMap<ValueType>>,
-			TypeTraits::THasFromJson<ValueType>>::Value, bool>::Type
+		!TypeTraits::TIsDerivedFromMap_V<ValueType> &&
+		TypeTraits::THasFromJson_V<ValueType>, bool>::Type
 	TryGet(const TSharedPtr<FJsonValue>& InJsonValue, ValueType& OutValue);
 
 	template <typename ValueType>
 	typename TEnableIf<
-		TAnd<
-			TNot<TypeTraits::TIsDerivedFromMap<ValueType>>,
-			TypeTraits::THasFromJson<ValueType>>::Value, bool>::Type
+		!TypeTraits::TIsDerivedFromMap_V<ValueType> &&
+		TypeTraits::THasFromJson_V<ValueType>, bool>::Type
 	TryGetField(const TSharedPtr<FJsonObject>& InJsonObject, const FString& InFieldName, ValueType& OutValue);
 
 	// Object (without FromJson)
@@ -414,7 +410,7 @@ namespace UE::Json
 		!TIsTArray_V<ValueType> &&
 		!TypeTraits::TIsStringLike<ValueType>::Value &&
 		!TIsPODType<ValueType>::Value &&
-		!TypeTraits::THasFromJson<ValueType>::Value, bool>::Type
+		!TypeTraits::THasFromJson_V<ValueType>, bool>::Type
 	TryGet(const TSharedPtr<FJsonValue>& InJsonValue, ValueType& OutValue);
 
 	// TIsTMap<ContainerType>,
@@ -426,7 +422,7 @@ namespace UE::Json
 		!TIsTArray_V<ValueType> &&
 		!TypeTraits::TIsStringLike<ValueType>::Value &&
 		!TIsPODType<ValueType>::Value &&
-		!TypeTraits::THasFromJson<ValueType>::Value, bool>::Type
+		!TypeTraits::THasFromJson_V<ValueType>, bool>::Type
 	TryGetField(const TSharedPtr<FJsonObject>& InJsonObject, const FString& InFieldName, ValueType& OutValue);
 
 	// UniqueObj
diff --git a/Engine/Plugins/Experimental/Web/WebAPI/Source/WebAPIEditor/Public/WebAPIJsonUtilities.inl b/Engine/Plugins/Experimental/Web/WebAPI/Source/WebAPIEditor/Public/WebAPIJsonUtilities.inl
--- a/Engine/Plugins/Experimental/Web/WebAPI/Source/WebAPIEditor/Public/WebAPIJsonUtilities.inl
+++ b/Engine/Plugins/Experimental/Web/WebAPI/Source/WebAPIEditor/Public/WebAPIJsonUtilities.inl
@@ -11,14 +11,14 @@ namespace UE::Json
 {
 	// Numeric
 	template <typename ValueType>
-	constexpr typename TEnableIf<TypeTraits::TIsNumeric<ValueType>::Value, bool>::Type
+	constexpr typename TEnableIf<TypeTraits::TIsNumeric_V<ValueType>, bool>::Type
 	TryGet(const TSharedPtr<FJsonValue>& InJsonValue, ValueType& OutValue)
 	{
 		return InJsonValue->TryGetNumber(OutValue);
 	}
 
 	template <typename ValueType>
-	constexpr typename TEnableIf<TypeTraits::TIsNumeric<ValueType>::Value, bool>::Type
+	constexpr typename TEnableIf<TypeTraits::TIsNumeric_V<ValueType>, bool>::Type
 	TryGetField(const TSharedPtr<FJsonObject>& InJsonObject, const FString& InFieldName, ValueType& OutValue)
 	{
 		return InJsonObject->TryGetNumberField(InFieldName, OutValue);
@@ -303,9 +303,8 @@ namespace UE::Json
 	// Object (with FromJson)
 	template <typename ValueType>
 	typename TEnableIf<
-		TAnd<
-			TNot<TypeTraits::TIsDerivedFromMap<ValueType>>,
-			TypeTraits::THasFromJson<ValueType>>::Value, bool>::Type
+		!TypeTraits::TIsDerivedFromMap_V<ValueType> &&
+		TypeTraits::THasFromJson_V<ValueType>, bool>::Type
 	TryGet(const TSharedPtr<FJsonValue>& InJsonValue, ValueType& OutValue)
 	{
 		const TSharedPtr<FJsonObject>* Value;
@@ -318,9 +317,8 @@ namespace UE::Json
 
 	template <typename ValueType>
 	typename TEnableIf<
-		TAnd<
-			TNot<TypeTraits::TIsDerivedFromMap<ValueType>>,
-			TypeTraits::THasFromJson<ValueType>>::Value, bool>::Type
+		!TypeTraits::TIsDerivedFromMap_V<ValueType> &&
+		TypeTraits::THasFromJson_V<ValueType>, bool>::Type
 	TryGetField(const TSharedPtr<FJsonObject>& InJsonObject, const FString& InFieldName, ValueType& OutValue)
 	{
 		if(const TSharedPtr<FJsonValue> JsonField = InJsonObject->TryGetField(InFieldName))
@@ -339,10 +337,10 @@ namespace UE::Json
 		!TIsTArray_V<ValueType> &&
 		!TypeTraits::TIsStringLike<ValueType>::Value &&
 		!TIsPODType<ValueType>::Value &&
-		!TypeTraits::THasFromJson<ValueType>::Value, bool>::Type
+		!TypeTraits::THasFromJson_V<ValueType>, bool>::Type
 	TryGet(const TSharedPtr<FJsonValue>& InJsonValue, ValueType& OutValue)
 	{
-		static_assert(TypeTraits::THasFromJson<ValueType>::Value, "ValueType must implement FromJson");
+		static_assert(TypeTraits::THasFromJson_V<ValueType>, "ValueType must implement FromJson");
 		return false;
 	}
 
@@ -354,10 +352,10 @@ namespace UE::Json
 		!TIsTArray_V<ValueType> &&
 		!TypeTraits::TIsStringLike<ValueType>::Value &&
 		!TIsPODType<ValueType>::Value &&
-		!TypeTraits::THasFromJson<ValueType>::Value, bool>::Type
+		!TypeTraits::THasFromJson_V<ValueType>, bool>::Type
 	TryGetField(const TSharedPtr<FJsonObject>& InJsonObject, const FString& InFieldName, ValueType& OutValue)
 	{
-		static_assert(TypeTraits::THasFromJson<ValueType>::Value, "ValueType must implement FromJson");
+		static_assert(TypeTraits::THasFromJson_V<ValueType>, "ValueType must implement FromJson");
 		return false;
 	}
 
diff --git a/Engine/Plugins/NetcodeUnitTest/NetcodeUnitTest/Source/NetcodeUnitTest/Public/NUTUtilReflection.h b/Engine/Plugins/NetcodeUnitTest/NetcodeUnitTest/Source/NetcodeUnitTest/Public/NUTUtilReflection.h
--- a/Engine/Plugins/NetcodeUnitTest/NetcodeUnitTest/Source/NetcodeUnitTest/Public/NUTUtilReflection.h
+++ b/Engine/Plugins/NetcodeUnitTest/NetcodeUnitTest/Source/NetcodeUnitTest/Public/NUTUtilReflection.h
@@ -474,7 +474,7 @@ public:
 	template<typename T>
 	FORCEINLINE T* GetStructRef()
 	{
-		static_assert(TModels<CStaticStructProvider, T>::Value, "Only USTRUCT types can be passed to GetStructRef");
+		static_assert(TModels_V<CStaticStructProvider, T>, "Only USTRUCT types can be passed to GetStructRef");
 
 		return (T*)(void*)(FVMReflection(*this)[TCHAR_TO_ANSI(*FString::Printf(TEXT("F%s"), *T::StaticStruct()->GetName()))]);
 	}
diff --git a/Engine/Plugins/Online/OnlineServices/Source/OnlineServicesCommon/Private/Online/OnlineSyntaxTests.cpp b/Engine/Plugins/Online/OnlineServices/Source/OnlineServicesCommon/Private/Online/OnlineSyntaxTests.cpp
--- a/Engine/Plugins/Online/OnlineServices/Source/OnlineServicesCommon/Private/Online/OnlineSyntaxTests.cpp
+++ b/Engine/Plugins/Online/OnlineServices/Source/OnlineServicesCommon/Private/Online/OnlineSyntaxTests.cpp
@@ -144,10 +144,10 @@ END_ONLINE_STRUCT_META()
 /* Meta */ }
 
 // confirm that a couple concepts are working as expected
-static_assert(TModels<UE::Online::Meta::CSuperDefined, Test::FTestInterface>::Value);
-static_assert(!TModels<UE::Online::Meta::CSuperDefined, Test::ITestInterface>::Value);
-static_assert(!TModels<UE::Online::CLocalAccountIdDefined, Test::FTestOp::Params>::Value);
-static_assert(TModels<UE::Online::CLocalAccountIdDefined, Test::FTestUserOp::Params>::Value);
+static_assert(TModels_V<UE::Online::Meta::CSuperDefined, Test::FTestInterface>);
+static_assert(!TModels_V<UE::Online::Meta::CSuperDefined, Test::ITestInterface>);
+static_assert(!TModels_V<UE::Online::CLocalAccountIdDefined, Test::FTestOp::Params>);
+static_assert(TModels_V<UE::Online::CLocalAccountIdDefined, Test::FTestUserOp::Params>);
 
 /* UE::Online */ }
 
diff --git a/Engine/Plugins/Online/OnlineServices/Source/OnlineServicesCommon/Public/Online/OnlineAsyncOpCache.h b/Engine/Plugins/Online/OnlineServices/Source/OnlineServicesCommon/Public/Online/OnlineAsyncOpCache.h
--- a/Engine/Plugins/Online/OnlineServices/Source/OnlineServicesCommon/Public/Online/OnlineAsyncOpCache.h
+++ b/Engine/Plugins/Online/OnlineServices/Source/OnlineServicesCommon/Public/Online/OnlineAsyncOpCache.h
@@ -191,7 +191,7 @@ struct TMergeableOpParamsFuncs
 		Meta::VisitFields<typename OpType::Params>([&bResult, &First, &Second](const auto& Field)
 			{
 				constexpr bool bIsMutationFieldType = std::is_same_v<decltype(Field.Pointer), decltype(&OpType::Params::Mutations)>;
-				if constexpr (TModels<Private::COnlineGetTypeHashable, decltype(First.*Field.Pointer)>::Value || !bIsMutationFieldType)
+				if constexpr (TModels_V<Private::COnlineGetTypeHashable, decltype(First.*Field.Pointer)> || !bIsMutationFieldType)
 				{
 					if constexpr (bIsMutationFieldType)
 					{
@@ -214,7 +214,7 @@ struct TMergeableOpParamsFuncs
 		Meta::VisitFields<typename OpType::Params>([&Params, &CombinedHash](const auto& Field)
 			{
 				constexpr bool bIsMutationFieldType = std::is_same_v<decltype(Field.Pointer), decltype(&OpType::Params::Mutations)>;
-				if constexpr (TModels<Private::COnlineGetTypeHashable, decltype(Params.*Field.Pointer)>::Value || !bIsMutationFieldType)
+				if constexpr (TModels_V<Private::COnlineGetTypeHashable, decltype(Params.*Field.Pointer)> || !bIsMutationFieldType)
 				{
 					if constexpr (bIsMutationFieldType)
 					{
@@ -249,7 +249,7 @@ public:
 	{
 		TSharedRef<TOnlineAsyncOp<OpType>> Op = MakeShared<TOnlineAsyncOp<OpType>>(Services, MoveTemp(Params));
 
-		if constexpr (TModels<CLocalAccountIdDefined, typename OpType::Params>::Value)
+		if constexpr (TModels_V<CLocalAccountIdDefined, typename OpType::Params>)
 		{
 			// add LocalAccountId to the Op Data
 			Op->Data.template Set<decltype(Params.LocalAccountId)>(TEXT("LocalAccountId"), Op->GetParams().LocalAccountId);
@@ -265,7 +265,7 @@ public:
 		TSharedPtr<TOnlineAsyncOp<OpType>> Op;
 
 		// find existing op
-		if constexpr (TModels<CLocalAccountIdDefined, typename OpType::Params>::Value)
+		if constexpr (TModels_V<CLocalAccountIdDefined, typename OpType::Params>)
 		{
 			if (TUniquePtr<IWrappedOperation>* OpPtr = UserOperations.FindOrAdd(Params.LocalAccountId).Find(FWrappedOperationKey::Create<OpType, ParamsFuncsType>(Params)))
 			{
@@ -298,7 +298,7 @@ public:
 					// remove from cache if operation was canceled
 					if (ThisOp.GetState() == EAsyncOpState::Cancelled)
 					{
-						if constexpr (TModels<CLocalAccountIdDefined, typename OpType::Params>::Value)
+						if constexpr (TModels_V<CLocalAccountIdDefined, typename OpType::Params>)
 						{
 							UserOperations.FindOrAdd(ThisOp.GetParams().LocalAccountId).Remove(FWrappedOperationKey::Create<OpType, ParamsFuncsType>(ThisOp.GetParams()));
 						}
@@ -320,7 +320,7 @@ public:
 		TSharedPtr<TOnlineAsyncOp<OpType>> Op;
 
 		// find existing op
-		if constexpr (TModels<CLocalAccountIdDefined, typename OpType::Params>::Value)
+		if constexpr (TModels_V<CLocalAccountIdDefined, typename OpType::Params>)
 		{
 			if (TUniquePtr<IWrappedOperation>* OpPtr = UserOperations.FindOrAdd(Params.LocalAccountId).Find(FWrappedOperationKey::Create<OpType, ParamsFuncsType>(Params)))
 			{
@@ -349,7 +349,7 @@ public:
 			// remove from cache once operation has started. It is no longer mergeable at that point
 			FOnlineEventDelegateHandle Handle = Op->OnStart().Add(GetSharedThis(), [this](const TOnlineAsyncOp<OpType>& ThisOp)
 				{
-					if constexpr (TModels<CLocalAccountIdDefined, typename OpType::Params>::Value)
+					if constexpr (TModels_V<CLocalAccountIdDefined, typename OpType::Params>)
 					{
 						UserOperations.FindOrAdd(ThisOp.GetParams().LocalAccountId).Remove(FWrappedOperationKey::Create<OpType, ParamsFuncsType>(ThisOp.GetParams()));
 					}
@@ -382,7 +382,7 @@ private:
 
 		TSharedRef<TOnlineAsyncOp<OpType>> Op = WrappedOp->GetDataRef();
 
-		if constexpr (TModels<CLocalAccountIdDefined, typename OpType::Params>::Value)
+		if constexpr (TModels_V<CLocalAccountIdDefined, typename OpType::Params>)
 		{
 			// add LocalAccountId to the Op Data
 			Op->Data.template Set<decltype(Params.LocalAccountId)>(TEXT("LocalAccountId"), Op->GetParams().LocalAccountId);
diff --git a/Engine/Plugins/Online/OnlineServices/Source/OnlineServicesCommon/Public/Online/OnlineConfig.h b/Engine/Plugins/Online/OnlineServices/Source/OnlineServicesCommon/Public/Online/OnlineConfig.h
--- a/Engine/Plugins/Online/OnlineServices/Source/OnlineServicesCommon/Public/Online/OnlineConfig.h
+++ b/Engine/Plugins/Online/OnlineServices/Source/OnlineServicesCommon/Public/Online/OnlineConfig.h
@@ -86,7 +86,7 @@ public:
 	template <typename T>
 	bool GetValue(const TCHAR* Key, T& Value)
 	{
-		if constexpr (TModels<Meta::COnlineMetadataAvailable, T>::Value)
+		if constexpr (TModels_V<Meta::COnlineMetadataAvailable, T>)
 		{
 			IOnlineConfigStructPtr ConfigStructPtr;
 			if (GetValue(Key, ConfigStructPtr))
@@ -118,7 +118,7 @@ public:
 	{
 		Value.Empty();
 
-		if constexpr (TModels<Meta::COnlineMetadataAvailable, T>::Value)
+		if constexpr (TModels_V<Meta::COnlineMetadataAvailable, T>)
 		{
 			TArray<IOnlineConfigStructPtr> ConfigStructPtrArray;
 			if (GetValue(Key, ConfigStructPtrArray) > 0)
@@ -206,7 +206,7 @@ public:
 	template <typename T>
 	bool GetValue(const TCHAR* Section, const TCHAR* Key, T& Value)
 	{
-		if constexpr (TModels<Meta::COnlineMetadataAvailable, T>::Value)
+		if constexpr (TModels_V<Meta::COnlineMetadataAvailable, T>)
 		{
 			IOnlineConfigStructPtr ConfigStructPtr;
 			if (GetValue(Section, Key, ConfigStructPtr))
@@ -237,7 +237,7 @@ public:
 	int32 GetValue(const TCHAR* Section, const TCHAR* Key, TArray<T>& Value)
 	{
 		int32 NumInitialElements = Value.Num();
-		if constexpr (TModels<Meta::COnlineMetadataAvailable, T>::Value)
+		if constexpr (TModels_V<Meta::COnlineMetadataAvailable, T>)
 		{
 			TArray<IOnlineConfigStructPtr> ConfigStructPtrArray;
 			if (GetValue(Section, Key, ConfigStructPtrArray) > 0)
@@ -363,7 +363,7 @@ bool LoadConfigValue(IOnlineConfigProvider& Provider, const TCHAR* Section, cons
  */
 template <typename T>
 auto LoadConfig(IOnlineConfigProvider& Provider, const FString& Section, T& Value)
-	-> std::enable_if_t<TModels<Meta::COnlineMetadataAvailable, T>::Value, bool>
+	-> std::enable_if_t<TModels_V<Meta::COnlineMetadataAvailable, T>, bool>
 {
 	bool bLoadedValue = false;
 	Meta::VisitFields(Value,
@@ -386,7 +386,7 @@ auto LoadConfig(IOnlineConfigProvider& Provider, const FString& Section, T& Valu
  */
 template <typename T>
 auto LoadConfig(IOnlineConfigProvider& Provider, const TArray<FString>& SectionHeirarchy, T& Value)
-	-> std::enable_if_t<TModels<Meta::COnlineMetadataAvailable, T>::Value, bool>
+	-> std::enable_if_t<TModels_V<Meta::COnlineMetadataAvailable, T>, bool>
 {
 	bool bLoadedValue = false;
 	for (const FString& Section : SectionHeirarchy)
@@ -409,7 +409,7 @@ auto LoadConfig(IOnlineConfigProvider& Provider, const TArray<FString>& SectionH
  */
 template <typename T>
 auto LoadConfig(IOnlineConfigProvider& Provider, const FString& Section, const TCHAR* Key, T& OutValue)
--> std::enable_if_t<TModels<Meta::COnlineMetadataAvailable, T>::Value, bool>
+-> std::enable_if_t<TModels_V<Meta::COnlineMetadataAvailable, T>, bool>
 {
 	bool bLoadedValue = false;
 	IOnlineConfigStructPtr ConfigStructPtr;
@@ -442,7 +442,7 @@ auto LoadConfig(IOnlineConfigProvider& Provider, const FString& Section, const T
  */
 template <typename T>
 auto LoadConfig(IOnlineConfigProvider& Provider, const TArray<FString>& SectionHeirarchy, const TCHAR* Key, T& Value)
--> std::enable_if_t<TModels<Meta::COnlineMetadataAvailable, T>::Value, bool>
+-> std::enable_if_t<TModels_V<Meta::COnlineMetadataAvailable, T>, bool>
 {
 	bool bLoadedValue = false;
 	for (const FString& Section : SectionHeirarchy)
diff --git a/Engine/Plugins/Online/OnlineServices/Source/OnlineServicesCommon/Public/Online/OnlineExecHandler.h b/Engine/Plugins/Online/OnlineServices/Source/OnlineServicesCommon/Public/Online/OnlineExecHandler.h
--- a/Engine/Plugins/Online/OnlineServices/Source/OnlineServicesCommon/Public/Online/OnlineExecHandler.h
+++ b/Engine/Plugins/Online/OnlineServices/Source/OnlineServicesCommon/Public/Online/OnlineExecHandler.h
@@ -122,8 +122,8 @@ template <typename T> inline bool ParseOnlineExecParams(const TCHAR*& Cmd, TShar
 template <typename T> inline bool ParseOnlineExecParams(const TCHAR*& Cmd, TOptional<T>& Optional, IOnlineServices* Services = nullptr);
 template <typename... Ts> inline bool ParseOnlineExecParams(const TCHAR*& Cmd, TVariant<Ts...>& Variant, IOnlineServices* Services = nullptr);
 template <typename IdType> inline bool ParseOnlineExecParams(const TCHAR*& Cmd, TOnlineId<IdType>& Value, IOnlineServices* Services = nullptr);
-template <typename T> std::enable_if_t<!TModels<Meta::COnlineMetadataAvailable, T>::Value, bool> ParseOnlineExecParams(const TCHAR*& Cmd, T& Value, IOnlineServices* Services = nullptr);
-template <typename T> std::enable_if_t<TModels<Meta::COnlineMetadataAvailable, T>::Value, bool> ParseOnlineExecParams(const TCHAR*& Cmd, T& Value, IOnlineServices* Services = nullptr);
+template <typename T> std::enable_if_t<!TModels_V<Meta::COnlineMetadataAvailable, T>, bool> ParseOnlineExecParams(const TCHAR*& Cmd, T& Value, IOnlineServices* Services = nullptr);
+template <typename T> std::enable_if_t<TModels_V<Meta::COnlineMetadataAvailable, T>, bool> ParseOnlineExecParams(const TCHAR*& Cmd, T& Value, IOnlineServices* Services = nullptr);
 
 template<typename T, typename VariantObject>
 struct TOnlineVariantVisitInfo
@@ -829,7 +829,7 @@ inline bool ParseOnlineExecParams<OnlineIdHandleTags::FAccount>(const TCHAR*& Cm
 }
 
 template <typename T>
-std::enable_if_t<!TModels<Meta::COnlineMetadataAvailable, T>::Value, bool> ParseOnlineExecParams(const TCHAR*& Cmd, T& Value, IOnlineServices* Services)
+std::enable_if_t<!TModels_V<Meta::COnlineMetadataAvailable, T>, bool> ParseOnlineExecParams(const TCHAR*& Cmd, T& Value, IOnlineServices* Services)
 {
 	FString Token;
 	if (FParse::Token(Cmd, Token, true))
@@ -845,7 +845,7 @@ std::enable_if_t<!TModels<Meta::COnlineMetadataAvailable, T>::Value, bool> Parse
 }
 
 template <typename T>
-std::enable_if_t<TModels<Meta::COnlineMetadataAvailable, T>::Value, bool> ParseOnlineExecParams(const TCHAR*& Cmd, T& Value, IOnlineServices* Services)
+std::enable_if_t<TModels_V<Meta::COnlineMetadataAvailable, T>, bool> ParseOnlineExecParams(const TCHAR*& Cmd, T& Value, IOnlineServices* Services)
 {
 	bool bSuccess = true;
 
diff --git a/Engine/Plugins/Online/OnlineServices/Source/OnlineServicesCommon/Public/Online/OnlineTypeInfo.h b/Engine/Plugins/Online/OnlineServices/Source/OnlineServicesCommon/Public/Online/OnlineTypeInfo.h
--- a/Engine/Plugins/Online/OnlineServices/Source/OnlineServicesCommon/Public/Online/OnlineTypeInfo.h
+++ b/Engine/Plugins/Online/OnlineServices/Source/OnlineServicesCommon/Public/Online/OnlineTypeInfo.h
@@ -46,7 +46,7 @@ struct TOnlineTypeInfo
 			return true;
 		}
 
-		if constexpr (TModels<Meta::CSuperDefined, T>::Value)
+		if constexpr (TModels_V<Meta::CSuperDefined, T>)
 		{
 			return TOnlineTypeInfo<typename T::Super>::IsA(TypeName);
 		}
diff --git a/Engine/Plugins/Online/OnlineServices/Source/OnlineServicesInterface/Public/Online/OnlineAsyncOpHandle.h b/Engine/Plugins/Online/OnlineServices/Source/OnlineServicesInterface/Public/Online/OnlineAsyncOpHandle.h
--- a/Engine/Plugins/Online/OnlineServices/Source/OnlineServicesInterface/Public/Online/OnlineAsyncOpHandle.h
+++ b/Engine/Plugins/Online/OnlineServices/Source/OnlineServicesInterface/Public/Online/OnlineAsyncOpHandle.h
@@ -74,7 +74,7 @@ namespace Private
 	};
 
 	template <typename T>
-	struct TSharedPtrHelper<T, std::enable_if_t<TModels<CAsSharedCallable, T>::Value && !TIsSharedPtrOrRef<T>::value, void>> : public TSharedPtrHelper2<T>
+	struct TSharedPtrHelper<T, std::enable_if_t<TModels_V<CAsSharedCallable, T> && !TIsSharedPtrOrRef<T>::value, void>> : public TSharedPtrHelper2<T>
 	{
 	};
 
diff --git a/Engine/Plugins/Online/OnlineServices/Source/OnlineServicesInterface/Public/Online/OnlineMeta.h b/Engine/Plugins/Online/OnlineServices/Source/OnlineServicesInterface/Public/Online/OnlineMeta.h
--- a/Engine/Plugins/Online/OnlineServices/Source/OnlineServicesInterface/Public/Online/OnlineMeta.h
+++ b/Engine/Plugins/Online/OnlineServices/Source/OnlineServicesInterface/Public/Online/OnlineMeta.h
@@ -17,7 +17,7 @@ struct CSuperDefined
 
 namespace Private {
 
-template <typename T, bool THasSuper = TModels<CSuperDefined, T>::Value>
+template <typename T, bool THasSuper = TModels_V<CSuperDefined, T>>
 struct TBaseClassHelper
 {
 };
diff --git a/Engine/Plugins/Online/OnlineServices/Source/OnlineServicesInterface/Public/Online/OnlineUtils.h b/Engine/Plugins/Online/OnlineServices/Source/OnlineServicesInterface/Public/Online/OnlineUtils.h
--- a/Engine/Plugins/Online/OnlineServices/Source/OnlineServicesInterface/Public/Online/OnlineUtils.h
+++ b/Engine/Plugins/Online/OnlineServices/Source/OnlineServicesInterface/Public/Online/OnlineUtils.h
@@ -34,8 +34,8 @@ inline FString ToLogString(int64 Value);
 inline FString ToLogString(bool Value);
 inline FString ToLogString(float Value);
 inline FString ToLogString(double Value);
-template <typename T> std::enable_if_t<!TModels<Meta::COnlineMetadataAvailable, T>::Value, FString> ToLogString(const T& Value);
-template <typename T> std::enable_if_t<TModels<Meta::COnlineMetadataAvailable, T>::Value, FString> ToLogString(const T& Value);
+template <typename T> std::enable_if_t<!TModels_V<Meta::COnlineMetadataAvailable, T>, FString> ToLogString(const T& Value);
+template <typename T> std::enable_if_t<TModels_V<Meta::COnlineMetadataAvailable, T>, FString> ToLogString(const T& Value);
 
 template <typename T>
 FString ToLogString(const TArray<T>& Array)
@@ -184,13 +184,13 @@ inline FString ToLogString(FPlatformUserId PlatformUserId)
 }
 
 template<typename T>
-std::enable_if_t<!TModels<Meta::COnlineMetadataAvailable, T>::Value, FString> ToLogString(const T& Value)
+std::enable_if_t<!TModels_V<Meta::COnlineMetadataAvailable, T>, FString> ToLogString(const T& Value)
 {
 	return LexToString(Value);
 }
 
 template<typename T>
-std::enable_if_t<TModels<Meta::COnlineMetadataAvailable, T>::Value, FString> ToLogString(const T& Value)
+std::enable_if_t<TModels_V<Meta::COnlineMetadataAvailable, T>, FString> ToLogString(const T& Value)
 {
 	FString LogString;
 	LogString += TEXT("{ ");
diff --git a/Engine/Plugins/Runtime/NetworkPrediction/Source/NetworkPrediction/Public/NetworkPredictionDriver.h b/Engine/Plugins/Runtime/NetworkPrediction/Source/NetworkPrediction/Public/NetworkPredictionDriver.h
--- a/Engine/Plugins/Runtime/NetworkPrediction/Source/NetworkPrediction/Public/NetworkPredictionDriver.h
+++ b/Engine/Plugins/Runtime/NetworkPrediction/Source/NetworkPrediction/Public/NetworkPredictionDriver.h
@@ -274,7 +274,7 @@ struct FNetworkPredictionDriverBase
 		auto Requires(InDriverType* Driver, SyncType* Sync, AuxType* Aux) -> decltype(Driver->InitializeSimulationState(Sync, Aux));
 	};
 
-	static constexpr bool HasInitializeSimulationState = TModels<CInitializeSimulationStateFuncable, DriverType, SyncType, AuxType>::Value;
+	static constexpr bool HasInitializeSimulationState = TModels_V<CInitializeSimulationStateFuncable, DriverType, SyncType, AuxType>;
 	
 	template<bool HasFunc=HasInitializeSimulationState>
 	static typename TEnableIf<HasFunc>::Type CallInitializeSimulationStateMemberFunc(DriverType* Driver, SyncType* Sync, AuxType* Aux)
@@ -308,7 +308,7 @@ struct FNetworkPredictionDriverBase
 		auto Requires(InDriverType* Driver, int32 TimeMS, InputType* InputCmd) -> decltype(Driver->ProduceInput(TimeMS, InputCmd));
 	};
 
-	static constexpr bool HasProduceInput = TModels<CProduceInputMemberFuncable, DriverType, int32, InputType>::Value;
+	static constexpr bool HasProduceInput = TModels_V<CProduceInputMemberFuncable, DriverType, int32, InputType>;
 
 	template<bool HasFunc=HasProduceInput>
 	static typename TEnableIf<HasFunc>::Type CallProduceInputMemberFunc(DriverType* Driver, int32 DeltaTimeMS, InputType* InputCmd)
@@ -339,7 +339,7 @@ struct FNetworkPredictionDriverBase
 		auto Requires(InDriverType* Driver, const SyncType* S, const AuxType* A) -> decltype(Driver->FinalizeFrame(S, A));
 	};
 
-	static constexpr bool HasFinalizeFrame = TModels<CFinalizeFrameMemberFuncable, DriverType, SyncType, AuxType>::Value;
+	static constexpr bool HasFinalizeFrame = TModels_V<CFinalizeFrameMemberFuncable, DriverType, SyncType, AuxType>;
 
 	template<bool HasFunc=HasFinalizeFrame>
 	static typename TEnableIf<HasFunc>::Type CallFinalizeFrameMemberFunc(DriverType* Driver, const SyncType* SyncState, const AuxType* AuxState)
@@ -379,7 +379,7 @@ struct FNetworkPredictionDriverBase
 		auto Requires(InDriverType* Driver, const SyncType* S, const AuxType* A) -> decltype(Driver->RestoreFrame(S, A));
 	};
 
-	static constexpr bool HasRestoreFrame = TModels<CRestoreFrameMemberFuncable, DriverType, SyncType, AuxType>::Value;
+	static constexpr bool HasRestoreFrame = TModels_V<CRestoreFrameMemberFuncable, DriverType, SyncType, AuxType>;
 
 	template<bool HasFunc=HasRestoreFrame>
 	static typename TEnableIf<HasFunc>::Type CallRestoreFrameMemberFunc(DriverType* Driver, const SyncType* SyncState, const AuxType* AuxState)
@@ -415,7 +415,7 @@ struct FNetworkPredictionDriverBase
 		auto Requires(InDriverType* Driver) -> decltype(Driver->CallServerRPC());
 	};
 
-	static constexpr bool HasCallServerRPC = TModels<CCallServerRPCMemberFuncable, DriverType>::Value;
+	static constexpr bool HasCallServerRPC = TModels_V<CCallServerRPCMemberFuncable, DriverType>;
 
 	template<bool HasFunc=HasCallServerRPC>
 	static typename TEnableIf<HasFunc>::Type CallServerRPCMemberFunc(DriverType* Driver)
@@ -504,7 +504,7 @@ struct FNetworkPredictionDriverBase
 		auto Requires(InDriverType* Driver, bool bHidden) -> decltype(Driver->SetHiddenForInterpolation(bHidden));
 	};
 	
-	static constexpr bool HasSetHiddenForInterpolation = TModels<CSetHiddenForInterpolationFuncable, DriverType>::Value;
+	static constexpr bool HasSetHiddenForInterpolation = TModels_V<CSetHiddenForInterpolationFuncable, DriverType>;
 
 	template<bool HasFunc=HasSetHiddenForInterpolation>
 	static typename TEnableIf<HasFunc>::Type CallSetHiddenForInterpolation(DriverType* Driver, bool bHide)
@@ -550,7 +550,7 @@ struct FNetworkPredictionDriverBase
 		auto Requires(InDriverType* Driver) -> decltype(Driver->GetPhysicsPrimitiveComponent());
 	};
 
-	static constexpr bool HasGetPrimitiveComp = TModels<CGetPhysicsPrimitiveComponentMemberFuncable, DriverType>::Value;
+	static constexpr bool HasGetPrimitiveComp = TModels_V<CGetPhysicsPrimitiveComponentMemberFuncable, DriverType>;
 	
 	template<bool HasFunc=HasGetPrimitiveComp>
 	static typename TEnableIf<HasFunc, UPrimitiveComponent*>::Type GetPhysicsPrimitiveComponent(DriverType* Driver)
@@ -592,7 +592,7 @@ struct FNetworkPredictionDriverBase
 		auto Requires(InDriverType* Driver) -> decltype(Driver->GetPhysicsBodyInstance());
 	};
 
-	static constexpr bool HasGetBodyInstance = TModels<CGetPhysicsBodyInstanceMemberFuncable, DriverType>::Value;
+	static constexpr bool HasGetBodyInstance = TModels_V<CGetPhysicsBodyInstanceMemberFuncable, DriverType>;
 
 	template<bool HasFunc=HasGetBodyInstance>
 	static typename TEnableIf<HasFunc, FBodyInstance*>::Type GetPhysicsBodyInstance(DriverType* Driver)
diff --git a/Engine/Plugins/Runtime/ReplicationSystemTestPlugin/Source/Private/Tests/ReplicationState/TestFastArrayReplicationState.cpp b/Engine/Plugins/Runtime/ReplicationSystemTestPlugin/Source/Private/Tests/ReplicationState/TestFastArrayReplicationState.cpp
--- a/Engine/Plugins/Runtime/ReplicationSystemTestPlugin/Source/Private/Tests/ReplicationState/TestFastArrayReplicationState.cpp
+++ b/Engine/Plugins/Runtime/ReplicationSystemTestPlugin/Source/Private/Tests/ReplicationState/TestFastArrayReplicationState.cpp
@@ -775,7 +775,7 @@ UE_NET_TEST_FIXTURE(FReplicationSystemServerClientTestFixture, TestChangeMaskWra
 	UE_NET_ASSERT_EQ(ServerObject->FastArray.GetItemArray()[1].bRepBool, ClientObject->FastArray.GetItemArray()[1].bRepBool);
 }
 
-static_assert(TModels<FFastArraySerializer::CPostReplicatedReceiveFuncable, FTestFastArrayReplicationState_FastArraySerializer, const FFastArraySerializer::FPostReplicatedReceiveParameters&>::Value, "FTestFastArrayReplicationState_FastArraySerializer should have a function matching the requirements of FFastArraySerializer::CPostReplicatedReceiveFuncable");
+static_assert(TModels_V<FFastArraySerializer::CPostReplicatedReceiveFuncable, FTestFastArrayReplicationState_FastArraySerializer, const FFastArraySerializer::FPostReplicatedReceiveParameters&>, "FTestFastArrayReplicationState_FastArraySerializer should have a function matching the requirements of FFastArraySerializer::CPostReplicatedReceiveFuncable");
 
 // Test partial out of order resolve of references in  fast array
 UE_NET_TEST_FIXTURE(FReplicationSystemServerClientTestFixture, TestPostReplicatedReceiveCallback)
diff --git a/Engine/Plugins/VirtualProduction/RemoteControl/Source/RemoteControlCommon/Public/RCTypeTraits.h b/Engine/Plugins/VirtualProduction/RemoteControl/Source/RemoteControlCommon/Public/RCTypeTraits.h
--- a/Engine/Plugins/VirtualProduction/RemoteControl/Source/RemoteControlCommon/Public/RCTypeTraits.h
+++ b/Engine/Plugins/VirtualProduction/RemoteControl/Source/RemoteControlCommon/Public/RCTypeTraits.h
@@ -83,7 +83,7 @@ namespace RemoteControlTypeTraits
 
 	/** Ensures ValueType is a numeric type. */
 	template <typename ValueType>
-	using TNumericValueConstraint = TAnd<TModels<RemoteControlTypeTraits::Concepts::CNumerical, ValueType>>;
+	constexpr bool TNumericValueConstraint_V = TModels_V<RemoteControlTypeTraits::Concepts::CNumerical, ValueType>;
 }
 
 /** Various RemoteControl type traits */
@@ -100,10 +100,9 @@ struct TRemoteControlPropertyTypeTraits;
 template <typename ValueType>
 struct TRemoteControlTypeTraits<ValueType,
 	typename TEnableIf<
-		TAnd<
-			TIsArithmetic<ValueType>,
-			TNot<TIsFloatingPoint<ValueType>>,
-			TModels<RemoteControlTypeTraits::Concepts::CNumerical, ValueType>>::Value, void>::Type>
+		std::is_arithmetic_v<ValueType> &&
+		!std::is_floating_point_v<ValueType> &&
+		TModels_V<RemoteControlTypeTraits::Concepts::CNumerical, ValueType>, void>::Type>
 {
 	using Type = ValueType;
 
@@ -168,7 +167,7 @@ struct TRemoteControlPropertyTypeTraits<FNumericProperty>
 /** RemoteControlPropertyTypeTraits for any property with a TCppType  */
 template <typename InPropertyType>
 struct TRemoteControlPropertyTypeTraits<InPropertyType,
-	typename TEnableIf<TModels<RemoteControlTypeTraits::Concepts::CPropertyHasTCppType, InPropertyType>::Value, void>::Type>
+	typename TEnableIf<TModels_V<RemoteControlTypeTraits::Concepts::CPropertyHasTCppType, InPropertyType>, void>::Type>
 {
 	using PropertyType = InPropertyType;
 	using ValueType = typename InPropertyType::TCppType;
diff --git a/Engine/Plugins/VirtualProduction/RemoteControl/Source/RemoteControlCommon/Public/RCTypeUtilities.h b/Engine/Plugins/VirtualProduction/RemoteControl/Source/RemoteControlCommon/Public/RCTypeUtilities.h
--- a/Engine/Plugins/VirtualProduction/RemoteControl/Source/RemoteControlCommon/Public/RCTypeUtilities.h
+++ b/Engine/Plugins/VirtualProduction/RemoteControl/Source/RemoteControlCommon/Public/RCTypeUtilities.h
@@ -159,9 +159,8 @@ namespace RemoteControlTypeUtilities
 	/** Returns typed metadata for numeric property if it exists, otherwise returns the DefaultValue. */
 	template <typename PropertyType, typename ValueType>
 	constexpr static typename TEnableIf<
-		TAnd<
-			TIsDerivedFrom<PropertyType, FProperty>,
-			RemoteControlTypeTraits::TNumericValueConstraint<ValueType>>::Value, ValueType>::Type
+		std::is_base_of_v<FProperty, PropertyType> &&
+		RemoteControlTypeTraits::TNumericValueConstraint_V<ValueType>, ValueType>::Type
 	GetMetadataValue(const PropertyType* InProperty, const FName& InKey, const ValueType& InDefaultValue)
 	{
 #if WITH_EDITORONLY_DATA
@@ -188,9 +187,8 @@ namespace RemoteControlTypeUtilities
 	/** Returns typed metadata if it exists, otherwise returns the DefaultValue. */
 	template <typename PropertyType, typename ValueType>
 	constexpr static typename TEnableIf<
-		TAnd<
-			TIsDerivedFrom<PropertyType, FProperty>,
-			TNot<RemoteControlTypeTraits::TNumericValueConstraint<ValueType>>>::Value, ValueType>::Type
+		std::is_base_of_v<FProperty, PropertyType> &&
+		!RemoteControlTypeTraits::TNumericValueConstraint_V<ValueType>, ValueType>::Type
 	GetMetadataValue(const PropertyType* InProperty, const FName& InKey, const ValueType InDefaultValue)
 	{
 		return InDefaultValue;
@@ -349,7 +347,7 @@ namespace RemoteControlTypeUtilities
 	}
 
 	template <typename ValueType>
-	static typename TEnableIf<TNot<RemoteControlTypeTraits::TNumericValueConstraint<ValueType>>::Value, ValueType>::Type
+	static typename TEnableIf<!RemoteControlTypeTraits::TNumericValueConstraint_V<ValueType>, ValueType>::Type
 	GetDefaultRangeValueMin(const FProperty* InProperty)
 	{
 		check(InProperty);
@@ -358,7 +356,7 @@ namespace RemoteControlTypeUtilities
 	}
 
 	template <typename ValueType>
-	static typename TEnableIf<RemoteControlTypeTraits::TNumericValueConstraint<ValueType>::Value, ValueType>::Type
+	static typename TEnableIf<RemoteControlTypeTraits::TNumericValueConstraint_V<ValueType>, ValueType>::Type
 	GetDefaultRangeValueMin(const FProperty* InProperty)
 	{
 		check(InProperty);
@@ -382,7 +380,7 @@ namespace RemoteControlTypeUtilities
 	}
 
 	template <typename ValueType>
-	static typename TEnableIf<TNot<RemoteControlTypeTraits::TNumericValueConstraint<ValueType>>::Value, ValueType>::Type
+	static typename TEnableIf<!RemoteControlTypeTraits::TNumericValueConstraint_V<ValueType>, ValueType>::Type
 	GetDefaultRangeValueMax(const FProperty* InProperty)
 	{
 		check(InProperty);
@@ -391,7 +389,7 @@ namespace RemoteControlTypeUtilities
 	}
 
 	template <typename ValueType>
-	static typename TEnableIf<RemoteControlTypeTraits::TNumericValueConstraint<ValueType>::Value, ValueType>::Type
+	static typename TEnableIf<RemoteControlTypeTraits::TNumericValueConstraint_V<ValueType>::Value, ValueType>::Type
 	GetDefaultRangeValueMax(const FProperty* InProperty)
 	{
 		check(InProperty);
@@ -415,7 +413,7 @@ namespace RemoteControlTypeUtilities
 	}
 
 	template <typename ValueType>
-	static typename TEnableIf<TNot<RemoteControlTypeTraits::TNumericValueConstraint<ValueType>>::Value, ValueType>::Type
+	static typename TEnableIf<!RemoteControlTypeTraits::TNumericValueConstraint_V<ValueType>, ValueType>::Type
 	GetDefaultMappingValueMin(const FProperty* InProperty)
 	{
 		check(InProperty);
@@ -424,7 +422,7 @@ namespace RemoteControlTypeUtilities
 	}
 
 	template <typename ValueType>
-	static typename TEnableIf<RemoteControlTypeTraits::TNumericValueConstraint<ValueType>::Value, ValueType>::Type
+	static typename TEnableIf<RemoteControlTypeTraits::TNumericValueConstraint_V<ValueType>, ValueType>::Type
 	GetDefaultMappingValueMin(const FProperty* InProperty)
 	{
 		check(InProperty);
@@ -467,7 +465,7 @@ namespace RemoteControlTypeUtilities
 	}
 
 	template <typename ValueType>
-	static typename TEnableIf<TNot<RemoteControlTypeTraits::TNumericValueConstraint<ValueType>>::Value, ValueType>::Type
+	static typename TEnableIf<!RemoteControlTypeTraits::TNumericValueConstraint_V<ValueType>, ValueType>::Type
 	GetDefaultMappingValueMax(const FProperty* InProperty)
 	{
 		check(InProperty);
@@ -476,7 +474,7 @@ namespace RemoteControlTypeUtilities
 	}
 
 	template <typename ValueType>
-	static typename TEnableIf<RemoteControlTypeTraits::TNumericValueConstraint<ValueType>::Value, ValueType>::Type
+	static typename TEnableIf<RemoteControlTypeTraits::TNumericValueConstraint_V<ValueType>, ValueType>::Type
 	GetDefaultMappingValueMax(const FProperty* InProperty)
 	{
 		check(InProperty);
diff --git a/Engine/Source/Developer/RigVMDeveloper/Public/RigVMModel/Nodes/RigVMUnitNode.h b/Engine/Source/Developer/RigVMDeveloper/Public/RigVMModel/Nodes/RigVMUnitNode.h
--- a/Engine/Source/Developer/RigVMDeveloper/Public/RigVMModel/Nodes/RigVMUnitNode.h
+++ b/Engine/Source/Developer/RigVMDeveloper/Public/RigVMModel/Nodes/RigVMUnitNode.h
@@ -58,7 +58,7 @@ public:
 	// Returns a copy of the struct with the current values
 	template <
 		typename T,
-		typename TEnableIf<TModels<CRigVMUStruct, T>::Value>::Type* = nullptr
+		typename TEnableIf<TModels_V<CRigVMUStruct, T>>::Type* = nullptr
 	>
 	T ConstructStructInstance() const
 	{
diff --git a/Engine/Source/Developer/RigVMDeveloper/Public/RigVMModel/RigVMController.h b/Engine/Source/Developer/RigVMDeveloper/Public/RigVMModel/RigVMController.h
--- a/Engine/Source/Developer/RigVMDeveloper/Public/RigVMModel/RigVMController.h
+++ b/Engine/Source/Developer/RigVMDeveloper/Public/RigVMModel/RigVMController.h
@@ -103,7 +103,7 @@ public:
 
 	template <
 		typename T,
-		typename TEnableIf<TModels<CRigVMUStruct, T>::Value>::Type * = nullptr
+		typename TEnableIf<TModels_V<CRigVMUStruct, T>>::Type * = nullptr
 	>
 	FRigStructScope(const T& InInstance)
 		: ScriptStruct(T::StaticStruct())
@@ -261,7 +261,7 @@ public:
 	// Adds a unit node using a template
 	template <
 		typename T,
-		typename TEnableIf<TModels<CRigVMUStruct, T>::Value>::Type * = nullptr
+		typename TEnableIf<TModels_V<CRigVMUStruct, T>>::Type * = nullptr
 	>
 	URigVMUnitNode* AddUnitNode(const T& InDefaults, const FName& InMethodName = TEXT("Execute"), const FVector2D& InPosition = FVector2D::ZeroVector, const FString& InNodeName = TEXT(""), bool bSetupUndoRedo = true, bool bPrintPythonCommand = false)
 	{
diff --git a/Engine/Source/Programs/Shared/EpicGames.UHT/Types/Properties/UhtSetProperty.cs b/Engine/Source/Programs/Shared/EpicGames.UHT/Types/Properties/UhtSetProperty.cs
--- a/Engine/Source/Programs/Shared/EpicGames.UHT/Types/Properties/UhtSetProperty.cs
+++ b/Engine/Source/Programs/Shared/EpicGames.UHT/Types/Properties/UhtSetProperty.cs
@@ -143,9 +143,9 @@ namespace EpicGames.UHT.Types
 			{
 				builder
 					.AppendTabs(tabs)
-					.Append("static_assert(TModels<CGetTypeHashable, ")
+					.Append("static_assert(TModels_V<CGetTypeHashable, ")
 					.Append(structProperty.ScriptStruct.SourceName)
-					.Append(">::Value, \"The structure '")
+					.Append(">, \"The structure '")
 					.Append(structProperty.ScriptStruct.SourceName)
 					.Append("' is used in a TSet but does not have a GetValueTypeHash defined\");\r\n");
 			}
diff --git a/Engine/Source/Programs/UnrealHeaderTool/Private/CodeGenerator.cpp b/Engine/Source/Programs/UnrealHeaderTool/Private/CodeGenerator.cpp
--- a/Engine/Source/Programs/UnrealHeaderTool/Private/CodeGenerator.cpp
+++ b/Engine/Source/Programs/UnrealHeaderTool/Private/CodeGenerator.cpp
@@ -1492,7 +1492,7 @@ void FNativeClassHeaderGenerator::OutputProperty(FOutputDevice& DeclOut, FOutput
 		if (ValueDef.IsStructOrStructStaticArray())
 		{
 			const FString& StructName = ValueDef.GetPropertyBase().ScriptStructDef->GetNameCPP();
-			Out.Logf(TEXT("%sstatic_assert(TModels<CGetTypeHashable, %s>::Value, \"The structure '%s' is used in a TSet but does not have a GetValueTypeHash defined\");\r\n"), Spaces, *StructName, *StructName);
+			Out.Logf(TEXT("%sstatic_assert(TModels_V<CGetTypeHashable, %s>, \"The structure '%s' is used in a TSet but does not have a GetValueTypeHash defined\");\r\n"), Spaces, *StructName, *StructName);
 		}
 
 		Out.Logf(
diff --git a/Engine/Source/Runtime/AIModule/Public/GraphAStar.h b/Engine/Source/Runtime/AIModule/Public/GraphAStar.h
--- a/Engine/Source/Runtime/AIModule/Public/GraphAStar.h
+++ b/Engine/Source/Runtime/AIModule/Public/GraphAStar.h
@@ -73,8 +73,8 @@ struct CQuery##ConditionalFunctionName	\
 	template<typename TemplateClass> auto Requires(TemplateClassParameter& Obj) -> decltype(Obj.ConditionalFunctionName());	\
 };	\
 template <typename TemplateClass> using ReturnType##ConditionalFunctionName = decltype((DeclVal<TemplateClass>().*DeclVal<decltype(&TemplateClass::ConditionalFunctionName)>())()); \
-template <typename TemplateClass> static FORCEINLINE typename TEnableIf< TModels<CQuery##ConditionalFunctionName, TemplateClass>::Value, ReturnType##ConditionalFunctionName<TemplateClass>>::Type ConditionalFunctionName(TemplateClassParameter & Obj) { return Obj.ConditionalFunctionName(); } \
-template <typename TemplateClass> static FORCEINLINE typename TEnableIf<!TModels<CQuery##ConditionalFunctionName, TemplateClass>::Value, ConditionalReturnType>::Type ConditionalFunctionName(TemplateClassParameter& Obj) { return DefaultImpl; }
+template <typename TemplateClass> static FORCEINLINE typename TEnableIf< TModels_V<CQuery##ConditionalFunctionName, TemplateClass>, ReturnType##ConditionalFunctionName<TemplateClass>>::Type ConditionalFunctionName(TemplateClassParameter & Obj) { return Obj.ConditionalFunctionName(); } \
+template <typename TemplateClass> static FORCEINLINE typename TEnableIf<!TModels_V<CQuery##ConditionalFunctionName, TemplateClass>, ConditionalReturnType>::Type ConditionalFunctionName(TemplateClassParameter& Obj) { return DefaultImpl; }
 #define DECLARE_OPTIONALLY_IMPLEMENTED_TEMPLATE_CLASS_FUNCTION( ConditionalReturnType, ConditionalFunctionName, DefaultImpl )  DECLARE_OPTIONALLY_IMPLEMENTED_TEMPLATE_CLASS_FUNCTION_IMPL( TemplateClass, TemplateClass, ConditionalReturnType, ConditionalFunctionName, DefaultImpl )
 #define DECLARE_OPTIONALLY_IMPLEMENTED_TEMPLATE_CLASS_FUNCTION_CONST( ConditionalReturnType, ConditionalFunctionName, DefaultImpl ) DECLARE_OPTIONALLY_IMPLEMENTED_TEMPLATE_CLASS_FUNCTION_IMPL( TemplateClass, const TemplateClass, ConditionalReturnType, ConditionalFunctionName, DefaultImpl )
 
@@ -84,8 +84,8 @@ struct CQuery##ConditionalFunctionName	\
 	template<typename TemplateClass> auto Requires(TemplateClassParameter& Obj, ConditionalParamType1 Param1) -> decltype(Obj.ConditionalFunctionName(Param1));	\
 };	\
 template <typename TemplateClass> using ReturnType##ConditionalFunctionName = decltype((DeclVal<TemplateClass>().*DeclVal<decltype(&TemplateClass::ConditionalFunctionName)>())(DeclVal<ConditionalParamType1>())); \
-template <typename TemplateClass> static FORCEINLINE typename TEnableIf< TModels<CQuery##ConditionalFunctionName, TemplateClass>::Value, ReturnType##ConditionalFunctionName<TemplateClass>>::Type ConditionalFunctionName(TemplateClassParameter & Obj, ConditionalParamType1 Param1) { return Obj.ConditionalFunctionName(Param1); } \
-template <typename TemplateClass> static FORCEINLINE typename TEnableIf<!TModels<CQuery##ConditionalFunctionName, TemplateClass>::Value, ConditionalReturnType>::Type ConditionalFunctionName(TemplateClassParameter& Obj, ConditionalParamType1 Param1) { return DefaultImpl; }
+template <typename TemplateClass> static FORCEINLINE typename TEnableIf< TModels_V<CQuery##ConditionalFunctionName, TemplateClass>, ReturnType##ConditionalFunctionName<TemplateClass>>::Type ConditionalFunctionName(TemplateClassParameter & Obj, ConditionalParamType1 Param1) { return Obj.ConditionalFunctionName(Param1); } \
+template <typename TemplateClass> static FORCEINLINE typename TEnableIf<!TModels_V<CQuery##ConditionalFunctionName, TemplateClass>, ConditionalReturnType>::Type ConditionalFunctionName(TemplateClassParameter& Obj, ConditionalParamType1 Param1) { return DefaultImpl; }
 #define DECLARE_OPTIONALLY_IMPLEMENTED_TEMPLATE_CLASS_FUNCTION_1PARAM( ConditionalReturnType, ConditionalFunctionName, ConditionalParamType1, DefaultImpl) DECLARE_OPTIONALLY_IMPLEMENTED_TEMPLATE_CLASS_FUNCTION_1PARAM_IMPL( TemplateClass, TemplateClass, ConditionalReturnType, ConditionalFunctionName, ConditionalParamType1, DefaultImpl) 
 #define DECLARE_OPTIONALLY_IMPLEMENTED_TEMPLATE_CLASS_FUNCTION_1PARAM_CONST( ConditionalReturnType, ConditionalFunctionName, ConditionalParamType1, DefaultImpl) DECLARE_OPTIONALLY_IMPLEMENTED_TEMPLATE_CLASS_FUNCTION_1PARAM_IMPL( TemplateClass, const TemplateClass, ConditionalReturnType, ConditionalFunctionName, ConditionalParamType1, DefaultImpl) 
 
@@ -95,8 +95,8 @@ struct CQuery##ConditionalFunctionName	\
 	template<typename TemplateClass> auto Requires(TemplateClassParameter& Obj, ConditionalParamType1 Param1, ConditionalParamType2 Param2) -> decltype(Obj.ConditionalFunctionName(Param1,Param2));	\
 };	\
 template <typename TemplateClass> using ReturnType##ConditionalFunctionName = decltype((DeclVal<TemplateClass>().*DeclVal<decltype(&TemplateClass::ConditionalFunctionName)>())(DeclVal<ConditionalParamType1>(), DeclVal<ConditionalParamType2>())); \
-template <typename TemplateClass> static FORCEINLINE typename TEnableIf< TModels<CQuery##ConditionalFunctionName, TemplateClass>::Value, ReturnType##ConditionalFunctionName<TemplateClass>>::Type ConditionalFunctionName(TemplateClassParameter & Obj, ConditionalParamType1 Param1, ConditionalParamType2 Param2) { return Obj.ConditionalFunctionName(Param1, Param2); } \
-template <typename TemplateClass> static FORCEINLINE typename TEnableIf<!TModels<CQuery##ConditionalFunctionName, TemplateClass>::Value, ConditionalReturnType>::Type ConditionalFunctionName(TemplateClassParameter& Obj, ConditionalParamType1 Param1, ConditionalParamType2 Param2) { return DefaultImpl; }
+template <typename TemplateClass> static FORCEINLINE typename TEnableIf< TModels_V<CQuery##ConditionalFunctionName, TemplateClass>, ReturnType##ConditionalFunctionName<TemplateClass>>::Type ConditionalFunctionName(TemplateClassParameter & Obj, ConditionalParamType1 Param1, ConditionalParamType2 Param2) { return Obj.ConditionalFunctionName(Param1, Param2); } \
+template <typename TemplateClass> static FORCEINLINE typename TEnableIf<!TModels_V<CQuery##ConditionalFunctionName, TemplateClass>, ConditionalReturnType>::Type ConditionalFunctionName(TemplateClassParameter& Obj, ConditionalParamType1 Param1, ConditionalParamType2 Param2) { return DefaultImpl; }
 #define DECLARE_OPTIONALLY_IMPLEMENTED_TEMPLATE_CLASS_FUNCTION_2PARAMS( ConditionalReturnType, ConditionalFunctionName, ConditionalParamType1, ConditionalParamType2, DefaultImpl) DECLARE_OPTIONALLY_IMPLEMENTED_TEMPLATE_CLASS_FUNCTION_2PARAMS_IMPL( TemplateClass, TemplateClass, ConditionalReturnType, ConditionalFunctionName, ConditionalParamType1, ConditionalParamType2, DefaultImpl)
 #define DECLARE_OPTIONALLY_IMPLEMENTED_TEMPLATE_CLASS_FUNCTION_2PARAMS_CONST( ConditionalReturnType, ConditionalFunctionName, ConditionalParamType1, ConditionalParamType2, DefaultImpl) DECLARE_OPTIONALLY_IMPLEMENTED_TEMPLATE_CLASS_FUNCTION_2PARAMS_IMPL(TemplateClass, const TemplateClass, ConditionalReturnType, ConditionalFunctionName, ConditionalParamType1, ConditionalParamType2, DefaultImpl)
 
@@ -105,8 +105,8 @@ struct CQuery##QueryFunctionName	\
 {	\
 	template<typename TemplateClass> auto Requires(const TemplateClass& Obj) -> decltype(Obj.ConditionalFunctionName());	\
 };	\
-template <typename TemplateClass> static FORCEINLINE typename TEnableIf< TModels<CQuery##QueryFunctionName, TemplateClass>::Value, QueryReturnType>::Type QueryFunctionName(const TemplateClass& Obj, QueryParam) { return QueryImpl; } \
-template <typename TemplateClass> static FORCEINLINE typename TEnableIf<!TModels<CQuery##QueryFunctionName, TemplateClass>::Value, QueryReturnType>::Type QueryFunctionName(const TemplateClass& Obj, QueryParam) { return QueryDefaultImpl; }
+template <typename TemplateClass> static FORCEINLINE typename TEnableIf< TModels_V<CQuery##QueryFunctionName, TemplateClass>, QueryReturnType>::Type QueryFunctionName(const TemplateClass& Obj, QueryParam) { return QueryImpl; } \
+template <typename TemplateClass> static FORCEINLINE typename TEnableIf<!TModels_V<CQuery##QueryFunctionName, TemplateClass>, QueryReturnType>::Type QueryFunctionName(const TemplateClass& Obj, QueryParam) { return QueryDefaultImpl; }
 
 template<bool DoRangeCheck>
 class FRangeChecklessAllocator : public FDefaultAllocator
diff --git a/Engine/Source/Runtime/Core/Private/Containers/ContainersTest.cpp b/Engine/Source/Runtime/Core/Private/Containers/ContainersTest.cpp
--- a/Engine/Source/Runtime/Core/Private/Containers/ContainersTest.cpp
+++ b/Engine/Source/Runtime/Core/Private/Containers/ContainersTest.cpp
@@ -240,7 +240,7 @@ namespace
 	template <typename Container>
 	void CheckContainerSelfEquality(Container& Cont)
 	{
-		if constexpr (TModels<CEqualityComparable, Container>::Value)
+		if constexpr (TModels_V<CEqualityComparable, Container>)
 		{
 			check(Cont == Cont);
 		}
diff --git a/Engine/Source/Runtime/Core/Public/Math/GenericOctree.h b/Engine/Source/Runtime/Core/Public/Math/GenericOctree.h
--- a/Engine/Source/Runtime/Core/Public/Math/GenericOctree.h
+++ b/Engine/Source/Runtime/Core/Public/Math/GenericOctree.h
@@ -1047,12 +1047,12 @@ private:
 
 	// Function overload set which calls the V2 version if it's supported or the old version if it's not
 	template <typename Semantics>
-	typename TEnableIf<!TModels<COctreeSemanticsV2, Semantics>::Value>::Type SetOctreeSemanticsElementId(const ElementType& Element, FOctreeElementId2 Id)
+	typename TEnableIf<!TModels_V<COctreeSemanticsV2, Semantics>>::Type SetOctreeSemanticsElementId(const ElementType& Element, FOctreeElementId2 Id)
 	{
 		Semantics::SetElementId(Element, Id);
 	}
 	template <typename Semantics>
-	typename TEnableIf<TModels<COctreeSemanticsV2, Semantics>::Value>::Type SetOctreeSemanticsElementId(const ElementType& Element, FOctreeElementId2 Id)
+	typename TEnableIf<TModels_V<COctreeSemanticsV2, Semantics>>::Type SetOctreeSemanticsElementId(const ElementType& Element, FOctreeElementId2 Id)
 	{
 		Semantics::SetElementId(static_cast<typename Semantics::FOctree&>(*this), Element, Id);
 	}
@@ -1784,12 +1784,12 @@ private:
 
 	// Function overload set which calls the V2 version if it's supported or the old version if it's not
 	template <typename Semantics>
-	typename TEnableIf<!TModels<COctreeSemanticsV2, Semantics>::Value>::Type SetOctreeSemanticsElementId(const ElementType& Element, FOctreeElementId Id)
+	typename TEnableIf<!TModels_V<COctreeSemanticsV2, Semantics>>::Type SetOctreeSemanticsElementId(const ElementType& Element, FOctreeElementId Id)
 	{
 		Semantics::SetElementId(Element, Id);
 	}
 	template <typename Semantics>
-	typename TEnableIf<TModels<COctreeSemanticsV2, Semantics>::Value>::Type SetOctreeSemanticsElementId(const ElementType& Element, FOctreeElementId Id)
+	typename TEnableIf<TModels_V<COctreeSemanticsV2, Semantics>>::Type SetOctreeSemanticsElementId(const ElementType& Element, FOctreeElementId Id)
 	{
 		Semantics::SetElementId(*this, Element, Id);
 	}
diff --git a/Engine/Source/Runtime/Core/Public/Misc/HashBuilder.h b/Engine/Source/Runtime/Core/Public/Misc/HashBuilder.h
--- a/Engine/Source/Runtime/Core/Public/Misc/HashBuilder.h
+++ b/Engine/Source/Runtime/Core/Public/Misc/HashBuilder.h
@@ -34,13 +34,13 @@ public:
 	}
 
 	template <typename T>
-	typename TEnableIf<!TModels<CGetTypeHashable, T>::Value, FHashBuilder&>::Type Append(const T& InData)
+	typename TEnableIf<!TModels_V<CGetTypeHashable, T>, FHashBuilder&>::Type Append(const T& InData)
 	{
 		return AppendRaw(InData);
 	}
 
 	template <typename T>
-	typename TEnableIf<TModels<CGetTypeHashable, T>::Value, FHashBuilder&>::Type Append(const T& InData)
+	typename TEnableIf<TModels_V<CGetTypeHashable, T>, FHashBuilder&>::Type Append(const T& InData)
 	{
 		Hash = HashCombineFast(Hash, GetTypeHash(InData));
 		return *this;
diff --git a/Engine/Source/Runtime/Core/Public/Misc/TVariantMeta.h b/Engine/Source/Runtime/Core/Public/Misc/TVariantMeta.h
--- a/Engine/Source/Runtime/Core/Public/Misc/TVariantMeta.h
+++ b/Engine/Source/Runtime/Core/Public/Misc/TVariantMeta.h
@@ -246,7 +246,7 @@ namespace Private
 	{
 		using VariantType = TVariant<Ts...>;
 		static_assert((std::is_default_constructible<Ts>::value && ...), "Each type in TVariant template parameter pack must be default constructible in order to use FArchive serialization");
-		static_assert((TModels<CInsertable<FArchive&>, Ts>::Value && ...), "Each type in TVariant template parameter pack must be able to use operator<< with an FArchive");
+		static_assert((TModels_V<CInsertable<FArchive&>, Ts> && ...), "Each type in TVariant template parameter pack must be able to use operator<< with an FArchive");
 
 		/** Load the type at the specified index from the FArchive and emplace it into the TVariant */
 		static void Load(SIZE_T TypeIndex, FArchive& Ar, VariantType& OutVariant)
diff --git a/Engine/Source/Runtime/Core/Public/Serialization/MemoryLayout.h b/Engine/Source/Runtime/Core/Public/Serialization/MemoryLayout.h
--- a/Engine/Source/Runtime/Core/Public/Serialization/MemoryLayout.h
+++ b/Engine/Source/Runtime/Core/Public/Serialization/MemoryLayout.h
@@ -175,7 +175,7 @@ const T* GetDefault();
 template<typename T>
 struct THasCustomDefaultObject
 {
-	static constexpr bool Value = TModels<CProvidesDefaultUObject, T>::Value;
+	static constexpr bool Value = TModels_V<CProvidesDefaultUObject, T>;
 };
 
 template<typename T>
@@ -321,7 +321,7 @@ namespace Freeze
 template<typename T>
 struct TUsePropertyFreezing
 {
-	static constexpr bool Value = (TModels<CStaticClassProvider, T>::Value || TModels<CStaticStructProvider, T>::Value);
+	static constexpr bool Value = (TModels_V<CStaticClassProvider, T> || TModels_V<CStaticStructProvider, T>);
 };
 
 template <typename T, bool bUsePropertyFreezing=TUsePropertyFreezing<T>::Value>
@@ -330,7 +330,7 @@ struct TGetFreezeImageHelper
 	static FORCEINLINE FTypeLayoutDesc::FWriteFrozenMemoryImageFunc* Do() { return &Freeze::DefaultWriteMemoryImage; }
 };
 
-template <typename T, bool bProvidesStaticStruct=TModels<CStaticStructProvider, T>::Value>
+template <typename T, bool bProvidesStaticStruct=TModels_V<CStaticStructProvider, T>>
 struct TGetFreezeImageFieldHelper
 {
 	static FORCEINLINE FFieldLayoutDesc::FWriteFrozenMemoryImageFunc* Do() { return &Freeze::DefaultWriteMemoryImageField; }
@@ -620,7 +620,7 @@ template<typename T>
 struct THasTypeLayout
 {
 	UE_STATIC_ONLY(THasTypeLayout);
-	static constexpr bool Value = TModels<CTypeLayout, T>::Value;
+	static constexpr bool Value = TModels_V<CTypeLayout, T>;
 };
 
 template<typename T>
diff --git a/Engine/Source/Runtime/Core/Public/Serialization/StructuredArchiveAdapters.h b/Engine/Source/Runtime/Core/Public/Serialization/StructuredArchiveAdapters.h
--- a/Engine/Source/Runtime/Core/Public/Serialization/StructuredArchiveAdapters.h
+++ b/Engine/Source/Runtime/Core/Public/Serialization/StructuredArchiveAdapters.h
@@ -122,7 +122,7 @@ private:
  */
 template <typename T>
 typename TEnableIf<
-	!TModels<CInsertable<FArchive&>, T>::Value&& TModels<CInsertable<FStructuredArchiveSlot>, T>::Value,
+	!TModels_V<CInsertable<FArchive&>, T> && TModels_V<CInsertable<FStructuredArchiveSlot>, T>,
 	FArchive&
 >::Type operator<<(FArchive& Ar, T& Obj)
 {
@@ -139,8 +139,8 @@ typename TEnableIf<
  */
 template <typename T>
 typename TEnableIf<
-	TModels<CInsertable<FArchive&>, T>::Value &&
-	!TModels<CInsertable<FStructuredArchiveSlot>, T>::Value
+	TModels_V<CInsertable<FArchive&>, T> &&
+	!TModels_V<CInsertable<FStructuredArchiveSlot>, T>
 >::Type operator<<(FStructuredArchiveSlot Slot, T& Obj)
 {
 #if WITH_TEXT_ARCHIVE_SUPPORT
diff --git a/Engine/Source/Runtime/Core/Public/Templates/Models.h b/Engine/Source/Runtime/Core/Public/Templates/Models.h
--- a/Engine/Source/Runtime/Core/Public/Templates/Models.h
+++ b/Engine/Source/Runtime/Core/Public/Templates/Models.h
@@ -25,8 +25,9 @@
  * CIsComparable        - bad
  * CHasEqualsOperator   - bad
  *
- * Concepts can be checked using the TModels trait:
+ * Concepts can be checked using the TModels_V trait or TModels traits class:
  *
+ * TModels_V<Concept, [...arguments...]>
  * TModels<Concept, [...arguments...]>::Value
  *
  * The arguments are forwarded to the template parameters of the concept's Requires() function, which will attempt to
@@ -39,6 +40,21 @@
  * See the CContainerLvalueAddable example below.
  */
 
+namespace UE::Core::Private
+{
+	template <typename Concept, typename... Ts>
+	static char (&ModelsResolve(decltype(&Concept::template Requires<Ts...>)*))[2];
+
+	template <typename Concept, typename... Ts>
+	static char (&ModelsResolve(...))[1];
+}
+
+
+/**
+ * Trait which does concept checking.
+ */
+template <typename Concept, typename... Args>
+constexpr bool TModels_V = sizeof(UE::Core::Private::ModelsResolve<Concept, Args...>(0)) == 2;
 
 /**
  * Traits class which does concept checking.
@@ -46,24 +62,18 @@
 template <typename Concept, typename... Args>
 struct TModels
 {
-	template <typename... Ts>
-	static char (&Resolve(decltype(&Concept::template Requires<Ts...>)*))[2];
-
-	template <typename... Ts>
-	static char (&Resolve(...))[1];
-
-	static constexpr bool Value = sizeof(Resolve<Args...>(0)) == 2;
+	static constexpr bool Value = TModels_V<Concept, Args...>;
 };
 
 /**
  * Helper function which can be used as an expression in a concept to refine ('inherit') another concept.
- * It should be used as expression-based variant of the TModels traits class.  If the arguments model
+ * It should be used as expression-based variant of the TModels_V trait.  If the arguments model
  * the given concept, Refines<Concept, Args...> is a valid expression, otherwise it is not.
  *
  * See the CCopyablePointer example below.
  */
 template <typename Concept, typename... Args>
-auto Refines() -> int(&)[!!TModels<Concept, Args...>::Value * 2 - 1];
+auto Refines() -> int(&)[!!TModels_V<Concept, Args...> * 2 - 1];
 
 
 /************
@@ -80,8 +90,8 @@ auto Refines() -> int(&)[!!TModels<Concept, Args...>::Value * 2 - 1];
  *     );
  * };
  *
- * static_assert( TModels<CNegatable, int >::Value); // ints are negatable
- * static_assert(!TModels<CNegatable, int*>::Value); // pointers are not negatable
+ * static_assert( TModels_V<CNegatable, int >); // ints are negatable
+ * static_assert(!TModels_V<CNegatable, int*>); // pointers are not negatable
  */
 
 /**
@@ -96,9 +106,9 @@ auto Refines() -> int(&)[!!TModels<Concept, Args...>::Value * 2 - 1];
  *     );
  * };
  *
- * static_assert( TModels<CIncrementable, int  >::Value); // ints are incrementable
- * static_assert( TModels<CIncrementable, int* >::Value); // pointers are incrementable
- * static_assert(!TModels<CIncrementable, int[]>::Value); // arrays are not incrementable
+ * static_assert( TModels_V<CIncrementable, int  >); // ints are incrementable
+ * static_assert( TModels_V<CIncrementable, int* >); // pointers are incrementable
+ * static_assert(!TModels_V<CIncrementable, int[]>); // arrays are not incrementable
  */
 
 /**
@@ -115,8 +125,8 @@ auto Refines() -> int(&)[!!TModels<Concept, Args...>::Value * 2 - 1];
  *     );
  * };
  *
- * static_assert( TModels<CEqualityComparable, FArchive*, FArchiveUObject*>::Value); // base pointers are comparable with derived pointers
- * static_assert(!TModels<CEqualityComparable, int*,      float*          >::Value); // unrelated pointers are not comparable
+ * static_assert( TModels_V<CEqualityComparable, FArchive*, FArchiveUObject*>); // base pointers are comparable with derived pointers
+ * static_assert(!TModels_V<CEqualityComparable, int*,      float*          >); // unrelated pointers are not comparable
  */
 
 /**
@@ -130,7 +140,7 @@ auto Refines() -> int(&)[!!TModels<Concept, Args...>::Value * 2 - 1];
  *     );
  * };
  *
- * static_assert(TModels<CLvalueAddableContainer, TArray<TUniquePtr<int>>>::Value); // success, but...
+ * static_assert(TModels_V<CLvalueAddableContainer, TArray<TUniquePtr<int>>>); // success, but...
  *
  * TUniquePtr<int> Temp;
  * Array.Add(Temp); // compile error when TArray attempts to copy the TUniquePtr
@@ -149,8 +159,8 @@ auto Refines() -> int(&)[!!TModels<Concept, Args...>::Value * 2 - 1];
  *     );
  * };
  *
- * static_assert( TModels<CGroupSubtractable, int >::Value); // ints form a group under subtraction
- * static_assert(!TModels<CGroupSubtractable, int*>::Value); // pointers do not result in another pointer under subtraction and so do not form a group
+ * static_assert( TModels_V<CGroupSubtractable, int >); // ints form a group under subtraction
+ * static_assert(!TModels_V<CGroupSubtractable, int*>); // pointers do not result in another pointer under subtraction and so do not form a group
  */
 
 /**
@@ -185,9 +195,9 @@ auto Refines() -> int(&)[!!TModels<Concept, Args...>::Value * 2 - 1];
  *     );
  * };
  *
- * static_assert( TModels<CCopyablePointer, int*           >::Value); // raw pointers model this concept
- * static_assert( TModels<CCopyablePointer, TSharedPtr<int>>::Value); // TSharedPtrs model this concept
- * static_assert(!TModels<CCopyablePointer, TUniquePtr<int>>::Value); // TUniquePtrs are dereferenceable and bool-testable, but not copyable
- * static_assert(!TModels<CCopyablePointer, int            >::Value); // ints are copyable and bool-testable, but not dereferencable
- * static_assert(!TModels<CCopyablePointer, FString        >::Value); // FStrings are copyable and dereferencable (to get the const TCHAR*), but not bool-testable
+ * static_assert( TModels_V<CCopyablePointer, int*           >); // raw pointers model this concept
+ * static_assert( TModels_V<CCopyablePointer, TSharedPtr<int>>); // TSharedPtrs model this concept
+ * static_assert(!TModels_V<CCopyablePointer, TUniquePtr<int>>); // TUniquePtrs are dereferenceable and bool-testable, but not copyable
+ * static_assert(!TModels_V<CCopyablePointer, int            >); // ints are copyable and bool-testable, but not dereferencable
+ * static_assert(!TModels_V<CCopyablePointer, FString        >); // FStrings are copyable and dereferencable (to get the const TCHAR*), but not bool-testable
  */
diff --git a/Engine/Source/Runtime/CoreUObject/Public/Serialization/ArchiveCrc32.h b/Engine/Source/Runtime/CoreUObject/Public/Serialization/ArchiveCrc32.h
--- a/Engine/Source/Runtime/CoreUObject/Public/Serialization/ArchiveCrc32.h
+++ b/Engine/Source/Runtime/CoreUObject/Public/Serialization/ArchiveCrc32.h
@@ -49,7 +49,7 @@ private:
 	* @param Value The value to serialize.
 	*/
 template <typename StructType>
-FORCEINLINE typename TEnableIf<TModels<CStaticStructProvider, StructType>::Value, FArchive&>::Type operator <<(FArchive& Ar, const StructType& Value)
+FORCEINLINE typename TEnableIf<TModels_V<CStaticStructProvider, StructType>, FArchive&>::Type operator <<(FArchive& Ar, const StructType& Value)
 {
 	StructType* MutableValue = const_cast<StructType*>(&Value);
 	StructType::StaticStruct()->SerializeItem(Ar, MutableValue, nullptr);
diff --git a/Engine/Source/Runtime/CoreUObject/Public/UObject/Class.h b/Engine/Source/Runtime/CoreUObject/Public/UObject/Class.h
--- a/Engine/Source/Runtime/CoreUObject/Public/UObject/Class.h
+++ b/Engine/Source/Runtime/CoreUObject/Public/UObject/Class.h
@@ -1167,7 +1167,7 @@ public:
 				(TIsPODType<CPPSTRUCT>::Value ? CPF_IsPlainOldData : CPF_None)
 				| (TIsTriviallyDestructible<CPPSTRUCT>::Value ? CPF_NoDestructor : CPF_None)
 				| (TIsZeroConstructType<CPPSTRUCT>::Value ? CPF_ZeroConstructor : CPF_None)
-				| (TModels<CGetTypeHashable, CPPSTRUCT>::Value ? CPF_HasGetValueTypeHash : CPF_None),
+				| (TModels_V<CGetTypeHashable, CPPSTRUCT> ? CPF_HasGetValueTypeHash : CPF_None),
 				TTraits::WithNoInitConstructor,
 				TTraits::WithZeroConstructor,
 				!(TTraits::WithNoDestructor || TIsPODType<CPPSTRUCT>::Value),
@@ -1188,7 +1188,7 @@ public:
 				TTraits::WithAddStructReferencedObjects,
 				TTraits::WithSerializeFromMismatchedTag,
 				TTraits::WithStructuredSerializeFromMismatchedTag,
-				TModels<CGetTypeHashable, CPPSTRUCT>::Value,
+				TModels_V<CGetTypeHashable, CPPSTRUCT>,
 				TIsAbstract<CPPSTRUCT>::Value,
 #if WITH_EDITOR
 				TTraits::WithCanEditChange,
@@ -1448,7 +1448,7 @@ public:
 		{
 			ensure(HasGetTypeHash());
 
-			if constexpr (TModels<CGetTypeHashable, CPPSTRUCT>::Value)
+			if constexpr (TModels_V<CGetTypeHashable, CPPSTRUCT>)
 			{
 				return GetTypeHashHelper(*(const CPPSTRUCT*)Src);
 			}
diff --git a/Engine/Source/Runtime/CoreUObject/Public/UObject/UnrealType.h b/Engine/Source/Runtime/CoreUObject/Public/UObject/UnrealType.h
--- a/Engine/Source/Runtime/CoreUObject/Public/UObject/UnrealType.h
+++ b/Engine/Source/Runtime/CoreUObject/Public/UObject/UnrealType.h
@@ -1356,7 +1356,7 @@ protected:
 			(TIsPODType<TCppType>::Value ? CPF_IsPlainOldData : CPF_None) 
 			| (TIsTriviallyDestructible<TCppType>::Value ? CPF_NoDestructor : CPF_None) 
 			| (TIsZeroConstructType<TCppType>::Value ? CPF_ZeroConstructor : CPF_None)
-			| (TModels<CGetTypeHashable, TCppType>::Value ? CPF_HasGetValueTypeHash : CPF_None);
+			| (TModels_V<CGetTypeHashable, TCppType> ? CPF_HasGetValueTypeHash : CPF_None);
 
 	}
 };
diff --git a/Engine/Source/Runtime/CoreUObject/Tests/ObjectPtrTest.cpp b/Engine/Source/Runtime/CoreUObject/Tests/ObjectPtrTest.cpp
--- a/Engine/Source/Runtime/CoreUObject/Tests/ObjectPtrTest.cpp
+++ b/Engine/Source/Runtime/CoreUObject/Tests/ObjectPtrTest.cpp
@@ -77,63 +77,63 @@ static_assert(!std::is_assignable<FConstInterfacePtr, const FMutableObjectPtr&>:
 static_assert(!std::is_assignable<FConstInterfacePtr, const FConstObjectPtr&>::value, "Invalid assignment (TObjectPtr<const UInterface> from TObjectPtr<const UObject>)");
 
 // Ensure that TObjectPtr<[const] UObject> is comparable with another TObjectPtr<[const] UObject> regardless of constness
-static_assert(TModels<CEqualityComparableWith, FConstObjectPtr, FConstObjectPtr>::Value, "Must be able to compare equality and inequality bidirectionally between TObjectPtr<const UObject> and TObjectPtr<const UObject>");
-static_assert(TModels<CEqualityComparableWith, FMutableObjectPtr, FConstObjectPtr>::Value, "Must be able to compare equality and inequality bidirectionally between TObjectPtr<UObject> and TObjectPtr<const UObject>");
+static_assert(TModels_V<CEqualityComparableWith, FConstObjectPtr, FConstObjectPtr>, "Must be able to compare equality and inequality bidirectionally between TObjectPtr<const UObject> and TObjectPtr<const UObject>");
+static_assert(TModels_V<CEqualityComparableWith, FMutableObjectPtr, FConstObjectPtr>, "Must be able to compare equality and inequality bidirectionally between TObjectPtr<UObject> and TObjectPtr<const UObject>");
 
 // Ensure that TObjectPtr<[const] UObject> is comparable with another TObjectPtr<[const] UInterface> regardless of constness
-static_assert(TModels<CEqualityComparableWith, FConstObjectPtr, FConstInterfacePtr>::Value, "Must be able to compare equality and inequality bidirectionally between TObjectPtr<const UObject> and TObjectPtr<const UInterface>");
-static_assert(TModels<CEqualityComparableWith, FMutableObjectPtr, FConstInterfacePtr>::Value, "Must be able to compare equality and inequality bidirectionally between TObjectPtr<UObject> and TObjectPtr<const UInterface>");
-static_assert(TModels<CEqualityComparableWith, FConstObjectPtr, FMutableInterfacePtr>::Value, "Must be able to compare equality and inequality bidirectionally between TObjectPtr<const UObject> and TObjectPtr<UInterface>");
-static_assert(TModels<CEqualityComparableWith, FMutableObjectPtr, FMutableInterfacePtr>::Value, "Must be able to compare equality and inequality bidirectionally between TObjectPtr<UObject> and TObjectPtr<UInterface>");
+static_assert(TModels_V<CEqualityComparableWith, FConstObjectPtr, FConstInterfacePtr>, "Must be able to compare equality and inequality bidirectionally between TObjectPtr<const UObject> and TObjectPtr<const UInterface>");
+static_assert(TModels_V<CEqualityComparableWith, FMutableObjectPtr, FConstInterfacePtr>, "Must be able to compare equality and inequality bidirectionally between TObjectPtr<UObject> and TObjectPtr<const UInterface>");
+static_assert(TModels_V<CEqualityComparableWith, FConstObjectPtr, FMutableInterfacePtr>, "Must be able to compare equality and inequality bidirectionally between TObjectPtr<const UObject> and TObjectPtr<UInterface>");
+static_assert(TModels_V<CEqualityComparableWith, FMutableObjectPtr, FMutableInterfacePtr>, "Must be able to compare equality and inequality bidirectionally between TObjectPtr<UObject> and TObjectPtr<UInterface>");
 
 // Ensure that TObjectPtr<[const] UPackage> is not comparable with a TObjectPtr<[const] UInterface> regardless of constness
 // TODO: This only ensures that at least one of the A==B,B==A,A!=B,B!=A operations fail, not that they all fail.
 #if !(PLATFORM_MICROSOFT) || !defined(_MSC_EXTENSIONS) // MSVC static analyzer is run in non-conformance mode, and that causes these checks to fail.
-static_assert(!TModels<CEqualityComparableWith, FConstPackagePtr, FConstInterfacePtr>::Value, "Must not be able to compare equality and inequality bidirectionally between TObjectPtr<const UPackage> and TObjectPtr<const UInterface>");
-static_assert(!TModels<CEqualityComparableWith, FMutablePackagePtr, FConstInterfacePtr>::Value, "Must not be able to compare equality and inequality bidirectionally between TObjectPtr<UPackage> and TObjectPtr<const UInterface>");
-static_assert(!TModels<CEqualityComparableWith, FConstPackagePtr, FMutableInterfacePtr>::Value, "Must not be able to compare equality and inequality bidirectionally between TObjectPtr<const UPackage> and TObjectPtr<UInterface>");
-static_assert(!TModels<CEqualityComparableWith, FMutablePackagePtr, FMutableInterfacePtr>::Value, "Must not be able to compare equality and inequality bidirectionally between TObjectPtr<UPackage> and TObjectPtr<UInterface>");
+static_assert(!TModels_V<CEqualityComparableWith, FConstPackagePtr, FConstInterfacePtr>, "Must not be able to compare equality and inequality bidirectionally between TObjectPtr<const UPackage> and TObjectPtr<const UInterface>");
+static_assert(!TModels_V<CEqualityComparableWith, FMutablePackagePtr, FConstInterfacePtr>, "Must not be able to compare equality and inequality bidirectionally between TObjectPtr<UPackage> and TObjectPtr<const UInterface>");
+static_assert(!TModels_V<CEqualityComparableWith, FConstPackagePtr, FMutableInterfacePtr>, "Must not be able to compare equality and inequality bidirectionally between TObjectPtr<const UPackage> and TObjectPtr<UInterface>");
+static_assert(!TModels_V<CEqualityComparableWith, FMutablePackagePtr, FMutableInterfacePtr>, "Must not be able to compare equality and inequality bidirectionally between TObjectPtr<UPackage> and TObjectPtr<UInterface>");
 #endif // #if !(PLATFORM_MICROSOFT) || !defined(_MSC_EXTENSIONS)
 
 // Ensure that TObjectPtr<[const] UObject> is comparable with a raw pointer of the same referenced type regardless of constness
-static_assert(TModels<CEqualityComparableWith, FConstObjectPtr, const UObject*>::Value, "Must be able to compare equality and inequality bidirectionally between TObjectPtr<const UObject> and const UObject*");
-static_assert(TModels<CEqualityComparableWith, FMutableObjectPtr, const UObject*>::Value, "Must be able to compare equality and inequality bidirectionally between TObjectPtr<UObject> and const UObject*");
-static_assert(TModels<CEqualityComparableWith, FConstObjectPtr, UObject*>::Value, "Must be able to compare equality and inequality bidirectionally between TObjectPtr<const UObject> and UObject*");
-static_assert(TModels<CEqualityComparableWith, FMutableObjectPtr, UObject*>::Value, "Must be able to compare equality and inequality bidirectionally between TObjectPtr<UObject> and UObject*");
+static_assert(TModels_V<CEqualityComparableWith, FConstObjectPtr, const UObject*>, "Must be able to compare equality and inequality bidirectionally between TObjectPtr<const UObject> and const UObject*");
+static_assert(TModels_V<CEqualityComparableWith, FMutableObjectPtr, const UObject*>, "Must be able to compare equality and inequality bidirectionally between TObjectPtr<UObject> and const UObject*");
+static_assert(TModels_V<CEqualityComparableWith, FConstObjectPtr, UObject*>, "Must be able to compare equality and inequality bidirectionally between TObjectPtr<const UObject> and UObject*");
+static_assert(TModels_V<CEqualityComparableWith, FMutableObjectPtr, UObject*>, "Must be able to compare equality and inequality bidirectionally between TObjectPtr<UObject> and UObject*");
 
 // Ensure that TObjectPtr<[const] UObject> is comparable with a UInterface raw pointer regardless of constness
-static_assert(TModels<CEqualityComparableWith, FConstObjectPtr, const UInterface*>::Value, "Must be able to compare equality and inequality bidirectionally between TObjectPtr<const UObject> and const UInterface*");
-static_assert(TModels<CEqualityComparableWith, FMutableObjectPtr, const UInterface*>::Value, "Must be able to compare equality and inequality bidirectionally between TObjectPtr<UObject> and const UInterface*");
-static_assert(TModels<CEqualityComparableWith, FConstObjectPtr, UInterface*>::Value, "Must be able to compare equality and inequality bidirectionally between TObjectPtr<const UObject> and UInterface*");
-static_assert(TModels<CEqualityComparableWith, FMutableObjectPtr, UInterface*>::Value, "Must be able to compare equality and inequality bidirectionally between TObjectPtr<UObject> and UInterface*");
+static_assert(TModels_V<CEqualityComparableWith, FConstObjectPtr, const UInterface*>, "Must be able to compare equality and inequality bidirectionally between TObjectPtr<const UObject> and const UInterface*");
+static_assert(TModels_V<CEqualityComparableWith, FMutableObjectPtr, const UInterface*>, "Must be able to compare equality and inequality bidirectionally between TObjectPtr<UObject> and const UInterface*");
+static_assert(TModels_V<CEqualityComparableWith, FConstObjectPtr, UInterface*>, "Must be able to compare equality and inequality bidirectionally between TObjectPtr<const UObject> and UInterface*");
+static_assert(TModels_V<CEqualityComparableWith, FMutableObjectPtr, UInterface*>, "Must be able to compare equality and inequality bidirectionally between TObjectPtr<UObject> and UInterface*");
 
 // Ensure that TObjectPtr<[const] UInterface> is comparable with a UObject raw pointer regardless of constness
-static_assert(TModels<CEqualityComparableWith, FConstInterfacePtr, const UObject*>::Value, "Must be able to compare equality and inequality bidirectionally between TObjectPtr<const UInterface> and const UObject*");
-static_assert(TModels<CEqualityComparableWith, FMutableInterfacePtr, const UObject*>::Value, "Must be able to compare equality and inequality bidirectionally between TObjectPtr<UInterface> and const UObject*");
-static_assert(TModels<CEqualityComparableWith, FConstInterfacePtr, UObject*>::Value, "Must be able to compare equality and inequality bidirectionally between TObjectPtr<const UInterface> and UObject*");
-static_assert(TModels<CEqualityComparableWith, FMutableInterfacePtr, UObject*>::Value, "Must be able to compare equality and inequality bidirectionally between TObjectPtr<UInterface> and UObject*");
+static_assert(TModels_V<CEqualityComparableWith, FConstInterfacePtr, const UObject*>, "Must be able to compare equality and inequality bidirectionally between TObjectPtr<const UInterface> and const UObject*");
+static_assert(TModels_V<CEqualityComparableWith, FMutableInterfacePtr, const UObject*>, "Must be able to compare equality and inequality bidirectionally between TObjectPtr<UInterface> and const UObject*");
+static_assert(TModels_V<CEqualityComparableWith, FConstInterfacePtr, UObject*>, "Must be able to compare equality and inequality bidirectionally between TObjectPtr<const UInterface> and UObject*");
+static_assert(TModels_V<CEqualityComparableWith, FMutableInterfacePtr, UObject*>, "Must be able to compare equality and inequality bidirectionally between TObjectPtr<UInterface> and UObject*");
 
 // Ensure that TObjectPtr<[const] UInterface> is not comparable with a UPackage raw pointer regardless of constness
 // TODO: This only ensures that at least one of the A==B,B==A,A!=B,B!=A operations fail, not that they all fail.
-static_assert(!TModels<CEqualityComparableWith, FConstInterfacePtr, const UPackage*>::Value, "Must not be able to compare equality and inequality bidirectionally between TObjectPtr<const UInterface> and const UPackage*");
-static_assert(!TModels<CEqualityComparableWith, FMutableInterfacePtr, const UPackage*>::Value, "Must not be able to compare equality and inequality bidirectionally between TObjectPtr<UInterface> and const UPackage*");
-static_assert(!TModels<CEqualityComparableWith, FConstInterfacePtr, UPackage*>::Value, "Must not be able to compare equality and inequality bidirectionally between TObjectPtr<const UInterface> and UPackage*");
-static_assert(!TModels<CEqualityComparableWith, FMutableInterfacePtr, UPackage*>::Value, "Must not be able to compare equality and inequality bidirectionally between TObjectPtr<UInterface> and UPackage*");
+static_assert(!TModels_V<CEqualityComparableWith, FConstInterfacePtr, const UPackage*>, "Must not be able to compare equality and inequality bidirectionally between TObjectPtr<const UInterface> and const UPackage*");
+static_assert(!TModels_V<CEqualityComparableWith, FMutableInterfacePtr, const UPackage*>, "Must not be able to compare equality and inequality bidirectionally between TObjectPtr<UInterface> and const UPackage*");
+static_assert(!TModels_V<CEqualityComparableWith, FConstInterfacePtr, UPackage*>, "Must not be able to compare equality and inequality bidirectionally between TObjectPtr<const UInterface> and UPackage*");
+static_assert(!TModels_V<CEqualityComparableWith, FMutableInterfacePtr, UPackage*>, "Must not be able to compare equality and inequality bidirectionally between TObjectPtr<UInterface> and UPackage*");
 
 // Ensure that TObjectPtr<[const] UInterface> is not comparable with a char raw pointer regardless of constness
 // TODO: This only ensures that at least one of the A==B,B==A,A!=B,B!=A operations fail, not that they all fail.
-static_assert(!TModels<CEqualityComparableWith, FConstObjectPtr, const char*>::Value, "Must not be able to compare equality and inequality bidirectionally between TObjectPtr<const UObject> and const UObject*");
-static_assert(!TModels<CEqualityComparableWith, FMutableObjectPtr, const char*>::Value, "Must not be able to compare equality and inequality bidirectionally between TObjectPtr<UObject> and const UObject*");
-static_assert(!TModels<CEqualityComparableWith, FConstObjectPtr, char*>::Value, "Must not be able to compare equality and inequality bidirectionally between TObjectPtr<const UObject> and UObject*");
-static_assert(!TModels<CEqualityComparableWith, FMutableObjectPtr, char*>::Value, "Must not be able to compare equality and inequality bidirectionally between TObjectPtr<UObject> and UObject*");
+static_assert(!TModels_V<CEqualityComparableWith, FConstObjectPtr, const char*>, "Must not be able to compare equality and inequality bidirectionally between TObjectPtr<const UObject> and const UObject*");
+static_assert(!TModels_V<CEqualityComparableWith, FMutableObjectPtr, const char*>, "Must not be able to compare equality and inequality bidirectionally between TObjectPtr<UObject> and const UObject*");
+static_assert(!TModels_V<CEqualityComparableWith, FConstObjectPtr, char*>, "Must not be able to compare equality and inequality bidirectionally between TObjectPtr<const UObject> and UObject*");
+static_assert(!TModels_V<CEqualityComparableWith, FMutableObjectPtr, char*>, "Must not be able to compare equality and inequality bidirectionally between TObjectPtr<UObject> and UObject*");
 
 // Ensure that TObjectPtr<[const] UObject> is comparable with nullptr regardless of constness
-static_assert(TModels<CEqualityComparableWith, FConstObjectPtr, TYPE_OF_NULLPTR>::Value, "Must be able to compare equality and inequality bidirectionally between TObjectPtr<const UObject> and nullptr");
-static_assert(TModels<CEqualityComparableWith, FMutableObjectPtr, TYPE_OF_NULLPTR>::Value, "Must be able to compare equality and inequality bidirectionally between TObjectPtr<UObject> and nullptr");
+static_assert(TModels_V<CEqualityComparableWith, FConstObjectPtr, TYPE_OF_NULLPTR>, "Must be able to compare equality and inequality bidirectionally between TObjectPtr<const UObject> and nullptr");
+static_assert(TModels_V<CEqualityComparableWith, FMutableObjectPtr, TYPE_OF_NULLPTR>, "Must be able to compare equality and inequality bidirectionally between TObjectPtr<UObject> and nullptr");
 
 #if !UE_OBJECT_PTR_NONCONFORMANCE_SUPPORT // Specialized NULL support causes these checks to fail.
-static_assert(!TModels<CEqualityComparableWith, FConstObjectPtr, long>::Value, "Should not be able to compare equality and inequality bidirectionally between TObjectPtr<const UObject> and long");
-static_assert(!TModels<CEqualityComparableWith, FMutableObjectPtr, long>::Value, "Should not be able to compare equality and inequality bidirectionally between TObjectPtr<UObject> and long");
+static_assert(!TModels_V<CEqualityComparableWith, FConstObjectPtr, long>, "Should not be able to compare equality and inequality bidirectionally between TObjectPtr<const UObject> and long");
+static_assert(!TModels_V<CEqualityComparableWith, FMutableObjectPtr, long>, "Should not be able to compare equality and inequality bidirectionally between TObjectPtr<UObject> and long");
 #endif // #if !UE_OBJECT_PTR_NONCONFORMANCE_SUPPORT
 
 
diff --git a/Engine/Source/Runtime/Engine/Public/Animation/AttributeTypes.h b/Engine/Source/Runtime/Engine/Public/Animation/AttributeTypes.h
--- a/Engine/Source/Runtime/Engine/Public/Animation/AttributeTypes.h
+++ b/Engine/Source/Runtime/Engine/Public/Animation/AttributeTypes.h
@@ -93,11 +93,11 @@ namespace UE
 				
 				if constexpr (UE::Anim::TAttributeTypeTraits<AttributeType>::IsBlendable)	
 				{
-					static_assert(TModels<CBlendableAttribute, AttributeType>::Value, "Missing function implementations required for Attribute blending");
+					static_assert(TModels_V<CBlendableAttribute, AttributeType>, "Missing function implementations required for Attribute blending");
 
 					if  constexpr (UE::Anim::TAttributeTypeTraits<AttributeType>::RequiresNormalization)
 					{
-						static_assert(TModels<CNormalizedAttribute, AttributeType>::Value, "Missing function implementations required for Attribute normalization");
+						static_assert(TModels_V<CNormalizedAttribute, AttributeType>, "Missing function implementations required for Attribute normalization");
 					}
 					
 					if constexpr (!UE::Anim::TAttributeTypeTraits<AttributeType>::StepInterpolate)
diff --git a/Engine/Source/Runtime/Engine/Public/Net/UnrealNetwork.h b/Engine/Source/Runtime/Engine/Public/Net/UnrealNetwork.h
--- a/Engine/Source/Runtime/Engine/Public/Net/UnrealNetwork.h
+++ b/Engine/Source/Runtime/Engine/Public/Net/UnrealNetwork.h
@@ -366,7 +366,7 @@ struct CGetFastArrayCreateReplicationFragmentFuncable
 
 #if UE_WITH_IRIS
 template<typename T>
-inline typename TEnableIf<TModels<CGetFastArrayCreateReplicationFragmentFuncable, T>::Value, const FDoRepLifetimeParams>::Type FixupParams(const FDoRepLifetimeParams& Params)
+inline typename TEnableIf<TModels_V<CGetFastArrayCreateReplicationFragmentFuncable, T>, const FDoRepLifetimeParams>::Type FixupParams(const FDoRepLifetimeParams& Params)
 {
 	// Use passed in CreateAndRegisterReplicationFragmentFunction if set
 	if (Params.CreateAndRegisterReplicationFragmentFunction)
@@ -381,7 +381,7 @@ inline typename TEnableIf<TModels<CGetFastArrayCreateReplicationFragmentFuncable
 #endif
 
 template<typename T>
-inline typename TEnableIf<!TModels<CGetFastArrayCreateReplicationFragmentFuncable, T>::Value, const FDoRepLifetimeParams&>::Type FixupParams(const FDoRepLifetimeParams& Params)
+inline typename TEnableIf<!TModels_V<CGetFastArrayCreateReplicationFragmentFuncable, T>, const FDoRepLifetimeParams&>::Type FixupParams(const FDoRepLifetimeParams& Params)
 {
 	return Params;
 }
diff --git a/Engine/Source/Runtime/Experimental/Chaos/Public/Chaos/AABBTree.h b/Engine/Source/Runtime/Experimental/Chaos/Public/Chaos/AABBTree.h
--- a/Engine/Source/Runtime/Experimental/Chaos/Public/Chaos/AABBTree.h
+++ b/Engine/Source/Runtime/Experimental/Chaos/Public/Chaos/AABBTree.h
@@ -127,13 +127,13 @@ struct CIsUpdatableElement
 	auto Requires(ElementT& InElem, const ElementT& InOtherElem) -> decltype(InElem.UpdateFrom(InOtherElem));
 };
 
-template<typename T, typename TEnableIf<!TModels<CIsUpdatableElement, T>::Value>::Type* = nullptr>
+template<typename T, typename TEnableIf<!TModels_V<CIsUpdatableElement, T>>::Type* = nullptr>
 static void UpdateElementHelper(T& InElem, const T& InFrom)
 {
 
 }
 
-template<typename T, typename TEnableIf<TModels<CIsUpdatableElement, T>::Value>::Type* = nullptr>
+template<typename T, typename TEnableIf<TModels_V<CIsUpdatableElement, T>>::Type* = nullptr>
 static void UpdateElementHelper(T& InElem, const T& InFrom)
 {
 	if (FAABBTreeCVars::UpdateDirtyElementPayloadData != 0)
diff --git a/Engine/Source/Runtime/Experimental/Chaos/Public/Chaos/BoundingVolumeHierarchy.h b/Engine/Source/Runtime/Experimental/Chaos/Public/Chaos/BoundingVolumeHierarchy.h
--- a/Engine/Source/Runtime/Experimental/Chaos/Public/Chaos/BoundingVolumeHierarchy.h
+++ b/Engine/Source/Runtime/Experimental/Chaos/Public/Chaos/BoundingVolumeHierarchy.h
@@ -63,7 +63,7 @@ class TBoundingVolumeHierarchy final : public ISpatialAcceleration<int32, T,d>
 	static constexpr bool DefaultAllowMultipleSplitting = false;
 	static constexpr bool DefaultUseVelocity = false;
 	static constexpr T DefaultDt = 0;
-	using TPayloadType = typename TBVHLeafTraits<LEAF_TYPE, TModels<CComplexBVHLeaf, LEAF_TYPE>::Value>::TPayloadType;
+	using TPayloadType = typename TBVHLeafTraits<LEAF_TYPE, TModels_V<CComplexBVHLeaf, LEAF_TYPE>>::TPayloadType;
 
 	CHAOS_API TBoundingVolumeHierarchy()
 		: MObjects(nullptr)
diff --git a/Engine/Source/Runtime/Experimental/Chaos/Public/Chaos/BoundingVolumeUtilities.h b/Engine/Source/Runtime/Experimental/Chaos/Public/Chaos/BoundingVolumeUtilities.h
--- a/Engine/Source/Runtime/Experimental/Chaos/Public/Chaos/BoundingVolumeUtilities.h
+++ b/Engine/Source/Runtime/Experimental/Chaos/Public/Chaos/BoundingVolumeUtilities.h
@@ -381,7 +381,7 @@ struct CParticleView
 
 //todo: how do we protect ourselves and make it const?
 template<typename ParticleView, typename T, int d>
-typename TEnableIf<TModels<CParticleView, ParticleView>::Value>::Type ComputeAllWorldSpaceBoundingBoxes(const ParticleView& Particles, const TArray<bool>& RequiresBounds, const bool bUseVelocity, const T Dt, TArray<TAABB<T, d>>& WorldSpaceBoxes)
+typename TEnableIf<TModels_V<CParticleView, ParticleView>>::Type ComputeAllWorldSpaceBoundingBoxes(const ParticleView& Particles, const TArray<bool>& RequiresBounds, const bool bUseVelocity, const T Dt, TArray<TAABB<T, d>>& WorldSpaceBoxes)
 {
 	WorldSpaceBoxes.AddUninitialized(Particles.Num());
 	ParticlesParallelFor(Particles, [&RequiresBounds, &WorldSpaceBoxes, bUseVelocity, Dt](const auto& Particle, int32 Index)
@@ -401,7 +401,7 @@ typename TEnableIf<TModels<CParticleView, ParticleView>::Value>::Type ComputeAll
 }
 
 template<typename ParticleView, typename T, int d>
-typename TEnableIf<!TModels<CParticleView, ParticleView>::Value>::Type ComputeAllWorldSpaceBoundingBoxes(const ParticleView& Particles, const TArray<bool>& RequiresBounds, const bool bUseVelocity, const T Dt, TArray<TAABB<T, d>>& WorldSpaceBoxes)
+typename TEnableIf<!TModels_V<CParticleView, ParticleView>>::Type ComputeAllWorldSpaceBoundingBoxes(const ParticleView& Particles, const TArray<bool>& RequiresBounds, const bool bUseVelocity, const T Dt, TArray<TAABB<T, d>>& WorldSpaceBoxes)
 {
 	WorldSpaceBoxes.AddUninitialized(Particles.Num());
 	ParticlesParallelFor(Particles, [&RequiresBounds, &WorldSpaceBoxes, bUseVelocity, Dt](const auto& Particle, int32 Index)
diff --git a/Engine/Source/Runtime/Experimental/Chaos/Public/Chaos/ChaosArchive.h b/Engine/Source/Runtime/Experimental/Chaos/Public/Chaos/ChaosArchive.h
--- a/Engine/Source/Runtime/Experimental/Chaos/Public/Chaos/ChaosArchive.h
+++ b/Engine/Source/Runtime/Experimental/Chaos/Public/Chaos/ChaosArchive.h
@@ -380,7 +380,7 @@ struct CSerializablePtr
 };
 
 template <typename T>
-constexpr typename TEnableIf<TModels<CSerializablePtr, T>::Value, bool>::Type IsSerializablePtr()
+constexpr typename TEnableIf<TModels_V<CSerializablePtr, T>, bool>::Type IsSerializablePtr()
 {
 	return true;
 }
diff --git a/Engine/Source/Runtime/Experimental/Iris/Core/Public/Iris/ReplicationSystem/Private/FastArrayReplicationFragmentInternal.h b/Engine/Source/Runtime/Experimental/Iris/Core/Public/Iris/ReplicationSystem/Private/FastArrayReplicationFragmentInternal.h
--- a/Engine/Source/Runtime/Experimental/Iris/Core/Public/Iris/ReplicationSystem/Private/FastArrayReplicationFragmentInternal.h
+++ b/Engine/Source/Runtime/Experimental/Iris/Core/Public/Iris/ReplicationSystem/Private/FastArrayReplicationFragmentInternal.h
@@ -78,7 +78,7 @@ struct FFastArrayReplicationFragmentHelper
 	 * We only want to do this for FastArrays that define PostReplicatedReceive since it might require extra work to calculate the required parameters
 	 */
 	template<typename FastArrayType>
-	static inline typename TEnableIf<TModels<FFastArraySerializer::CPostReplicatedReceiveFuncable, FastArrayType, const FFastArraySerializer::FPostReplicatedReceiveParameters&>::Value, void>::Type CallPostReplicatedReceiveOrNot(FastArrayType& ArraySerializer, bool bHasUnresolvedReferences)
+	static inline typename TEnableIf<TModels_V<FFastArraySerializer::CPostReplicatedReceiveFuncable, FastArrayType, const FFastArraySerializer::FPostReplicatedReceiveParameters&>, void>::Type CallPostReplicatedReceiveOrNot(FastArrayType& ArraySerializer, bool bHasUnresolvedReferences)
 	{
 		FFastArraySerializer::FPostReplicatedReceiveParameters PostReceivedParameters;
 		PostReceivedParameters.bHasMoreUnmappedReferences = bHasUnresolvedReferences;
@@ -86,7 +86,7 @@ struct FFastArrayReplicationFragmentHelper
 	}
 
 	template<typename FastArrayType>
-	static inline typename TEnableIf<!TModels<FFastArraySerializer::CPostReplicatedReceiveFuncable, FastArrayType, const FFastArraySerializer::FPostReplicatedReceiveParameters&>::Value, void>::Type CallPostReplicatedReceiveOrNot(FastArrayType& ArraySerializer, bool bHasUnresolvedReferences) {}
+	static inline typename TEnableIf<!TModels_V<FFastArraySerializer::CPostReplicatedReceiveFuncable, FastArrayType, const FFastArraySerializer::FPostReplicatedReceiveParameters&>, void>::Type CallPostReplicatedReceiveOrNot(FastArrayType& ArraySerializer, bool bHasUnresolvedReferences) {}
 };
 
 class FFastArrayReplicationFragmentBase : public FReplicationFragment
diff --git a/Engine/Source/Runtime/InteractiveToolsFramework/Public/InteractiveTool.h b/Engine/Source/Runtime/InteractiveToolsFramework/Public/InteractiveTool.h
--- a/Engine/Source/Runtime/InteractiveToolsFramework/Public/InteractiveTool.h
+++ b/Engine/Source/Runtime/InteractiveToolsFramework/Public/InteractiveTool.h
@@ -100,7 +100,7 @@ public:
 		// "not equal" function.
 
 		// Two-argument constructor, enabled only if PropType is CInequalityComparable
-		template<typename Q = PropType, typename = typename TEnableIf<TModels<CInequalityComparable, Q>::Value>::Type>
+		template<typename Q = PropType, typename = typename TEnableIf<TModels_V<CInequalityComparable, Q>>::Type>
 		TPropertyWatcher(const PropType& Property,
 						 FChangedCallback OnChangedIn)
 			: Cached(),
@@ -110,7 +110,7 @@ public:
 		{}
 
 		// Two-argument constructor, enabled only if PropType is CInequalityComparable
-		template<typename Q = PropType, typename = typename TEnableIf<TModels<CInequalityComparable, Q>::Value>::Type>
+		template<typename Q = PropType, typename = typename TEnableIf<TModels_V<CInequalityComparable, Q>>::Type>
 		TPropertyWatcher(FValueGetter GetValueIn,
 						 FChangedCallback OnChangedIn)
 			: Cached(),
diff --git a/Engine/Source/Runtime/Net/Core/Classes/Net/Serialization/FastArraySerializer.h b/Engine/Source/Runtime/Net/Core/Classes/Net/Serialization/FastArraySerializer.h
--- a/Engine/Source/Runtime/Net/Core/Classes/Net/Serialization/FastArraySerializer.h
+++ b/Engine/Source/Runtime/Net/Core/Classes/Net/Serialization/FastArraySerializer.h
@@ -227,7 +227,7 @@ private:
 	// Helper to always return a Type even if GetFastArrayItemTypePtr is not defined
 	static constexpr auto FastArrayTypePtr = []
 	{
-		if constexpr (TModels<CGetFastArrayItemTypeFuncable, FastArrayType>::Value)
+		if constexpr (TModels_V<CGetFastArrayItemTypeFuncable, FastArrayType>)
 		{
 			return FastArrayType::GetFastArrayItemTypePtr();
 		}
@@ -697,7 +697,7 @@ private:
 	
 		/** Conditionally invoke PostReplicatedReceive method depending on if is defined or not */
 		template<typename FastArrayType = SerializerType>
-		inline typename TEnableIf<TModels<CPostReplicatedReceiveFuncable, FastArrayType, const FFastArraySerializer::FPostReplicatedReceiveParameters>::Value, void>::Type CallPostReplicatedReceiveOrNot(int32 OldArraySize)
+		inline typename TEnableIf<TModels_V<CPostReplicatedReceiveFuncable, FastArrayType, const FFastArraySerializer::FPostReplicatedReceiveParameters>, void>::Type CallPostReplicatedReceiveOrNot(int32 OldArraySize)
 		{
 			FFastArraySerializer::FPostReplicatedReceiveParameters PostReceivedParameters;
 			PostReceivedParameters.OldArraySize = OldArraySize;
@@ -706,7 +706,7 @@ private:
 		}
 
 		template<typename FastArrayType = SerializerType>
-		inline typename TEnableIf<!TModels<CPostReplicatedReceiveFuncable, FastArrayType, const FFastArraySerializer::FPostReplicatedReceiveParameters>::Value, void>::Type CallPostReplicatedReceiveOrNot(int32) {}
+		inline typename TEnableIf<!TModels_V<CPostReplicatedReceiveFuncable, FastArrayType, const FFastArraySerializer::FPostReplicatedReceiveParameters>, void>::Type CallPostReplicatedReceiveOrNot(int32) {}
 
 		// Validate that deduced FastArrayItemType is valid and that it is the same as the specified one		
 		static_assert(std::is_same_v<typename TFastArrayTypeHelper<SerializerType>::FastArrayItemType, Type>, "Auto deduced FastArrayItemType is invalid or differs from the specified type. Make sure that the FastArraySerializer has a single replicated array property.");
diff --git a/Engine/Source/Runtime/Net/Core/Public/Net/Core/Analytics/NetStatsUtils.h b/Engine/Source/Runtime/Net/Core/Public/Net/Core/Analytics/NetStatsUtils.h
--- a/Engine/Source/Runtime/Net/Core/Public/Net/Core/Analytics/NetStatsUtils.h
+++ b/Engine/Source/Runtime/Net/Core/Public/Net/Core/Analytics/NetStatsUtils.h
@@ -44,7 +44,7 @@ class TSampleProducer
 		);
 	};
 
-	static_assert(TModels<CSampleConsumer, ConsumerType>::Value, "ConsumerType must implement AddSample(double Val)");
+	static_assert(TModels_V<CSampleConsumer, ConsumerType>, "ConsumerType must implement AddSample(double Val)");
 
 public:
 	/**
@@ -77,7 +77,7 @@ protected:
 	 * @param Value		The measurement value being collected for sampling
 	 */
 
-	template<typename InnerConsumerType=ConsumerType, typename = typename TEnableIf<TModels<CPeekMeasurements, InnerConsumerType>::Value>::Type>
+	template<typename InnerConsumerType=ConsumerType, typename = typename TEnableIf<TModels_V<CPeekMeasurements, InnerConsumerType>>::Type>
 	inline void PeekMeasurement(double TimeVal, double Value)
 	{
 		if (Consumer != nullptr)
@@ -86,7 +86,7 @@ protected:
 		}
 	}
 
-	template<typename InnerConsumerType=ConsumerType, typename = typename TEnableIf<!TModels<CPeekMeasurements, InnerConsumerType>::Value>::Type, int32 UnusedParamForODR=0>
+	template<typename InnerConsumerType=ConsumerType, typename = typename TEnableIf<!TModels_V<CPeekMeasurements, InnerConsumerType>>::Type, int32 UnusedParamForODR=0>
 	inline void PeekMeasurement(double TimeVal, double Value)
 	{
 	}
diff --git a/Engine/Source/Runtime/Net/Core/Public/Net/Core/Connection/EscalationStates.h b/Engine/Source/Runtime/Net/Core/Public/Net/Core/Connection/EscalationStates.h
--- a/Engine/Source/Runtime/Net/Core/Public/Net/Core/Connection/EscalationStates.h
+++ b/Engine/Source/Runtime/Net/Core/Public/Net/Core/Connection/EscalationStates.h
@@ -631,13 +631,13 @@ class TEscalationManager final : public FEscalationManager
 	};
 
 	// 'typename = typename TEnableIf' does not work here, unusually
-	template<typename T, typename TEnableIf<TModels<CHasEnumPreallocNum, T>::Value>::Type* = nullptr>
+	template<typename T, typename TEnableIf<TModels_V<CHasEnumPreallocNum, T>>::Type* = nullptr>
 	inline static constexpr int32 GetCounterNumPrealloc()
 	{
 		return FMath::Max((int32)T::NumPrealloc, (int32)T::Max);
 	}
 
-	template<typename T, typename TEnableIf<!TModels<CHasEnumPreallocNum, T>::Value>::Type* = nullptr>
+	template<typename T, typename TEnableIf<!TModels_V<CHasEnumPreallocNum, T>>::Type* = nullptr>
 	inline static constexpr int32 GetCounterNumPrealloc()
 	{
 		return (int32)T::Max;
diff --git a/Engine/Source/Runtime/RigVM/Public/RigVMCore/RigVMExternalVariable.h b/Engine/Source/Runtime/RigVM/Public/RigVMCore/RigVMExternalVariable.h
--- a/Engine/Source/Runtime/RigVM/Public/RigVMCore/RigVMExternalVariable.h
+++ b/Engine/Source/Runtime/RigVM/Public/RigVMCore/RigVMExternalVariable.h
@@ -472,7 +472,7 @@ struct RIGVM_API FRigVMExternalVariable
 
 	template <
 		typename T,
-		typename TEnableIf<TModels<CRigVMUStruct, T>::Value>::Type * = nullptr
+		typename TEnableIf<TModels_V<CRigVMUStruct, T>>::Type * = nullptr
 	>
 	static FRigVMExternalVariable Make(const FName& InName, T& InValue)
 	{
@@ -488,7 +488,7 @@ struct RIGVM_API FRigVMExternalVariable
 
 	template <
 		typename T,
-		typename TEnableIf<TModels<CRigVMUStruct, T>::Value>::Type * = nullptr
+		typename TEnableIf<TModels_V<CRigVMUStruct, T>>::Type * = nullptr
 	>
 	static FRigVMExternalVariable Make(const FName& InName, TArray<T>& InValue)
 	{
@@ -504,7 +504,7 @@ struct RIGVM_API FRigVMExternalVariable
 
 	template <
 		typename T,
-		typename TEnableIf<TModels<CRigVMUClass, T>::Value>::Type * = nullptr
+		typename TEnableIf<TModels_V<CRigVMUClass, T>>::Type * = nullptr
 	>
 	static FRigVMExternalVariable Make(const FName& InName, T& InValue)
 	{
@@ -520,7 +520,7 @@ struct RIGVM_API FRigVMExternalVariable
 
 	template <
 		typename T,
-		typename TEnableIf<TModels<CRigVMUClass, T>::Value>::Type * = nullptr
+		typename TEnableIf<TModels_V<CRigVMUClass, T>>::Type * = nullptr
 	>
 	static FRigVMExternalVariable Make(const FName& InName, TArray<T>& InValue)
 	{
diff --git a/Engine/Source/Runtime/RigVM/Public/RigVMCore/RigVMMemoryStorage.h b/Engine/Source/Runtime/RigVM/Public/RigVMCore/RigVMMemoryStorage.h
--- a/Engine/Source/Runtime/RigVM/Public/RigVMCore/RigVMMemoryStorage.h
+++ b/Engine/Source/Runtime/RigVM/Public/RigVMCore/RigVMMemoryStorage.h
@@ -516,7 +516,7 @@ public:
 	// returns true if the handle is of the given type
 	template <
 		typename T,
-		typename TEnableIf<TModels<CRigVMUStruct, T>::Value>::Type * = nullptr
+		typename TEnableIf<TModels_V<CRigVMUStruct, T>>::Type * = nullptr
 	>
 	bool IsType(bool bIsHiddenArgument = false) const
 	{
@@ -526,7 +526,7 @@ public:
 	// returns true if the handle is of the given type
 	template <
 		typename T,
-		typename TEnableIf<TModels<CRigVMUStruct, T>::Value>::Type * = nullptr
+		typename TEnableIf<TModels_V<CRigVMUStruct, T>>::Type * = nullptr
 	>
 	bool IsTypeArray(bool bIsHiddenArgument = false) const
 	{
@@ -556,7 +556,7 @@ public:
 	// returns true if the handle is of the given type
 	template <
 		typename T,
-		typename TEnableIf<TModels<CRigVMUClass, T>::Value>::Type * = nullptr
+		typename TEnableIf<TModels_V<CRigVMUClass, T>>::Type * = nullptr
 	>
 	bool IsType(bool bIsHiddenArgument = false) const
 	{
@@ -566,7 +566,7 @@ public:
 	// returns true if the handle is of the given type
 	template <
 		typename T,
-		typename TEnableIf<TModels<CRigVMUClass, T>::Value>::Type * = nullptr
+		typename TEnableIf<TModels_V<CRigVMUClass, T>>::Type * = nullptr
 	>
 	bool IsTypeArray(bool bIsHiddenArgument = false) const
 	{
@@ -716,7 +716,7 @@ public:
 	// Returns true if this memory handle maps to a struct property
 	template <
 		typename T,
-		typename TEnableIf<TModels<CRigVMUStruct, T>::Value>::Type * = nullptr
+		typename TEnableIf<TModels_V<CRigVMUStruct, T>>::Type * = nullptr
 	>
 	bool IsStruct(bool bIsHiddenArgument = false) const
 	{
@@ -736,7 +736,7 @@ public:
 	// Returns true if this memory handle maps to a struct array property
 	template <
 		typename T,
-		typename TEnableIf<TModels<CRigVMUStruct, T>::Value>::Type * = nullptr
+		typename TEnableIf<TModels_V<CRigVMUStruct, T>>::Type * = nullptr
 	>
 	bool IsStructArray(bool bIsHiddenArgument = false) const
 	{
@@ -768,7 +768,7 @@ public:
 	// Returns true if this memory handle maps to an object property
 	template <
 		typename T,
-		typename TEnableIf<TModels<CRigVMUClass, T>::Value>::Type * = nullptr
+		typename TEnableIf<TModels_V<CRigVMUClass, T>>::Type * = nullptr
 	>
 	bool IsObject(bool bIsHiddenArgument = false) const
 	{
@@ -778,7 +778,7 @@ public:
 	// Returns true if this memory handle maps to an objct array property
 	template <
 		typename T,
-		typename TEnableIf<TModels<CRigVMUClass, T>::Value>::Type * = nullptr
+		typename TEnableIf<TModels_V<CRigVMUClass, T>>::Type * = nullptr
 	>
 	bool IsObjectArray(bool bIsHiddenArgument = false) const
 	{
diff --git a/Engine/Source/Runtime/RigVM/Public/RigVMCore/RigVMNativized.h b/Engine/Source/Runtime/RigVM/Public/RigVMCore/RigVMNativized.h
--- a/Engine/Source/Runtime/RigVM/Public/RigVMCore/RigVMNativized.h
+++ b/Engine/Source/Runtime/RigVM/Public/RigVMCore/RigVMNativized.h
@@ -328,7 +328,7 @@ public:
 
 	template <
 		typename T,
-		typename TEnableIf<TModels<CRigVMUStruct, T>::Value>::Type* = nullptr
+		typename TEnableIf<TModels_V<CRigVMUStruct, T>>::Type* = nullptr
 	>
 	static T GetStructConstant(const FString& InDefaultValue)
 	{
@@ -358,7 +358,7 @@ public:
 
 	template <
 		typename T,
-		typename TEnableIf<TModels<CRigVMUStruct, T>::Value>::Type* = nullptr
+		typename TEnableIf<TModels_V<CRigVMUStruct, T>>::Type* = nullptr
 	>
 	static TArray<T> GetStructArrayConstant(const FString& InDefaultValue)
 	{
diff --git a/Engine/Source/Runtime/RigVM/Public/RigVMCore/RigVMRegistry.h b/Engine/Source/Runtime/RigVM/Public/RigVMCore/RigVMRegistry.h
--- a/Engine/Source/Runtime/RigVM/Public/RigVMCore/RigVMRegistry.h
+++ b/Engine/Source/Runtime/RigVM/Public/RigVMCore/RigVMRegistry.h
@@ -132,7 +132,7 @@ public:
 	// Returns the type index given a struct
 	template <
 		typename T,
-		typename TEnableIf<TModels<CRigVMUStruct, T>::Value>::Type * = nullptr
+		typename TEnableIf<TModels_V<CRigVMUStruct, T>>::Type * = nullptr
 	>
 	TRigVMTypeIndex GetTypeIndex(bool bAsArray = false) const
 	{
@@ -147,7 +147,7 @@ public:
 	// Returns the type index given an object
 	template <
 		typename T,
-		typename TEnableIf<TModels<CRigVMUClass, T>::Value>::Type * = nullptr
+		typename TEnableIf<TModels_V<CRigVMUClass, T>>::Type * = nullptr
 	>
 	TRigVMTypeIndex GetTypeIndex(bool bAsArray = false) const
 	{
diff --git a/Engine/Source/Runtime/RigVM/Public/RigVMCore/RigVMStruct.h b/Engine/Source/Runtime/RigVM/Public/RigVMCore/RigVMStruct.h
--- a/Engine/Source/Runtime/RigVM/Public/RigVMCore/RigVMStruct.h
+++ b/Engine/Source/Runtime/RigVM/Public/RigVMCore/RigVMStruct.h
@@ -110,7 +110,7 @@ public:
 	/** Returns the name of the first variable given a type (or NAME_None if not found) */
 	template <
 		typename T,
-		typename TEnableIf<TModels<CRigVMUStruct, T>::Value>::Type * = nullptr
+		typename TEnableIf<TModels_V<CRigVMUStruct, T>>::Type * = nullptr
 	>
 	FName FindFirstVariableOfType() const
 	{
@@ -120,7 +120,7 @@ public:
 	/** Returns the name of the first variable given a type (or NAME_None if not found) */
 	template <
 		typename T,
-		typename TEnableIf<TModels<CRigVMUClass, T>::Value>::Type * = nullptr
+		typename TEnableIf<TModels_V<CRigVMUClass, T>>::Type * = nullptr
 	>
 	FName FindFirstVariableOfType() const
 	{
@@ -213,7 +213,7 @@ public:
 	
 	template <
 		typename T,
-		typename TEnableIf<TModels<CRigVMUStruct, T>::Value>::Type * = nullptr
+		typename TEnableIf<TModels_V<CRigVMUStruct, T>>::Type * = nullptr
 	>
 	static FString ExportToFullyQualifiedText(const T& InStructValue)
 	{
diff --git a/Engine/Source/Runtime/RigVM/Public/RigVMCore/RigVMTraits.h b/Engine/Source/Runtime/RigVM/Public/RigVMCore/RigVMTraits.h
--- a/Engine/Source/Runtime/RigVM/Public/RigVMCore/RigVMTraits.h
+++ b/Engine/Source/Runtime/RigVM/Public/RigVMCore/RigVMTraits.h
@@ -230,7 +230,7 @@ void RigVMCopy(void* InTargetPtr, const void* InSourcePtr, int32 InCount = 1)
 
 template <
 	typename T,
-	typename TEnableIf<TModels<CRigVMUStruct, T>::Value>::Type* = nullptr
+	typename TEnableIf<TModels_V<CRigVMUStruct, T>>::Type* = nullptr
 >
 void RigVMInitialize(void* InPtr, int32 InCount)
 {
@@ -240,7 +240,7 @@ void RigVMInitialize(void* InPtr, int32 InCount)
 
 template <
 	typename T,
-	typename TEnableIf<TModels<CRigVMUStruct, T>::Value>::Type* = nullptr
+	typename TEnableIf<TModels_V<CRigVMUStruct, T>>::Type* = nullptr
 >
 void RigVMDestroy(void* InPtr, int32 InCount)
 {
@@ -250,7 +250,7 @@ void RigVMDestroy(void* InPtr, int32 InCount)
 
 template <
 	typename T,
-	typename TEnableIf<TModels<CRigVMUStruct, T>::Value>::Type* = nullptr
+	typename TEnableIf<TModels_V<CRigVMUStruct, T>>::Type* = nullptr
 >
 void RigVMCopy(void* InTargetPtr, const void* InSourcePtr, int32 InCount = 1)
 {
