Add a way for objects to specify property overrides local to the save context.
The initial implementation is simplistic and allow to mark object property transient without much else

diff --git a/Engine/Source/Runtime/CoreUObject/Private/UObject/SavePackage/PackageHarvester.cpp b/Engine/Source/Runtime/CoreUObject/Private/UObject/SavePackage/PackageHarvester.cpp
--- a/Engine/Source/Runtime/CoreUObject/Private/UObject/SavePackage/PackageHarvester.cpp
+++ b/Engine/Source/Runtime/CoreUObject/Private/UObject/SavePackage/PackageHarvester.cpp
@@ -186,6 +186,8 @@ FPackageHarvester::FPackageHarvester(FSaveContext& InContext)
 	this->SetCookData(SaveContext.GetSaveArgs().ArchiveCookData);
 	this->SetSerializeContext(SaveContext.GetSerializeContext());
 	this->SetUseUnversionedPropertySerialization(SaveContext.IsSaveUnversionedProperties());
+
+	ResolveOverrides();
 }
 
 FPackageHarvester::FExportWithContext FPackageHarvester::PopExportToProcess()
@@ -573,6 +575,42 @@ FArchive& FPackageHarvester::operator<<(FName& Name)
 	return *this;
 }
 
+bool FPackageHarvester::ShouldSkipProperty(const FProperty* InProperty) const
+{
+	const TSet<FProperty*>* Props = TransientPropertyOverrides.Find(CurrentExportDependencies.CurrentExport);
+	if (Props && Props->Contains(InProperty))
+	{
+		return true;
+	}
+	return false;
+}
+
+void FPackageHarvester::ResolveOverrides()
+{
+	for (auto& PairObjectOverrides : SaveContext.GetObjectSaveContext().SaveOverrides)
+	{
+		TSet<FProperty*> Props;
+		for (const FPropertySaveOverride& Override : PairObjectOverrides.Value.PropOverrides)
+		{
+			if (FProperty* Prop = Override.bMarkTransient ? CastField<FProperty>(Override.PropertyPath.GetTyped(FProperty::StaticClass())) : nullptr)
+			{
+				// We currently only support object property
+				if (ensureAlwaysMsgf(Prop->IsA<FObjectProperty>(), TEXT("Save Overrides supports only object property at the moment. Name: %s, Type: %s"), *Prop->GetName(), *Prop->GetClass()->GetName()))
+				{
+					// Harvest the name of the property, since it is only made transient for the purpose of the package harvest, it will be needed for the LinkerSave
+					HarvestExportDataName(Prop->GetFName());
+
+					Props.Add(Prop);
+				}
+			}
+		}
+		if (!Props.IsEmpty())
+		{
+			TransientPropertyOverrides.Add(PairObjectOverrides.Key, MoveTemp(Props));
+		}
+	}
+}
+
 void FPackageHarvester::HarvestDependency(UObject* InObj, bool bIsNative)
 {
 	// if we aren't currently processing an export or the referenced object is a package, do not harvest the dependency
diff --git a/Engine/Source/Runtime/CoreUObject/Private/UObject/SavePackage/PackageHarvester.h b/Engine/Source/Runtime/CoreUObject/Private/UObject/SavePackage/PackageHarvester.h
--- a/Engine/Source/Runtime/CoreUObject/Private/UObject/SavePackage/PackageHarvester.h
+++ b/Engine/Source/Runtime/CoreUObject/Private/UObject/SavePackage/PackageHarvester.h
@@ -2,6 +2,7 @@
 
 #pragma once
 
+#include "Containers/Map.h"
 #include "Containers/Set.h"
 #include "Containers/Queue.h"
 #include "Serialization/ArchiveUObject.h"
@@ -90,7 +91,11 @@ public:
 	virtual FArchive& operator<<(FSoftObjectPath& Value) override;
 	virtual FArchive& operator<<(FName& Name) override;
 
+	virtual bool ShouldSkipProperty(const FProperty* InProperty) const;
+
 private:
+	void ResolveOverrides();
+
 	ESaveRealm GetObjectHarvestingRealm(UObject* InObject, EIllegalRefReason& OutReason) const;
 
 	void HarvestExport(UObject* InObject, ESaveRealm InContext);
@@ -112,6 +117,7 @@ private:
 
 	TQueue<FExportWithContext> ExportsToProcess;
 	FExportDependencies CurrentExportDependencies;
+	TMap<UObject*, TSet<FProperty*>> TransientPropertyOverrides;
 	ESaveRealm CurrentExportHarvestingRealm;
 	//@todo: bIsEditorOnlyExportOnStack can be probably be folded in CurrentExportHarvestingContext
 	bool bIsEditorOnlyExportOnStack;
diff --git a/Engine/Source/Runtime/CoreUObject/Public/UObject/ObjectSaveContext.h b/Engine/Source/Runtime/CoreUObject/Public/UObject/ObjectSaveContext.h
--- a/Engine/Source/Runtime/CoreUObject/Public/UObject/ObjectSaveContext.h
+++ b/Engine/Source/Runtime/CoreUObject/Public/UObject/ObjectSaveContext.h
@@ -4,6 +4,7 @@
 
 #include "Containers/UnrealString.h"
 #include "HAL/Platform.h"
+#include "UObject/ObjectSaveOverride.h"
 
 class FPackagePath;
 class ITargetPlatform;
@@ -63,13 +64,14 @@ struct FObjectSaveContextData
 	/** Set to false if the save failed, before calling any PostSaves. */
 	bool bSaveSucceeded = true;
 
-
 	// Per-object Output variables; writable from PreSave functions, readable from PostSave functions
 
+	/** List of property overrides per object to apply to during save */
+	TMap<UObject*, FObjectSaveOverride> SaveOverrides;
+
 	/** A bool that can be set from PreSave to indicate PostSave needs to take some extra cleanup steps. */
 	bool bCleanupRequired = false;
 
-
 	// Variables set/read per call to PreSave/PostSave functions
 	/** PreSave contract enforcement; records whether PreSave is overridden. */
 	int32 NumRefPasses = 0;
@@ -130,6 +132,15 @@ public:
 	 */
 	bool IsFirstConcurrentSave() const { return Data.bOuterConcurrentSave; }
 
+	/**
+	 * Add a save override to specific object. (i.e. mark certain properties transient for this save)
+	 * @note only object property are supported at the moment
+	 */
+	void AddSaveOverride(UObject* Target, FObjectSaveOverride InOverride)
+	{
+		Data.SaveOverrides.Add(Target, MoveTemp(InOverride));
+	}
+
 protected:
 	FObjectSaveContextData& Data;
 	friend class UObject;
diff --git a/Engine/Source/Runtime/CoreUObject/Public/UObject/ObjectSaveOverride.h b/Engine/Source/Runtime/CoreUObject/Public/UObject/ObjectSaveOverride.h
new file mode 100644
--- /dev/null
+++ b/Engine/Source/Runtime/CoreUObject/Public/UObject/ObjectSaveOverride.h
@@ -0,0 +1,24 @@
+// Copyright Epic Games, Inc. All Rights Reserved.
+
+#pragma once
+
+#include "Containers/Array.h"
+#include "UObject/FieldPath.h"
+#include "UObject/ObjectMacros.h"
+
+/** 
+ * Data collected during SavePackage that modifies the EPropertyFlags for a single FProperty on a single object instance when that object is serialized by SavePackage.
+ * The specified changes apply during both the harvesting phase (discovery of referenced imports and exports) and the serialization to disk phase.
+ * @note currently only support marking a reference property transient (i.e object pointers)
+ */
+struct FPropertySaveOverride
+{
+	FFieldPath PropertyPath;
+	bool bMarkTransient;
+};
+
+/** Data to specify an override to apply to an object during save without mutating the object itself. */
+struct FObjectSaveOverride
+{
+	TArray<FPropertySaveOverride> PropOverrides;
+};
\ No newline at end of file
