PackageName DirectoryTree Bugfix: Downgrade assertion to an Error and give more diagnostic information.
#rnx

diff --git a/Engine/Source/Runtime/CoreUObject/Private/Misc/PackageName.cpp b/Engine/Source/Runtime/CoreUObject/Private/Misc/PackageName.cpp
--- a/Engine/Source/Runtime/CoreUObject/Private/Misc/PackageName.cpp
+++ b/Engine/Source/Runtime/CoreUObject/Private/Misc/PackageName.cpp
@@ -1536,7 +1536,20 @@ bool FPackageName::TryGetMountPointForPath(FStringView InFilePathOrPackageName,
 		const FMountPoint* MountPoint = *MountPointPtr;
 		if (!FPathViews::TryMakeChildPathRelativeTo(PossibleAbsFilePath, MountPoint->ContentPathAbsolute, RelPath))
 		{
-			verify(FPathViews::TryMakeChildPathRelativeTo(PossibleAbsFilePath, MountPoint->ContentPathRelative, RelPath));
+			if (!FPathViews::TryMakeChildPathRelativeTo(PossibleAbsFilePath, MountPoint->ContentPathRelative, RelPath))
+			{
+				UE_LOG(LogPackageName, Error, TEXT("TryGetMountPointForPath failed: found a mountpoint, but TryMakeChildPathRelativeTo failed. Path=%s, MountAbsPath=%s, MountRelPath=%s"),
+					*PossibleAbsFilePath, *MountPoint->ContentPathAbsolute, *MountPoint->ContentPathRelative);
+				if (OutFlexNameType)
+				{
+					*OutFlexNameType = EFlexNameType::Invalid;
+				}
+				if (OutFailureReason)
+				{
+					*OutFailureReason = EErrorCode::PackageNamePathNotMounted;
+				}
+				return false;
+			}
 		}
 		OutMountPointPackageName << MountPoint->RootPath;
 		OutMountPointFilePath << MountPoint->ContentPathRelative;
