Reverted the change to CoreNetTypes until all the buildfarm jobs are updated to not compile the C++ UHT.

diff --git a/Engine/Source/Runtime/CoreUObject/Public/UObject/CoreNetTypes.h b/Engine/Source/Runtime/CoreUObject/Public/UObject/CoreNetTypes.h
--- a/Engine/Source/Runtime/CoreUObject/Public/UObject/CoreNetTypes.h
+++ b/Engine/Source/Runtime/CoreUObject/Public/UObject/CoreNetTypes.h
@@ -4,8 +4,6 @@
 
 #include "UObject/ObjectMacros.h"
 
-#include "CoreNetTypes.generated.h"
-
 /** Whether or not caching of actor/subobject names to the stack should be enabled, for async demo recording crashdumps */
 #ifndef UE_NET_REPACTOR_NAME_DEBUG
 	#define UE_NET_REPACTOR_NAME_DEBUG 0
