PackageName DirectoryTree Bugfix: Downgrade all cases where DirectoryTree finds a childdir but TryMakeChildPathRelativeTo does not to a silent PackageNamePathNotMounted. It can occur for bad paths such as d:/Dir//Root. Silent failure is the same behavior as pre-DirectoryTree for those cases.
#rnx

diff --git a/Engine/Source/Runtime/CoreUObject/Private/Misc/PackageName.cpp b/Engine/Source/Runtime/CoreUObject/Private/Misc/PackageName.cpp
--- a/Engine/Source/Runtime/CoreUObject/Private/Misc/PackageName.cpp
+++ b/Engine/Source/Runtime/CoreUObject/Private/Misc/PackageName.cpp
@@ -736,12 +736,18 @@ void FPackageName::InternalFilenameToLongPackageName(FStringView InFilename, FSt
 			return false;
 		}
 		const FMountPoint* MountPoint = *MountPointPtr;
-		OutPackageName << MountPoint->RootPath;
 		FStringView RelPath;
 		if (!FPathViews::TryMakeChildPathRelativeTo(SearchName, MountPoint->ContentPathRelative, RelPath))
 		{
-			verify(FPathViews::TryMakeChildPathRelativeTo(SearchName, MountPoint->ContentPathAbsolute, RelPath));
+			if (!FPathViews::TryMakeChildPathRelativeTo(SearchName, MountPoint->ContentPathAbsolute, RelPath))
+			{
+				// DirectoryTree is less conservative than TryMakeChildPathRelativeTo; e.g. d:/Dir//Root will match
+				// d:/Dir/Root/Child in DirectoryTree but not in TryMakeChildPathRelativeTo. Treat this as PathNotMounted,
+				// since no mountdir will match it according to TryMakeChildPathRelativeTo.
+				return false;
+			}
 		}
+		OutPackageName << MountPoint->RootPath;
 		FPathViews::AppendPath(OutPackageName, RelPath);
 		return true;
 	};
@@ -881,7 +887,13 @@ bool FPackageName::TryConvertLongPackageNameToFilename(const FString& InLongPack
 
 	const FMountPoint* MountPoint = *MountPointPtr;
 	FStringView RelPath;
-	verify(FPathViews::TryMakeChildPathRelativeTo(InLongPackageName, MountPoint->RootPath, RelPath));
+	if (!FPathViews::TryMakeChildPathRelativeTo(InLongPackageName, MountPoint->RootPath, RelPath))
+	{
+		// DirectoryTree is less conservative than TryMakeChildPathRelativeTo; e.g. d:/Dir//Root will match
+		// d:/Dir/Root/Child in DirectoryTree but not in TryMakeChildPathRelativeTo. Treat this as PathNotMounted,
+		// since no mountdir will match it according to TryMakeChildPathRelativeTo.
+		return false;
+	}
 	TStringBuilder<256> Builder;
 	Builder << MountPoint->ContentPathRelative;
 	FPathViews::AppendPath(Builder, RelPath);
@@ -1508,7 +1520,21 @@ bool FPackageName::TryGetMountPointForPath(FStringView InFilePathOrPackageName,
 	if (MountPointPtr)
 	{
 		const FMountPoint* MountPoint = *MountPointPtr;
-		verify(FPathViews::TryMakeChildPathRelativeTo(InFilePathOrPackageName, MountPoint->RootPath, RelPath));
+		if (!FPathViews::TryMakeChildPathRelativeTo(InFilePathOrPackageName, MountPoint->RootPath, RelPath))
+		{
+			// DirectoryTree is less conservative than TryMakeChildPathRelativeTo; e.g. d:/Dir//Root will match
+			// d:/Dir/Root/Child in DirectoryTree but not in TryMakeChildPathRelativeTo. Treat this as PathNotMounted,
+			// since no mountdir will match it according to TryMakeChildPathRelativeTo.
+			if (OutFlexNameType)
+			{
+				*OutFlexNameType = EFlexNameType::Invalid;
+			}
+			if (OutFailureReason)
+			{
+				*OutFailureReason = EErrorCode::PackageNamePathNotMounted;
+			}
+			return false;
+		}
 		OutMountPointPackageName << MountPoint->RootPath;
 		OutMountPointFilePath << MountPoint->ContentPathRelative;
 		OutRelPath << RelPath;
@@ -1538,8 +1564,9 @@ bool FPackageName::TryGetMountPointForPath(FStringView InFilePathOrPackageName,
 		{
 			if (!FPathViews::TryMakeChildPathRelativeTo(PossibleAbsFilePath, MountPoint->ContentPathRelative, RelPath))
 			{
-				UE_LOG(LogPackageName, Error, TEXT("TryGetMountPointForPath failed: found a mountpoint, but TryMakeChildPathRelativeTo failed. Path=%s, MountAbsPath=%s, MountRelPath=%s"),
-					*PossibleAbsFilePath, *MountPoint->ContentPathAbsolute, *MountPoint->ContentPathRelative);
+				// DirectoryTree is less conservative than TryMakeChildPathRelativeTo; e.g. d:/Dir//Root will match
+				// d:/Dir/Root/Child in DirectoryTree but not in TryMakeChildPathRelativeTo. Treat this as PathNotMounted,
+				// since no mountdir will match it according to TryMakeChildPathRelativeTo.
 				if (OutFlexNameType)
 				{
 					*OutFlexNameType = EFlexNameType::Invalid;
@@ -2538,6 +2565,40 @@ IMPLEMENT_SIMPLE_AUTOMATION_TEST(FPackageNameTests, "System.Core.Misc.PackageNam
 
 bool FPackageNameTests::RunTest(const FString& Parameters)
 {
+	auto AggressiveNormalize = [](FStringView Path)
+	{
+		FString Norm(Path);
+		Norm.ReplaceCharInline('\\', '/');
+		// Strip trailing /
+		if (Norm.Len() >= 2)
+		{
+			TCHAR Last = Norm[Norm.Len() - 1];
+			TCHAR Next = Norm[Norm.Len() - 2];
+			if (Last == '/' && (Next != '/' && Next != ':'))
+			{
+				Norm.LeftChopInline(1);
+			}
+		}
+		// Replace // with / except at beginning
+		if (Norm.StartsWith(TEXT("//")))
+		{
+			FString Suffix = Norm.RightChop(2);
+			Suffix.ReplaceInline(TEXT("//"), TEXT("/"));
+			Norm = TEXT("//") + Suffix;
+		}
+		else
+		{
+			Norm.ReplaceInline(TEXT("//"), TEXT("/"));
+		}
+		return Norm;
+	};
+	auto IsSamePath = [&AggressiveNormalize](FStringView A, FStringView B)
+	{
+		FString NormA(AggressiveNormalize(A));
+		FString NormB(AggressiveNormalize(B));
+		return FPaths::IsSamePath(NormA, NormB);
+	};
+
 	// Localized paths tests
 	{
 		auto TestIsLocalizedPackage = [&](const FString& InPath, const bool InExpected)
@@ -2671,6 +2732,69 @@ bool FPackageNameTests::RunTest(const FString& Parameters)
 			FPackageName::EFlexNameType::Invalid, FPackageName::EErrorCode::PackageNameEmptyPath);
 	}
 
+	// TryGetMountPointForPath
+	{
+		auto TestTryGetMountPointForPath = [this, &IsSamePath](FStringView InPath, bool bExpectedResult, FStringView ExpectedPackageName,
+			FStringView ExpectedFilePath, FStringView ExpectedRelPath, FPackageName::EFlexNameType ExpectedFlexNameType,
+			FPackageName::EErrorCode ExpectedFailureReason)
+		{
+			bool bActualResult;
+			TStringBuilder<256> ActualMountPointPackageName;
+			TStringBuilder<256> ActualMountPointFilePath;
+			TStringBuilder<256> ActualRelPath;
+			FPackageName::EFlexNameType ActualFlexNameType;
+			FPackageName::EErrorCode ActualFailureReason;
+
+			bActualResult = FPackageName::TryGetMountPointForPath(InPath, ActualMountPointPackageName,
+				ActualMountPointFilePath, ActualRelPath, &ActualFlexNameType, &ActualFailureReason);
+			bool bActualMatchesExpected = bActualResult == bExpectedResult &&
+				ActualFlexNameType == ExpectedFlexNameType && ActualFailureReason == ExpectedFailureReason;
+			if (bActualMatchesExpected && bExpectedResult)
+			{
+				bActualMatchesExpected =
+					IsSamePath(FString(ExpectedPackageName), FString(ActualMountPointPackageName)) &&
+					IsSamePath(FString(ExpectedFilePath), FString(ActualMountPointFilePath)) &&
+					IsSamePath(FString(ExpectedRelPath), FString(ActualRelPath));
+			}
+			if (!bActualMatchesExpected)
+			{
+				if (!bActualResult)
+				{
+					ActualMountPointPackageName.Reset();
+					ActualMountPointFilePath.Reset();
+					ActualRelPath.Reset();
+				}
+				AddError(FString::Printf(TEXT("Path '%.*s' failed FPackageName::TestTryGetMountPointForPath\n")
+					TEXT("got      %s,'%s','%s','%s',%d,%d,\nexpected %s,'%s','%s','%s',%d,%d."),
+					InPath.Len(), InPath.GetData(),
+					(bActualResult ? TEXT("true") : TEXT("false")),
+					*FString(ActualMountPointPackageName), *FString(ActualMountPointFilePath),
+					*FString(ActualRelPath), (int32)ActualFlexNameType, (int32)ActualFailureReason,
+					(bExpectedResult ? TEXT("true") : TEXT("false")),
+					*FString(ExpectedPackageName), *FString(ExpectedFilePath),
+					*FString(ExpectedRelPath), (int32)ExpectedFlexNameType, (int32)ExpectedFailureReason));
+			}
+		};
+
+		TestTryGetMountPointForPath(FPaths::ProjectContentDir(), true, TEXT("/Game"), FPaths::ProjectContentDir(),
+			TEXT(""), FPackageName::EFlexNameType::LocalPath, FPackageName::EErrorCode::PackageNameUnknown);
+		TestTryGetMountPointForPath(TEXT("/Game"), true, TEXT("/Game"), FPaths::ProjectContentDir(),
+			TEXT(""), FPackageName::EFlexNameType::PackageName, FPackageName::EErrorCode::PackageNameUnknown);
+
+		// For a malformed path of this type:
+		// d:/root/QAGame/Content -> d:/root/QAGame//Content
+		// TestTryGetMountPointForPath should fail
+		FString MalformedContentDir = FPaths::ProjectContentDir();
+		MalformedContentDir.ReplaceCharInline('\\', '/');
+		int32 LastSlashIndex;
+		if (FStringView(MalformedContentDir).LeftChop(1).FindLastChar('/', LastSlashIndex))
+		{
+			MalformedContentDir = MalformedContentDir.Left(LastSlashIndex) + TEXT("/") + MalformedContentDir.RightChop(LastSlashIndex);
+			TestTryGetMountPointForPath(MalformedContentDir, false, TEXT(""), TEXT(""),
+				TEXT(""), FPackageName::EFlexNameType::Invalid, FPackageName::EErrorCode::PackageNamePathNotMounted);
+		}
+	}
+
 	// Mounting and Unmounting Paths
 	{
 		FString TestMountLongPackageName(TEXT("/PackageNameTestDNE"));
