Removed 1 extra virtual function call (or more if ArrayDim > 1) from object-based Copy*ScriptVM property functions by having an implementation per property type instead of a top level override.  This also removes the need for the FSoftObjectPtrProperty::Copy*ScriptVM overloads.

diff --git a/Engine/Source/Runtime/CoreUObject/Private/UObject/PropertyBaseObject.cpp b/Engine/Source/Runtime/CoreUObject/Private/UObject/PropertyBaseObject.cpp
--- a/Engine/Source/Runtime/CoreUObject/Private/UObject/PropertyBaseObject.cpp
+++ b/Engine/Source/Runtime/CoreUObject/Private/UObject/PropertyBaseObject.cpp
@@ -685,30 +685,3 @@ bool FObjectPropertyBase::SameType(const FProperty* Other) const
 	return (Super::SameType(Other) || (Other && Other->IsA<FObjectPropertyBase>() && ((FObjectPropertyBase*)Other)->AllowObjectTypeReinterpretationTo(this))) && 
 			 (PropertyClass == ((FObjectPropertyBase*)Other)->PropertyClass);
 }
-
-void FObjectPropertyBase::CopySingleValueToScriptVM( void* Dest, void const* Src ) const
-{
-	*(UObject**)Dest = GetObjectPropertyValue(Src);
-}
-
-void FObjectPropertyBase::CopyCompleteValueToScriptVM( void* Dest, void const* Src ) const
-{
-	for (int32 Index = 0; Index < ArrayDim; Index++)
-	{
-		((UObject**)Dest)[Index] = GetObjectPropertyValue(((uint8*)Src) + Index * ElementSize);
-	}
-}
-
-void FObjectPropertyBase::CopySingleValueFromScriptVM( void* Dest, void const* Src ) const
-{
-	SetObjectPropertyValue(Dest, *(UObject**)Src);
-}
-
-void FObjectPropertyBase::CopyCompleteValueFromScriptVM( void* Dest, void const* Src ) const
-{
-	checkSlow(ElementSize == sizeof(UObject*)); // the idea that script pointers are the same size as weak pointers is maybe required, maybe not
-	for (int32 Index = 0; Index < ArrayDim; Index++)
-	{
-		SetObjectPropertyValue(((uint8*)Dest) + Index * ElementSize, ((UObject**)Src)[Index]);
-	}
-}
diff --git a/Engine/Source/Runtime/CoreUObject/Private/UObject/PropertyClassPtr.cpp b/Engine/Source/Runtime/CoreUObject/Private/UObject/PropertyClassPtr.cpp
--- a/Engine/Source/Runtime/CoreUObject/Private/UObject/PropertyClassPtr.cpp
+++ b/Engine/Source/Runtime/CoreUObject/Private/UObject/PropertyClassPtr.cpp
@@ -87,6 +87,26 @@ uint32 FClassPtrProperty::GetValueTypeHashInternal(const void* Src) const
 	return GetTypeHash((FObjectPtr&)GetPropertyValue(Src));
 }
 
+void FClassPtrProperty::CopySingleValueToScriptVM( void* Dest, void const* Src ) const
+{
+	*(UObject**)Dest = ((const FObjectPtr*)Src)->Get();
+}
+
+void FClassPtrProperty::CopySingleValueFromScriptVM( void* Dest, void const* Src ) const
+{
+	*(FObjectPtr*)Dest = *(UObject**)Src;
+}
+
+void FClassPtrProperty::CopyCompleteValueToScriptVM( void* Dest, void const* Src ) const
+{
+	GetWrappedUObjectPtrValues<FObjectPtr>((UObject**)Dest, Src, EPropertyMemoryAccess::Direct, 0, ArrayDim);
+}
+
+void FClassPtrProperty::CopyCompleteValueFromScriptVM( void* Dest, void const* Src ) const
+{
+	SetWrappedUObjectPtrValues<FObjectPtr>(Dest, EPropertyMemoryAccess::Direct, (UObject**)Src, 0, ArrayDim);
+}
+
 void FClassPtrProperty::CopyCompleteValueToScriptVM_InContainer(void* OutValue, void const* InContainer) const
 {
 	GetWrappedUObjectPtrValues<FObjectPtr>((UObject**)OutValue, InContainer, EPropertyMemoryAccess::InContainer, 0, ArrayDim);
diff --git a/Engine/Source/Runtime/CoreUObject/Private/UObject/PropertyLazyObjectPtr.cpp b/Engine/Source/Runtime/CoreUObject/Private/UObject/PropertyLazyObjectPtr.cpp
--- a/Engine/Source/Runtime/CoreUObject/Private/UObject/PropertyLazyObjectPtr.cpp
+++ b/Engine/Source/Runtime/CoreUObject/Private/UObject/PropertyLazyObjectPtr.cpp
@@ -156,6 +156,26 @@ uint32 FLazyObjectProperty::GetValueTypeHashInternal(const void* Src) const
 	return GetTypeHash(GetPropertyValue(Src));
 }
 
+void FLazyObjectProperty::CopySingleValueToScriptVM( void* Dest, void const* Src ) const
+{
+	*(UObject**)Dest = ((const FLazyObjectPtr*)Src)->Get();
+}
+
+void FLazyObjectProperty::CopySingleValueFromScriptVM( void* Dest, void const* Src ) const
+{
+	*(FLazyObjectPtr*)Dest = *(UObject**)Src;
+}
+
+void FLazyObjectProperty::CopyCompleteValueToScriptVM( void* Dest, void const* Src ) const
+{
+	GetWrappedUObjectPtrValues<FLazyObjectPtr>((UObject**)Dest, Src, EPropertyMemoryAccess::Direct, 0, ArrayDim);
+}
+
+void FLazyObjectProperty::CopyCompleteValueFromScriptVM( void* Dest, void const* Src ) const
+{
+	SetWrappedUObjectPtrValues<FLazyObjectPtr>(Dest, EPropertyMemoryAccess::Direct, (UObject**)Src, 0, ArrayDim);
+}
+
 void FLazyObjectProperty::CopyCompleteValueToScriptVM_InContainer(void* OutValue, void const* InContainer) const
 {
 	GetWrappedUObjectPtrValues<FLazyObjectPtr>((UObject**)OutValue, InContainer, EPropertyMemoryAccess::InContainer, 0, ArrayDim);
diff --git a/Engine/Source/Runtime/CoreUObject/Private/UObject/PropertyObject.cpp b/Engine/Source/Runtime/CoreUObject/Private/UObject/PropertyObject.cpp
--- a/Engine/Source/Runtime/CoreUObject/Private/UObject/PropertyObject.cpp
+++ b/Engine/Source/Runtime/CoreUObject/Private/UObject/PropertyObject.cpp
@@ -221,6 +221,26 @@ void FObjectProperty::SetObjectPropertyValue_InContainer(void* ContainerAddress,
 	}
 }
 
+void FObjectProperty::CopySingleValueToScriptVM(void* Dest, const void* Src) const
+{
+	*(UObject**)Dest = ((const FObjectPtr*)Src)->Get();
+}
+
+void FObjectProperty::CopySingleValueFromScriptVM(void* Dest, const void* Src) const
+{
+	*(FObjectPtr*)Dest = *(UObject**)Src;
+}
+
+void FObjectProperty::CopyCompleteValueToScriptVM(void* Dest, const void* Src) const
+{
+	GetWrappedUObjectPtrValues<FObjectPtr>((UObject**)Dest, Src, EPropertyMemoryAccess::Direct, 0, ArrayDim);
+}
+
+void FObjectProperty::CopyCompleteValueFromScriptVM(void* Dest, const void* Src) const
+{
+	SetWrappedUObjectPtrValues<FObjectPtr>(Dest, EPropertyMemoryAccess::Direct, (UObject**)Src, 0, ArrayDim);
+}
+
 void FObjectProperty::CopyCompleteValueToScriptVM_InContainer(void* OutValue, void const* InContainer) const
 {
 	GetWrappedUObjectPtrValues<FObjectPtr>((UObject**)OutValue, InContainer, EPropertyMemoryAccess::InContainer, 0, ArrayDim);
diff --git a/Engine/Source/Runtime/CoreUObject/Private/UObject/PropertyObjectPtr.cpp b/Engine/Source/Runtime/CoreUObject/Private/UObject/PropertyObjectPtr.cpp
--- a/Engine/Source/Runtime/CoreUObject/Private/UObject/PropertyObjectPtr.cpp
+++ b/Engine/Source/Runtime/CoreUObject/Private/UObject/PropertyObjectPtr.cpp
@@ -197,6 +197,26 @@ uint32 FObjectPtrProperty::GetValueTypeHashInternal(const void* Src) const
 	return GetTypeHash((FObjectPtr&)GetPropertyValue(Src));
 }
 
+void FObjectPtrProperty::CopySingleValueToScriptVM(void* Dest, const void* Src) const
+{
+	*(UObject**)Dest = ((const FObjectPtr*)Src)->Get();
+}
+
+void FObjectPtrProperty::CopySingleValueFromScriptVM(void* Dest, const void* Src) const
+{
+	*(FObjectPtr*)Dest = *(UObject**)Src;
+}
+
+void FObjectPtrProperty::CopyCompleteValueToScriptVM(void* Dest, const void* Src) const
+{
+	GetWrappedUObjectPtrValues<FObjectPtr>((UObject**)Dest, Src, EPropertyMemoryAccess::Direct, 0, ArrayDim);
+}
+
+void FObjectPtrProperty::CopyCompleteValueFromScriptVM(void* Dest, const void* Src) const
+{
+	SetWrappedUObjectPtrValues<FObjectPtr>(Dest, EPropertyMemoryAccess::Direct, (UObject**)Src, 0, ArrayDim);
+}
+
 void FObjectPtrProperty::CopyCompleteValueToScriptVM_InContainer(void* OutValue, void const* InContainer) const
 {
 	GetWrappedUObjectPtrValues<FObjectPtr>((UObject**)OutValue, InContainer, EPropertyMemoryAccess::InContainer, 0, ArrayDim);
diff --git a/Engine/Source/Runtime/CoreUObject/Private/UObject/PropertySoftObjectPtr.cpp b/Engine/Source/Runtime/CoreUObject/Private/UObject/PropertySoftObjectPtr.cpp
--- a/Engine/Source/Runtime/CoreUObject/Private/UObject/PropertySoftObjectPtr.cpp
+++ b/Engine/Source/Runtime/CoreUObject/Private/UObject/PropertySoftObjectPtr.cpp
@@ -267,24 +267,3 @@ uint32 FSoftObjectProperty::GetValueTypeHashInternal(const void* Src) const
 {
 	return GetTypeHash(GetPropertyValue(Src));
 }
-
-void FSoftObjectProperty::CopySingleValueToScriptVM(void* Dest, void const* Src) const
-{
-	CopySingleValue(Dest, Src);
-}
-
-void FSoftObjectProperty::CopyCompleteValueToScriptVM(void* Dest, void const* Src) const
-{
-	CopyCompleteValue(Dest, Src);
-}
-
-void FSoftObjectProperty::CopySingleValueFromScriptVM(void* Dest, void const* Src) const
-{
-	CopySingleValue(Dest, Src);
-}
-
-void FSoftObjectProperty::CopyCompleteValueFromScriptVM(void* Dest, void const* Src) const
-{
-	CopyCompleteValue(Dest, Src);
-}
-
diff --git a/Engine/Source/Runtime/CoreUObject/Private/UObject/PropertyWeakObjectPtr.cpp b/Engine/Source/Runtime/CoreUObject/Private/UObject/PropertyWeakObjectPtr.cpp
--- a/Engine/Source/Runtime/CoreUObject/Private/UObject/PropertyWeakObjectPtr.cpp
+++ b/Engine/Source/Runtime/CoreUObject/Private/UObject/PropertyWeakObjectPtr.cpp
@@ -118,6 +118,26 @@ uint32 FWeakObjectProperty::GetValueTypeHashInternal(const void* Src) const
 	return GetTypeHash(*(FWeakObjectPtr*)Src);
 }
 
+void FWeakObjectProperty::CopySingleValueToScriptVM( void* Dest, void const* Src ) const
+{
+	*(UObject**)Dest = ((const FWeakObjectPtr*)Src)->Get();
+}
+
+void FWeakObjectProperty::CopySingleValueFromScriptVM( void* Dest, void const* Src ) const
+{
+	*(FWeakObjectPtr*)Dest = *(UObject**)Src;
+}
+
+void FWeakObjectProperty::CopyCompleteValueToScriptVM( void* Dest, void const* Src ) const
+{
+	GetWrappedUObjectPtrValues<FWeakObjectPtr>((UObject**)Dest, Src, EPropertyMemoryAccess::Direct, 0, ArrayDim);
+}
+
+void FWeakObjectProperty::CopyCompleteValueFromScriptVM( void* Dest, void const* Src ) const
+{
+	SetWrappedUObjectPtrValues<FWeakObjectPtr>(Dest, EPropertyMemoryAccess::Direct, (UObject**)Src, 0, ArrayDim);
+}
+
 void FWeakObjectProperty::CopyCompleteValueToScriptVM_InContainer(void* OutValue, void const* InContainer) const
 {
 	GetWrappedUObjectPtrValues<FWeakObjectPtr>((UObject**)OutValue, InContainer, EPropertyMemoryAccess::InContainer, 0, ArrayDim);
diff --git a/Engine/Source/Runtime/CoreUObject/Public/UObject/UnrealType.h b/Engine/Source/Runtime/CoreUObject/Public/UObject/UnrealType.h
--- a/Engine/Source/Runtime/CoreUObject/Public/UObject/UnrealType.h
+++ b/Engine/Source/Runtime/CoreUObject/Public/UObject/UnrealType.h
@@ -2633,49 +2633,6 @@ public:
 	virtual FName GetID() const override;
 	virtual void InstanceSubobjects( void* Data, void const* DefaultData, UObject* Owner, struct FObjectInstancingGraph* InstanceGraph ) override;
 	virtual bool SameType(const FProperty* Other) const override;
-	/**
-	 * Copy the value for a single element of this property. To the script VM.
-	 * 
-	 * @param	Dest				the address where the value should be copied to.  This should always correspond to the BASE + OFFSET + INDEX * SIZE, where
-	 *									BASE = (for member properties) the address of the UObject which contains this data, (for locals/parameters) the address of the space allocated for the function's locals
-	 *									OFFSET = the Offset of this FProperty
-	 *									INDEX = the index that you want to copy.  for properties which are not arrays, this should always be 0
-	 *									SIZE = the ElementSize of this FProperty
-	 * @param	Src					the address of the value to copy from. should be evaluated the same way as Dest
-	 */
-	virtual void CopySingleValueToScriptVM( void* Dest, void const* Src ) const override;
-
-	/**
-	 * Copy the value for all elements of this property. To the script VM.
-	 * 
-	 * @param	Dest				the address where the value should be copied to.  This should always correspond to the BASE + OFFSET, where
-	 *									BASE = (for member properties) the address of the UObject which contains this data, (for locals/parameters) the address of the space allocated for the function's locals
-	 *									OFFSET = the Offset of this FProperty
-	 * @param	Src					the address of the value to copy from. should be evaluated the same way as Dest
-	 */
-	virtual void CopyCompleteValueToScriptVM( void* Dest, void const* Src ) const override;
-
-	/**
-	 * Copy the value for a single element of this property. From the script VM.
-	 * 
-	 * @param	Dest				the address where the value should be copied to.  This should always correspond to the BASE + OFFSET + INDEX * SIZE, where
-	 *									BASE = (for member properties) the address of the UObject which contains this data, (for locals/parameters) the address of the space allocated for the function's locals
-	 *									OFFSET = the Offset of this FProperty
-	 *									INDEX = the index that you want to copy.  for properties which are not arrays, this should always be 0
-	 *									SIZE = the ElementSize of this FProperty
-	 * @param	Src					the address of the value to copy from. should be evaluated the same way as Dest
-	 */
-	virtual void CopySingleValueFromScriptVM( void* Dest, void const* Src ) const override;
-
-	/**
-	 * Copy the value for all elements of this property. From the script VM.
-	 * 
-	 * @param	Dest				the address where the value should be copied to.  This should always correspond to the BASE + OFFSET, where
-	 *									BASE = (for member properties) the address of the UObject which contains this data, (for locals/parameters) the address of the space allocated for the function's locals
-	 *									OFFSET = the Offset of this FProperty
-	 * @param	Src					the address of the value to copy from. should be evaluated the same way as Dest
-	 */
-	virtual void CopyCompleteValueFromScriptVM( void* Dest, void const* Src ) const override;
 	// End of FProperty interface
 
 	// FObjectPropertyBase interface
@@ -2992,6 +2949,10 @@ PRAGMA_ENABLE_DEPRECATION_WARNINGS
 private:
 	virtual uint32 GetValueTypeHashInternal(const void* Src) const override;
 public:
+	virtual void CopySingleValueToScriptVM( void* Dest, void const* Src ) const override;
+	virtual void CopySingleValueFromScriptVM( void* Dest, void const* Src ) const override;
+	virtual void CopyCompleteValueToScriptVM( void* Dest, void const* Src ) const override;
+	virtual void CopyCompleteValueFromScriptVM( void* Dest, void const* Src ) const override;
 	virtual void CopyCompleteValueToScriptVM_InContainer( void* OutValue, void const* InContainer ) const override;
 	virtual void CopyCompleteValueFromScriptVM_InContainer( void* OutContainer, void const* InValue ) const override;
 	// End of FProperty interface
@@ -3041,6 +3002,10 @@ class COREUOBJECT_API FObjectPtrProperty : public FObjectProperty
 	virtual bool SameType(const FProperty* Other) const override;
 	virtual bool Identical(const void* A, const void* B, uint32 PortFlags) const override;
 	virtual void SerializeItem(FStructuredArchive::FSlot Slot, void* Value, void const* Defaults) const override;
+	virtual void CopySingleValueToScriptVM( void* Dest, void const* Src ) const override;
+	virtual void CopySingleValueFromScriptVM( void* Dest, void const* Src ) const override;
+	virtual void CopyCompleteValueToScriptVM( void* Dest, void const* Src ) const override;
+	virtual void CopyCompleteValueFromScriptVM( void* Dest, void const* Src ) const override;
 	virtual void CopyCompleteValueToScriptVM_InContainer( void* OutValue, void const* InContainer ) const override;
 	virtual void CopyCompleteValueFromScriptVM_InContainer( void* OutContainer, void const* InValue ) const override;
 	// End of FProperty interface
@@ -3115,6 +3080,10 @@ protected:
 private:
 	virtual uint32 GetValueTypeHashInternal(const void* Src) const override;
 public:
+	virtual void CopySingleValueToScriptVM( void* Dest, void const* Src ) const override;
+	virtual void CopySingleValueFromScriptVM( void* Dest, void const* Src ) const override;
+	virtual void CopyCompleteValueToScriptVM( void* Dest, void const* Src ) const override;
+	virtual void CopyCompleteValueFromScriptVM( void* Dest, void const* Src ) const override;
 	virtual void CopyCompleteValueToScriptVM_InContainer( void* OutValue, void const* InContainer ) const override;
 	virtual void CopyCompleteValueFromScriptVM_InContainer( void* OutContainer, void const* InValue ) const override;
 	// End of FProperty interface
@@ -3175,6 +3144,10 @@ PRAGMA_ENABLE_DEPRECATION_WARNINGS
 	virtual bool Identical( const void* A, const void* B, uint32 PortFlags ) const override;
 	virtual void SerializeItem( FStructuredArchive::FSlot Slot, void* Value, void const* Defaults ) const override;
 	virtual void EmitReferenceInfo(UE::GC::FTokenStreamBuilder& TokenStream, int32 BaseOffset, TArray<const FStructProperty*>& EncounteredStructProps, FGCStackSizeHelper& StackSizeHelper) override;
+	virtual void CopySingleValueToScriptVM( void* Dest, void const* Src ) const override;
+	virtual void CopySingleValueFromScriptVM( void* Dest, void const* Src ) const override;
+	virtual void CopyCompleteValueToScriptVM( void* Dest, void const* Src ) const override;
+	virtual void CopyCompleteValueFromScriptVM( void* Dest, void const* Src ) const override;
 	virtual void CopyCompleteValueToScriptVM_InContainer( void* OutValue, void const* InContainer ) const override;
 	virtual void CopyCompleteValueFromScriptVM_InContainer( void* OutContainer, void const* InValue ) const override;
 	// End of FProperty interface
@@ -3278,13 +3251,8 @@ public:
 private:
 	virtual uint32 GetValueTypeHashInternal(const void* Src) const override;
 public:
+	// Note: FSoftObjectProperty does not override the Copy*VM functions, as ScriptVM should store Asset as a FSoftObjectPtr not as a UObject*.
 
-	// ScriptVM should store Asset as FSoftObjectPtr not as UObject.
-
-	virtual void CopySingleValueToScriptVM(void* Dest, void const* Src) const override;
-	virtual void CopyCompleteValueToScriptVM(void* Dest, void const* Src) const override;
-	virtual void CopySingleValueFromScriptVM(void* Dest, void const* Src) const override;
-	virtual void CopyCompleteValueFromScriptVM(void* Dest, void const* Src) const override;
 	// End of FObjectProperty interface
 };
 
@@ -3396,6 +3364,10 @@ class COREUOBJECT_API FClassPtrProperty : public FClassProperty
 	virtual bool SameType(const FProperty* Other) const override;
 	virtual bool Identical(const void* A, const void* B, uint32 PortFlags) const override;
 	virtual void SerializeItem(FStructuredArchive::FSlot Slot, void* Value, void const* Defaults) const override;
+	virtual void CopySingleValueToScriptVM( void* Dest, void const* Src ) const override;
+	virtual void CopySingleValueFromScriptVM( void* Dest, void const* Src ) const override;
+	virtual void CopyCompleteValueToScriptVM( void* Dest, void const* Src ) const override;
+	virtual void CopyCompleteValueFromScriptVM( void* Dest, void const* Src ) const override;
 	virtual void CopyCompleteValueToScriptVM_InContainer( void* OutValue, void const* InContainer ) const override;
 	virtual void CopyCompleteValueFromScriptVM_InContainer( void* OutContainer, void const* InValue ) const override;
 	// End of FProperty interface
