changed `FClassProperty::Identical` to use TObjectPtr to avoid causing unnecessary reads

diff --git a/Engine/Source/Runtime/CoreUObject/Private/UObject/PropertyClass.cpp b/Engine/Source/Runtime/CoreUObject/Private/UObject/PropertyClass.cpp
--- a/Engine/Source/Runtime/CoreUObject/Private/UObject/PropertyClass.cpp
+++ b/Engine/Source/Runtime/CoreUObject/Private/UObject/PropertyClass.cpp
@@ -187,11 +187,11 @@ bool FClassProperty::SameType(const FProperty* Other) const
 
 bool FClassProperty::Identical( const void* A, const void* B, uint32 PortFlags ) const
 {
-	UObject* ObjectA = A ? GetObjectPropertyValue(A) : nullptr;
-	UObject* ObjectB = B ? GetObjectPropertyValue(B) : nullptr;
+	TObjectPtr<UObject> ObjectA = A ? GetObjectPtrPropertyValue(A) : TObjectPtr<UObject>();
+	TObjectPtr<UObject> ObjectB = B ? GetObjectPtrPropertyValue(B) : TObjectPtr<UObject>();
 
-	check(ObjectA == nullptr || ObjectA->IsA<UClass>());
-	check(ObjectB == nullptr || ObjectB->IsA<UClass>());
+	check(ObjectA == nullptr || ObjectA.IsA<UClass>());
+	check(ObjectB == nullptr || ObjectB.IsA<UClass>());
 	return (ObjectA == ObjectB);
 }
 
diff --git a/Engine/Source/Runtime/CoreUObject/Tests/ObjectPtrTestClass.cpp b/Engine/Source/Runtime/CoreUObject/Tests/ObjectPtrTestClass.cpp
--- a/Engine/Source/Runtime/CoreUObject/Tests/ObjectPtrTestClass.cpp
+++ b/Engine/Source/Runtime/CoreUObject/Tests/ObjectPtrTestClass.cpp
@@ -61,6 +61,14 @@ IMPLEMENT_CORE_INTRINSIC_CLASS(UObjectPtrTestClassWithRef, UObject,
 	});
 #endif
 
+IMPLEMENT_CORE_INTRINSIC_CLASS(UObjectWithClassProperty, UObject,
+	{
+		auto Property1 = new FClassPtrProperty(Class, TEXT("ClassPtr"), EObjectFlags::RF_NoFlags);
+		Property1->PropertyClass = UClass::StaticClass();
+		Class->AddCppProperty(Property1);
+	}
+);
+
 IMPLEMENT_CORE_INTRINSIC_CLASS(UObjectPtrDerrivedTestClass, UObjectPtrTestClass, {});
 
 IMPLEMENT_CORE_INTRINSIC_CLASS(UObjectPtrNotLazyTestClass, UObject, {});
diff --git a/Engine/Source/Runtime/CoreUObject/Tests/ObjectPtrTestClass.h b/Engine/Source/Runtime/CoreUObject/Tests/ObjectPtrTestClass.h
--- a/Engine/Source/Runtime/CoreUObject/Tests/ObjectPtrTestClass.h
+++ b/Engine/Source/Runtime/CoreUObject/Tests/ObjectPtrTestClass.h
@@ -26,6 +26,15 @@ public:
 };
 
 
+//test class with typed reference to another class
+class UObjectWithClassProperty : public UObject
+{
+	DECLARE_CLASS_INTRINSIC(UObjectWithClassProperty, UObject, CLASS_MatchedSerializers, TEXT("/Script/CoreUObject"))
+
+public:
+	TObjectPtr<UClass> ClassPtr;
+};
+
 //derived test class
 class UObjectPtrDerrivedTestClass : public UObjectPtrTestClass
 {
diff --git a/Engine/Source/Runtime/CoreUObject/Tests/PropertyClassTest.cpp b/Engine/Source/Runtime/CoreUObject/Tests/PropertyClassTest.cpp
new file mode 100644
--- /dev/null
+++ b/Engine/Source/Runtime/CoreUObject/Tests/PropertyClassTest.cpp
@@ -0,0 +1,46 @@
+// Copyright Epic Games, Inc. All Rights Reserved.
+
+#if WITH_LOW_LEVEL_TESTS
+
+#include "TestHarness.h"
+#include "TestMacros/Assertions.h"
+#include "ObjectPtrTestClass.h"
+#include "UObject/Package.h"
+#include "Misc/ScopeExit.h"
+#include "UObject/UnrealType.h"
+#include "UObject/LinkerPlaceholderExportObject.h"
+#include "UObject/LinkerPlaceholderClass.h"
+#include "UObject/ObjectHandleTracking.h"
+#include "LowLevelTestsRunner/WarnFilterScope.h"
+
+TEST_CASE("UE::CoreUObject::FClassProperty::Identical")
+{
+	bool bAllowRead = false;
+#if UE_WITH_OBJECT_HANDLE_TRACKING
+	auto CallbackHandle = UE::CoreUObject::AddObjectHandleReadCallback([&bAllowRead](UObject* ReadObject)
+		{
+			if (!bAllowRead)
+				FAIL("Unexpected read during CheckValidObject");
+		});
+	ON_SCOPE_EXIT
+	{
+		UE::CoreUObject::RemoveObjectHandleReadCallback(CallbackHandle);
+	};
+#endif
+	UClass* Class = UObjectWithClassProperty::StaticClass();
+	FObjectProperty* Property = CastField<FObjectProperty>(Class->FindPropertyByName(TEXT("ClassPtr")));
+	REQUIRE(Property != nullptr);
+
+
+	UPackage* TestPackage = NewObject<UPackage>(nullptr, TEXT("TestPackageName"), RF_Transient);
+	TestPackage->AddToRoot();
+	ON_SCOPE_EXIT
+	{
+		TestPackage->RemoveFromRoot();
+	};
+	UObjectWithClassProperty* Obj = NewObject<UObjectWithClassProperty>(TestPackage, TEXT("UObjectWithClassProperty"));
+
+	CHECK(Property->Identical(Obj, Obj, 0u));
+
+}
+#endif
\ No newline at end of file
