Wrapped the CoreNetTypes.generated.h include with USE_COMPILED_IN_NATIVES so it compiles properly with C++ UHT.

diff --git a/Engine/Source/Runtime/CoreUObject/Public/UObject/CoreNetTypes.h b/Engine/Source/Runtime/CoreUObject/Public/UObject/CoreNetTypes.h
--- a/Engine/Source/Runtime/CoreUObject/Public/UObject/CoreNetTypes.h
+++ b/Engine/Source/Runtime/CoreUObject/Public/UObject/CoreNetTypes.h
@@ -4,6 +4,10 @@
 
 #include "UObject/ObjectMacros.h"
 
+#if USE_COMPILED_IN_NATIVES 
+#include "CoreNetTypes.generated.h"
+#endif
+
 /** Whether or not caching of actor/subobject names to the stack should be enabled, for async demo recording crashdumps */
 #ifndef UE_NET_REPACTOR_NAME_DEBUG
 	#define UE_NET_REPACTOR_NAME_DEBUG 0
