changed FObjectProperty::SerializeItem to not cause object handle reads

diff --git a/Engine/Source/Runtime/CoreUObject/Private/UObject/PropertyObject.cpp b/Engine/Source/Runtime/CoreUObject/Private/UObject/PropertyObject.cpp
--- a/Engine/Source/Runtime/CoreUObject/Private/UObject/PropertyObject.cpp
+++ b/Engine/Source/Runtime/CoreUObject/Private/UObject/PropertyObject.cpp
@@ -110,10 +110,16 @@ void FObjectProperty::SerializeItem( FStructuredArchive::FSlot Slot, void* Value
 	}
 	else
 	{
-		UObject* ObjectValue = GetObjectPropertyValue(Value);
+		TObjectPtr<UObject> ObjectValuePtr = GetObjectPtrPropertyValue(Value);
+		check(ObjectValuePtr.IsResolved());
+		UObject* ObjectValue = UE::CoreUObject::Private::ReadObjectHandlePointerNoCheck(ObjectValuePtr.GetHandle());
+
 		Slot << ObjectValue;
 
-		UObject* CurrentValue = GetObjectPropertyValue(Value);
+		TObjectPtr<UObject> CurrentValuePtr = GetObjectPtrPropertyValue(Value);
+		check(CurrentValuePtr.IsResolved());
+		UObject* CurrentValue = UE::CoreUObject::Private::ReadObjectHandlePointerNoCheck(CurrentValuePtr.GetHandle());
+
 		if (ObjectValue != CurrentValue)
 		{
 			SetObjectPropertyValue(Value, ObjectValue);
diff --git a/Engine/Source/Runtime/CoreUObject/Private/UObject/PropertyObjectPtr.cpp b/Engine/Source/Runtime/CoreUObject/Private/UObject/PropertyObjectPtr.cpp
--- a/Engine/Source/Runtime/CoreUObject/Private/UObject/PropertyObjectPtr.cpp
+++ b/Engine/Source/Runtime/CoreUObject/Private/UObject/PropertyObjectPtr.cpp
@@ -56,7 +56,7 @@ void FObjectPtrProperty::StaticSerializeItem(const FObjectPropertyBase* ObjectPr
 	FArchive& UnderlyingArchive = Slot.GetUnderlyingArchive();
 	FObjectPtr* ObjectPtr = (FObjectPtr*)GetPropertyValuePtr(Value);
 	const bool IsHandleResolved = IsObjectHandleResolved(ObjectPtr->GetHandle());
-	UObject* CurrentValue = IsHandleResolved ? ObjectPtr->Get() : nullptr;
+	UObject* CurrentValue = IsHandleResolved && *ObjectPtr ? UE::CoreUObject::Private::ReadObjectHandlePointerNoCheck(ObjectPtr->GetHandle()) : nullptr;
 
 	if (UnderlyingArchive.IsObjectReferenceCollector())
 	{
@@ -77,14 +77,17 @@ void FObjectPtrProperty::StaticSerializeItem(const FObjectPropertyBase* ObjectPr
 		FObjectHandle CurrentHandle = ObjectPtr->GetHandle();
 		if ((OriginalHandle != CurrentHandle) && IsObjectHandleResolved(CurrentHandle))
 		{
-			UObject* ResolvedObject = ObjectPtr->Get();
 	#if USE_CIRCULAR_DEPENDENCY_LOAD_DEFERRING
-			if (ULinkerPlaceholderExportObject* PlaceholderVal = Cast<ULinkerPlaceholderExportObject>(ResolvedObject))
+			if (ObjectPtr->IsA<ULinkerPlaceholderExportObject>())
 			{
+				//resolve the handle with no read to avoid trigger a handle read. 
+				ULinkerPlaceholderExportObject* PlaceholderVal = static_cast<ULinkerPlaceholderExportObject*>(UE::CoreUObject::Private::ReadObjectHandlePointerNoCheck(CurrentHandle));
 				PlaceholderVal->AddReferencingPropertyValue(ObjectProperty, Value);
 			}
-			else if (ULinkerPlaceholderClass* PlaceholderClass = Cast<ULinkerPlaceholderClass>(ResolvedObject))
+			else if (ObjectPtr->IsA<ULinkerPlaceholderClass>())
 			{
+				//resolve the handle with no read to avoid trigger a handle read. 
+				ULinkerPlaceholderClass* PlaceholderClass = static_cast<ULinkerPlaceholderClass*>(UE::CoreUObject::Private::ReadObjectHandlePointerNoCheck(CurrentHandle));
 				PlaceholderClass->AddReferencingPropertyValue(ObjectProperty, Value);
 			}
 			// NOTE: we don't remove this from CurrentValue if it is a 
diff --git a/Engine/Source/Runtime/CoreUObject/Tests/ObjectPropertyTest.cpp b/Engine/Source/Runtime/CoreUObject/Tests/ObjectPropertyTest.cpp
--- a/Engine/Source/Runtime/CoreUObject/Tests/ObjectPropertyTest.cpp
+++ b/Engine/Source/Runtime/CoreUObject/Tests/ObjectPropertyTest.cpp
@@ -8,10 +8,11 @@
 #include "UObject/Package.h"
 #include "Misc/ScopeExit.h"
 #include "UObject/UnrealType.h"
+#include "UObject/LinkerPlaceholderExportObject.h"
+#include "UObject/LinkerPlaceholderClass.h"
 #include "UObject/ObjectHandleTracking.h"
 #include "LowLevelTestsRunner/WarnFilterScope.h"
 
-
 TEST_CASE("UE::CoreUObject::FObjectProperty::CheckValidAddress")
 {
 	bool bAllowRead = false;
@@ -19,7 +20,7 @@ TEST_CASE("UE::CoreUObject::FObjectProperty::CheckValidAddress")
 	auto CallbackHandle = UE::CoreUObject::AddObjectHandleReadCallback([&bAllowRead](UObject* ReadObject)
 		{
 			if (!bAllowRead)
-			FAIL("Unexpected read during CheckValidObject");
+				FAIL("Unexpected read during CheckValidObject");
 		});
 	ON_SCOPE_EXIT
 	{
@@ -129,4 +130,90 @@ TEST_CASE("UE::CoreUObject::FObjectProperty::CheckValidAddressNonNullable")
 	REQUIRE_CHECK(Property->CheckValidObject(&Obj->ObjectPtrNonNullable, nullptr));
 }
 
-#endif
\ No newline at end of file
+class FMockArchive : public FArchive
+{
+public:
+	TFunction<void(FObjectPtr& Value)> OnFObjectPtr;
+	virtual FArchive& operator<<(FObjectPtr& Value)
+	{
+		if (OnFObjectPtr)
+			OnFObjectPtr(Value);
+		return *this;
+	}
+};
+
+template<typename T>
+static void TestSerializeItem(FName ObjectName)
+{
+	UClass* Class = T::StaticClass();
+	FObjectProperty* Property = CastField<FObjectProperty>(Class->FindPropertyByName(TEXT("ObjectPtr")));
+	REQUIRE(Property != nullptr);
+
+	UPackage* TestPackage = NewObject<UPackage>(nullptr, TEXT("TestPackageName"), RF_Transient);
+	TestPackage->AddToRoot();
+	ON_SCOPE_EXIT
+	{
+		TestPackage->RemoveFromRoot();
+	};
+	T* Obj = NewObject<T>(TestPackage, ObjectName);
+	UObjectPtrTestClass* Other = NewObject<UObjectPtrTestClass>(Obj, TEXT("Other"));
+	ULinkerPlaceholderExportObject* PlaceHolderExport = NewObject<ULinkerPlaceholderExportObject>(TestPackage, TEXT("PlaceHolderExport"));
+	ULinkerPlaceholderClass* PlaceHolderClass = NewObject<ULinkerPlaceholderClass>(TestPackage, TEXT("PlaceHolderClass"));
+	PlaceHolderClass->Bind(); //must call bind or crashes on shutdown
+
+	int ReadCount = 0;
+#if UE_WITH_OBJECT_HANDLE_TRACKING
+	auto CallbackId = UE::CoreUObject::AddObjectHandleReadCallback([&ReadCount](UObject* ReadObject)
+		{
+			++ReadCount;
+		});
+	ON_SCOPE_EXIT
+	{
+		UE::CoreUObject::RemoveObjectHandleReadCallback(CallbackId);
+	};
+#endif
+	{
+		//verify that if the property is null no reads are triggered
+		FMockArchive MockArchive;
+		MockArchive.OnFObjectPtr = [&](FObjectPtr& Value)
+		{
+			Value = PlaceHolderExport;
+		};
+
+		FBinaryArchiveFormatter Formatter(MockArchive);
+		FStructuredArchive Ar(Formatter);
+		FStructuredArchiveSlot Slot = Ar.Open();
+		FObjectPtrProperty::StaticSerializeItem(Property, Slot, &Obj->ObjectPtr, nullptr);
+		CHECK(ReadCount == 0);
+		CHECK(*reinterpret_cast<ULinkerPlaceholderExportObject**>(&Obj->ObjectPtr) == PlaceHolderExport);
+
+	}
+	{
+		//verify that if the property is not null no reads are triggered
+		FMockArchive MockArchive;
+		MockArchive.OnFObjectPtr = [&](FObjectPtr& Value)
+		{
+			Value = PlaceHolderClass;
+		};
+
+		FBinaryArchiveFormatter Formatter(MockArchive);
+		FStructuredArchive Ar(Formatter);
+		FStructuredArchiveSlot Slot = Ar.Open();
+		Obj->ObjectPtr = Other;
+		FObjectPtrProperty::StaticSerializeItem(Property, Slot, &Obj->ObjectPtr, nullptr);
+		CHECK(ReadCount == 0);
+		CHECK(*reinterpret_cast<ULinkerPlaceholderClass**>(&Obj->ObjectPtr) == PlaceHolderClass);
+	}
+}
+
+TEST_CASE("UE::CoreUObject::FObjectPtrProperty::StaticSerializeItem")
+{
+	TestSerializeItem<UObjectPtrTestClassWithRef>(TEXT("Object1"));
+}
+
+TEST_CASE("UE::CoreUObject::FObjectProperty::StaticSerializeItem")
+{
+	TestSerializeItem<UObjectWithRawProperty>(TEXT("Object2"));
+}
+
+#endif
diff --git a/Engine/Source/Runtime/CoreUObject/Tests/ObjectPtrTestClass.cpp b/Engine/Source/Runtime/CoreUObject/Tests/ObjectPtrTestClass.cpp
--- a/Engine/Source/Runtime/CoreUObject/Tests/ObjectPtrTestClass.cpp
+++ b/Engine/Source/Runtime/CoreUObject/Tests/ObjectPtrTestClass.cpp
@@ -69,6 +69,19 @@ IMPLEMENT_CORE_INTRINSIC_CLASS(UObjectWithClassProperty, UObject,
 	}
 );
 
+IMPLEMENT_CORE_INTRINSIC_CLASS(UObjectWithRawProperty, UObject,
+	{
+		auto Property1 = new FObjectProperty(Class, TEXT("ObjectPtr"), EObjectFlags::RF_NoFlags);
+		Property1->PropertyClass = UClass::StaticClass();
+		Class->AddCppProperty(Property1);
+
+		auto Property2 = new FObjectPtrProperty(Class, TEXT("ObjectPtrNonNullable"), EObjectFlags::RF_NoFlags);
+		Property2->PropertyClass = UObjectPtrTestClass::StaticClass();
+		Property2->SetPropertyFlags(EPropertyFlags::CPF_NonNullable); //make non nullable
+		Class->AddCppProperty(Property2);
+	}
+);
+
 IMPLEMENT_CORE_INTRINSIC_CLASS(UObjectPtrDerrivedTestClass, UObjectPtrTestClass, {});
 
 IMPLEMENT_CORE_INTRINSIC_CLASS(UObjectPtrNotLazyTestClass, UObject, {});
diff --git a/Engine/Source/Runtime/CoreUObject/Tests/ObjectPtrTestClass.h b/Engine/Source/Runtime/CoreUObject/Tests/ObjectPtrTestClass.h
--- a/Engine/Source/Runtime/CoreUObject/Tests/ObjectPtrTestClass.h
+++ b/Engine/Source/Runtime/CoreUObject/Tests/ObjectPtrTestClass.h
@@ -35,6 +35,17 @@ public:
 	TObjectPtr<UClass> ClassPtr;
 };
 
+//test class with raw pointer
+class UObjectWithRawProperty : public UObject
+{
+	DECLARE_CLASS_INTRINSIC(UObjectWithRawProperty, UObject, CLASS_MatchedSerializers, TEXT("/Script/CoreUObject"))
+
+public:
+	UObjectPtrTestClass* ObjectPtr;
+	UObjectPtrTestClass* ObjectPtrNonNullable;
+};
+
+
 //derived test class
 class UObjectPtrDerrivedTestClass : public UObjectPtrTestClass
 {
