Fix PERF_DETAILED_PER_CLASS_GC_STATS for new GC reachability code.

diff --git a/Engine/Source/Runtime/CoreUObject/Private/UObject/GarbageCollection.cpp b/Engine/Source/Runtime/CoreUObject/Private/UObject/GarbageCollection.cpp
--- a/Engine/Source/Runtime/CoreUObject/Private/UObject/GarbageCollection.cpp
+++ b/Engine/Source/Runtime/CoreUObject/Private/UObject/GarbageCollection.cpp
@@ -48,6 +48,8 @@
 
 DEFINE_LOG_CATEGORY(LogGarbage);
 
+#define PERF_DETAILED_PER_CLASS_GC_STATS				(LOOKING_FOR_PERF_ISSUES || 0) 
+
 static_assert(sizeof(FGCSkipInfo) == sizeof(uint32),		"Must have size of a GC token");
 static_assert(sizeof(FGCReferenceInfo) == sizeof(uint32),	"Must have size of a GC token");
 
@@ -249,87 +251,219 @@ static FAutoConsoleVariableRef CGarbageReferenceTrackingEnabled(
 	TEXT("Causes the Garbage Collector to track and log unreleased garbage objects. If 1, will dump every reference. If 2, will dump a sample of the references to highlight problematic properties."),
 	ECVF_Default
 );
+#endif // UE_BUILD_SHIPPING
+
+namespace UE::GC
+{
+	struct FStats
+	{
+#if PERF_DETAILED_PER_CLASS_GC_STATS
+		/** Map from a UClass' FName to the number of objects that were purged during the last purge phase of this class.	*/
+		TMap<const FName,uint64> ClassToPurgeCountMap;
+		/** Map from a UClass' FName to the number of "Disregard For GC" object references followed for all instances.		*/
+		TMap<const FName,uint64> ClassToDisregardedObjectRefsMap;
+		/** Map from a UClass' FName to the number of regular object references followed for all instances.					*/
+		TMap<const FName,uint64> ClassToRegularObjectRefsMap;
+		/** Map from a UClass' FName to the number of clsuter references followed for all instances.						*/
+		TMap<const FName,uint64> ClassToClustersRefsMap;
+		/** Map from a UClass' FName to the number of cycles spent with GC.													*/
+		TMap<const FName,uint64> ClassToCyclesMap;
+
+		/** Number of objects traversed by GC in total.			 															*/
+		uint32 NumObjectsTraversed = 0;
+		/** Number of clusters traversed by GC in total.			 														*/
+		uint32 NumClustersTraversed = 0;
+		/** Number of references from cluster to cluster checked.	 														*/
+		uint32 NumClusterToClusterRefs = 0;
+		/** Number of references from cluster to mutable object checked.													*/
+		uint32 NumClusterToObjectRefs = 0;
+		/** Number of references from clusters to mutable objects which were later clustered. 								*/
+		uint32 NumClusterToLateClusterRefs = 0;
+
+		/** Number of disregarded object refs for current object.															*/
+		uint32 CurrentObjectDisregardedObjectRefs;
+		/** Number of regular object refs for current object.																*/
+		uint32 CurrentObjectRegularObjectRefs;
+		/** Number of references to clusters from the current object */
+		uint32 CurrentObjectClusterRefs;
+
+		// Timestamp for starting to profile the current object's outgoing references
+		uint64 CurrentObjectStartCyles;
 #endif
+		
+		void LogClassCountInfo(const TCHAR* LogText, TMap<const FName, uint64>& ClassToCountMap, int32 NumItemsToLog, uint64 TotalCount);
+		void IncClusterToObjectRefs(FUObjectItem* Item);
+		void IncClusterToClusterRefs(int32 Count);
+		void IncNumClustersTraversed();
+
+		void LogDetailedStatsSummary();
+		void BeginTimingObject(UObject* CurrentObject);
+		void UpdateDetailedStats(UObject* CurrentObject);
+		void IncreaseObjectRefStats(UObject* RefToObject);
+
+		void IncPurgeCount(UObject* Object);
+		void LogPurgeStats(int32 Total);
+	};
+	static FStats Stats;
 
+	/**
+	 * Helper function to log the various class to count info maps.
+	 *
+	 * @param	LogText				Text to emit between number and class 
+	 * @param	ClassToCountMap		TMap from a class' FName to "count"
+	 * @param	NumItemsToList		Number of items to log
+	 * @param	TotalCount			Total count, if 0 will be calculated
+	 */
+	void FStats::LogClassCountInfo( const TCHAR* LogText, TMap<const FName, uint64>& ClassToCountMap, int32 NumItemsToLog, uint64 TotalCount )
+	{
 #if PERF_DETAILED_PER_CLASS_GC_STATS
-/** Map from a UClass' FName to the number of objects that were purged during the last purge phase of this class.	*/
-static TMap<const FName,uint32> GClassToPurgeCountMap;
-/** Map from a UClass' FName to the number of "Disregard For GC" object references followed for all instances.		*/
-static TMap<const FName,uint32> GClassToDisregardedObjectRefsMap;
-/** Map from a UClass' FName to the number of regular object references followed for all instances.					*/
-static TMap<const FName,uint32> GClassToRegularObjectRefsMap;
-/** Map from a UClass' FName to the number of cycles spent with GC.													*/
-static TMap<const FName,uint32> GClassToCyclesMap;
-
-/** Number of disregarded object refs for current object.															*/
-static uint32 GCurrentObjectDisregardedObjectRefs;
-/** Number of regulard object refs for current object.																*/
-static uint32 GCurrentObjectRegularObjectRefs;
+		// Figure out whether we need to calculate the total count.
+		bool bNeedToCalculateCount = false;
+		if( TotalCount == 0 )
+		{
+			for (const TPair<FName, uint64>& Pair : ClassToCountMap) 
+			{
+				TotalCount += Pair.Value;
+			}
+		}
+		ClassToCountMap.ValueSort(TGreater<uint64>{});
 
-/**
- * Helper structure used for sorting class to count map.
- */
-struct FClassCountInfo
-{
-	FName	ClassName;
-	uint32	InstanceCount;
-};
+		// Log top NumItemsToLog class counts
+		for (const TPair<FName, uint64>& Pair : ClassToCountMap)
+		{
+			const float Percent = 100.f * Pair.Value / TotalCount;
+			const FString PercentString = (TotalCount > 0) ? FString::Printf(TEXT("%6.2f%%"), Percent) : FString(TEXT("  N/A  "));
+			UE_LOG(LogGarbage, Log, TEXT("%5" UINT64_FMT "[% s] % s Class % s"), Pair.Value, *PercentString, LogText, *Pair.Key.ToString() ); 
+		}
 
-/**
- * Helper function to log the various class to count info maps.
- *
- * @param	LogText				Text to emit between number and class 
- * @param	ClassToCountMap		TMap from a class' FName to "count"
- * @param	NumItemsToList		Number of items to log
- * @param	TotalCount			Total count, if 0 will be calculated
- */
-static void LogClassCountInfo( const TCHAR* LogText, TMap<const FName,uint32>& ClassToCountMap, int32 NumItemsToLog, uint32 TotalCount )
-{
-	// Array of class name and counts.
-	TArray<FClassCountInfo> ClassCountArray;
-	ClassCountArray.Empty( ClassToCountMap.Num() );
+		// Empty the map for the next run.
+		ClassToCountMap.Empty();
+#endif
+	};
 
-	// Figure out whether we need to calculate the total count.
-	bool bNeedToCalculateCount = false;
-	if( TotalCount == 0 )
+	FORCEINLINE void FStats::IncClusterToObjectRefs(FUObjectItem* ObjectItem)
 	{
-		bNeedToCalculateCount = true;
-	}
-	// Copy TMap to TArray for sorting purposes (and to calculate count if needed).
-	for( TMap<const FName,uint32>::TIterator It(ClassToCountMap); It; ++It )
-	{
-		FClassCountInfo ClassCountInfo;
-		ClassCountInfo.ClassName		= It.Key();
-		ClassCountInfo.InstanceCount	= It.Value();
-		ClassCountArray.Add( ClassCountInfo );
-		if( bNeedToCalculateCount )
+#if PERF_DETAILED_PER_CLASS_GC_STATS
+		NumClusterToObjectRefs++;
+		if (ObjectItem->GetOwnerIndex() != 0)
 		{
-			TotalCount += ClassCountInfo.InstanceCount;
+			NumClusterToLateClusterRefs++;
 		}
+#endif
 	}
-	// Sort array by instance count.
-	struct FCompareFClassCountInfo
+
+	FORCEINLINE void FStats::IncClusterToClusterRefs(int32 Count)
+	{
+#if PERF_DETAILED_PER_CLASS_GC_STATS
+		NumClusterToClusterRefs += Count;
+#endif
+	}
+
+	FORCEINLINE void FStats::IncNumClustersTraversed()
+	{
+#if PERF_DETAILED_PER_CLASS_GC_STATS
+		NumClusterToObjectRefs++;
+#endif
+	}
+
+	void FStats::BeginTimingObject(UObject* CurrentObject)
+	{
+#if PERF_DETAILED_PER_CLASS_GC_STATS
+		CurrentObjectStartCyles = FPlatformTime::Cycles64();
+#endif
+	}
+
+	void FStats::UpdateDetailedStats(UObject* CurrentObject)
+	{
+#if PERF_DETAILED_PER_CLASS_GC_STATS
+		uint64 DeltaCycles = FPlatformTime::Cycles64() - CurrentObjectStartCyles;
+		// Keep track of how many refs we encountered for the object's class.
+		const FName& ClassName = CurrentObject->GetClass()->GetFName();
+		// Refs to objects that reside in permanent object pool.
+		ClassToDisregardedObjectRefsMap.FindOrAdd(ClassName) += CurrentObjectDisregardedObjectRefs;
+		// Refs to regular objects.
+		ClassToRegularObjectRefsMap.FindOrAdd(ClassName) += CurrentObjectRegularObjectRefs;
+		// Refs to clusters
+		ClassToClustersRefsMap.FindOrAdd(ClassName) += CurrentObjectClusterRefs;
+		// Track per class cycle count spent in GC.
+		uint64  ClassCycles = ClassToCyclesMap.FindRef(ClassName);
+		ClassToCyclesMap.Add(ClassName, ClassCycles + DeltaCycles);
+		++NumObjectsTraversed;
+		// Reset current counts.
+		CurrentObjectDisregardedObjectRefs = 0;
+		CurrentObjectRegularObjectRefs = 0;
+		CurrentObjectClusterRefs = 0;
+#endif
+	}
+
+	void FStats::LogDetailedStatsSummary()
+	{
+#if PERF_DETAILED_PER_CLASS_GC_STATS
+		LogClassCountInfo(TEXT("references to regular objects from"), ClassToRegularObjectRefsMap, 20, 0);
+		LogClassCountInfo(TEXT("references to permanent objects from"), ClassToDisregardedObjectRefsMap, 20, 0);
+		LogClassCountInfo(TEXT("references to clusters from"), ClassToClustersRefsMap, 20, 0);
+		LogClassCountInfo(TEXT("cycles for GC"), ClassToCyclesMap, 20, 0);
+
+		UE_LOG(LogGarbage, Log, TEXT("%d objects traversed"), NumObjectsTraversed);
+		UE_LOG(LogGarbage, Log, TEXT("%d clusters traversed"), NumClustersTraversed);
+		UE_LOG(LogGarbage, Log, TEXT("%d cluster to cluster refs followed"), NumClusterToClusterRefs);
+		UE_LOG(LogGarbage, Log, TEXT("%d cluster to object refs followed"), NumClusterToObjectRefs);
+		UE_LOG(LogGarbage, Log, TEXT("%d cluster to clustered mutable object refs followed"), NumClusterToLateClusterRefs);
+
+		NumObjectsTraversed = 0;
+		NumClustersTraversed = 0;
+		NumClusterToClusterRefs = 0;
+		NumClusterToObjectRefs = 0;
+		NumClusterToLateClusterRefs = 0;
+#endif
+	}
+
+	FORCEINLINE void FStats::IncreaseObjectRefStats(UObject* RefToObject)
 	{
-		FORCEINLINE bool operator()( const FClassCountInfo& A, const FClassCountInfo& B ) const
+#if PERF_DETAILED_PER_CLASS_GC_STATS
+		if (!RefToObject)
 		{
-			return B.InstanceCount < A.InstanceCount;
+			return;
 		}
-	};
-	ClassCountArray.Sort( FCompareFClassCountInfo() );
 
-	// Log top NumItemsToLog class counts
-	for( int32 Index=0; Index<FMath::Min(NumItemsToLog,ClassCountArray.Num()); Index++ )
+		FPermanentObjectPoolExtents Extents;
+		if (Extents.Contains(RefToObject))
+		{
+			++CurrentObjectDisregardedObjectRefs;
+		}
+		else 
+		{
+			FUObjectItem* Item = GUObjectArray.ObjectToObjectItem(RefToObject);
+			if (Item->HasAnyFlags(EInternalObjectFlags::ClusterRoot) || Item->GetOwnerIndex() > 0)
+			{
+				++CurrentObjectClusterRefs;
+			}
+			else
+			{
+				++CurrentObjectRegularObjectRefs;
+			}
+		}
+#endif
+	}
+
+	FORCEINLINE void FStats::IncPurgeCount(UObject* Object)
 	{
-		const FClassCountInfo& ClassCountInfo = ClassCountArray[Index];
-		const float Percent = 100.f * ClassCountInfo.InstanceCount / TotalCount;
-		const FString PercentString = (TotalCount > 0) ? FString::Printf(TEXT("%6.2f%%"), Percent) : FString(TEXT("  N/A  "));
-		UE_LOG(LogGarbage, Log, TEXT("%5d [%s] %s Class %s"), ClassCountInfo.InstanceCount, *PercentString, LogText, *ClassCountInfo.ClassName.ToString() ); 
+#if PERF_DETAILED_PER_CLASS_GC_STATS
+		// Keep track of how many objects of a certain class we're purging.
+		const FName& ClassName = Object->GetClass()->GetFName();
+		int32 InstanceCount = ClassToPurgeCountMap.FindRef( ClassName );
+		ClassToPurgeCountMap.Add( ClassName, ++InstanceCount );
+#endif
 	}
 
-	// Empty the map for the next run.
-	ClassToCountMap.Empty();
-};
+	FORCEINLINE void FStats::LogPurgeStats(int32 Total)
+	{
+#if PERF_DETAILED_PER_CLASS_GC_STATS
+		LogClassCountInfo(TEXT("objects of"), ClassToPurgeCountMap, 10, Total);
 #endif
+	}
+}
+
 
 /**
  * Helper class for destroying UObjects on a worker thread
@@ -771,6 +905,7 @@ static bool MarkClusterMutableObjectsAsReachable(FUObjectCluster& Cluster, Conta
 		if (ReferencedMutableObjectIndex >= 0) // Pending kill support
 		{
 			FUObjectItem* ReferencedMutableObjectItem = GUObjectArray.IndexToObjectUnsafeForGC(ReferencedMutableObjectIndex);
+			UE::GC::Stats.IncClusterToObjectRefs(ReferencedMutableObjectItem);
 			if constexpr (IsParallel(Options))
 			{
 				PRAGMA_DISABLE_DEPRECATION_WARNINGS
@@ -863,12 +998,15 @@ static FORCENOINLINE void MarkReferencedClustersAsReachable(int32 ClusterIndex,
 {
 	static_assert(IsWithClusters(Options));
 
+	UE::GC::Stats.IncNumClustersTraversed();
+
 	// If we run across some PendingKill objects we need to add all objects from this cluster
 	// to ObjectsToSerialize so that we can properly null out all the references.
 	// It also means this cluster will have to be dissolved because we may no longer guarantee all cross-cluster references are correct.
 
 	bool bAddClusterObjectsToSerialize = false;
 	FUObjectCluster& Cluster = GUObjectClusters[ClusterIndex];
+	UE::GC::Stats.IncClusterToClusterRefs(Cluster.ReferencedClusters.Num());
 	// Also mark all referenced objects from outside of the cluster as reachable
 	for (int32& ReferncedClusterIndex : Cluster.ReferencedClusters)
 	{
@@ -1187,7 +1325,7 @@ private:
 			UObject* Object = GetObject(UnvalidatedReferences[Idx].Reference);
 			ValidsB.Set(Idx, (!!Object) & IsObjectHandleResolved(reinterpret_cast<FObjectHandle&>(Object))); //-V792
 		}
-
+		
 		FValidatedBitmask Validations = FValidatedBitmask::And(ValidsA, ValidsB);
 		uint32 NumValid = Validations.CountBits();
 		uint32 UnvalidatedIdx = 0;
@@ -2272,43 +2410,32 @@ FORCEINLINE static bool ClearUnreachableInterlocked(int32& Flags)
 	return false;
 }
 
+
+
+
 template <EGCOptions InOptions>
-class TReachabilityProcessor
+class TReachabilityProcessor 
 {
 public:
-	static constexpr EGCOptions Options = InOptions;
-
-	static constexpr FORCEINLINE bool IsWithPendingKill() {	return !!(Options & EGCOptions::WithPendingKill);}
+	FORCEINLINE void BeginTimingObject(UObject* CurrentObject)
+	{
+		UE::GC::Stats.BeginTimingObject(CurrentObject);
+	}
 
-	static void UpdateDetailedStats(UObject* CurrentObject, uint32 DeltaCycles)
+	FORCEINLINE void UpdateDetailedStats(UObject* CurrentObject)
 	{
-#if PERF_DETAILED_PER_CLASS_GC_STATS
-		// Keep track of how many refs we encountered for the object's class.
-		const FName& ClassName = CurrentObject->GetClass()->GetFName();
-		// Refs to objects that reside in permanent object pool.
-		uint32 ClassDisregardedObjRefs = GClassToDisregardedObjectRefsMap.FindRef(ClassName);
-		GClassToDisregardedObjectRefsMap.Add(ClassName, ClassDisregardedObjRefs + GCurrentObjectDisregardedObjectRefs);
-		// Refs to regular objects.
-		uint32 ClassRegularObjRefs = GClassToRegularObjectRefsMap.FindRef(ClassName);
-		GClassToRegularObjectRefsMap.Add(ClassName, ClassRegularObjRefs + GCurrentObjectRegularObjectRefs);
-		// Track per class cycle count spent in GC.
-		uint32 ClassCycles = GClassToCyclesMap.FindRef(ClassName);
-		GClassToCyclesMap.Add(ClassName, ClassCycles + DeltaCycles);
-		// Reset current counts.
-		GCurrentObjectDisregardedObjectRefs = 0;
-		GCurrentObjectRegularObjectRefs = 0;
-#endif
+		UE::GC::Stats.UpdateDetailedStats(CurrentObject);
 	}
 
-	static void LogDetailedStatsSummary()
+	FORCEINLINE void LogDetailedStatsSummary()
 	{
-#if PERF_DETAILED_PER_CLASS_GC_STATS
-		LogClassCountInfo(TEXT("references to regular objects from"), GClassToRegularObjectRefsMap, 20, 0);
-		LogClassCountInfo(TEXT("references to permanent objects from"), GClassToDisregardedObjectRefsMap, 20, 0);
-		LogClassCountInfo(TEXT("cycles for GC"), GClassToCyclesMap, 20, 0);
-#endif
+		UE::GC::Stats.LogDetailedStatsSummary();
 	}
 
+	static constexpr EGCOptions Options = InOptions;
+
+	static constexpr FORCEINLINE bool IsWithPendingKill() {	return !!(Options & EGCOptions::WithPendingKill);}
+
 	PRAGMA_DISABLE_DEPRECATION_WARNINGS
 	static constexpr EInternalObjectFlags KillFlag = IsWithPendingKill() ? EInternalObjectFlags::PendingKill : EInternalObjectFlags::Garbage;
 	PRAGMA_ENABLE_DEPRECATION_WARNINGS
@@ -2338,6 +2465,7 @@ public:
 	template<EKillable Killable>
 	static FORCEINLINE_DEBUGGABLE void ProcessReferenceDirectly(FWorkerContext& Context, FPermanentObjectPoolExtents PermanentPool, const UObject* ReferencingObject, UObject*& Object, FTokenId TokenId)
 	{
+		UE::GC::Stats.IncreaseObjectRefStats(Object);
 		if (ValidateReference(Object, PermanentPool, ReferencingObject, TokenId))
 		{
 			const int32 ObjectIndex = GUObjectArray.ObjectToIndex(Object);
@@ -2359,6 +2487,7 @@ public:
 
 	FORCEINLINE static void HandleBatchedReference(FWorkerContext& Context, FResolvedMutableReference Reference, FReferenceMetadata Metadata)
 	{
+		UE::GC::Stats.IncreaseObjectRefStats(GetObject(Reference));
 		if (Metadata.Has(KillFlag))
 		{
 			checkSlow(Metadata.ObjectItem->GetOwnerIndex() <= 0);
@@ -2372,6 +2501,7 @@ public:
 
 	FORCEINLINE static void HandleBatchedReference(FWorkerContext& Context, FImmutableReference Reference, FReferenceMetadata Metadata)
 	{
+		UE::GC::Stats.IncreaseObjectRefStats(GetObject(Reference));
 		DetectGarbageReference(Context, Metadata);
 		HandleValidReference(Context, Reference, Metadata);
 	}
@@ -2431,9 +2561,6 @@ public:
 					}
 				}
 			}
-	#if PERF_DETAILED_PER_CLASS_GC_STATS
-			GCurrentObjectRegularObjectRefs++;
-	#endif
 		}
 
 		return false;
@@ -2633,9 +2760,24 @@ public:
 #endif
 
 template <EGCOptions InOptions>
-class TDebugReachabilityProcessor
+class TDebugReachabilityProcessor 
 {
 public:
+	FORCEINLINE void BeginTimingObject(UObject* CurrentObject)
+	{
+		UE::GC::Stats.BeginTimingObject(CurrentObject);
+	}
+
+	FORCEINLINE void UpdateDetailedStats(UObject* CurrentObject)
+	{
+		UE::GC::Stats.UpdateDetailedStats(CurrentObject);
+	}
+
+	FORCEINLINE void LogDetailedStatsSummary()
+	{
+		UE::GC::Stats.LogDetailedStatsSummary();
+	}
+
 	static constexpr EGCOptions Options = InOptions;
 
 	TDebugReachabilityProcessor()
@@ -2648,6 +2790,7 @@ public:
 
 	FORCENOINLINE void HandleTokenStreamObjectReference(FWorkerContext& Context, const UObject* ReferencingObject, UObject*& Object, FTokenId TokenId, EGCTokenType TokenType, bool bAllowReferenceElimination)
 	{
+		UE::GC::Stats.IncreaseObjectRefStats(Object);
 		if (ValidateReference(Object, PermanentPool, ReferencingObject, TokenId))
 		{
 			FReferenceMetadata Metadata(GUObjectArray.ObjectToIndex(Object));
@@ -3539,12 +3682,7 @@ bool IncrementalDestroyGarbage(bool bUseTimeLimit, double TimeLimit)
 					// Only proceed with destroying the object if the asynchronous cleanup started by BeginDestroy has finished.
 					if(Object->IsReadyForFinishDestroy())
 					{
-#if PERF_DETAILED_PER_CLASS_GC_STATS
-						// Keep track of how many objects of a certain class we're purging.
-						const FName& ClassName = Object->GetClass()->GetFName();
-						int32 InstanceCount = GClassToPurgeCountMap.FindRef( ClassName );
-						GClassToPurgeCountMap.Add( ClassName, ++InstanceCount );
-#endif
+						UE::GC::Stats.IncPurgeCount(Object);
 						// Send FinishDestroy message.
 						Object->ConditionalFinishDestroy();
 					}
@@ -3604,12 +3742,7 @@ bool IncrementalDestroyGarbage(bool bUseTimeLimit, double TimeLimit)
 					// Only proceed with destroying the object if the asynchronous cleanup started by BeginDestroy has finished.
 					if( Object->IsReadyForFinishDestroy() )
 					{
-#if PERF_DETAILED_PER_CLASS_GC_STATS
-						// Keep track of how many objects of a certain class we're purging.
-						const FName& ClassName = Object->GetClass()->GetFName();
-						int32 InstanceCount = GClassToPurgeCountMap.FindRef( ClassName );
-						GClassToPurgeCountMap.Add( ClassName, ++InstanceCount );
-#endif
+						UE::GC::Stats.IncPurgeCount(Object);
 						// Send FinishDestroy message.
 						Object->ConditionalFinishDestroy();
 
@@ -3774,9 +3907,7 @@ bool IncrementalDestroyGarbage(bool bUseTimeLimit, double TimeLimit)
 				GObjectCountDuringLastMarkPhase.GetValue(), 
 				GObjectCountDuringLastMarkPhase.GetValue() - PurgedObjectCountSinceLastMarkPhase,
 				(FPlatformTime::Seconds() - IncrementalDestroyGarbageStartTime) * 1000);
-#if PERF_DETAILED_PER_CLASS_GC_STATS
-			LogClassCountInfo( TEXT("objects of"), GClassToPurgeCountMap, 10, PurgedObjectCountSinceLastMarkPhase);
-#endif
+			UE::GC::Stats.LogPurgeStats(PurgedObjectCountSinceLastMarkPhase);
 			GAsyncPurge->ResetObjectsDestroyedSinceLastMarkPhase();
 		}
 	}
@@ -4163,7 +4294,7 @@ void CollectGarbageImpl(EObjectFlags KeepFlags)
 				DECLARE_SCOPE_CYCLE_COUNTER(TEXT("FRealtimeGC::PerformReachabilityAnalysis"), STAT_FArchiveRealtimeGC_PerformReachabilityAnalysis, STATGROUP_GC);
 				const double StartTime = FPlatformTime::Seconds();
 				FRealtimeGC().PerformReachabilityAnalysis(KeepFlags, Options);
-				UE_LOG(LogGarbage, Log, TEXT("%f ms for GC"), (FPlatformTime::Seconds() - StartTime) * 1000);
+				UE_LOG(LogGarbage, Log, TEXT("%f ms for GC with %d clusters"), (FPlatformTime::Seconds() - StartTime) * 1000, GUObjectClusters.GetNumAllocatedClusters());
 			}
 
 			if (GRedoReachabilityToTrackGarbage && GGarbageReferenceTrackingEnabled > 0)
diff --git a/Engine/Source/Runtime/CoreUObject/Public/UObject/FastReferenceCollector.h b/Engine/Source/Runtime/CoreUObject/Public/UObject/FastReferenceCollector.h
--- a/Engine/Source/Runtime/CoreUObject/Public/UObject/FastReferenceCollector.h
+++ b/Engine/Source/Runtime/CoreUObject/Public/UObject/FastReferenceCollector.h
@@ -565,6 +565,8 @@ StoleContext:
 					continue;
 				}
 
+				Processor.BeginTimingObject(CurrentObject);
+
 				uint32 TokenStreamIndex = 0;
 				// Keep track of index to reference info. Used to avoid LHSs.
 				uint32 ReferenceTokenStreamIndex = 0;
@@ -1088,6 +1090,7 @@ StoleContext:
 					} // switch (Type)
 				} // while (true)
 TokensDone:
+				Processor.UpdateDetailedStats(CurrentObject);
 				check(StackEntry == Stack.GetBottomFrame());
 			} // for (FPrefetchingObjectIterator)
 
@@ -1132,11 +1135,7 @@ StoleARO:
 			CurrentObjects = MakeArrayView(Block->Objects, BlockSize);
 		} // while (true)
 		
-	#if PERF_DETAILED_PER_CLASS_GC_STATS
-		// Detailed per class stats should not be performed when parallel GC is running
-		check(!IsParallel());
 		Processor.LogDetailedStatsSummary();
-	#endif
 	}
 };
 
@@ -1182,7 +1181,11 @@ class FSimpleReferenceProcessorBase
 {
 public:
 	static constexpr EGCOptions Options = UE::GC::DefaultOptions;
-	void UpdateDetailedStats(UObject* CurrentObject, uint32 DeltaCycles) {}
+
+	// These functions are implemented in the GC collectors to generate detail stats on objects considered by GC.
+	// They are generally not needed for other reference collection tasks.
+	void BeginTimingObject(UObject* CurrentObject) {}
+	void UpdateDetailedStats(UObject* CurrentObject) {}
 	void LogDetailedStatsSummary() {}
 
 	// Implement this in your derived class, don't make this virtual as it will affect performance!
diff --git a/Engine/Source/Runtime/CoreUObject/Public/UObject/GarbageCollection.h b/Engine/Source/Runtime/CoreUObject/Public/UObject/GarbageCollection.h
--- a/Engine/Source/Runtime/CoreUObject/Public/UObject/GarbageCollection.h
+++ b/Engine/Source/Runtime/CoreUObject/Public/UObject/GarbageCollection.h
@@ -32,7 +32,6 @@ class UObject;
 #define GARBAGE_COLLECTION_KEEPFLAGS	(GIsEditor ? RF_Standalone : RF_NoFlags)
 
 #define	ENABLE_GC_DEBUG_OUTPUT					1
-#define PERF_DETAILED_PER_CLASS_GC_STATS				(LOOKING_FOR_PERF_ISSUES || 0) 
 
 /** UObject pointer checks are disabled by default in shipping and test builds as they add roughly 20% overhead to GC times */
 #ifndef ENABLE_GC_OBJECT_CHECKS
