Fix truncation warning when serializing FObjectExport

#rnx

- This code branch should only run when loading older packages as we should never be saving with an older package version, therefor it is safe to just cast.

diff --git a/Engine/Source/Runtime/CoreUObject/Private/UObject/ObjectResource.cpp b/Engine/Source/Runtime/CoreUObject/Private/UObject/ObjectResource.cpp
--- a/Engine/Source/Runtime/CoreUObject/Private/UObject/ObjectResource.cpp
+++ b/Engine/Source/Runtime/CoreUObject/Private/UObject/ObjectResource.cpp
@@ -144,13 +144,13 @@ void operator<<(FStructuredArchive::FSlot Slot, FObjectExport& E)
 
 	if (BaseArchive.UEVer() < VER_UE4_64BIT_EXPORTMAP_SERIALSIZES)
 	{
-		int32 SerialSize = E.SerialSize;
+		int32 SerialSize = (int32)E.SerialSize;
 		Record << SA_VALUE(TEXT("SerialSize"), SerialSize);
 		E.SerialSize = (int64)SerialSize;
 
-		int32 SerialOffset = E.SerialOffset;
+		int32 SerialOffset = (int32)E.SerialOffset;
 		Record << SA_VALUE(TEXT("SerialOffset"), SerialOffset);
-		E.SerialOffset = SerialOffset;
+		E.SerialOffset = (int64)SerialOffset;
 	}
 	else
 	{
