FLinkerInstancingContex: refactor SoftObjectPath fixup code so it is shared between AsyncPackage2 and FLinkerload
- Fixes cooked softobjectpath fixup on instanced WorldPartitions when not using iostore

diff --git a/Engine/Source/Runtime/CoreUObject/Private/Serialization/AsyncLoading2.cpp b/Engine/Source/Runtime/CoreUObject/Private/Serialization/AsyncLoading2.cpp
--- a/Engine/Source/Runtime/CoreUObject/Private/Serialization/AsyncLoading2.cpp
+++ b/Engine/Source/Runtime/CoreUObject/Private/Serialization/AsyncLoading2.cpp
@@ -1508,13 +1508,6 @@ public:
 	}
 };
 
-struct FAsyncPackageInstanceContext
-{
-	TArray<TTuple<FString, FString>, TInlineAllocator<2>> InstancedPackageMappings;
-	FString GeneratedPackagesFolder;
-	FString InstancedPackageSuffix;
-};
-
 class FExportArchive final : public FArchive
 {
 public:
@@ -1867,7 +1860,7 @@ private:
 	FPackageImportStore* ImportStore = nullptr;
 	TArray<FExternalReadCallback>* ExternalReadDependencies;
 	const FAsyncPackageHeaderData* HeaderData = nullptr;
-	const FAsyncPackageInstanceContext* InstanceContext = nullptr;
+	const FLinkerInstancingContext* InstanceContext = nullptr;
 	TArrayView<const FExportObject> Exports;
 	UObject* CurrentExport = nullptr;
 	uint64 CookedSerialOffset = 0;
@@ -1878,44 +1871,9 @@ private:
 
 	void FixupSoftObjectPathForInstancedPackage(FSoftObjectPath& InOutSoftObjectPath)
 	{
-		if (InstanceContext && InstanceContext->InstancedPackageMappings.Num() > 0)
+		if (InstanceContext)
 		{
-			FNameBuilder TmpSoftObjectPathBuilder;
-			InOutSoftObjectPath.ToString(TmpSoftObjectPathBuilder);
-
-			// Is this a basic package remapping?
-			for (const TTuple<FString, FString>& InstancedPackageMapping : InstanceContext->InstancedPackageMappings)
-			{
-				FStringView TmpSoftObjectPathView = TmpSoftObjectPathBuilder.ToView();
-				if (TmpSoftObjectPathView.StartsWith(InstancedPackageMapping.Key) && (TmpSoftObjectPathView.Len() == InstancedPackageMapping.Key.Len() || TmpSoftObjectPathView[InstancedPackageMapping.Key.Len()] == TEXT('.')))
-				{
-					TmpSoftObjectPathBuilder.ReplaceAt(0, InstancedPackageMapping.Key.Len(), InstancedPackageMapping.Value);
-					InOutSoftObjectPath.SetPath(TmpSoftObjectPathBuilder.ToView());
-					return;
-				}
-			}
-
-			// Is this a generated package reference that we care about?
-			if (!InstanceContext->GeneratedPackagesFolder.IsEmpty())
-			{
-				check(!InstanceContext->InstancedPackageSuffix.IsEmpty());
-
-				// Does this package path start with the generated folder path?
-				FStringView TmpSoftObjectPathView = TmpSoftObjectPathBuilder.ToView();
-				if (TmpSoftObjectPathView.StartsWith(InstanceContext->GeneratedPackagesFolder))
-				{
-					// ... and is that generated folder path immediately preceding the package name?
-					if (const int32 ExtraSlashIndex = TmpSoftObjectPathView.Find(TEXTVIEW("/"), InstanceContext->GeneratedPackagesFolder.Len()); ExtraSlashIndex == INDEX_NONE)
-					{
-						FNameBuilder PackageNameBuilder;
-						PackageNameBuilder.Append(InOutSoftObjectPath.GetLongPackageName());
-						PackageNameBuilder.Append(InstanceContext->InstancedPackageSuffix);
-						FTopLevelAssetPath SuffixTopLevelAsset(FName(PackageNameBuilder.ToString()), InOutSoftObjectPath.GetAssetFName());
-						InOutSoftObjectPath = FSoftObjectPath(SuffixTopLevelAsset, InOutSoftObjectPath.GetSubPathString());
-						return;
-					}
-				}
-			}
+			InstanceContext->FixupSoftObjectPath(InOutSoftObjectPath);
 		}
 	}
 };
@@ -2768,7 +2726,7 @@ private:
 	TArray<FCompletionCallback, TInlineAllocator<2>> CompletionCallbacks;
 
 	/** Set when the package is being loaded as an instance; null otherwise. */
-	TUniquePtr<FAsyncPackageInstanceContext> InstanceContext;
+	TUniquePtr<FLinkerInstancingContext> InstanceContext;
 
 public:
 
@@ -5400,43 +5358,8 @@ EEventLoadNodeExecutionResult FAsyncPackage2::Event_ProcessPackageSummary(FAsync
 		FName PackageNameToLoad = Package->Desc.PackagePathToLoad.GetPackageFName();
 		if (Package->Desc.UPackageName != PackageNameToLoad)
 		{
-			Package->InstanceContext = MakeUnique<FAsyncPackageInstanceContext>();
-			Package->InstanceContext->InstancedPackageMappings.Add(MakeTuple(PackageNameToLoad.ToString(), Package->Desc.UPackageName.ToString()));
-
-			// Stash the suffix used for this instance so we can also apply it to generated packages
-			if (Package->InstanceContext->InstancedPackageMappings[0].Value.StartsWith(Package->InstanceContext->InstancedPackageMappings[0].Key))
-			{
-				Package->InstanceContext->InstancedPackageSuffix = Package->InstanceContext->InstancedPackageMappings[0].Value.Mid(Package->InstanceContext->InstancedPackageMappings[0].Key.Len());
-			}
-
-			// Is this a generated partitioned map package? If so, we'll also need to handle re-mapping paths to our persistent map package
-			if (!Package->InstanceContext->InstancedPackageSuffix.IsEmpty())
-			{
-				check(Package->InstanceContext->InstancedPackageMappings.Num() == 1);
-
-				const FStringView GeneratedFolderName = TEXTVIEW("/_Generated_/");
-				const FStringView PackageNameToLoadView = Package->InstanceContext->InstancedPackageMappings[0].Key;
-
-				// Does this package path include the generated folder?
-				if (const int32 GeneratedFolderStartIndex = PackageNameToLoadView.Find(GeneratedFolderName); GeneratedFolderStartIndex != INDEX_NONE)
-				{
-					// ... and is that generated folder immediately preceding the package name?
-					const int32 GeneratedFolderEndIndex = GeneratedFolderStartIndex + GeneratedFolderName.Len();
-					if (const int32 ExtraSlashIndex = PackageNameToLoadView.Find(TEXTVIEW("/"), GeneratedFolderEndIndex); ExtraSlashIndex == INDEX_NONE)
-					{
-						Package->InstanceContext->GeneratedPackagesFolder = PackageNameToLoadView.Left(GeneratedFolderEndIndex);
-
-						FString PersistentPackageName = FString(PackageNameToLoadView.Left(GeneratedFolderStartIndex));
-						FString PersistentPackageInstanceName = PersistentPackageName + Package->InstanceContext->InstancedPackageSuffix;
-						Package->InstanceContext->InstancedPackageMappings.Add(MakeTuple(MoveTemp(PersistentPackageName), MoveTemp(PersistentPackageInstanceName)));
-					}
-				}
-				if (Package->InstanceContext->GeneratedPackagesFolder.IsEmpty())
-				{
-					Package->InstanceContext->GeneratedPackagesFolder =  PackageNameToLoadView;
-					Package->InstanceContext->GeneratedPackagesFolder += GeneratedFolderName;
-				}
-			}
+			Package->InstanceContext = MakeUnique<FLinkerInstancingContext>();
+			Package->InstanceContext->BuildPackageMapping(PackageNameToLoad, Package->Desc.UPackageName);
 		}
 
 		TRACE_LOADTIME_PACKAGE_SUMMARY(Package, Package->HeaderData.PackageName, Package->HeaderData.PackageSummary->HeaderSize, Package->HeaderData.ImportMap.Num(), Package->HeaderData.ExportMap.Num());
diff --git a/Engine/Source/Runtime/CoreUObject/Private/UObject/LinkerInstancingContext.cpp b/Engine/Source/Runtime/CoreUObject/Private/UObject/LinkerInstancingContext.cpp
new file mode 100644
--- /dev/null
+++ b/Engine/Source/Runtime/CoreUObject/Private/UObject/LinkerInstancingContext.cpp
@@ -0,0 +1,103 @@
+// Copyright Epic Games, Inc. All Rights Reserved.
+
+#include "UObject/LinkerInstancingContext.h"
+#include "UObject/SoftObjectPath.h"
+#include "UObject/TopLevelAssetPath.h"
+#include "UObject/NameTypes.h"
+#include "Containers/StringView.h"
+#include "Containers/UnrealString.h"
+#include "Templates/Tuple.h"
+
+void FLinkerInstancingContext::FixupSoftObjectPath(FSoftObjectPath& InOutSoftObjectPath) const
+{
+	if (IsInstanced() && GetSoftObjectPathRemappingEnabled())
+	{
+		// Try remapping AssetPathName before remapping LongPackageName
+		if (FSoftObjectPath RemappedAssetPath = RemapPath(InOutSoftObjectPath); RemappedAssetPath != InOutSoftObjectPath)
+		{
+			InOutSoftObjectPath = RemappedAssetPath;
+		}
+		else if (FName LongPackageName = InOutSoftObjectPath.GetLongPackageFName(), RemappedPackage = RemapPackage(LongPackageName); RemappedPackage != LongPackageName)
+		{
+			InOutSoftObjectPath = FSoftObjectPath(RemappedPackage, InOutSoftObjectPath.GetAssetFName(), InOutSoftObjectPath.GetSubPathString());
+		}
+#if !WITH_EDITOR
+		else if (!GeneratedPackagesFolder.IsEmpty())
+		{
+			check(!InstancedPackageSuffix.IsEmpty());
+
+			FNameBuilder TmpSoftObjectPathBuilder;
+			InOutSoftObjectPath.ToString(TmpSoftObjectPathBuilder);
+
+			// Does this package path start with the generated folder path?
+			FStringView TmpSoftObjectPathView = TmpSoftObjectPathBuilder.ToView();
+			if (TmpSoftObjectPathView.StartsWith(GeneratedPackagesFolder))
+			{
+				// ... and is that generated folder path immediately preceding the package name?
+				if (const int32 ExtraSlashIndex = TmpSoftObjectPathView.Find(TEXTVIEW("/"), GeneratedPackagesFolder.Len()); ExtraSlashIndex == INDEX_NONE)
+				{
+					FNameBuilder PackageNameBuilder;
+					PackageNameBuilder.Append(InOutSoftObjectPath.GetLongPackageName());
+					PackageNameBuilder.Append(InstancedPackageSuffix);
+					FTopLevelAssetPath SuffixTopLevelAsset(FName(PackageNameBuilder.ToString()), InOutSoftObjectPath.GetAssetFName());
+					InOutSoftObjectPath = FSoftObjectPath(SuffixTopLevelAsset, InOutSoftObjectPath.GetSubPathString());
+				}
+			}
+		}
+#endif
+	}
+}
+
+void FLinkerInstancingContext::BuildPackageMapping(FName Original, FName Instanced)
+{
+	check(GeneratedPackagesFolder.IsEmpty() && InstancedPackageSuffix.IsEmpty());
+
+	AddPackageMapping(Original, Instanced);
+
+#if !WITH_EDITOR
+	if (bSoftObjectPathRemappingEnabled)
+	{
+		FNameBuilder TmpOriginal(Original);
+		FNameBuilder TmpInstanced(Instanced);
+		FStringView OriginalView = TmpOriginal.ToView();
+		FStringView InstancedView = TmpInstanced.ToView();
+
+		// Stash the suffix used for this instance so we can also apply it to generated packages
+		if (InstancedView.StartsWith(OriginalView))
+		{
+			InstancedPackageSuffix = InstancedView.Mid(OriginalView.Len());
+		}
+
+		// Is this a generated partitioned map package? If so, we'll also need to handle re-mapping paths to our persistent map package
+		if (!InstancedPackageSuffix.IsEmpty())
+		{
+			const FStringView GeneratedFolderName = TEXTVIEW("/_Generated_/");
+			
+			// Does this package path include the generated folder?
+			if (const int32 GeneratedFolderStartIndex = OriginalView.Find(GeneratedFolderName); GeneratedFolderStartIndex != INDEX_NONE)
+			{
+				// ... and is that generated folder immediately preceding the package name?
+				const int32 GeneratedFolderEndIndex = GeneratedFolderStartIndex + GeneratedFolderName.Len();
+				if (const int32 ExtraSlashIndex = OriginalView.Find(TEXTVIEW("/"), GeneratedFolderEndIndex); ExtraSlashIndex == INDEX_NONE)
+				{
+					GeneratedPackagesFolder = OriginalView.Left(GeneratedFolderEndIndex);
+
+					FNameBuilder PackageNameBuilder(FName(OriginalView.Left(GeneratedFolderStartIndex)));
+					FName PersistentPackageName(PackageNameBuilder);
+					PackageNameBuilder.Append(InstancedPackageSuffix);
+					FName PersistentPackageInstanceName(PackageNameBuilder);
+										
+					PackageMapping.Add(MakeTuple(PersistentPackageName, PersistentPackageInstanceName));
+				}
+			}
+
+			if (GeneratedPackagesFolder.IsEmpty())
+			{
+				GeneratedPackagesFolder = OriginalView;
+				GeneratedPackagesFolder += GeneratedFolderName;
+			}
+		}
+	}
+#endif
+}
+
diff --git a/Engine/Source/Runtime/CoreUObject/Private/UObject/LinkerLoad.cpp b/Engine/Source/Runtime/CoreUObject/Private/UObject/LinkerLoad.cpp
--- a/Engine/Source/Runtime/CoreUObject/Private/UObject/LinkerLoad.cpp
+++ b/Engine/Source/Runtime/CoreUObject/Private/UObject/LinkerLoad.cpp
@@ -1005,7 +1005,7 @@ FLinkerLoad::FLinkerLoad(UPackage* InParent, const FPackagePath& InPackagePath,
 	FName PackageNameToLoad = GetPackagePath().GetPackageFName();
 	if (LinkerRoot->GetFName() != PackageNameToLoad)
 	{
-		InstancingContext.AddPackageMapping(PackageNameToLoad, LinkerRoot->GetFName());
+		InstancingContext.BuildPackageMapping(PackageNameToLoad, LinkerRoot->GetFName());
 	}
 
 	TRACE_LOADTIME_NEW_LINKER(this);
@@ -6644,22 +6644,7 @@ bool FLinkerLoad::IsSoftObjectRemappingEnabled() const
 
 void FLinkerLoad::FixupSoftObjectPathForInstancedPackage(FSoftObjectPath& InOutSoftObjectPath)
 {
-	if (IsSoftObjectRemappingEnabled())
-	{
-		// Try remapping AssetPathName before remapping LongPackageName
-		if (FSoftObjectPath RemappedAssetPath = InstancingContext.RemapPath(InOutSoftObjectPath); RemappedAssetPath != InOutSoftObjectPath)
-		{
-			InOutSoftObjectPath = RemappedAssetPath;
-		}
-		else
-		{
-			FName LongPackageName = InOutSoftObjectPath.GetLongPackageFName();
-			if (FName RemappedPackage = InstancingContext.RemapPackage(LongPackageName); RemappedPackage != LongPackageName)
-			{
-				InOutSoftObjectPath = FSoftObjectPath(RemappedPackage, InOutSoftObjectPath.GetAssetFName(), InOutSoftObjectPath.GetSubPathString());
-			}
-		}
-	}
+	InstancingContext.FixupSoftObjectPath(InOutSoftObjectPath);
 }
 
 #if WITH_EDITORONLY_DATA
diff --git a/Engine/Source/Runtime/CoreUObject/Public/UObject/LinkerInstancingContext.h b/Engine/Source/Runtime/CoreUObject/Public/UObject/LinkerInstancingContext.h
--- a/Engine/Source/Runtime/CoreUObject/Public/UObject/LinkerInstancingContext.h
+++ b/Engine/Source/Runtime/CoreUObject/Public/UObject/LinkerInstancingContext.h
@@ -105,7 +105,11 @@ public:
 		return FString::Printf(TEXT("%s_InstanceOf_%s"), *InOuterPackageName, *InPackageName);
 	}
 
+	void FixupSoftObjectPath(FSoftObjectPath& InOutSoftObjectPath) const;
+
 private:
+	void BuildPackageMapping(FName Original, FName Instanced);
+
 	/** Used internally by the linker to try to fix references on relocated packages. */
 	FName RelocatePackage(const FName& PackageName) const
 	{
@@ -117,6 +121,7 @@ private:
 	}
 
 	friend class FLinkerLoad;
+	friend struct FAsyncPackage2;
 
 	/** Map of original package name to their instance counterpart. */
 	TMap<FName, FName> PackageMapping;
@@ -129,4 +134,7 @@ private:
 	TSet<FName> Tags;
 	/** Remap soft object paths */
 	bool bSoftObjectPathRemappingEnabled = true;
+
+	FString GeneratedPackagesFolder;
+	FString InstancedPackageSuffix;
 };
