added ObjectSerializeScope which will log a warning when enabled about accessing referenced objects during serialize

diff --git a/Engine/Source/Runtime/CoreUObject/Private/Serialization/AsyncLoading.cpp b/Engine/Source/Runtime/CoreUObject/Private/Serialization/AsyncLoading.cpp
--- a/Engine/Source/Runtime/CoreUObject/Private/Serialization/AsyncLoading.cpp
+++ b/Engine/Source/Runtime/CoreUObject/Private/Serialization/AsyncLoading.cpp
@@ -30,6 +30,7 @@
 #include "UObject/PackageResourceManager.h"
 #include "Serialization/DeferredMessageLog.h"
 #include "UObject/UObjectThreadContext.h"
+#include "UObject/ObjectSerializeAccessScope.h"
 #include "UObject/LinkerManager.h"
 #include "Misc/Paths.h"
 #include "Serialization/AsyncLoadingThread.h"
@@ -3523,6 +3524,7 @@ void FAsyncPackage::EventDrivenSerializeExport(int32 LocalExportIndex)
 			else
 			{
 				UE_SCOPED_IO_ACTIVITY(*Object->GetName());
+				UE_SERIALIZE_ACCCESS_SCOPE(Object);
 				Object->Serialize(ExportArchive);
 			}
 		}
diff --git a/Engine/Source/Runtime/CoreUObject/Private/Serialization/AsyncLoading2.cpp b/Engine/Source/Runtime/CoreUObject/Private/Serialization/AsyncLoading2.cpp
--- a/Engine/Source/Runtime/CoreUObject/Private/Serialization/AsyncLoading2.cpp
+++ b/Engine/Source/Runtime/CoreUObject/Private/Serialization/AsyncLoading2.cpp
@@ -57,6 +57,7 @@
 #include "Serialization/MemoryReader.h"
 #include "UObject/UObjectClusters.h"
 #include "UObject/LinkerInstancingContext.h"
+#include "UObject/ObjectSerializeAccessScope.h"
 #include "ProfilingDebugging/CountersTrace.h"
 #include "ProfilingDebugging/AssetMetadataTrace.h"
 #include "Async/Async.h"
@@ -5925,6 +5926,7 @@ bool FAsyncPackage2::EventDrivenSerializeExport(const FAsyncPackageHeaderData& H
 	else
 	{
 		TRACE_CPUPROFILER_EVENT_SCOPE(SerializeObject);
+		UE_SERIALIZE_ACCCESS_SCOPE(Object);
 		Object->Serialize(Ar);
 	}
 	Ar.TemplateForGetArchetypeFromLoader = nullptr;
diff --git a/Engine/Source/Runtime/CoreUObject/Private/UObject/LinkerLoad.cpp b/Engine/Source/Runtime/CoreUObject/Private/UObject/LinkerLoad.cpp
--- a/Engine/Source/Runtime/CoreUObject/Private/UObject/LinkerLoad.cpp
+++ b/Engine/Source/Runtime/CoreUObject/Private/UObject/LinkerLoad.cpp
@@ -34,6 +34,7 @@
 #include "UObject/LinkerPlaceholderExportObject.h"
 #include "UObject/LinkerPlaceholderFunction.h"
 #include "UObject/LinkerManager.h"
+#include "UObject/ObjectSerializeAccessScope.h"
 #include "Serialization/DeferredMessageLog.h"
 #include "UObject/UObjectThreadContext.h"
 #include "Serialization/AsyncLoading.h"
@@ -1116,6 +1117,7 @@ FLinkerLoad::ELinkerStatus FLinkerLoad::CreateLoader(
 #if WITH_EDITOR
 		if (LoadProgressScope)
 		{
+			UE_SERIALIZE_ACCCESS_SCOPE_SUSPEND();
 			static const FTextFormat LoadingFileTextFormat = NSLOCTEXT("Core", "LoadingFileWithFilename", "Loading file: {CleanFilename}...");
 			FFormatNamedArguments FeedbackArgs;
 			FeedbackArgs.Add( TEXT("CleanFilename"), FText::FromString( FPaths::GetCleanFilename( *GetDebugName() ) ) );
@@ -1292,6 +1294,7 @@ FLinkerLoad::ELinkerStatus FLinkerLoad::SerializePackageFileSummaryInternal()
 #if WITH_EDITOR
 	if (LoadProgressScope)
 	{
+		UE_SERIALIZE_ACCCESS_SCOPE_SUSPEND();
 		LoadProgressScope->EnterProgressFrame(1);
 	}
 #endif
@@ -1847,6 +1850,7 @@ FLinkerLoad::ELinkerStatus FLinkerLoad::FixupImportMap()
 #if WITH_EDITOR
 		if (LoadProgressScope)
 		{
+			UE_SERIALIZE_ACCCESS_SCOPE_SUSPEND();
 			LoadProgressScope->EnterProgressFrame(1);
 		}
 #endif
@@ -2581,6 +2585,7 @@ FLinkerLoad::ELinkerStatus FLinkerLoad::FindExistingExports()
 #if WITH_EDITOR
 		if (LoadProgressScope)
 		{
+			UE_SERIALIZE_ACCCESS_SCOPE_SUSPEND();
 			LoadProgressScope->EnterProgressFrame(1);
 		}
 		if( GIsEditor && GIsRunning )
@@ -2619,7 +2624,8 @@ FLinkerLoad::ELinkerStatus FLinkerLoad::FinalizeCreation(TMap<TPair<FName, FPack
 #if WITH_EDITOR
 		if (LoadProgressScope)
 		{
-		LoadProgressScope->EnterProgressFrame(1);
+			UE_SERIALIZE_ACCCESS_SCOPE_SUSPEND();
+			LoadProgressScope->EnterProgressFrame(1);
 		}
 #endif
 
@@ -2850,6 +2856,7 @@ void FLinkerLoad::Verify()
 #if WITH_EDITOR
 				if (SlowTask)
 				{
+					UE_SERIALIZE_ACCCESS_SCOPE_SUSPEND();
 					static const FText LoadingImportText = NSLOCTEXT("Core", "LinkerLoad_LoadingImportName", "Loading Import '{0}'");
 					SlowTask->EnterProgressFrame(1, FText::Format(LoadingImportText, FText::FromString(Import.ObjectName.ToString())));
 				}
@@ -3350,6 +3357,7 @@ bool FLinkerLoad::VerifyImportInner(const int32 ImportIndex, FString& WarningSuf
 #if WITH_EDITOR
 		if (SlowTask)
 		{
+			UE_SERIALIZE_ACCCESS_SCOPE_SUSPEND();
 			SlowTask->EnterProgressFrame(30);
 		}
 #endif
@@ -3429,6 +3437,7 @@ bool FLinkerLoad::VerifyImportInner(const int32 ImportIndex, FString& WarningSuf
 #if WITH_EDITOR
 		if (SlowTask)
 		{
+			UE_SERIALIZE_ACCCESS_SCOPE_SUSPEND();
 			SlowTask->EnterProgressFrame(30);
 		}
 #endif
@@ -3467,6 +3476,7 @@ bool FLinkerLoad::VerifyImportInner(const int32 ImportIndex, FString& WarningSuf
 #if WITH_EDITOR
 		if (SlowTask)
 		{
+			UE_SERIALIZE_ACCCESS_SCOPE_SUSPEND();
 			SlowTask->EnterProgressFrame(40);
 		}
 #endif
@@ -3547,6 +3557,7 @@ bool FLinkerLoad::VerifyImportInner(const int32 ImportIndex, FString& WarningSuf
 #if WITH_EDITOR
 		if (SlowTask)
 		{
+			UE_SERIALIZE_ACCCESS_SCOPE_SUSPEND();
 			SlowTask->EnterProgressFrame(50);
 		}
 #endif
@@ -3605,6 +3616,7 @@ bool FLinkerLoad::VerifyImportInner(const int32 ImportIndex, FString& WarningSuf
 #if WITH_EDITOR
 		if (SlowTask)
 		{
+			UE_SERIALIZE_ACCCESS_SCOPE_SUSPEND();
 			SlowTask->EnterProgressFrame(50);
 		}
 #endif
@@ -4038,6 +4050,7 @@ void FLinkerLoad::LoadAllObjects(bool bForcePreload)
 #if WITH_EDITOR
 		if (SlowTask)
 		{
+			UE_SERIALIZE_ACCCESS_SCOPE_SUSPEND();
 			SlowTask->EnterProgressFrame(1);
 		}
 #endif
@@ -4534,6 +4547,7 @@ void FLinkerLoad::Preload( UObject* Object )
 						else
 #endif
 						{
+							UE_SERIALIZE_ACCCESS_SCOPE(Object);
 							Object->Serialize(*this);
 						}
 
diff --git a/Engine/Source/Runtime/CoreUObject/Private/UObject/ObjectSerializeAccessScope.cpp b/Engine/Source/Runtime/CoreUObject/Private/UObject/ObjectSerializeAccessScope.cpp
new file mode 100644
--- /dev/null
+++ b/Engine/Source/Runtime/CoreUObject/Private/UObject/ObjectSerializeAccessScope.cpp
@@ -0,0 +1,173 @@
+// Copyright Epic Games, Inc. All Rights Reserved.
+
+#include "ObjectSerializeAccessScope.h"
+#include "UObject/Object.h"
+#include "Misc/CommandLine.h"
+#include "Misc/ConfigCacheIni.h"
+#include "UObject/Linker.h"
+#include "HAL/PlatformStackWalk.h"
+
+namespace UE::CoreUObject::Private
+{
+#if UE_WITH_OBJECT_HANDLE_TRACKING
+	bool IsEnableSerializeAccessWarnings()
+	{
+		auto IsEnabled = []()
+		{
+			if (FParse::Param(FCommandLine::Get(), TEXT("EnableSerializeAccessWarnings")))
+			{
+				return true;
+			}
+			else if (GConfig)
+			{
+				bool bConfig = false;
+				GConfig->GetBool(TEXT("Core.System.Experimental"), TEXT("EnableSerializeAccessWarnings"), bConfig, GEngineIni);
+				return bConfig;
+			}
+			return false;
+		};
+		static const bool bIsEnabled = IsEnabled();
+		return bIsEnabled;
+	}
+	static void ScopeCallback(const UObject& Object, const UObject* ReadObject)
+	{
+		
+	}
+
+	struct FObjectSerializeAccessScopeImpl
+	{
+		FObjectSerializeAccessScopeImpl(bool bEnabled)
+			: bEnabled(bEnabled)
+		{
+			if (bEnabled)
+			{
+				FObjectHandleReadFunc Func = [this](UObject* ReadObject)
+				{
+					FLocalData& Data = GetThreadLocalData();
+					if (!IsEnableSerializeAccessWarnings() || (Data.CurrentScope == nullptr || Data.CurrentScope->bSuspended))
+						return;
+
+					const UObject& CurrentObject = Data.CurrentScope->Object;
+
+					bool bIsAlreadyInSet;
+					Data.Types.FindOrAdd(CurrentObject.GetClass(), &bIsAlreadyInSet);
+					if (bIsAlreadyInSet)
+					{
+						return;
+					}
+
+					const SIZE_T StackTraceSize = 65535;
+					ANSICHAR StackTrace[StackTraceSize];
+					StackTrace[0] = 0;
+					FPlatformStackWalk::StackWalkAndDump(StackTrace, StackTraceSize, 3);
+					UE_LOG(LogLinker, Warning, TEXT("%s trying to access referenced object %s during serialize\n%s"), *CurrentObject.GetClass()->GetFullName(), (ReadObject ? *ReadObject->GetFullName() : TEXT("None")), ANSI_TO_TCHAR(StackTrace));
+				};
+
+				CallbackHandle = AddObjectHandleReadCallback(Func);
+			}
+		}
+
+		~FObjectSerializeAccessScopeImpl()
+		{
+			if (CallbackHandle.IsValid())
+			{
+				RemoveObjectHandleReadCallback(CallbackHandle);
+			}
+		}
+
+		static FObjectSerializeAccessScopeImpl& Get()
+		{
+			static FObjectSerializeAccessScopeImpl Impl(IsEnableSerializeAccessWarnings());
+			return Impl;
+		}
+
+		//thread local data
+		struct FLocalData
+		{
+			FObjectSerializeAccessScope* CurrentScope;
+
+			//types that have been seen
+			TSet<const UClass*> Types;
+		};
+
+		static FLocalData& GetThreadLocalData()
+		{
+			static thread_local FLocalData Data;
+			return Data;
+		}
+
+		void Push(FObjectSerializeAccessScope& Scope)
+		{
+			if (!bEnabled)
+				return;
+			FLocalData& LocalData = GetThreadLocalData();
+			FObjectSerializeAccessScope* Current = LocalData.CurrentScope;
+			Scope.Parent = Current;
+			LocalData.CurrentScope = &Scope;
+		}
+
+		void Pop(FObjectSerializeAccessScope& Scope)
+		{
+			if (!bEnabled)
+				return;
+			
+			FLocalData& LocalData = GetThreadLocalData();
+			FObjectSerializeAccessScope* Current = LocalData.CurrentScope;
+			check(&Scope == Current);
+			LocalData.CurrentScope = Scope.Parent;
+		}
+
+		void Suspend()
+		{
+			if (!bEnabled)
+				return;
+			
+			FLocalData LocalData = GetThreadLocalData();
+			FObjectSerializeAccessScope* Current = LocalData.CurrentScope;
+			if (Current)
+			{
+				Current->bSuspended = true;
+			}
+		}
+		void Resume()
+		{
+			if (!bEnabled)
+				return;
+			
+			FLocalData LocalData = GetThreadLocalData();
+			FObjectSerializeAccessScope* Current = LocalData.CurrentScope;
+			if (Current)
+			{
+				Current->bSuspended = false;
+			}
+		}
+
+		FObjectHandleTrackingCallbackId CallbackHandle;
+		bool bEnabled;
+
+	};
+
+	FObjectSerializeAccessScope::FObjectSerializeAccessScope(const UObject& Object)
+		: Object(Object)
+		, Parent(nullptr)
+		, bSuspended(false)
+	{
+		FObjectSerializeAccessScopeImpl::Get().Push(*this);
+	}
+
+	FObjectSerializeAccessScope::~FObjectSerializeAccessScope()
+	{
+		FObjectSerializeAccessScopeImpl::Get().Pop(*this);
+	}
+
+	FObjectSerializeAccessScope::FSuspendScope::FSuspendScope()
+	{
+		FObjectSerializeAccessScopeImpl::Get().Suspend();
+	}
+
+	FObjectSerializeAccessScope::FSuspendScope::~FSuspendScope()
+	{
+		FObjectSerializeAccessScopeImpl::Get().Resume();
+	}
+#endif
+}
\ No newline at end of file
diff --git a/Engine/Source/Runtime/CoreUObject/Private/UObject/ObjectSerializeAccessScope.h b/Engine/Source/Runtime/CoreUObject/Private/UObject/ObjectSerializeAccessScope.h
new file mode 100644
--- /dev/null
+++ b/Engine/Source/Runtime/CoreUObject/Private/UObject/ObjectSerializeAccessScope.h
@@ -0,0 +1,53 @@
+// Copyright Epic Games, Inc. All Rights Reserved.
+
+#pragma once
+#include "UObject/ObjectHandleTracking.h"
+
+#if UE_WITH_OBJECT_HANDLE_TRACKING
+
+namespace UE::CoreUObject::Private
+{
+	/// @brief scope to log messages when a TObjectPtr is accessed
+	struct FObjectSerializeAccessScope
+	{
+		/// @param object to being tracking for
+		FObjectSerializeAccessScope(const UObject& Object);
+		~FObjectSerializeAccessScope();
+
+		FObjectSerializeAccessScope(const FObjectSerializeAccessScope&) = delete;
+		FObjectSerializeAccessScope& operator=(const FObjectSerializeAccessScope&) = delete;
+
+		///@brief suspends the threads current FObjectSerializeAccessScope if any
+		struct FSuspendScope
+		{
+			FSuspendScope();
+			~FSuspendScope();
+
+			FSuspendScope(const FSuspendScope&) = delete;
+			FSuspendScope& operator=(const FSuspendScope&) = delete;
+		};
+
+	private:
+		friend struct FObjectSerializeAccessScopeImpl;
+
+		//object being serialized
+		const UObject& Object;
+
+		//parent scope for the serialization stack
+		FObjectSerializeAccessScope* Parent;
+		
+		//suspend to pause logging.
+		//progress bars for instance can "pause" serialize
+		//and do other work that kick object handle reads
+		bool bSuspended;
+	};
+}
+#define UE_SERIALIZE_ACCCESS_SCOPE(Object) UE::CoreUObject::Private::FObjectSerializeAccessScope PREPROCESSOR_JOIN(_Scope, __LINE__)(*(Object))
+#define UE_SERIALIZE_ACCCESS_SCOPE_SUSPEND() UE::CoreUObject::Private::FObjectSerializeAccessScope::FSuspendScope PREPROCESSOR_JOIN(_Suspend, __LINE__)
+
+#else
+
+#define UE_SERIALIZE_ACCCESS_SCOPE(Object)
+#define UE_SERIALIZE_ACCCESS_SCOPE_SUSPEND()
+
+#endif
\ No newline at end of file
diff --git a/Engine/Source/Runtime/CoreUObject/Public/UObject/ObjectHandleTracking.h b/Engine/Source/Runtime/CoreUObject/Public/UObject/ObjectHandleTracking.h
--- a/Engine/Source/Runtime/CoreUObject/Public/UObject/ObjectHandleTracking.h
+++ b/Engine/Source/Runtime/CoreUObject/Public/UObject/ObjectHandleTracking.h
@@ -57,6 +57,7 @@ namespace UE::CoreUObject
 	struct FObjectHandleTrackingCallbackId
 	{
 		int32 Id = -1;
+		bool IsValid() { return Id != -1; }
 	};
 
 	/**
