Deprecate FCoreUObjectDelegates::TraceExternalRootsForReachabilityAnalysis to simplify GC

#rnx
#locdown mark.lintott

diff --git a/Engine/Source/Runtime/CoreUObject/Private/UObject/GarbageCollection.cpp b/Engine/Source/Runtime/CoreUObject/Private/UObject/GarbageCollection.cpp
--- a/Engine/Source/Runtime/CoreUObject/Private/UObject/GarbageCollection.cpp
+++ b/Engine/Source/Runtime/CoreUObject/Private/UObject/GarbageCollection.cpp
@@ -3423,9 +3423,11 @@ public:
 			UE_LOG(LogGarbage, Verbose, TEXT("%f ms for Reachability Analysis"), (FPlatformTime::Seconds() - StartTime) * 1000);
 		}
 
+PRAGMA_DISABLE_DEPRECATION_WARNINGS
 		// Allowing external systems to add object roots. This can't be done through AddReferencedObjects
 		// because it may require tracing objects (via FGarbageCollectionTracer) multiple times
 		FCoreUObjectDelegates::TraceExternalRootsForReachabilityAnalysis.Broadcast(*this, KeepFlags, !(Options & EGCOptions::Parallel));
+PRAGMA_ENABLE_DEPRECATION_WARNINGS
 	}
 
 	virtual void PerformReachabilityAnalysisOnObjects(FWorkerContext* Context, EGCOptions Options) override
diff --git a/Engine/Source/Runtime/CoreUObject/Public/UObject/UObjectGlobals.h b/Engine/Source/Runtime/CoreUObject/Public/UObject/UObjectGlobals.h
--- a/Engine/Source/Runtime/CoreUObject/Public/UObject/UObjectGlobals.h
+++ b/Engine/Source/Runtime/CoreUObject/Public/UObject/UObjectGlobals.h
@@ -2882,6 +2882,7 @@ struct COREUOBJECT_API FCoreUObjectDelegates
 	DECLARE_MULTICAST_DELEGATE_ThreeParams(FTraceExternalRootsForReachabilityAnalysisDelegate, FGarbageCollectionTracer&, EObjectFlags, bool);
 
 	/** Called as last phase of reachability analysis. Allow external systems to add UObject roots *after* first reachability pass has been done */
+	UE_DEPRECATED(5.2, "Ability to inject extra root objects will be removed to reduce GC complexity")
 	static FTraceExternalRootsForReachabilityAnalysisDelegate TraceExternalRootsForReachabilityAnalysis;
 
 	/** Called after reachability analysis, before any purging */
