UObject: Fix StaticFindObject/ResolveName bug when searching for a sub object in a specified object, and the subobject has the same name as a short script package name.
Only look for short script package names and package redirects when we haven't yet found the correct package.

Tiny drive-by opt: Never try to find a UPackage inside another object, we know it will never succeed.

Example: /Game/AnimBP.AnimBP:AnimGraph.AnimGraphNode_StateMachine_0
Searching for AnimGraph.AnimGraphNode_StateMachine_0 with the /Game/AnimBP.AnimBP asset as outer would previously have converted AnimGraph to /Script/AnimGraph and the subobject would not have been found.

#rnx

diff --git a/Engine/Source/Runtime/CoreUObject/Private/UObject/UObjectGlobals.cpp b/Engine/Source/Runtime/CoreUObject/Private/UObject/UObjectGlobals.cpp
--- a/Engine/Source/Runtime/CoreUObject/Private/UObject/UObjectGlobals.cpp
+++ b/Engine/Source/Runtime/CoreUObject/Private/UObject/UObjectGlobals.cpp
@@ -1132,21 +1132,24 @@ bool ResolveName(UObject*& InPackage, FString& InOutName, bool Create, bool Thro
 		FString PartialName = InOutName.Left(DotIndex);
 
 		bool bIsScriptPackage = false;
-		if (!bSubobjectPath)
+		if (!InPackage)
 		{
-			// In case this is a short script package name, convert to long name before passing to CreatePackage/FindObject.
-			FName* ScriptPackageName = FPackageName::FindScriptPackageName(*PartialName);
-			if (ScriptPackageName)
+			if (!bSubobjectPath)
 			{
-				PartialName = ScriptPackageName->ToString();
+				// In case this is a short script package name, convert to long name before passing to CreatePackage/FindObject.
+				FName* ScriptPackageName = FPackageName::FindScriptPackageName(*PartialName);
+				if (ScriptPackageName)
+				{
+					PartialName = ScriptPackageName->ToString();
+				}
+				bIsScriptPackage = ScriptPackageName || FPackageName::IsScriptPackage(PartialName);
 			}
-			bIsScriptPackage = ScriptPackageName || FPackageName::IsScriptPackage(PartialName);
-		}
 
-		// Process any package redirects before calling CreatePackage/FindObject
-		{
-			const FCoreRedirectObjectName NewPackageName = FCoreRedirects::GetRedirectedName(ECoreRedirectFlags::Type_Package, FCoreRedirectObjectName(NAME_None, NAME_None, *PartialName));
-			PartialName = NewPackageName.PackageName.ToString();
+			// Process any package redirects before calling CreatePackage/FindObject
+			{
+				const FCoreRedirectObjectName NewPackageName = FCoreRedirects::GetRedirectedName(ECoreRedirectFlags::Type_Package, FCoreRedirectObjectName(NAME_None, NAME_None, *PartialName));
+				PartialName = NewPackageName.PackageName.ToString();
+			}
 		}
 
 		// Only long package names are allowed so don't even attempt to create one because whatever the name represents
@@ -1154,7 +1157,7 @@ bool ResolveName(UObject*& InPackage, FString& InOutName, bool Create, bool Thro
 		
 		if (!Create)
 		{
-			UObject* NewPackage = FindObject<UPackage>( InPackage, *PartialName );
+			UObject* NewPackage = InPackage ? nullptr : FindObject<UPackage>( InPackage, *PartialName );
 			if( !NewPackage )
 			{
 				if (InPackage)
@@ -1175,7 +1178,7 @@ bool ResolveName(UObject*& InPackage, FString& InOutName, bool Create, bool Thro
 		else if (!FPackageName::IsShortPackageName(PartialName))
 		{
 			// Try to find the package in memory first, should be faster than attempting to load or create
-			InPackage = StaticFindObjectFast(UPackage::StaticClass(), InPackage, *PartialName);
+			InPackage = InPackage ? nullptr : StaticFindObjectFast(UPackage::StaticClass(), InPackage, *PartialName);
 			if (!bIsScriptPackage && !InPackage)
 			{
 				InPackage = LoadPackage(Cast<UPackage>(InPackage), *PartialName, LoadFlags, nullptr, InstancingContext);
