Change GC clustering system to cluster only objects within the same package to avoid 'leaky' behavior.
Enable asset clustering by default.

diff --git a/Engine/Config/BaseEngine.ini b/Engine/Config/BaseEngine.ini
--- a/Engine/Config/BaseEngine.ini
+++ b/Engine/Config/BaseEngine.ini
@@ -1324,7 +1324,6 @@ gc.CreateGCClusters=True
 gc.MinGCClusterSize=5
 gc.AssetClustreringEnabled=False
 gc.ActorClusteringEnabled=False
-gc.BlueprintClusteringEnabled=False
 gc.UseDisregardForGCOnDedicatedServers=False
 gc.MultithreadedDestructionEnabled=True
 gc.VerifyGCObjectNames=True
diff --git a/Engine/Plugins/FX/Niagara/Source/Niagara/Classes/NiagaraSystem.h b/Engine/Plugins/FX/Niagara/Source/Niagara/Classes/NiagaraSystem.h
--- a/Engine/Plugins/FX/Niagara/Source/Niagara/Classes/NiagaraSystem.h
+++ b/Engine/Plugins/FX/Niagara/Source/Niagara/Classes/NiagaraSystem.h
@@ -245,6 +245,7 @@ public:
 	void PostInitProperties();
 	virtual void Serialize(FArchive& Ar) override;
 	virtual void PostLoad() override; 
+	virtual bool CanBeClusterRoot() const override { return true; }
 #if WITH_EDITORONLY_DATA
 	static void DeclareConstructClasses(TArray<FTopLevelAssetPath>& OutConstructClasses, const UClass* SpecificSubclass);
 #endif
diff --git a/Engine/Source/Runtime/CoreUObject/Private/UObject/GarbageCollectionVerification.cpp b/Engine/Source/Runtime/CoreUObject/Private/UObject/GarbageCollectionVerification.cpp
--- a/Engine/Source/Runtime/CoreUObject/Private/UObject/GarbageCollectionVerification.cpp
+++ b/Engine/Source/Runtime/CoreUObject/Private/UObject/GarbageCollectionVerification.cpp
@@ -215,35 +215,20 @@ public:
 			FUObjectItem* ObjectItem = GUObjectArray.ObjectToObjectItem(Object);
 			if (ObjectItem->GetOwnerIndex() <= 0)
 			{
-				// We are allowed to reference other clusters, root set objects and objects from diregard for GC pool
-				if (!ObjectItem->HasAnyFlags(EInternalObjectFlags::ClusterRoot | EInternalObjectFlags::RootSet)
-					&& !GUObjectArray.IsDisregardForGC(Object) && Object->CanBeInCluster() &&
-					!Cluster->MutableObjects.Contains(GUObjectArray.ObjectToIndex(Object))) // This is for objects that had RF_NeedLoad|RF_NeedPostLoad set when creating the cluster
+				// Referenced object is a cluster root or not clustered
+				if (ObjectItem->HasAnyFlags(EInternalObjectFlags::ClusterRoot))
 				{
-					UE_LOG(LogGarbage, Warning, TEXT("Object %s (0x%016llx) from cluster %s (0x%016llx / 0x%016llx) is referencing 0x%016llx %s which is not part of root set or cluster."),
-						*CurrentObject->GetFullName(),
-						(int64)(PTRINT)CurrentObject,
-						*ClusterRootObject->GetFullName(),
-						(int64)(PTRINT)ClusterRootObject,
-						(int64)(PTRINT)Cluster,
-						(int64)(PTRINT)Object,
-						*Object->GetFullName());
-					NumErrors.Increment();
-#if UE_BUILD_DEBUG
-					FReferenceChainSearch RefChainSearch(Object, EReferenceChainSearchMode::Shortest | EReferenceChainSearchMode::PrintResults);
-#endif
-				}
-				else if (ObjectItem->HasAnyFlags(EInternalObjectFlags::ClusterRoot))
-				{
-					// However, clusters need to be referenced by the current cluster otherwise they can also get GC'd too early.
+					// Clusters need to be referenced by the current cluster otherwise they can also get GC'd too early.
 					const FUObjectItem* ClusterRootObjectItem = GUObjectArray.ObjectToObjectItem(ClusterRootObject);
 					const int32 OtherClusterRootIndex = GUObjectArray.ObjectToIndex(Object);
 					const FUObjectItem* OtherClusterRootItem = GUObjectArray.IndexToObjectUnsafeForGC(OtherClusterRootIndex);
 					check(OtherClusterRootItem && OtherClusterRootItem->Object);
 					UObject* OtherClusterRootObject = static_cast<UObject*>(OtherClusterRootItem->Object);
-					UE_CLOG(OtherClusterRootIndex != Cluster->RootIndex &&
-						!Cluster->ReferencedClusters.Contains(OtherClusterRootIndex) &&
-						!Cluster->MutableObjects.Contains(OtherClusterRootIndex), LogGarbage, Warning,
+					UE_CLOG(
+						OtherClusterRootIndex != Cluster->RootIndex  // Same cluster is legal 
+					&&	!Cluster->ReferencedClusters.Contains(OtherClusterRootIndex)  // cluster-cluster reference is legal
+					&&  !Cluster->MutableObjects.Contains(OtherClusterRootIndex),  // reference to an external object which later became a cluster root is legal 
+						LogGarbage, Warning,
 						TEXT("Object %s from source cluster %s (%d) is referencing cluster root object %s (0x%016llx) (%d) which is not referenced by the source cluster."),
 						*GetFullNameSafe(ReferencingObject),
 						*ClusterRootObject->GetFullName(),
@@ -252,6 +237,25 @@ public:
 						(int64)(PTRINT)Object,
 						OtherClusterRootItem->GetClusterIndex());
 				}
+				else if (	!ObjectItem->HasAnyFlags(EInternalObjectFlags::RootSet) // Root set objects will stay alive that way 
+						&&	!GUObjectArray.IsDisregardForGC(Object) // Disregard-for-GC objects are never freed
+						&&  !Cluster->MutableObjects.Contains(GUObjectArray.ObjectToIndex(Object)) // Mutable object ref is traversed during GC regardless of if the object is cluster root or not 
+				) 
+				{
+					// There is a danger this object could be freed leaving a dangling pointer in an object inside the cluster
+					UE_LOG(LogGarbage, Warning, TEXT("Object %s (0x%016llx) from cluster %s (0x%016llx / 0x%016llx) is referencing 0x%016llx %s which is not part of root set or cluster."),
+						*CurrentObject->GetFullName(),
+						(int64)(PTRINT)CurrentObject,
+						*ClusterRootObject->GetFullName(),
+						(int64)(PTRINT)ClusterRootObject,
+						(int64)(PTRINT)Cluster,
+						(int64)(PTRINT)Object,
+						*Object->GetFullName());
+					NumErrors.Increment();
+#if UE_BUILD_DEBUG
+					FReferenceChainSearch RefChainSearch(Object, EReferenceChainSearchMode::Shortest | EReferenceChainSearchMode::PrintResults);
+#endif
+				}
 			}
 			else if (ObjectItem->GetOwnerIndex() != Cluster->RootIndex)
 			{
@@ -262,9 +266,11 @@ public:
 				const FUObjectItem* OtherClusterRootItem = GUObjectArray.IndexToObjectUnsafeForGC(OtherClusterRootIndex);
 				check(OtherClusterRootItem && OtherClusterRootItem->Object);
 				UObject* OtherClusterRootObject = static_cast<UObject*>(OtherClusterRootItem->Object);
-				UE_CLOG(OtherClusterRootIndex != Cluster->RootIndex &&
-					!Cluster->ReferencedClusters.Contains(OtherClusterRootIndex) &&
-					!Cluster->MutableObjects.Contains(GUObjectArray.ObjectToIndex(Object)), LogGarbage, Warning,
+				UE_CLOG(
+						OtherClusterRootIndex != Cluster->RootIndex  // Same cluster is legal 
+					&&	!Cluster->ReferencedClusters.Contains(OtherClusterRootIndex)  // Cluster-cluster reference 
+					&&	!Cluster->MutableObjects.Contains(GUObjectArray.ObjectToIndex(Object)), // Reference to an object which was later clustered
+					LogGarbage, Warning,
 					TEXT("Object %s from source cluster %s (%d) is referencing object %s (0x%016llx) from cluster %s (%d) which is not referenced by the source cluster."),
 					*GetFullNameSafe(ReferencingObject),
 					*ClusterRootObject->GetFullName(),
diff --git a/Engine/Source/Runtime/CoreUObject/Private/UObject/UObjectClusters.cpp b/Engine/Source/Runtime/CoreUObject/Private/UObject/UObjectClusters.cpp
--- a/Engine/Source/Runtime/CoreUObject/Private/UObject/UObjectClusters.cpp
+++ b/Engine/Source/Runtime/CoreUObject/Private/UObject/UObjectClusters.cpp
@@ -37,7 +37,8 @@ static FAutoConsoleVariableRef CVarAssetClustreringEnabled(
 	ECVF_Default
 );
 
-int32 GMinGCClusterSize = 5;
+// By default cluster anything that asks to be clustered with its containing package.
+int32 GMinGCClusterSize = 2;
 static FAutoConsoleVariableRef CMinGCClusterSize(
 	TEXT("gc.MinGCClusterSize"),
 	GMinGCClusterSize,
@@ -303,7 +304,8 @@ void DumpClusterToLog(const FUObjectCluster& Cluster, bool bHierarchy, bool bInd
 
 	FUObjectItem* RootItem = GUObjectArray.IndexToObjectUnsafeForGC(Cluster.RootIndex);
 	UObject* RootObject = static_cast<UObject*>(RootItem->Object);
-	UE_LOG(LogObj, Display, TEXT("%s (Index: %d), Size: %d, ReferencedClusters: %d"), *RootObject->GetFullName(), Cluster.RootIndex, Cluster.Objects.Num(), Cluster.ReferencedClusters.Num());
+	UE_LOG(LogObj, Display, TEXT("%s (Index: %d), Size: %d, ReferencedClusters: %d, MutableObjects: %d, ReferencedByClusters: %d"), 
+		*RootObject->GetFullName(), Cluster.RootIndex, Cluster.Objects.Num(), Cluster.ReferencedClusters.Num(), Cluster.MutableObjects.Num(), Cluster.ReferencedByClusters.Num());
 	if (bHierarchy)
 	{
 		int32 Index = 0;
@@ -450,6 +452,7 @@ void ListClusters(const TArray<FString>& Args)
 	UE_LOG(LogObj, Display, TEXT("Maximum cluster size: %d"), MaxClusterSize);
 	UE_LOG(LogObj, Display, TEXT("Average cluster size: %d"), AllClusters.Num() ? (TotalClusterObjects / AllClusters.Num()) : 0);
 	UE_LOG(LogObj, Display, TEXT("Number of objects in GC clusters: %d"), TotalClusterObjects);
+	UE_LOG(LogObj, Display, TEXT("Max number of gc objects: %d"), GUObjectArray.GetObjectArrayNumPermanent());
 	UE_LOG(LogObj, Display, TEXT("Maximum number of custer-to-cluster references: %d"), MaxInterClusterReferences);
 	UE_LOG(LogObj, Display, TEXT("Average number of custer-to-cluster references: %d"), AllClusters.Num() ? (TotalInterClusterReferences / AllClusters.Num()) : 0);
 }
@@ -542,6 +545,8 @@ void DumpRefsToCluster(const TArray<FString>& Args)
 	}
 }
 
+static void SuggestClusters(const TArray<FString>& Args);
+
 
 static FAutoConsoleCommand ListClustersCommand(
 	TEXT("gc.ListClusters"),
@@ -561,6 +566,12 @@ static FAutoConsoleCommand DumpRefsToClusterCommand(
 	FConsoleCommandWithArgsDelegate::CreateStatic(DumpRefsToCluster)
 );
 
+static FAutoConsoleCommand SuggestClustersCommand(
+	TEXT("gc.SuggestClusters"),
+	TEXT("Searches for assets which contain many internal objects which are not clustered."),
+	FConsoleCommandWithArgsDelegate::CreateStatic(SuggestClusters)
+);
+
 #endif // !UE_BUILD_SHIPPING
 
 /**
@@ -570,12 +581,14 @@ class FClusterReferenceProcessor : public FSimpleReferenceProcessorBase
 {
 	int32 ClusterRootIndex;
 	FUObjectCluster& Cluster;
+	UObject* Outermost;
 
 public:
 
-	FClusterReferenceProcessor(int32 InClusterRootIndex, FUObjectCluster& InCluster)
+	FClusterReferenceProcessor(int32 InClusterRootIndex, FUObjectCluster& InCluster, UObject* InOutermost)
 		: ClusterRootIndex(InClusterRootIndex)
 		, Cluster(InCluster)
+		, Outermost(InOutermost)
 	{}
 
 	static FString LoadFlagsToString(UObject* Obj)
@@ -662,61 +675,142 @@ public:
 	*/
 	FORCEINLINE void HandleTokenStreamObjectReference(FGCArrayStruct& ObjectsToSerializeStruct, UObject* ReferencingObject, UObject*& Object, UE::GC::FTokenId TokenIndex, const EGCTokenType TokenType, bool bAllowReferenceElimination)
 	{
-		if (Object)
+		if (!Object)
+		{
+			return;
+		}
+
+		FUObjectItem* ObjectItem = GUObjectArray.ObjectToObjectItem(Object);
+		// Add encountered object reference to list of to be serialized objects if it hasn't already been added.
+		if (ObjectItem->GetOwnerIndex() == ClusterRootIndex)
 		{
-			// If we haven't finished loading, we can't be sure we know all the references so the object will be added as mutable reference
-			UE_CLOG(Object->HasAnyFlags(RF_NeedLoad), LogObj, Log, TEXT("%s hasn't been loaded (%s) but is being added to cluster %s"),
-				*Object->GetFullName(),
-				*LoadFlagsToString(Object),
-				*GetClusterRoot()->GetFullName());
+			return;
+		}
 
-			FUObjectItem* ObjectItem = GUObjectArray.ObjectToObjectItem(Object);
+		// If we haven't finished loading, we can't be sure we know all the references so the object will be added as mutable reference
+		UE_CLOG(Object->HasAnyFlags(RF_NeedLoad | RF_NeedPostLoad), LogObj, Log, TEXT("%s hasn't been loaded (%s) but is being added to cluster %s"),
+			*Object->GetFullName(),
+			*LoadFlagsToString(Object),
+			*GetClusterRoot()->GetFullName());
 
-			// Add encountered object reference to list of to be serialized objects if it hasn't already been added.
-			if (ObjectItem->GetOwnerIndex() != ClusterRootIndex)
+		if (ObjectItem->HasAnyFlags(EInternalObjectFlags::ClusterRoot) || ObjectItem->GetOwnerIndex() != 0)
+		{					
+			// Simply reference this cluster and all clusters it's referencing
+			const int32 OtherClusterRootIndex = ObjectItem->HasAnyFlags(EInternalObjectFlags::ClusterRoot) ? GUObjectArray.ObjectToIndex(Object) : ObjectItem->GetOwnerIndex();
+			FUObjectItem* OtherClusterRootItem = GUObjectArray.IndexToObject(OtherClusterRootIndex);
+			const int32 OtherClusterIndex = OtherClusterRootItem->GetClusterIndex();
+			FUObjectCluster& OtherCluster = GUObjectClusters[OtherClusterIndex];
+			Cluster.ReferencedClusters.AddUnique(OtherClusterRootIndex);
+			OtherCluster.ReferencedByClusters.AddUnique(ClusterRootIndex);
+
+			for (int32 OtherClusterReferencedCluster : OtherCluster.ReferencedClusters)
 			{
-				if (ObjectItem->HasAnyFlags(EInternalObjectFlags::ClusterRoot) || ObjectItem->GetOwnerIndex() != 0)
-				{					
-					// Simply reference this cluster and all clusters it's referencing
-					const int32 OtherClusterRootIndex = ObjectItem->HasAnyFlags(EInternalObjectFlags::ClusterRoot) ? GUObjectArray.ObjectToIndex(Object) : ObjectItem->GetOwnerIndex();
-					FUObjectItem* OtherClusterRootItem = GUObjectArray.IndexToObject(OtherClusterRootIndex);
-					const int32 OtherClusterIndex = OtherClusterRootItem->GetClusterIndex();
-					FUObjectCluster& OtherCluster = GUObjectClusters[OtherClusterIndex];
-					Cluster.ReferencedClusters.AddUnique(OtherClusterRootIndex);
-					OtherCluster.ReferencedByClusters.AddUnique(ClusterRootIndex);
-
-					for (int32 OtherClusterReferencedCluster : OtherCluster.ReferencedClusters)
-					{
-						if (OtherClusterReferencedCluster != ClusterRootIndex)
-						{
-							Cluster.ReferencedClusters.AddUnique(OtherClusterReferencedCluster);
-						}
-					}
-					for (int32 OtherClusterReferencedMutableObjectIndex : OtherCluster.MutableObjects)
-					{
-						Cluster.MutableObjects.AddUnique(OtherClusterReferencedMutableObjectIndex);
-					}
-				}
-				else if (!GUObjectArray.IsDisregardForGC(Object)) // We know that disregard for GC objects will never be GC'd so no reference is necessary
+				if (OtherClusterReferencedCluster != ClusterRootIndex)
 				{
-					check(ObjectItem->GetOwnerIndex() == 0);
-
-					// New object, add it to the cluster.
-					if (Object->CanBeInCluster() && !Object->HasAnyFlags(RF_NeedLoad | RF_NeedPostLoad) && !Object->IsRooted())
-					{
-						AddObjectToCluster(GUObjectArray.ObjectToIndex(Object), ObjectItem, Object, ObjectsToSerializeStruct, true);
-					}
-					else
-					{
-						// If the object can't be in a cluster or is being loaded, adding to the mutable objects list (and we won't be processing it further)
-						Cluster.MutableObjects.AddUnique(GUObjectArray.ObjectToIndex(Object));
-					}
+					Cluster.ReferencedClusters.AddUnique(OtherClusterReferencedCluster);
 				}
 			}
+			for (int32 OtherClusterReferencedMutableObjectIndex : OtherCluster.MutableObjects)
+			{
+				Cluster.MutableObjects.AddUnique(OtherClusterReferencedMutableObjectIndex);
+			}
+		}
+		else if (GUObjectArray.IsDisregardForGC(Object)) 
+		{
+			// We know that disregard for GC objects will never be GC'd so no reference is necessary
+		}
+		else if (!Object->CanBeInCluster() || Object->HasAnyFlags(RF_NeedLoad | RF_NeedPostLoad) || Object->IsRooted())
+		{
+			// This object may change state in ways which would invalidate the cluster, so keep it in a list of objects to fully explore when this cluster is encountered during GC		
+			Cluster.MutableObjects.AddUnique(GUObjectArray.ObjectToIndex(Object));
+		}
+		else if (!Object->IsIn(Outermost) && Object != Outermost)
+		{
+			// This object is in a different package so we don't want it to be part of this cluster 
+			// TODO: Consider exploring this object to find if it is part of a package containing an object which may add it to a cluster, so we can add a cluster ref instead of a mutable object ref 
+			Cluster.MutableObjects.AddUnique(GUObjectArray.ObjectToIndex(Object));
+		}
+		else 
+		{
+			check(ObjectItem->GetOwnerIndex() == 0);
+
+			// New object in the same package, add it to the cluster.
+			AddObjectToCluster(GUObjectArray.ObjectToIndex(Object), ObjectItem, Object, ObjectsToSerializeStruct, true);
 		}
 	}
 };
 
+static void SuggestClusters(const TArray<FString>& Args)
+{
+	struct FStats 
+	{
+		int32 NumAssets = 0;
+		int32 NumObjects = 0;
+		int32 NumClusterRefs = 0;
+		int32 NumObjectRefs = 0;
+	};
+	TMap<UClass*, FStats> SuggestedClasses;
+	for (TObjectIterator<UPackage> It; It; ++It)
+	{
+		bool FoundCluster = false;
+
+		// If this package contains a cluster, the package should have been added to it
+		if (GUObjectClusters.GetObjectCluster(*It) != nullptr)
+		{
+			continue;
+		}
+
+		ForEachObjectWithOuter(*It, [&](UObject* Obj)
+		{
+			if (Obj->HasAllFlags(RF_Standalone))
+			{
+				FUObjectItem* RootItem = GUObjectArray.ObjectToObjectItem(Obj);
+				const int32 InternalIndex = GUObjectArray.ObjectToIndex(Obj);
+				const int32 ClusterIndex = GUObjectClusters.AllocateCluster(InternalIndex);
+				FUObjectCluster& Cluster = GUObjectClusters[ClusterIndex];
+				Cluster.Objects.Reserve(64);
+
+				FClusterReferenceProcessor Processor(InternalIndex, Cluster, *It);
+				FGCArrayStruct ArrayStruct;
+				TArray<UObject*> ObjectsToProcess = {Obj};
+				ArrayStruct.SetInitialObjectsUnpadded(ObjectsToProcess);
+				CollectReferences(Processor, ArrayStruct);
+				RootItem->SetClusterIndex(ClusterIndex);
+				RootItem->SetFlags(EInternalObjectFlags::ClusterRoot);
+
+				if (Cluster.Objects.Num() >= GUObjectClusters.GetMinClusterSize())
+				{
+					UE_LOG(LogObj, Display, TEXT("Potential cluster root %s with no cluster, %d internal objects, %d cluster refs, %d object refs."), 
+						*Obj->GetPathName(),
+						Cluster.Objects.Num(),
+						Cluster.ReferencedClusters.Num(),
+						Cluster.MutableObjects.Num()
+					);
+					FStats& Stats = SuggestedClasses.FindOrAdd(Obj->GetClass());
+					Stats.NumAssets++;
+					Stats.NumObjects += Cluster.Objects.Num();
+					Stats.NumClusterRefs += Cluster.ReferencedClusters.Num();
+					Stats.NumObjectRefs += Cluster.MutableObjects.Num();
+				}
+				for (int32 ClusterObjectIndex : Cluster.Objects)
+				{
+					FUObjectItem* ClusterObjectItem = GUObjectArray.IndexToObjectUnsafeForGC(ClusterObjectIndex);
+					ClusterObjectItem->SetOwnerIndex(0);
+				}
+				GUObjectClusters.FreeCluster(ClusterIndex);
+			}
+		});
+	}
+	SuggestedClasses.ValueSort([](const FStats& A, const FStats& B) { return A.NumObjects > B.NumObjects; });
+	for (const TPair<UClass*, FStats>& Pair : SuggestedClasses)
+	{
+		UClass* Class = Pair.Key;
+		FStats Stats = Pair.Value;
+		UE_LOG(LogObj, Display, TEXT("Suggested cluster class: %s - %d assets - %d objects - %d cluster refs - %d object refs"), 
+			*Class->GetPathName(), Stats.NumAssets, Stats.NumObjects, Stats.NumClusterRefs, Stats.NumObjectRefs);
+	}
+}
+
 bool CanCreateObjectClusters()
 {
 	return FPlatformProperties::RequiresCookedData() && !GIsInitialLoad && GCreateGCClusters && GAssetClustreringEnabled && !GUObjectArray.IsOpenForDisregardForGC();
@@ -747,16 +841,16 @@ void UObjectBaseUtility::AddToCluster(UObjectBaseUtility* ClusterRootOrObjectFro
 		const int32 ClusterRootIndex = Cluster->RootIndex;
 		if (!bAddAsMutableObject)
 		{
-			FClusterReferenceProcessor Processor(ClusterRootIndex, *Cluster);
+			UObject* ClusterRootObject = static_cast<UObject*>(GUObjectArray.IndexToObjectUnsafeForGC(Cluster->RootIndex)->Object);
+			FClusterReferenceProcessor Processor(ClusterRootIndex, *Cluster, ClusterRootObject->GetOutermost());
 			UObject* ThisObject = static_cast<UObject*>(this);
+			UObject** ThisObjPtr = &ThisObject;
 			FGCArrayStruct ArrayStruct;
-			ArrayStruct.InitialNativeReferences = {&ThisObject};
+			ArrayStruct.InitialNativeReferences = TConstArrayView<UObject**>( &ThisObjPtr, 1 );
 			CollectReferences(Processor, ArrayStruct);
 
 #if UE_GCCLUSTER_VERBOSE_LOGGING
-			UObject* ClusterRootObject = static_cast<UObject*>(GUObjectArray.IndexToObjectUnsafeForGC(Cluster->RootIndex)->Object);
 			UE_LOG(LogObj, Log, TEXT("Added %s to cluster %s:"), *ThisObject->GetFullName(), *ClusterRootObject->GetFullName());
-
 			DumpClusterToLog(*Cluster, true, false);
 #endif
 		}
@@ -810,7 +904,7 @@ void UObjectBaseUtility::CreateCluster()
 	Cluster.Objects.Reserve(64);
 
 	// Collect all objects referenced by cluster root and by all objects it's referencing
-	FClusterReferenceProcessor Processor(InternalIndex, Cluster);
+	FClusterReferenceProcessor Processor(InternalIndex, Cluster, GetOutermost());
 	FGCArrayStruct ArrayStruct;
 	TArray<UObject*> ObjectsToProcess = {static_cast<UObject*>(this)};
 	ArrayStruct.SetInitialObjectsUnpadded(ObjectsToProcess);
@@ -828,7 +922,7 @@ void UObjectBaseUtility::CreateCluster()
 		Cluster.MutableObjects.Sort();		
 
 #if UE_GCCLUSTER_VERBOSE_LOGGING
-		UE_LOG(LogObj, Log, TEXT("Created Cluster (%d) with %d objects, %d referenced clusters and %d mutable objects."),
+		UE_LOG(LogObj, Log, TEXT("Created cluster (%d) with %d objects, %d referenced clusters and %d mutable objects."),
 			ClusterIndex, Cluster.Objects.Num(), Cluster.ReferencedClusters.Num(), Cluster.MutableObjects.Num());
 
 		DumpClusterToLog(Cluster, true, false);
@@ -836,6 +930,12 @@ void UObjectBaseUtility::CreateCluster()
 	}
 	else
 	{
+#if UE_GCCLUSTER_VERBOSE_LOGGING
+		UE_LOG(LogObj, Log, TEXT("Not creating cluster (%d) with %d objects, %d referenced clusters and %d mutable objects."),
+			ClusterIndex, Cluster.Objects.Num(), Cluster.ReferencedClusters.Num(), Cluster.MutableObjects.Num());
+
+		DumpClusterToLog(Cluster, true, false);
+#endif
 		for (int32 ClusterObjectIndex : Cluster.Objects)
 		{
 			FUObjectItem* ClusterObjectItem = GUObjectArray.IndexToObjectUnsafeForGC(ClusterObjectIndex);
diff --git a/Engine/Source/Runtime/Engine/Classes/Engine/CoreSettings.h b/Engine/Source/Runtime/Engine/Classes/Engine/CoreSettings.h
--- a/Engine/Source/Runtime/Engine/Classes/Engine/CoreSettings.h
+++ b/Engine/Source/Runtime/Engine/Classes/Engine/CoreSettings.h
@@ -210,11 +210,6 @@ protected:
 		ToolTip = "Whether to allow levels to create actor clusters for GC."))
 	uint32 ActorClusteringEnabled : 1;
 
-	UPROPERTY(EditAnywhere, config, Category = Optimization, meta = (
-		ConsoleVariable = "gc.BlueprintClusteringEnabled", DisplayName = "Blueprint Clustering Enabled",
-		ToolTip = "Whether to allow Blueprint classes to create GC clusters."))
-	uint32 BlueprintClusteringEnabled : 1;
-
 	UPROPERTY(EditAnywhere, config, Category = Optimization, meta = (
 		ConsoleVariable = "gc.UseDisregardForGCOnDedicatedServers", DisplayName = "Use DisregardForGC On Dedicated Servers",
 		ToolTip = "If false, DisregardForGC will be disabled for dedicated servers."))
diff --git a/Engine/Source/Runtime/Engine/Private/BlueprintGeneratedClass.cpp b/Engine/Source/Runtime/Engine/Private/BlueprintGeneratedClass.cpp
--- a/Engine/Source/Runtime/Engine/Private/BlueprintGeneratedClass.cpp
+++ b/Engine/Source/Runtime/Engine/Private/BlueprintGeneratedClass.cpp
@@ -39,14 +39,6 @@ extern UNREALED_API class UEditorEngine* GEditor;
 DEFINE_STAT(STAT_PersistentUberGraphFrameMemory);
 DEFINE_STAT(STAT_BPCompInstancingFastPathMemory);
 
-int32 GBlueprintClusteringEnabled = 0;
-static FAutoConsoleVariableRef CVarBlueprintClusteringEnabled(
-	TEXT("gc.BlueprintClusteringEnabled"),
-	GBlueprintClusteringEnabled,
-	TEXT("Whether to allow Blueprint classes to create GC clusters."),
-	ECVF_Default
-);
-
 int32 GBlueprintComponentInstancingFastPathDisabled = 0;
 static FAutoConsoleVariableRef CVarBlueprintComponentInstancingFastPathDisabled(
 	TEXT("bp.ComponentInstancingFastPathDisabled"),
@@ -1709,8 +1701,8 @@ bool UBlueprintGeneratedClass::NeedsLoadForEditorGame() const
 
 bool UBlueprintGeneratedClass::CanBeClusterRoot() const
 {
-	// Clustering level BPs doesn't work yet
-	return GBlueprintClusteringEnabled && !GetOutermost()->ContainsMap();
+	// We don't want to cluster level BPs with the rest of the contents of the map 
+	return !GetOutermost()->ContainsMap();
 }
 
 #if WITH_EDITOR
diff --git a/Engine/Source/Runtime/Engine/Private/CoreSettings.cpp b/Engine/Source/Runtime/Engine/Private/CoreSettings.cpp
--- a/Engine/Source/Runtime/Engine/Private/CoreSettings.cpp
+++ b/Engine/Source/Runtime/Engine/Private/CoreSettings.cpp
@@ -200,7 +200,6 @@ UGarbageCollectionSettings::UGarbageCollectionSettings()
 	MinGCClusterSize = 5;
 	AssetClusteringEnabled = true;
 	ActorClusteringEnabled = true;
-	BlueprintClusteringEnabled = false;
 	UseDisregardForGCOnDedicatedServers = false;
 	VerifyGCObjectNames = true;
 	VerifyUObjectsAreNotFGCObjects = true;
diff --git a/Engine/Source/Runtime/Engine/Private/SoundBase.cpp b/Engine/Source/Runtime/Engine/Private/SoundBase.cpp
--- a/Engine/Source/Runtime/Engine/Private/SoundBase.cpp
+++ b/Engine/Source/Runtime/Engine/Private/SoundBase.cpp
@@ -220,7 +220,7 @@ bool USoundBase::CanBeClusterRoot() const
 
 bool USoundBase::CanBeInCluster() const
 {
-	return false;
+	return true;
 }
 
 void USoundBase::Serialize(FArchive& Ar)
diff --git a/Engine/Source/Runtime/Engine/Private/SoundCue.cpp b/Engine/Source/Runtime/Engine/Private/SoundCue.cpp
--- a/Engine/Source/Runtime/Engine/Private/SoundCue.cpp
+++ b/Engine/Source/Runtime/Engine/Private/SoundCue.cpp
@@ -252,12 +252,12 @@ void USoundCue::PostLoad()
 
 bool USoundCue::CanBeClusterRoot() const
 {
-	return false;
+	return true;
 }
 
 bool USoundCue::CanBeInCluster() const
 {
-	return false;
+	return true;
 }
 
 void USoundCue::OnPostEngineInit()
diff --git a/Engine/Source/Runtime/Engine/Private/SoundNode.cpp b/Engine/Source/Runtime/Engine/Private/SoundNode.cpp
--- a/Engine/Source/Runtime/Engine/Private/SoundNode.cpp
+++ b/Engine/Source/Runtime/Engine/Private/SoundNode.cpp
@@ -58,7 +58,7 @@ bool USoundNode::CanBeClusterRoot() const
 
 bool USoundNode::CanBeInCluster() const
 {
-	return false;
+	return true;
 }
 
 #if WITH_EDITOR
