UE-173628 - Remove HasNativeBreak from FVector3f

There were two issues here:

1. BreakVector3f needs to be deprecated, and the FVector3f BP type should use BreakVector instead.
2. The float variants for the various Core types were missing HasNativeBreak/Make metadata. Without this, it would be impossible to break/split these types if they were exposed to BP.

diff --git a/Engine/Source/Runtime/CoreUObject/Public/UObject/NoExportTypes.h b/Engine/Source/Runtime/CoreUObject/Public/UObject/NoExportTypes.h
--- a/Engine/Source/Runtime/CoreUObject/Public/UObject/NoExportTypes.h
+++ b/Engine/Source/Runtime/CoreUObject/Public/UObject/NoExportTypes.h
@@ -481,7 +481,7 @@ struct FGuid
  * A point or direction FVector in 3d space.
  * @note The full C++ class is located here: Engine\Source\Runtime\Core\Public\Math\Vector.h
  */
-USTRUCT(immutable, noexport, BlueprintType, BlueprintInternalUseOnly, IsAlwaysAccessible, HasDefaults, HasNoOpConstructor, IsCoreType, meta = (HasNativeBreak = "/Script/Engine.KismetMathLibrary.BreakVector3f"))
+USTRUCT(immutable, noexport, BlueprintType, BlueprintInternalUseOnly, IsAlwaysAccessible, HasDefaults, HasNoOpConstructor, IsCoreType, meta = (ScriptDefaultMake, ScriptDefaultBreak, HasNativeMake = "/Script/Engine.KismetMathLibrary.MakeVector", HasNativeBreak = "/Script/Engine.KismetMathLibrary.BreakVector"))
 struct FVector3f
 {
 	UPROPERTY(EditAnywhere, Category = Vector, SaveGame)
@@ -533,7 +533,7 @@ struct FVector
 * A 4-D homogeneous vector.
 * @note The full C++ class is located here: Engine\Source\Runtime\Core\Public\Math\Vector4.h
 */
-USTRUCT(immutable, noexport, BlueprintType, BlueprintInternalUseOnly, IsAlwaysAccessible, HasDefaults, HasNoOpConstructor, IsCoreType)
+USTRUCT(immutable, noexport, BlueprintType, BlueprintInternalUseOnly, IsAlwaysAccessible, HasDefaults, HasNoOpConstructor, IsCoreType, meta = (ScriptDefaultMake, ScriptDefaultBreak, HasNativeMake = "/Script/Engine.KismetMathLibrary.MakeVector4", HasNativeBreak = "/Script/Engine.KismetMathLibrary.BreakVector4"))
 struct FVector4f
 {
 	UPROPERTY(EditAnywhere, Category = Vector4, SaveGame)
@@ -574,7 +574,7 @@ struct FVector4d
 * A 4-D homogeneous vector.
 * @note The full C++ class is located here: Engine\Source\Runtime\Core\Public\Math\Vector4.h
 */
-USTRUCT(immutable, noexport, BlueprintType, IsAlwaysAccessible, HasDefaults, HasNoOpConstructor, IsCoreType, meta = (HasNativeMake = "/Script/Engine.KismetMathLibrary.MakeVector4", HasNativeBreak = "/Script/Engine.KismetMathLibrary.BreakVector4"))
+USTRUCT(immutable, noexport, BlueprintType, IsAlwaysAccessible, HasDefaults, HasNoOpConstructor, IsCoreType, meta = (ScriptDefaultMake, ScriptDefaultBreak, HasNativeMake = "/Script/Engine.KismetMathLibrary.MakeVector4", HasNativeBreak = "/Script/Engine.KismetMathLibrary.BreakVector4"))
 struct FVector4
 {
 	UPROPERTY(EditAnywhere, BlueprintReadWrite, Category=Vector4, SaveGame)
@@ -595,13 +595,13 @@ struct FVector4
 * A vector in 2-D space composed of components (X, Y) with floating point precision.
 * @note The full C++ class is located here: Engine\Source\Runtime\Core\Public\Math\Vector2D.h
 */
-USTRUCT(immutable, noexport, BlueprintType, BlueprintInternalUseOnly, IsAlwaysAccessible, HasDefaults, HasNoOpConstructor, IsCoreType)
+USTRUCT(immutable, noexport, BlueprintType, BlueprintInternalUseOnly, IsAlwaysAccessible, HasDefaults, HasNoOpConstructor, IsCoreType, meta = (ScriptDefaultMake, ScriptDefaultBreak, HasNativeMake="/Script/Engine.KismetMathLibrary.MakeVector2D", HasNativeMake = "/Script/Engine.KismetMathLibrary.MakeVector2D", HasNativeBreak = "/Script/Engine.KismetMathLibrary.BreakVector2D"))
 struct FVector2f
 {
-	UPROPERTY(EditAnywhere, BlueprintReadWrite, Category=Vector2D, SaveGame)
+	UPROPERTY(EditAnywhere, Category=Vector2D, SaveGame)
 	float X;
 
-	UPROPERTY(EditAnywhere, BlueprintReadWrite, Category=Vector2D, SaveGame)
+	UPROPERTY(EditAnywhere, Category=Vector2D, SaveGame)
 	float Y;
 };
 
@@ -729,7 +729,7 @@ struct FRay
  * An orthogonal rotation in 3d space.
  * @note The full C++ class is located here: Engine\Source\Runtime\Core\Public\Math\Rotator.h
  */
-USTRUCT(immutable, noexport, BlueprintType, BlueprintInternalUseOnly, IsAlwaysAccessible, HasDefaults, HasNoOpConstructor, IsCoreType)
+USTRUCT(immutable, noexport, BlueprintType, BlueprintInternalUseOnly, IsAlwaysAccessible, HasDefaults, HasNoOpConstructor, IsCoreType, meta = (ScriptDefaultMake, ScriptDefaultBreak, HasNativeMake = "/Script/Engine.KismetMathLibrary.MakeRotator", HasNativeBreak = "/Script/Engine.KismetMathLibrary.BreakRotator"))
 struct FRotator3f
 {
 	/** Pitch (degrees) around Y axis */
@@ -835,7 +835,7 @@ struct FSphere
  * Quaternion.
  * @note The full C++ class is located here: Engine\Source\Runtime\Core\Public\Math\Quat.h
  */
-USTRUCT(immutable, noexport, BlueprintType, BlueprintInternalUseOnly, IsAlwaysAccessible, HasDefaults, HasNoOpConstructor, IsCoreType)
+USTRUCT(immutable, noexport, BlueprintType, BlueprintInternalUseOnly, IsAlwaysAccessible, HasDefaults, HasNoOpConstructor, IsCoreType, meta = (ScriptDefaultMake, ScriptDefaultBreak, HasNativeMake = "/Script/Engine.KismetMathLibrary.MakeQuat", HasNativeBreak = "/Script/Engine.KismetMathLibrary.BreakQuat"))
 struct FQuat4f
 {
 	UPROPERTY(EditAnywhere, Category=Quat, SaveGame)
@@ -1918,7 +1918,7 @@ struct FInterpCurveLinearColor
  * Transform composed of Quat/Translation/Scale.
  * @note This is implemented in either TransformVectorized.h or TransformNonVectorized.h depending on the platform.
  */
-USTRUCT(immutable, noexport, BlueprintType, BlueprintInternalUseOnly, IsAlwaysAccessible, HasDefaults, IsCoreType)
+USTRUCT(immutable, noexport, BlueprintType, BlueprintInternalUseOnly, IsAlwaysAccessible, HasDefaults, IsCoreType, meta = (ScriptDefaultMake, ScriptDefaultBreak, HasNativeMake = "/Script/Engine.KismetMathLibrary.MakeTransform", HasNativeBreak = "/Script/Engine.KismetMathLibrary.BreakTransform"))
 struct FTransform3f
 {
 	/** Rotation of this transformation, as a quaternion. */
diff --git a/Engine/Source/Runtime/Engine/Classes/Kismet/KismetMathLibrary.h b/Engine/Source/Runtime/Engine/Classes/Kismet/KismetMathLibrary.h
--- a/Engine/Source/Runtime/Engine/Classes/Kismet/KismetMathLibrary.h
+++ b/Engine/Source/Runtime/Engine/Classes/Kismet/KismetMathLibrary.h
@@ -1334,7 +1334,7 @@ class ENGINE_API UKismetMathLibrary : public UBlueprintFunctionLibrary
 	UFUNCTION(BlueprintPure, Category="Math|Vector", meta=(NativeBreakFunc))
 	static void BreakVector(FVector InVec, double& X, double& Y, double& Z);
 
-	UFUNCTION(BlueprintPure, Category = "Math|Vector", meta = (NativeBreakFunc))
+	UFUNCTION(BlueprintPure, Category = "Math|Vector", meta=(NativeBreakFunc, DeprecatedFunction, DeprecationMessage = "Explicit conversions between floats and doubles are not necessary. Please remove node."))
 	static void BreakVector3f(FVector3f InVec, float& X, float& Y, float& Z);
 
 	/** Converts a vector to LinearColor */
