Added WITH_ENGINE guard to fix issues when compiling C++ UHT.

diff --git a/Engine/Source/Runtime/CoreUObject/Public/UObject/CoreNetTypes.h b/Engine/Source/Runtime/CoreUObject/Public/UObject/CoreNetTypes.h
--- a/Engine/Source/Runtime/CoreUObject/Public/UObject/CoreNetTypes.h
+++ b/Engine/Source/Runtime/CoreUObject/Public/UObject/CoreNetTypes.h
@@ -4,8 +4,8 @@
 
 #include "UObject/ObjectMacros.h"
 
-#if USE_COMPILED_IN_NATIVES 
-#include "CoreNetTypes.generated.h"
+#if WITH_ENGINE 
+	#include "CoreNetTypes.generated.h"
 #endif
 
 /** Whether or not caching of actor/subobject names to the stack should be enabled, for async demo recording crashdumps */
