Early out in FLinkerSave::MapObject once object is found
  - 208ms to 87ms when saving level during cook of test map

#rnx

diff --git a/Engine/Source/Runtime/CoreUObject/Private/UObject/LinkerSave.cpp b/Engine/Source/Runtime/CoreUObject/Private/UObject/LinkerSave.cpp
--- a/Engine/Source/Runtime/CoreUObject/Private/UObject/LinkerSave.cpp
+++ b/Engine/Source/Runtime/CoreUObject/Private/UObject/LinkerSave.cpp
@@ -259,6 +259,7 @@ FPackageIndex FLinkerSave::MapObject( const UObject* Object ) const
 						if (DepListForErrorChecking[DepIndex] == *Found)
 						{
 							bFoundDep = true;
+							break;
 						}
 					}
 				}
