Reduce ping/pong/latency between threads by removing asyncpackage loader now that iostore is the default for cooked data in editor
   - Saves ~5s on 30s of a warm cookbythebook for a small project

#rnx

diff --git a/Engine/Source/Runtime/CoreUObject/Private/UObject/LinkerLoad.cpp b/Engine/Source/Runtime/CoreUObject/Private/UObject/LinkerLoad.cpp
--- a/Engine/Source/Runtime/CoreUObject/Private/UObject/LinkerLoad.cpp
+++ b/Engine/Source/Runtime/CoreUObject/Private/UObject/LinkerLoad.cpp
@@ -1126,9 +1126,10 @@ FLinkerLoad::ELinkerStatus FLinkerLoad::CreateLoader(
 		}
 #endif
 
-		// If want to be able to load cooked data in the editor we need to use FAsyncArchive which supports EDL cooked packages,
-		// otherwise the generic file reader is faster in the editor so use that
-		bool bCanUseAsyncLoader = FPlatformProperties::RequiresCookedData() || GAllowCookedDataInEditorBuilds;
+		// We don't need to go through FAsyncArchive in editor even if GAllowCookedDataInEditorBuilds is true
+		// as loading cooked data is handled by iostore now and the async archive has terrible performance causing
+		// heavy ping/pong latency between threads even for small reads.
+		bool bCanUseAsyncLoader = FPlatformProperties::RequiresCookedData();
 
 		if (bCanUseAsyncLoader)
 		{
