PackageName: Optimize GetPackageMountPoint and other functions using GetValidLongPackageRoots by replacing the linear search through GetValidLongPackageRoots with the O(log(n)) search in DirectoryTree.
1.2s -> 0.5s in a 7.3s CookOnTheFlyServer::Main for an empty DLC cook on top of a large base project.
#rnx

diff --git a/Engine/Source/Runtime/CoreUObject/Private/Misc/PackageName.cpp b/Engine/Source/Runtime/CoreUObject/Private/Misc/PackageName.cpp
--- a/Engine/Source/Runtime/CoreUObject/Private/Misc/PackageName.cpp
+++ b/Engine/Source/Runtime/CoreUObject/Private/Misc/PackageName.cpp
@@ -19,6 +19,7 @@
 #include "Misc/App.h"
 #include "Misc/AutomationTest.h"
 #include "Misc/CoreDelegates.h"
+#include "Misc/EnumClassFlags.h"
 #include "Misc/PackagePath.h"
 #include "Misc/PackageSegment.h"
 #include "Misc/Paths.h"
@@ -156,6 +157,16 @@ bool FPackageName::TryConvertGameRelativePackagePathToLocalPath(FStringView Rela
 
 struct FMountPoint
 {
+	enum class EMountFlags
+	{
+		Removable		= 0x1,
+		NotRemovable	= 0x0,
+		ReadOnly		= 0x2,
+		NotReadOnly		= 0x0,
+		Alias			= 0x4,
+		NotAlias		= 0x0,
+	};
+
 	/** The LongPackageName path(e.g., "/Engine/") */
 	const FString RootPath;
 
@@ -166,17 +177,29 @@ struct FMountPoint
 	const FString ContentPathAbsolute;
 
 	/** Whether the mountpoint was created by InsertMountPoint and can be removed by RemoveMountPoint. */
-	const bool bRemovable = false;
+	const bool bRemovable : 1;
+
+	/** Whether new Packages can be saved into the mountpoint. */
+	const bool bReadOnly : 1;
+
+	/** Whether this mountpoint is the same RootPath as another mountpoint but bound to a different ContentPath. */
+	const bool bAlias : 1;
 
 	FMountPoint(FString&& InRootPath, FString&& InContentPathRelative,
-		FString&& InContentPathAbsolute, bool bInRemovable = false)
-		: RootPath(MoveTemp(InRootPath))
-		, ContentPathRelative(MoveTemp(InContentPathRelative))
-		, ContentPathAbsolute(MoveTemp(InContentPathAbsolute))
-		, bRemovable(bInRemovable)
-	{
-	}
+		FString&& InContentPathAbsolute, EMountFlags Flags);
 };
+ENUM_CLASS_FLAGS(FMountPoint::EMountFlags);
+
+FMountPoint::FMountPoint(FString&& InRootPath, FString&& InContentPathRelative,
+	FString&& InContentPathAbsolute, EMountFlags Flags)
+	: RootPath(MoveTemp(InRootPath))
+	, ContentPathRelative(MoveTemp(InContentPathRelative))
+	, ContentPathAbsolute(MoveTemp(InContentPathAbsolute))
+	, bRemovable(!!(Flags & EMountFlags::Removable))
+	, bReadOnly(!!(Flags & EMountFlags::ReadOnly))
+	, bAlias(!!(Flags & EMountFlags::Alias))
+{
+}
 
 struct FLongPackagePathsSingleton
 {
@@ -208,9 +231,6 @@ struct FLongPackagePathsSingleton
 	TDirectoryTree<FMountPoint*> ContentPathTree;
 	TArray<TUniquePtr<FMountPoint>> MountPoints;
 
-	// Cached values
-	TArray<FString> ValidLongPackageRoots[2];
-
 	// singleton
 	static FLongPackagePathsSingleton& Get()
 	{
@@ -227,38 +247,12 @@ struct FLongPackagePathsSingleton
 		IPluginManager::Get().SetUnRegisterMountPointDelegate(IPluginManager::FRegisterMountPointDelegate::CreateStatic(&FPackageName::UnRegisterMountPoint));
 	}
 
-	const TArray<FString>& GetValidLongPackageRoots(bool bIncludeReadOnlyRoots) const
-	{
-		return ValidLongPackageRoots[bIncludeReadOnlyRoots ? 1 : 0];
-	}
-
-	void UpdateValidLongPackageRoots()
-	{
-		// List Writable roots (roots into which packages can be saved or resaved) in ValidLongPackageRoots[0]
-		ValidLongPackageRoots[0].Empty();
-		ValidLongPackageRoots[0].Add(EngineRootPath);
-		ValidLongPackageRoots[0].Add(GameRootPath);
-		for (TUniquePtr<FMountPoint>& MountPoint : MountPoints)
-		{
-			if (MountPoint->bRemovable)
-			{
-				ValidLongPackageRoots[0].Add(MountPoint->RootPath);
-			}
-		}
-
-		// List all roots in ValidLongPackageRoots[1] (Writable roots plus ReadOnly roots)
-		ValidLongPackageRoots[1] = ValidLongPackageRoots[0];
-		ValidLongPackageRoots[1].Add(ConfigRootPath);
-		ValidLongPackageRoots[1].Add(ScriptRootPath);
-		ValidLongPackageRoots[1].Add(MemoryRootPath);
-		ValidLongPackageRoots[1].Add(TempRootPath);
-	}
-
 	/**
 	 * Given a content path return the consistently-formatted RelativePath and AbsolutePath, suitable for lookup based only on text
 	 * The RelativePath is consistent with FileManager relative paths.
 	 */
-	static FMountPoint ConstructMountPoint(const FString& RootPath, const FString& ContentPath, bool bRemovable)
+	static FMountPoint ConstructMountPoint(const FString& RootPath, const FString& ContentPath,
+		FMountPoint::EMountFlags Flags)
 	{
 		FString MountPath = ContentPath;
 
@@ -287,14 +281,13 @@ struct FLongPackagePathsSingleton
 		TStringBuilder<256> AbsolutePathBuilder;
 		FPathViews::ToAbsolutePath(RelativeContentPath, AbsolutePathBuilder);
 
-		return FMountPoint(FString(RootPath), MoveTemp(RelativeContentPath), AbsolutePathBuilder.ToString(),
-			bRemovable);
+		return FMountPoint(FString(RootPath), MoveTemp(RelativeContentPath), AbsolutePathBuilder.ToString(), Flags);
 	}
 
 	void InsertMountPoint(const FString& RootPath, const FString& ContentPath)
 	{	
 		TUniquePtr<FMountPoint> MountPointOwner(MakeUnique<FMountPoint>(
-			ConstructMountPoint(RootPath, ContentPath, true /* bRemovable */)));
+			ConstructMountPoint(RootPath, ContentPath, FMountPoint::EMountFlags::Removable)));
 		FMountPoint* MountPoint = MountPointOwner.Get();
 		{
 			FWriteScopeLock ScopeLock(MountLock);
@@ -303,8 +296,6 @@ struct FLongPackagePathsSingleton
 			RootPathTree.FindOrAdd(MountPoint->RootPath) = MountPoint;
 			ContentPathTree.FindOrAdd(MountPoint->ContentPathRelative) = MountPoint;
 			ContentPathTree.FindOrAdd(MountPoint->ContentPathAbsolute) = MountPoint;
-
-			UpdateValidLongPackageRoots();
 		}		
 
 		UE_LOG(LogPackageName, Display, TEXT("FPackageName: Mount point added: '%s' mounted to '%s'"),
@@ -317,7 +308,7 @@ struct FLongPackagePathsSingleton
 	// This will remove a previously inserted mount point
 	void RemoveMountPoint(const FString& RootPath, const FString& ContentPath)
 	{
-		RemoveMountPoint(ConstructMountPoint(RootPath, ContentPath, true /* bRemovable */));
+		RemoveMountPoint(ConstructMountPoint(RootPath, ContentPath, FMountPoint::EMountFlags::Removable));
 	}
 
 	void RemoveMountPoint(const FMountPoint& MountPoint)
@@ -389,7 +380,6 @@ struct FLongPackagePathsSingleton
 			{
 				UE_LOG(LogPackageName, Display, TEXT("FPackageName: Mount point removed: '%s' unmounted from '%s'"),
 					*MountPoint.ContentPathRelative, *MountPoint.RootPath);
-				UpdateValidLongPackageRoots();
 			}
 			else
 			{
@@ -409,7 +399,7 @@ struct FLongPackagePathsSingleton
 	bool MountPointExists(const FString& RootPath)
 	{
 		FReadScopeLock ScopeLock(MountLock);
-		return GetValidLongPackageRoots(true/*bIncludeReadOnlyRoots*/).Contains(RootPath);
+		return RootPathTree.Contains(RootPath) || RootPath == MemoryRootPath;
 	}
 
 private:
@@ -461,8 +451,9 @@ private:
 		GameConfigPathRebased  = RebasedGameDir / TEXT("Config/");
 		GameScriptPathRebased  = RebasedGameDir / TEXT("Script/");
 		GameSavedPathRebased   = RebasedGameDir / TEXT("Saved/");
+		using EMountFlags = FMountPoint::EMountFlags;
 
-		auto LocalInsertMountPoint = [this](const FString& RootPath, const FString& ContentPathRelative)
+		auto LocalInsertMountPoint = [this](const FString& RootPath, const FString& ContentPathRelative, EMountFlags InFlags)
 		{
 			// Don't call ConstructMountPoint, because that calls FPaths::ConvertRelativePathToFull followed by
 			// IFileManager::Get().ConvertToRelativePath, and that combination will convert ../../../Engine/Content
@@ -470,7 +461,8 @@ private:
 			TStringBuilder<256> AbsolutePathBuilder;
 			FPathViews::ToAbsolutePath(ContentPathRelative, AbsolutePathBuilder);
 			FMountPoint& MountPoint = *MountPoints.Add_GetRef(MakeUnique<FMountPoint>(
-				FString(RootPath), FString(ContentPathRelative), AbsolutePathBuilder.ToString()));
+				FString(RootPath), FString(ContentPathRelative), AbsolutePathBuilder.ToString(), 
+				InFlags));
 
 			// We are appending to the MountPoint list, and will have only a few duplicates. If any duplicates do occur,
 			// keep the MountPoint that was appended earlier as the registered mountpoint.
@@ -491,32 +483,30 @@ private:
 			}
 		};
 		
-		LocalInsertMountPoint(EngineRootPath, EngineContentPath);
-		LocalInsertMountPoint(EngineRootPath, EngineShadersPath);
-		LocalInsertMountPoint(GameRootPath,   GameContentPath);
-		LocalInsertMountPoint(ScriptRootPath, GameScriptPath);
-		LocalInsertMountPoint(TempRootPath,   GameSavedPath);
-		LocalInsertMountPoint(ConfigRootPath, GameConfigPath);
+		LocalInsertMountPoint(EngineRootPath, EngineContentPath, EMountFlags::NotReadOnly);
+		LocalInsertMountPoint(EngineRootPath, EngineShadersPath, EMountFlags::NotReadOnly | EMountFlags::Alias);
+		LocalInsertMountPoint(GameRootPath,   GameContentPath, EMountFlags::NotReadOnly);
+		LocalInsertMountPoint(ScriptRootPath, GameScriptPath, EMountFlags::ReadOnly);
+		LocalInsertMountPoint(TempRootPath,   GameSavedPath, EMountFlags::ReadOnly);
+		LocalInsertMountPoint(ConfigRootPath, GameConfigPath, EMountFlags::ReadOnly);
 
 		// Add other LocalPaths that have different sets of .. but are the same location on disk
-		LocalInsertMountPoint(EngineRootPath, EngineShadersPathShort);
+		LocalInsertMountPoint(EngineRootPath, EngineShadersPathShort, EMountFlags::NotReadOnly | EMountFlags::Alias);
 		if (FPaths::IsSamePath(GameContentPath, ContentPathShort))
 		{
 			// ../../Content points to the the Engine directory, the same as ../../../Engine/Content
 			// But if the /Game is pointing to ../../../Engine/Content as well, then map ../../Content to
 			// /Game instead of /Engine
-			LocalInsertMountPoint(GameRootPath, ContentPathShort);
+			LocalInsertMountPoint(GameRootPath, ContentPathShort, EMountFlags::NotReadOnly | EMountFlags::Alias);
 		}
 		else
 		{
-			LocalInsertMountPoint(EngineRootPath, ContentPathShort);
+			LocalInsertMountPoint(EngineRootPath, ContentPathShort, EMountFlags::NotReadOnly | EMountFlags::Alias);
 		}
-		LocalInsertMountPoint(GameRootPath,   GameContentPathRebased);
-		LocalInsertMountPoint(ScriptRootPath, GameScriptPathRebased);
-		LocalInsertMountPoint(TempRootPath,   GameSavedPathRebased);
-		LocalInsertMountPoint(ConfigRootPath, GameConfigPathRebased);
-
-		UpdateValidLongPackageRoots();
+		LocalInsertMountPoint(GameRootPath,   GameContentPathRebased, EMountFlags::NotReadOnly | EMountFlags::Alias);
+		LocalInsertMountPoint(ScriptRootPath, GameScriptPathRebased, EMountFlags::ReadOnly | EMountFlags::Alias);
+		LocalInsertMountPoint(TempRootPath,   GameSavedPathRebased, EMountFlags::ReadOnly | EMountFlags::Alias);
+		LocalInsertMountPoint(ConfigRootPath, GameConfigPathRebased, EMountFlags::ReadOnly | EMountFlags::Alias);
 	}
 
 #if !UE_BUILD_SHIPPING
@@ -977,25 +967,34 @@ bool FPackageName::SplitLongPackageName(const FString& InLongPackageName, FStrin
 {
 	const FLongPackagePathsSingleton& Paths = FLongPackagePathsSingleton::Get();
 
-	const bool bIncludeReadOnlyRoots = true;
-
-	FReadScopeLock ScopeLock(Paths.MountLock);
-	const TArray<FString>& ValidRoots = Paths.GetValidLongPackageRoots(bIncludeReadOnlyRoots);
-
-	// Check to see whether our package came from a valid root
-	OutPackageRoot.Empty();
 	FStringView PackageRelPath;
-	for(auto RootIt = ValidRoots.CreateConstIterator(); RootIt; ++RootIt)
 	{
-		const FString& PackageRoot = *RootIt;
-		if (FPathViews::TryMakeChildPathRelativeTo(InLongPackageName, PackageRoot, PackageRelPath))
+		FReadScopeLock ScopeLock(Paths.MountLock);
+
+		// Check to see whether our package came from a valid root
+		OutPackageRoot.Empty();
+		const FMountPoint* const* MountPointPtr = Paths.RootPathTree.FindClosestValue(InLongPackageName);
+		if (MountPointPtr)
+		{
+			// DirectoryTree is less conservative than TryMakeChildPathRelativeTo; e.g. d:/Dir//Root will match
+			// d:/Dir/Root/Child in DirectoryTree but not in TryMakeChildPathRelativeTo. Treat this as PathNotMounted,
+			// since no mountdir will match it according to TryMakeChildPathRelativeTo.
+			if (FPathViews::TryMakeChildPathRelativeTo(InLongPackageName, (*MountPointPtr)->RootPath, PackageRelPath))
+			{
+				OutPackageRoot = (*MountPointPtr)->RootPath / "";
+			}
+		}
+		else
 		{
-			OutPackageRoot = PackageRoot / "";
-			break;
+			const FString& ExtraPackageRoot = Paths.MemoryRootPath;
+			if (FPathViews::TryMakeChildPathRelativeTo(InLongPackageName, ExtraPackageRoot, PackageRelPath))
+			{
+				OutPackageRoot = ExtraPackageRoot / "";
+			}
 		}
 	}
 
-	if(OutPackageRoot.IsEmpty() || InLongPackageName.Len() <= OutPackageRoot.Len())
+	if (OutPackageRoot.IsEmpty() || InLongPackageName.Len() <= OutPackageRoot.Len())
 	{
 		// Path is not part of a valid root, or the path given is too short to continue; splitting failed
 		return false;
@@ -1006,7 +1005,7 @@ bool FPackageName::SplitLongPackageName(const FString& InLongPackageName, FStrin
 	OutPackagePath = FPaths::GetPath(PackageRelPathStr) / "";
 	OutPackageName = FPaths::GetCleanFilename(PackageRelPathStr);
 
-	if(bStripRootLeadingSlash && OutPackageRoot.StartsWith(TEXT("/"), ESearchCase::CaseSensitive))
+	if (bStripRootLeadingSlash && OutPackageRoot.StartsWith(TEXT("/"), ESearchCase::CaseSensitive))
 	{
 		OutPackageRoot.RemoveAt(0);
 	}
@@ -1175,34 +1174,22 @@ bool FPackageName::IsValidLongPackageName(FStringView InLongPackageName, bool bI
 				const FLongPackagePathsSingleton& Paths = FLongPackagePathsSingleton::Get();
 
 				FReadScopeLock ScopeLock(Paths.MountLock);
-				const TArray<FString>& ValidRoots = Paths.GetValidLongPackageRoots(bIncludeReadOnlyRoots);
-				if (ValidRoots.Num() == 0)
-				{
-					*OutReason = NSLOCTEXT("Core", "LongPackageNames_NoValidRoots", "No valid roots exist!");
-				}
-				else
+				TStringBuilder<2048> ValidRootsString;
+				for (const TUniquePtr<FMountPoint>& MountPoint : Paths.MountPoints)
 				{
-					FString ValidRootsString = TEXT("");
-					if (ValidRoots.Num() == 1)
+					if (bIncludeReadOnlyRoots || !MountPoint->bReadOnly)
 					{
-						ValidRootsString = FString::Printf(TEXT("'%s'"), *ValidRoots[0]);
+						ValidRootsString << MountPoint->RootPath << TEXT(", ");
 					}
-					else
-					{
-						for (int32 RootIdx = 0; RootIdx < ValidRoots.Num(); ++RootIdx)
-						{
-							if (RootIdx < ValidRoots.Num() - 1)
-							{
-								ValidRootsString += FString::Printf(TEXT("'%s', "), *ValidRoots[RootIdx]);
-							}
-							else
-							{
-								ValidRootsString += FString::Printf(TEXT("or '%s'"), *ValidRoots[RootIdx]);
-							}
-						}
-					}
-					*OutReason = FText::Format( NSLOCTEXT("Core", "LongPackageNames_InvalidRoot", "Path does not start with a valid root. Path must begin with: {0}"), FText::FromString( ValidRootsString ) );
 				}
+				if (bIncludeReadOnlyRoots)
+				{
+					const FString& ExtraMountRootPath = Paths.MemoryRootPath;
+					ValidRootsString << ExtraMountRootPath << TEXT(", ");
+				}
+				ValidRootsString.RemoveSuffix(2); // Remove the trailing ", "
+				*OutReason = FText::Format( NSLOCTEXT("Core", "LongPackageNames_InvalidRoot", "Path does not start with a valid root. Path must begin with one of: {0}"),
+					FText::FromString( FString(ValidRootsString) ) );
 			}
 			else
 			{
@@ -1224,17 +1211,35 @@ bool FPackageName::IsValidLongPackageName(FStringView InLongPackageName, bool bI
 	// Check valid roots
 	const FLongPackagePathsSingleton& Paths = FLongPackagePathsSingleton::Get();
 	FReadScopeLock ScopeLock(Paths.MountLock);
-	const TArray<FString>& ValidRoots = Paths.GetValidLongPackageRoots(bIncludeReadOnlyRoots);
-	bool bValidRoot = false;
-	for (int32 RootIdx = 0; RootIdx < ValidRoots.Num(); ++RootIdx)
+
+	const FMountPoint* const* MountPointPtr = Paths.RootPathTree.FindClosestValue(InLongPackageName);
+	if (MountPointPtr)
 	{
-		const FString& Root = ValidRoots[RootIdx];
-		if (FPathViews::IsParentPathOf(Root, InLongPackageName))
+		if (bIncludeReadOnlyRoots || !(*MountPointPtr)->bReadOnly)
 		{
-			if (OutReason) *OutReason = EErrorCode::PackageNameUnknown;
-			return true;
+			// DirectoryTree is less conservative than IsParentPathOf; e.g. d:/Dir//Root will match
+			// d:/Dir/Root/Child in DirectoryTree but not in IsParentPathOf. Treat this as PathNotMounted,
+			// since no mountdir will match it according to IsParentPathOf.
+			if (FPathViews::IsParentPathOf((*MountPointPtr)->RootPath, InLongPackageName))
+			{
+				if (OutReason) *OutReason = EErrorCode::PackageNameUnknown;
+				return true;
+			}
+		}
+	}
+	else
+	{
+		if (bIncludeReadOnlyRoots)
+		{
+			const FString& ExtraPackageRoot = Paths.MemoryRootPath;
+			if (FPathViews::IsParentPathOf(ExtraPackageRoot, InLongPackageName))
+			{
+				if (OutReason) *OutReason = EErrorCode::PackageNameUnknown;
+				return true;
+			}
 		}
 	}
+
 	if (OutReason) *OutReason = EErrorCode::PackageNamePathNotMounted;
 	return false;
 }
@@ -1331,14 +1336,30 @@ FName FPackageName::GetPackageMountPoint(const FString& InPackagePath, bool InWi
 	FLongPackagePathsSingleton& Paths = FLongPackagePathsSingleton::Get();
 	
 	FReadScopeLock ScopeLock(Paths.MountLock);
-	const TArray<FString>& MountPoints = Paths.GetValidLongPackageRoots(true);
-
 	int32 WithoutSlashes = InWithoutSlashes ? 1 : 0;
-	for (auto RootIt = MountPoints.CreateConstIterator(); RootIt; ++RootIt)
+	auto TryRootPath = [&InPackagePath, WithoutSlashes](FStringView RootPath)
+	{
+		if (FPathViews::IsParentPathOf(RootPath, InPackagePath))
+		{
+			return FName(RootPath.Mid(WithoutSlashes, RootPath.Len() - (2 * WithoutSlashes)));
+		}
+		return FName();
+	};
+	const FMountPoint* const* MountPointPtr = Paths.RootPathTree.FindClosestValue(InPackagePath);
+	if (MountPointPtr)
+	{
+		// DirectoryTree is less conservative than IsParentPathOf; e.g. d:/Dir//Root will match
+		// d:/Dir/Root/Child in DirectoryTree but not in IsParentPathOf. Treat this as PathNotMounted,
+		// since no mountdir will match it according to IsParentPathOf.
+		return TryRootPath((*MountPointPtr)->RootPath);
+	}
+	else
 	{
-		if (FPathViews::IsParentPathOf(*RootIt, InPackagePath))
+		const FString& ExtraPackageRoot = Paths.MemoryRootPath;
+		FName Result = TryRootPath(ExtraPackageRoot);
+		if (Result.IsValid())
 		{
-			return FName(*RootIt->Mid(WithoutSlashes, RootIt->Len() - (2 * WithoutSlashes)));
+			return Result;
 		}
 	}
 
@@ -2303,7 +2324,34 @@ void FPackageName::QueryRootContentPaths(TArray<FString>& OutRootContentPaths, b
 
 	{
 		FReadScopeLock ScopeLock(Paths.MountLock);
-		OutRootContentPaths = Paths.GetValidLongPackageRoots( bIncludeReadOnlyRoots );
+		OutRootContentPaths.Reserve(Paths.MountPoints.Num() + (bIncludeReadOnlyRoots ? 1 : 0));
+		// Maintain the legacy order: non-removable non-readonly, then removable, then non-removable readonly
+		for (const TUniquePtr<FMountPoint>& MountPoint : Paths.MountPoints)
+		{
+			if (!MountPoint->bRemovable && !MountPoint->bReadOnly && !MountPoint->bAlias)
+			{
+				OutRootContentPaths.Add(MountPoint->RootPath);
+			}
+		}
+		for (const TUniquePtr<FMountPoint>& MountPoint : Paths.MountPoints)
+		{
+			if (MountPoint->bRemovable)
+			{
+				OutRootContentPaths.Add(MountPoint->RootPath);
+			}
+		}
+		if (bIncludeReadOnlyRoots)
+		{
+			for (const TUniquePtr<FMountPoint>& MountPoint : Paths.MountPoints)
+			{
+				if (!MountPoint->bRemovable && MountPoint->bReadOnly && !MountPoint->bAlias)
+				{
+					OutRootContentPaths.Add(MountPoint->RootPath);
+				}
+			}
+			const FString& ExtraRoot = Paths.MemoryRootPath;
+			OutRootContentPaths.Add(ExtraRoot);
+		}
 	}
 
 	if (bWithoutTrailingSlashes || bWithoutLeadingSlashes)
