2 updates to property functionality.

The following two changes were made to property functionality:
1. FDefaultConstructedPropertyElement has been moved to using a shared pointer to support copying and in turn allow it to be stored in a map for caching.
2. FScriptArrayHelper got the "GetElementPtr" wrapper around "GetRawPtr" in order to have a similar API as FScriptSetHelper for use in templates.

diff --git a/Engine/Source/Runtime/CoreUObject/Public/UObject/UnrealType.h b/Engine/Source/Runtime/CoreUObject/Public/UObject/UnrealType.h
--- a/Engine/Source/Runtime/CoreUObject/Public/UObject/UnrealType.h
+++ b/Engine/Source/Runtime/CoreUObject/Public/UObject/UnrealType.h
@@ -42,6 +42,7 @@
 #include "Templates/IsUEnumClass.h"
 #include "Templates/MemoryOps.h"
 #include "Templates/Models.h"
+#include "Templates/SharedPointer.h"
 #include "Templates/Tuple.h"
 #include "Templates/UnrealTemplate.h"
 #include "Templates/UnrealTypeTraits.h"
@@ -1256,32 +1257,24 @@ struct COREUOBJECT_API FDefinedProperty
 class COREUOBJECT_API FDefaultConstructedPropertyElement
 {
 public:
-	explicit FDefaultConstructedPropertyElement(FProperty* InProp)
-		: Prop(InProp)
-		, Obj(FMemory::Malloc(InProp->GetSize(), InProp->GetMinAlignment()))
-	{
-		InProp->InitializeValue(Obj);
-	}
-
-	~FDefaultConstructedPropertyElement()
+	FDefaultConstructedPropertyElement() = default;
+	explicit FDefaultConstructedPropertyElement(const FProperty* InProp)
+		: Obj(FMemory::Malloc(InProp->GetSize(), InProp->GetMinAlignment()), [InProp](void* Object)
+			{
+				InProp->DestroyValue(Object);
+				FMemory::Free(Object);
+			})
 	{
-		Prop->DestroyValue(Obj);
-		FMemory::Free(Obj);
+		InProp->InitializeValue(Obj.Get());
 	}
 
 	void* GetObjAddress() const
 	{
-		return Obj;
+		return Obj.Get();
 	}
 
 private:
-	// Non-copyable
-	FDefaultConstructedPropertyElement(const FDefaultConstructedPropertyElement&) = delete;
-	FDefaultConstructedPropertyElement& operator=(const FDefaultConstructedPropertyElement&) = delete;
-
-private:
-	FProperty* Prop;
-	void* Obj;
+	TSharedPtr<void> Obj;
 };
 
 
@@ -4179,6 +4172,16 @@ public:
 		checkSlow(IsValidIndex(Index)); 
 		return (uint8*)WithScriptArray([](auto* Array) { return Array->GetData(); }) + Index * ElementSize;
 	}
+	/**
+	 *	Returns a uint8 pointer to an element in the array. This call is identical to GetRawPtr and is
+	 *  here to provide interface parity with FScriptSetHelper*.
+	 *	@param	Index: index of the item to return a pointer to.
+	 *	@return	Pointer to this element, or NULL if the array is empty
+	**/
+	FORCEINLINE uint8* GetElementPtr(int32 Index = 0)
+	{
+		return GetRawPtr(Index);
+	}
 	/**
 	*	Empty the array, then add blank, constructed values to a given size.
 	*	@param	Count: the number of items the array will have on completion.
