Made FBoxes' IsValid flags editable, otherwise editor-modified boxes will compare equal to the CDO (all invalid boxes compare equal, regardless of the state of the extents) and so not be saved.

diff --git a/Engine/Source/Runtime/CoreUObject/Public/UObject/NoExportTypes.h b/Engine/Source/Runtime/CoreUObject/Public/UObject/NoExportTypes.h
--- a/Engine/Source/Runtime/CoreUObject/Public/UObject/NoExportTypes.h
+++ b/Engine/Source/Runtime/CoreUObject/Public/UObject/NoExportTypes.h
@@ -1369,14 +1369,14 @@ struct FLinearColor
 USTRUCT(immutable, noexport, BlueprintType, BlueprintInternalUseOnly, IsAlwaysAccessible, HasDefaults, HasNoOpConstructor, IsCoreType)
 struct FBox3f
 {
-	UPROPERTY(EditAnywhere, Category = Box, SaveGame)
+	UPROPERTY(EditAnywhere, Category = Box, SaveGame, meta=(EditCondition="IsValid"))
 	FVector3f Min;
 
-	UPROPERTY(EditAnywhere, Category = Box, SaveGame)
+	UPROPERTY(EditAnywhere, Category = Box, SaveGame, meta=(EditCondition="IsValid"))
 	FVector3f Max;
 
-	UPROPERTY()
-	uint8 IsValid;
+	UPROPERTY(EditAnywhere, Category = Box, SaveGame)
+	bool IsValid;
 };
 
 
@@ -1387,14 +1387,14 @@ struct FBox3f
 USTRUCT(immutable, noexport, IsAlwaysAccessible, HasDefaults, HasNoOpConstructor, IsCoreType)
 struct FBox3d
 {
-	UPROPERTY(EditAnywhere, Category = Box, SaveGame)
+	UPROPERTY(EditAnywhere, Category = Box, SaveGame, meta=(EditCondition="IsValid"))
 	FVector3d Min;
 
-	UPROPERTY(EditAnywhere, Category = Box, SaveGame)
+	UPROPERTY(EditAnywhere, Category = Box, SaveGame, meta=(EditCondition="IsValid"))
 	FVector3d Max;
 
-	UPROPERTY()
-	uint8 IsValid;
+	UPROPERTY(EditAnywhere, Category = Box, SaveGame)
+	bool IsValid;
 };
 
 /**
@@ -1404,14 +1404,14 @@ struct FBox3d
 USTRUCT(immutable, noexport, BlueprintType, IsAlwaysAccessible, HasDefaults, HasNoOpConstructor, IsCoreType, meta=(HasNativeMake="/Script/Engine.KismetMathLibrary.MakeBox"))
 struct FBox
 {
-	UPROPERTY(EditAnywhere, BlueprintReadWrite, Category=Box, SaveGame)
+	UPROPERTY(EditAnywhere, BlueprintReadWrite, Category=Box, SaveGame, meta=(EditCondition="IsValid"))
 	FVector Min;
 
-	UPROPERTY(EditAnywhere, BlueprintReadWrite, Category=Box, SaveGame)
+	UPROPERTY(EditAnywhere, BlueprintReadWrite, Category=Box, SaveGame, meta=(EditCondition="IsValid"))
 	FVector Max;
 
-	UPROPERTY()
-	uint8 IsValid;
+	UPROPERTY(EditAnywhere, BlueprintReadWrite, Category=Box, SaveGame)
+	bool IsValid;
 };
 
 /**
@@ -1421,14 +1421,14 @@ struct FBox
 USTRUCT(immutable, noexport, BlueprintType, BlueprintInternalUseOnly, IsAlwaysAccessible, HasDefaults, HasNoOpConstructor, IsCoreType)
 struct FBox2f
 {
-	UPROPERTY(EditAnywhere, Category=Box2D, SaveGame)
+	UPROPERTY(EditAnywhere, Category=Box2D, SaveGame, meta=(EditCondition="bIsValid"))
 	FVector2f Min;
 
-	UPROPERTY(EditAnywhere, Category=Box2D, SaveGame)
+	UPROPERTY(EditAnywhere, Category=Box2D, SaveGame, meta=(EditCondition="bIsValid"))
 	FVector2f Max;
 
-	UPROPERTY()
-	uint8 bIsValid;
+	UPROPERTY(EditAnywhere, Category=Box2D, SaveGame)
+	bool bIsValid;
 };
 
 /**
@@ -1439,14 +1439,14 @@ struct FBox2f
 // USTRUCT(immutable, noexport, IsAlwaysAccessible, HasDefaults, HasNoOpConstructor, IsCoreType)
 // struct FBox2d
 // {
-// 	UPROPERTY(EditAnywhere, Category=Box2D, SaveGame)
+// 	UPROPERTY(EditAnywhere, Category=Box2D, SaveGame, meta=(EditCondition="bIsValid"))
 // 	FVector2d Min;
 //
-// 	UPROPERTY(EditAnywhere, Category=Box2D, SaveGame)
+// 	UPROPERTY(EditAnywhere, Category=Box2D, SaveGame, meta=(EditCondition="bIsValid"))
 // 	FVector2d Max;
 //
-// 	UPROPERTY()
-// 	uint8 bIsValid;
+// 	UPROPERTY(EditAnywhere, Category=Box2D, SaveGame)
+// 	bool bIsValid;
 // };
 
 /**
@@ -1456,14 +1456,14 @@ struct FBox2f
 USTRUCT(immutable, noexport, BlueprintType, IsAlwaysAccessible, HasDefaults, HasNoOpConstructor, IsCoreType, meta=(HasNativeMake="/Script/Engine.KismetMathLibrary.MakeBox2D"))
 struct FBox2D
 {
-	UPROPERTY(EditAnywhere, BlueprintReadWrite, Category=Box2D, SaveGame)
+	UPROPERTY(EditAnywhere, BlueprintReadWrite, Category=Box2D, SaveGame, meta=(EditCondition="bIsValid"))
 	FVector2D Min;
 
-	UPROPERTY(EditAnywhere, BlueprintReadWrite, Category=Box2D, SaveGame)
+	UPROPERTY(EditAnywhere, BlueprintReadWrite, Category=Box2D, SaveGame, meta=(EditCondition="bIsValid"))
 	FVector2D Max;
 
-	UPROPERTY()
-	uint8 bIsValid;
+	UPROPERTY(EditAnywhere, BlueprintReadWrite, Category=Box2D, SaveGame)
+	bool bIsValid;
 };
 
 /**
