ZenLoader: Ensure that postload and serialize never overlaps for exports in related packages.

diff --git a/Engine/Source/Runtime/CoreUObject/Private/Serialization/AsyncLoading2.cpp b/Engine/Source/Runtime/CoreUObject/Private/Serialization/AsyncLoading2.cpp
--- a/Engine/Source/Runtime/CoreUObject/Private/Serialization/AsyncLoading2.cpp
+++ b/Engine/Source/Runtime/CoreUObject/Private/Serialization/AsyncLoading2.cpp
@@ -2453,6 +2453,8 @@ private:
 };
 #endif // ALT2_ENABLE_NEW_ARCHIVE_FOR_LINKERLOAD
 
+struct FAsyncLoadingPostLoadGroup;
+
 /**
 * Structure containing intermediate data required for async loading of all exports of a package.
 */
@@ -2602,6 +2604,7 @@ private:
 	UPackage* LinkerRoot = nullptr;
 	// The sync load context associated with this package
 	std::atomic<uint64>			SyncLoadContextId = 0;
+	FAsyncLoadingPostLoadGroup* PostLoadGroup = nullptr;
 	/** Time load begun. This is NOT the time the load was requested in the case of pending requests.	*/
 	double						LoadStartTime = 0.0;
 	std::atomic<int32>			RefCount = 0;
@@ -2710,7 +2713,6 @@ private:
 #if ALT2_ENABLE_LINKERLOAD_SUPPORT
 	FAllDependenciesState		AllDependenciesSetupState;
 #endif
-	FAllDependenciesState		AllDependenciesSerializedState;
 	FAllDependenciesState		AllDependenciesFullyLoadedState;
 
 	/** True if our load has failed */
@@ -2781,7 +2783,6 @@ private:
 	bool HaveAllDependenciesReachedState(FAsyncLoadingThreadState2& ThreadState, FAllDependenciesState FAsyncPackage2::* StateMemberPtr, EAsyncPackageLoadingState2 WaitForPackageState, uint32 CurrentTick, bool bExcludeZenPackages);
 	void UpdateDependenciesStateRecursive(FAsyncLoadingThreadState2& ThreadState, FAllDependenciesState FAsyncPackage2::* StateMemberPtr, EAsyncPackageLoadingState2 WaitForPackageState, uint32 CurrentTick, FAsyncPackage2* Root, bool bExcludeZenPackages);
 	void WaitForAllDependenciesToReachState(FAsyncLoadingThreadState2& ThreadState, FAllDependenciesState FAsyncPackage2::* StateMemberPtr, EAsyncPackageLoadingState2 WaitForPackageState, uint32& CurrentTickVariable, bool bExcludeZenPackages, TFunctionRef<void(FAsyncPackage2*)> OnStateReached);
-	void ConditionalBeginPostLoad(FAsyncLoadingThreadState2& ThreadState);
 	void ConditionalFinishLoading(FAsyncLoadingThreadState2& ThreadState);
 
 #if ALT2_ENABLE_LINKERLOAD_SUPPORT
@@ -2836,6 +2837,12 @@ public:
 	FUObjectSerializeContext* GetSerializeContext();
 };
 
+struct FAsyncLoadingPostLoadGroup
+{
+	TArray<FAsyncPackage2*> Packages;
+	int32 PackagesWithExportsToSerializeCount = 0;
+};
+
 class FAsyncLoadingThread2 final
 	: public FRunnable
 	, public IAsyncPackageLoader
@@ -2938,7 +2945,6 @@ private:
 #if ALT2_ENABLE_LINKERLOAD_SUPPORT
 	uint32 ConditionalProcessLinkerLoadPackageExportsTick = 0;
 #endif
-	uint32 ConditionalBeginPostLoadTick = 0;
 	uint32 ConditionalFinishLoadingTick = 0;
 
 public:
@@ -3042,7 +3048,7 @@ public:
 	void UpdatePackagePriority(FAsyncLoadingThreadState2& ThreadState, FAsyncPackage2* Package);
 	void UpdatePackagePriorityRecursive(FAsyncLoadingThreadState2& ThreadState, FAsyncPackage2* Package, int32 NewPriority);
 
-	FAsyncPackage2* FindOrInsertPackage(FAsyncLoadingThreadState2& ThreadState, FAsyncPackageDesc2& InDesc, bool& bInserted, TUniquePtr<FLoadPackageAsyncDelegate>&& PackageLoadedDelegate = TUniquePtr<FLoadPackageAsyncDelegate>());
+	FAsyncPackage2* FindOrInsertPackage(FAsyncLoadingThreadState2& ThreadState, FAsyncPackageDesc2& InDesc, bool& bInserted, FAsyncPackage2* ImportedByPackage, TUniquePtr<FLoadPackageAsyncDelegate>&& PackageLoadedDelegate = TUniquePtr<FLoadPackageAsyncDelegate>());
 	void QueueMissingPackage(FAsyncPackageDesc2& PackageDesc, TUniquePtr<FLoadPackageAsyncDelegate>&& LoadPackageAsyncDelegate);
 
 	/**
@@ -3193,6 +3199,9 @@ private:
 	}
 #endif
 
+	void ConditionalBeginPostLoad(FAsyncLoadingThreadState2& ThreadState, FAsyncLoadingPostLoadGroup* PostLoadGroup);
+	void MergePostLoadGroups(FAsyncLoadingThreadState2& ThreadState, FAsyncLoadingPostLoadGroup* Target, FAsyncLoadingPostLoadGroup* Source);
+
 	bool ProcessDeferredDeletePackagesQueue(int32 MaxCount = MAX_int32)
 	{
 		bool bDidSomething = false;
@@ -3578,9 +3587,50 @@ void FAsyncLoadingThread2::UpdatePackagePriorityRecursive(FAsyncLoadingThreadSta
 	UpdatePackagePriority(ThreadState, Package);
 }
 
-FAsyncPackage2* FAsyncLoadingThread2::FindOrInsertPackage(FAsyncLoadingThreadState2& ThreadState, FAsyncPackageDesc2& Desc, bool& bInserted, TUniquePtr<FLoadPackageAsyncDelegate>&& PackageLoadedDelegate)
+void FAsyncLoadingThread2::ConditionalBeginPostLoad(FAsyncLoadingThreadState2& ThreadState, FAsyncLoadingPostLoadGroup* PostLoadGroup)
+{
+	TRACE_CPUPROFILER_EVENT_SCOPE(ConditionalBeginPostLoad);
+	check(PostLoadGroup);
+	check(ThreadState.bCanAccessAsyncLoadingThreadData);
+	if (PostLoadGroup->PackagesWithExportsToSerializeCount == 0)
+	{
+		for (FAsyncPackage2* Package : PostLoadGroup->Packages)
+		{
+			check(Package->PostLoadGroup == PostLoadGroup);
+			check(Package->AsyncPackageLoadingState == EAsyncPackageLoadingState2::ExportsDone);
+			Package->AsyncPackageLoadingState = EAsyncPackageLoadingState2::PostLoad;
+			Package->PostLoadGroup = nullptr;
+			for (int32 ExportBundleIndex = 0; ExportBundleIndex < Package->Data.ExportBundleCount; ++ExportBundleIndex)
+			{
+				Package->GetExportBundleNode(EEventLoadNode2::ExportBundle_PostLoad, ExportBundleIndex).ReleaseBarrier(&ThreadState);
+			}
+		}
+		delete PostLoadGroup;
+	}
+}
+
+void FAsyncLoadingThread2::MergePostLoadGroups(FAsyncLoadingThreadState2& ThreadState, FAsyncLoadingPostLoadGroup* Target, FAsyncLoadingPostLoadGroup* Source)
+{
+	if (Target == Source)
+	{
+		return;
+	}
+	TRACE_CPUPROFILER_EVENT_SCOPE(MergePostLoadGroups);
+	check(ThreadState.bCanAccessAsyncLoadingThreadData);
+	for (FAsyncPackage2* Package : Source->Packages)
+	{
+		check(Package->PostLoadGroup == Source);
+		Package->PostLoadGroup = Target;
+	}
+	Target->Packages.Append(MoveTemp(Source->Packages));
+	Target->PackagesWithExportsToSerializeCount += Source->PackagesWithExportsToSerializeCount;
+	delete Source;
+}
+
+FAsyncPackage2* FAsyncLoadingThread2::FindOrInsertPackage(FAsyncLoadingThreadState2& ThreadState, FAsyncPackageDesc2& Desc, bool& bInserted, FAsyncPackage2* ImportedByPackage, TUniquePtr<FLoadPackageAsyncDelegate>&& PackageLoadedDelegate)
 {
 	TRACE_CPUPROFILER_EVENT_SCOPE(FindOrInsertPackage);
+	check(ThreadState.bCanAccessAsyncLoadingThreadData);
 	FAsyncPackage2* Package = nullptr;
 	bInserted = false;
 	{
@@ -3613,6 +3663,37 @@ FAsyncPackage2* FAsyncLoadingThread2::FindOrInsertPackage(FAsyncLoadingThreadSta
 			Package->AddCompletionCallback(MoveTemp(PackageLoadedDelegate));
 		}
 	}
+	if (bInserted)
+	{
+		// Created a new package, either create a new post load group or use the one from the importing package
+		FAsyncLoadingPostLoadGroup* PostLoadGroup = ImportedByPackage ? ImportedByPackage->PostLoadGroup : new FAsyncLoadingPostLoadGroup();
+		++PostLoadGroup->PackagesWithExportsToSerializeCount;
+		PostLoadGroup->Packages.Add(Package);
+		check(!Package->PostLoadGroup);
+		Package->PostLoadGroup = PostLoadGroup;
+	}
+	else if (ImportedByPackage)
+	{
+		// Importing a package that was already being loaded
+		check(Package->PostLoadGroup);
+		if (Package->AsyncPackageLoadingState >= EAsyncPackageLoadingState2::ExportsDone)
+		{
+			// The imported package has started postloading, wait for it to finish postloading before serializing any exports
+			for (int32 DependentExportBundleIndex = 0; DependentExportBundleIndex < ImportedByPackage->Data.ExportBundleCount; ++DependentExportBundleIndex)
+			{
+				for (int32 DependsOnExportBundleIndex = 0; DependsOnExportBundleIndex < Package->Data.ExportBundleCount; ++DependsOnExportBundleIndex)
+				{
+					ImportedByPackage->GetExportBundleNode(ExportBundle_Process, DependentExportBundleIndex).DependsOn(&Package->GetExportBundleNode(ExportBundle_DeferredPostLoad, DependsOnExportBundleIndex));
+				}
+			}
+		}
+		else if (ImportedByPackage->PostLoadGroup != Package->PostLoadGroup)
+		{
+			// The imported package hasn't started postloading yet, merge its post load group with the one for the importing package
+			check(ImportedByPackage->PostLoadGroup);
+			MergePostLoadGroups(ThreadState, ImportedByPackage->PostLoadGroup, Package->PostLoadGroup);
+		}
+	}
 	return Package;
 }
 
@@ -3778,7 +3859,7 @@ bool FAsyncLoadingThread2::CreateAsyncPackagesFromQueue(FAsyncLoadingThreadState
 			else
 			{
 				bool bInserted;
-				FAsyncPackage2* Package = FindOrInsertPackage(ThreadState, PackageDesc, bInserted, MoveTemp(Request.PackageLoadedDelegate));
+				FAsyncPackage2* Package = FindOrInsertPackage(ThreadState, PackageDesc, bInserted, nullptr, MoveTemp(Request.PackageLoadedDelegate));
 				checkf(Package, TEXT("Failed to find or insert package %s"), *PackageDesc.UPackageName.ToString());
 
 				if (bInserted)
@@ -4513,7 +4594,7 @@ void FAsyncPackage2::ImportPackagesRecursiveInner(FAsyncLoadingThreadState2& Thr
 		else
 		{
 			FAsyncPackageDesc2 PackageDesc = FAsyncPackageDesc2::FromPackageImport(Desc, ImportedPackageUPackageName, ImportedPackageId, ImportedPackageIdToLoad, MoveTemp(ImportedPackagePath));
-			ImportedPackage = AsyncLoadingThread.FindOrInsertPackage(ThreadState, PackageDesc, bInserted);
+			ImportedPackage = AsyncLoadingThread.FindOrInsertPackage(ThreadState, PackageDesc, bInserted, this);
 		}
 
 		checkf(ImportedPackage, TEXT("Failed to find or insert imported package with id '0x%llX'"), ImportedPackageId.Value());
@@ -5969,7 +6050,11 @@ EEventLoadNodeExecutionResult FAsyncPackage2::Event_ExportsDone(FAsyncLoadingThr
 		FCoreDelegates::ReleasePreloadedPackageShaderMaps.ExecuteIfBound(Package->Data.ShaderMapHashes);
 	}
 
-	Package->ConditionalBeginPostLoad(ThreadState);
+	FAsyncLoadingPostLoadGroup* PostLoadGroup = Package->PostLoadGroup;
+	check(PostLoadGroup);
+	check(PostLoadGroup->PackagesWithExportsToSerializeCount > 0);
+	--PostLoadGroup->PackagesWithExportsToSerializeCount;
+	Package->AsyncLoadingThread.ConditionalBeginPostLoad(ThreadState, PostLoadGroup);
 	return EEventLoadNodeExecutionResult::Complete;
 }
 
@@ -6172,20 +6257,6 @@ void FAsyncPackage2::WaitForAllDependenciesToReachState(FAsyncLoadingThreadState
 	}
 }
 
-void FAsyncPackage2::ConditionalBeginPostLoad(FAsyncLoadingThreadState2& ThreadState)
-{
-	TRACE_CPUPROFILER_EVENT_SCOPE(ConditionalBeginPostLoad);
-	constexpr bool bExcludeZenPackages = false;
-	WaitForAllDependenciesToReachState(ThreadState, &FAsyncPackage2::AllDependenciesSerializedState, EAsyncPackageLoadingState2::ExportsDone, AsyncLoadingThread.ConditionalBeginPostLoadTick, bExcludeZenPackages,
-		[&ThreadState](FAsyncPackage2* Package)
-		{
-			for (int32 ExportBundleIndex = 0; ExportBundleIndex < Package->Data.ExportBundleCount; ++ExportBundleIndex)
-			{
-				Package->GetExportBundleNode(EEventLoadNode2::ExportBundle_PostLoad, ExportBundleIndex).ReleaseBarrier(&ThreadState);
-			}
-		});
-}
-
 void FAsyncPackage2::ConditionalFinishLoading(FAsyncLoadingThreadState2& ThreadState)
 {
 	TRACE_CPUPROFILER_EVENT_SCOPE(ConditionalFinishLoading);
@@ -8096,10 +8167,21 @@ void FAsyncLoadingThread2::FlushLoading(int32 RequestId)
 			}
 			if (GameThreadState->bCanAccessAsyncLoadingThreadData)
 			{
-				while (FAsyncPackage2* WaitingPackage = CurrentlyExecutingPackage->AllDependenciesSerializedState.PackagesWaitingForThisHead)
+				FAsyncLoadingPostLoadGroup* PostLoadGroup = CurrentlyExecutingPackage->PostLoadGroup;
+				if (CurrentlyExecutingPackage->AsyncPackageLoadingState < EAsyncPackageLoadingState2::ExportsDone)
 				{
-					FAsyncPackage2::FAllDependenciesState::RemoveFromWaitList(&FAsyncPackage2::AllDependenciesSerializedState, CurrentlyExecutingPackage, WaitingPackage);
-					WaitingPackage->ConditionalBeginPostLoad(*GameThreadState);
+					check(PostLoadGroup);
+					check(PostLoadGroup->Packages.Contains(CurrentlyExecutingPackage));
+					check(PostLoadGroup->PackagesWithExportsToSerializeCount > 0);
+					if (PostLoadGroup->Packages.Num() > 1)
+					{
+						PostLoadGroup->Packages.Remove(CurrentlyExecutingPackage);
+						--PostLoadGroup->PackagesWithExportsToSerializeCount;
+						ConditionalBeginPostLoad(*GameThreadState, PostLoadGroup);
+						CurrentlyExecutingPackage->PostLoadGroup = new FAsyncLoadingPostLoadGroup();
+						CurrentlyExecutingPackage->PostLoadGroup->Packages.Add(CurrentlyExecutingPackage);
+						CurrentlyExecutingPackage->PostLoadGroup->PackagesWithExportsToSerializeCount = 1;
+					}
 				}
 
 #if ALT2_ENABLE_LINKERLOAD_SUPPORT
