AsyncLoading2 - set correct inline bulk data offset

#rnx

diff --git a/Engine/Source/Runtime/CoreUObject/Private/Serialization/AsyncLoading2.cpp b/Engine/Source/Runtime/CoreUObject/Private/Serialization/AsyncLoading2.cpp
--- a/Engine/Source/Runtime/CoreUObject/Private/Serialization/AsyncLoading2.cpp
+++ b/Engine/Source/Runtime/CoreUObject/Private/Serialization/AsyncLoading2.cpp
@@ -1765,9 +1765,9 @@ public:
 		{
 			if (int64 PayloadSize = Meta.GetSize(); PayloadSize > 0 && Meta.HasAnyFlags(BULKDATA_Unused) == false)
 			{
-				// Make inline bulk data reloadable
-				const int64 ExportBundleOffset = (ActiveFPLB->StartFastPathLoadBuffer - ActiveFPLB->OriginalFastPathLoadBuffer);
-				Meta.SetOffset(ExportBundleOffset);
+				// Set the offset from the beginning of the I/O chunk in order to make inline bulk data reloadable
+				const int64 ExportBundleChunkOffset = (ActiveFPLB->StartFastPathLoadBuffer - ExportBundleChunkData);
+				Meta.SetOffset(ExportBundleChunkOffset);
 				BulkData.BulkChunkId = CreateIoChunkId(PackageId.Value(), ChunkIndex, EIoChunkType::ExportBundleData);
 				Serialize(BulkData.ReallocateData(PayloadSize), PayloadSize);
 			}
@@ -1869,6 +1869,7 @@ private:
 	uint64 BufferSerialOffset = 0;
 	bool bIsOptionalSegment = false;
 	bool bExportsCookedToSeparateArchive = false;
+	const uint8* ExportBundleChunkData = nullptr;
 
 	void FixupSoftObjectPathForInstancedPackage(FSoftObjectPath& InOutSoftObjectPath)
 	{
@@ -5523,6 +5524,7 @@ EEventLoadNodeExecutionResult FAsyncPackage2::Event_ProcessExportBundle(FAsyncLo
 			Ar.InstanceContext = Package->InstanceContext.Get();
 			Ar.bIsOptionalSegment = bIsOptionalSegment;
 			Ar.bExportsCookedToSeparateArchive = Ar.UEVer() >= EUnrealEngineObjectUE5Version::DATA_RESOURCES;			
+			Ar.ExportBundleChunkData = SerializationState->IoRequest.GetResultOrDie().Data();
 		}
 
 		while (Package->ExportBundleEntryIndex < int32(ExportBundle->EntryCount))
