Added FTopLevelAssetPathFastLess and FTopLevelAssetPathLexicalLess function objects for performing comparisons on FTopLevelAssetPath

diff --git a/Engine/Source/Runtime/CoreUObject/Public/UObject/TopLevelAssetPath.h b/Engine/Source/Runtime/CoreUObject/Public/UObject/TopLevelAssetPath.h
--- a/Engine/Source/Runtime/CoreUObject/Public/UObject/TopLevelAssetPath.h
+++ b/Engine/Source/Runtime/CoreUObject/Public/UObject/TopLevelAssetPath.h
@@ -184,3 +184,21 @@ inline FStringBuilderBase& operator<<(FStringBuilderBase& Builder, const FTopLev
 	Path.AppendString(Builder);
 	return Builder;
 }
+
+/** Fast non-alphabetical order that is only stable during this process' lifetime */
+struct FTopLevelAssetPathFastLess
+{
+	FORCEINLINE bool operator()(const FTopLevelAssetPath& A, const FTopLevelAssetPath& B) const
+	{
+		return A.CompareFast(B) < 0;
+	}
+};
+
+/** Slow alphabetical order that is stable / deterministic over process runs */
+struct FTopLevelAssetPathLexicalLess
+{
+	FORCEINLINE bool operator()(const FTopLevelAssetPath& A, const FTopLevelAssetPath& B) const
+	{
+		return A.Compare(B) < 0;
+	}
+};
