Fix an issue where async loading leftover after tick timeout were mistaken for objects being loaded during postload

#rnx

UDN: https://epicgames.lightning.force.com/lightning/r/Case/5004z00001m9hZeAAI/view

diff --git a/Engine/Source/Runtime/CoreUObject/Private/Serialization/AsyncLoading.cpp b/Engine/Source/Runtime/CoreUObject/Private/Serialization/AsyncLoading.cpp
--- a/Engine/Source/Runtime/CoreUObject/Private/Serialization/AsyncLoading.cpp
+++ b/Engine/Source/Runtime/CoreUObject/Private/Serialization/AsyncLoading.cpp
@@ -10,6 +10,7 @@
 #include "HAL/Event.h"
 #include "HAL/RunnableThread.h"
 #include "Misc/PackageSegment.h"
+#include "Misc/ScopeExit.h"
 #include "Misc/ScopeLock.h"
 #include "Stats/StatsMisc.h"
 #include "Misc/CoreStats.h"
@@ -6970,6 +6971,13 @@ EAsyncPackageState::Type FAsyncPackage::PostLoadDeferredObjects(double InTickSta
 	TArray<UObject*>& ObjLoadedInPostLoad = LoadContext->PRIVATE_GetObjectsLoadedInternalUseOnly();
 	TArray<UObject*> ObjLoadedInPostLoadLocal;
 
+	// It may happen that the last tick timed out before having time to preload some objects
+	// and these will continue on the next tick. But we don't want to mistakenly think
+	// they were loaded as part of the postload we're going to execute in this function.
+	// So clear the list now and restore it for the next tick when we exit this function.
+	TArray<UObject*> LastTickLeftOver = MoveTemp(ObjLoadedInPostLoad);
+	ON_SCOPE_EXIT{ ObjLoadedInPostLoad.Append(LastTickLeftOver); };
+
 	STAT(double PostLoadStartTime = FPlatformTime::Seconds());
 
 	{
