Disable buffer overrun warning false positive in order to support FastGenPGO.

#rnx

diff --git a/Engine/Source/Runtime/Core/Public/GenericPlatform/GenericPlatformCompilerPreSetup.h b/Engine/Source/Runtime/Core/Public/GenericPlatform/GenericPlatformCompilerPreSetup.h
--- a/Engine/Source/Runtime/Core/Public/GenericPlatform/GenericPlatformCompilerPreSetup.h
+++ b/Engine/Source/Runtime/Core/Public/GenericPlatform/GenericPlatformCompilerPreSetup.h
@@ -36,3 +36,11 @@
 #ifndef PRAGMA_DEFAULT_VISIBILITY_END
 	#define PRAGMA_DEFAULT_VISIBILITY_END
 #endif
+
+#ifndef PRAGMA_DISABLE_BUFFER_OVERRUN_WARNING
+	#define PRAGMA_DISABLE_BUFFER_OVERRUN_WARNING
+#endif
+
+#ifndef PRAGMA_ENABLE_BUFFER_OVERRUN_WARNING
+	#define PRAGMA_ENABLE_BUFFER_OVERRUN_WARNING
+#endif
\ No newline at end of file
diff --git a/Engine/Source/Runtime/Core/Public/Serialization/BitReader.h b/Engine/Source/Runtime/Core/Public/Serialization/BitReader.h
--- a/Engine/Source/Runtime/Core/Public/Serialization/BitReader.h
+++ b/Engine/Source/Runtime/Core/Public/Serialization/BitReader.h
@@ -40,10 +40,8 @@ public:
 	/** Equivalent to SetData (reset position, copy from Src into internal buffer), but uses Reset not Empty to avoid a realloc if possible. */
 	void ResetData(FBitReader& Src, int64 CountBits);
 
-#if defined(_MSC_VER) && PLATFORM_DESKTOP
-#pragma warning( push )
-#pragma warning( disable : 4789 )	// Windows PGO (LTCG) is causing nonsensical errors in certain build environments
-#endif
+// Disable false positive buffer overrun warning during pgoprofile linking step
+PRAGMA_DISABLE_BUFFER_OVERRUN_WARNING
 	FORCEINLINE_DEBUGGABLE void SerializeBits( void* Dest, int64 LengthBits )
 	{
 		//@TODO: FLOATPRECISION: This function/class pretends it is 64-bit aware, e.g., in the type of LengthBits and the Pos member, but it is not as appBitsCpy is only 32 bits, the inner Buffer is a 32 bit TArray, etc...
@@ -74,9 +72,7 @@ public:
 			Pos += LengthBits;
 		}
 	}
-#if defined(_MSC_VER) && PLATFORM_DESKTOP
-#pragma warning( pop )
-#endif
+PRAGMA_ENABLE_BUFFER_OVERRUN_WARNING
 
 	virtual void SerializeBitsWithOffset( void* Dest, int32 DestBit, int64 LengthBits ) override;
 
diff --git a/Engine/Source/Runtime/Core/Public/Windows/WindowsPlatformCompilerPreSetup.h b/Engine/Source/Runtime/Core/Public/Windows/WindowsPlatformCompilerPreSetup.h
--- a/Engine/Source/Runtime/Core/Public/Windows/WindowsPlatformCompilerPreSetup.h
+++ b/Engine/Source/Runtime/Core/Public/Windows/WindowsPlatformCompilerPreSetup.h
@@ -63,6 +63,17 @@
 			__pragma (warning(pop))
 	#endif // DISABLE_DEPRECATION
 
+	#ifndef PRAGMA_DISABLE_BUFFER_OVERRUN_WARNING
+		#define PRAGMA_DISABLE_BUFFER_OVERRUN_WARNING \
+			__pragma (warning(push)) \
+			__pragma (warning(disable:4789)) /* buffer 'identifier' of size N bytes will be overrun; M bytes will be written starting at offset L */
+	#endif
+
+	#ifndef PRAGMA_ENABLE_BUFFER_OVERRUN_WARNING
+		#define PRAGMA_ENABLE_BUFFER_OVERRUN_WARNING \
+			__pragma (warning(pop))
+	#endif
+
 	#ifndef PRAGMA_DISABLE_SHADOW_VARIABLE_WARNINGS
 		#define PRAGMA_DISABLE_SHADOW_VARIABLE_WARNINGS \
 			__pragma (warning(push)) \
diff --git a/Engine/Source/Runtime/CoreUObject/Private/UObject/ScriptCore.cpp b/Engine/Source/Runtime/CoreUObject/Private/UObject/ScriptCore.cpp
--- a/Engine/Source/Runtime/CoreUObject/Private/UObject/ScriptCore.cpp
+++ b/Engine/Source/Runtime/CoreUObject/Private/UObject/ScriptCore.cpp
@@ -135,17 +135,6 @@ static struct F##inst##Registrar \
 	IMPLEMENT_FUNCTION(func) \
 	static uint8 UObject##func##BytecodeTemp = GRegisterNative( BytecodeIndex, &UObject::func );
 
-
-#ifndef FORCE_NO_INLINE_VM_FUNCTION
-	#define FORCE_NO_INLINE_VM_FUNCTION 0
-#endif
-
-#if FORCE_NO_INLINE_VM_FUNCTION
-	#define DEFINE_VM_FUNCTION FORCENOINLINE DEFINE_FUNCTION
-#else
-	#define DEFINE_VM_FUNCTION DEFINE_FUNCTION
-#endif
-
 //////////////////////////////////////////////////////////////////////////
 // FBlueprintCoreDelegates
 
@@ -1046,7 +1035,7 @@ void ProcessScriptFunction(UObject* Context, UFunction* Function, FFrame& Stack,
 	Stack.bAbortingExecution |= NewStack.bAbortingExecution;
 }
 
-DEFINE_VM_FUNCTION(UObject::execCallMathFunction)
+DEFINE_FUNCTION(UObject::execCallMathFunction)
 {
 	UFunction* Function = (UFunction*)Stack.ReadObject();
 	checkSlow(Function);
@@ -2155,7 +2144,6 @@ void UObject::ProcessEvent( UFunction* Function, void* Parms )
 #pragma warning (pop)
 #endif
 
-
 DEFINE_FUNCTION(UObject::execUndefined)
 {
 	const FText LocalizedErrorMessage = FText::Format(
@@ -2167,7 +2155,7 @@ DEFINE_FUNCTION(UObject::execUndefined)
 	Stack.Log(ELogVerbosity::Error, LocalizedErrorMessage.ToString());
 }
 
-DEFINE_VM_FUNCTION(UObject::execLocalVariable)
+DEFINE_FUNCTION(UObject::execLocalVariable)
 {
 	checkSlow(Stack.Object == P_THIS);
 	checkSlow(Stack.Locals != NULL);
@@ -2194,7 +2182,7 @@ DEFINE_VM_FUNCTION(UObject::execLocalVariable)
 }
 IMPLEMENT_VM_FUNCTION( EX_LocalVariable, execLocalVariable );
 
-DEFINE_VM_FUNCTION(UObject::execInstanceVariable)
+DEFINE_FUNCTION(UObject::execInstanceVariable)
 {
 	FProperty* VarProperty = (FProperty*)Stack.ReadObject();
 	Stack.MostRecentProperty = VarProperty;
@@ -2221,7 +2209,7 @@ DEFINE_VM_FUNCTION(UObject::execInstanceVariable)
 }
 IMPLEMENT_VM_FUNCTION( EX_InstanceVariable, execInstanceVariable );
 
-DEFINE_VM_FUNCTION(UObject::execClassSparseDataVariable)
+DEFINE_FUNCTION(UObject::execClassSparseDataVariable)
 {
 	FProperty* VarProperty = (FProperty*)Stack.ReadObject();
 	Stack.MostRecentProperty = VarProperty;
@@ -2248,7 +2236,7 @@ DEFINE_VM_FUNCTION(UObject::execClassSparseDataVariable)
 }
 IMPLEMENT_VM_FUNCTION(EX_ClassSparseDataVariable, execClassSparseDataVariable);
 
-DEFINE_VM_FUNCTION(UObject::execDefaultVariable)
+DEFINE_FUNCTION(UObject::execDefaultVariable)
 {
 	FProperty* VarProperty = (FProperty*)Stack.ReadObject();
 	Stack.MostRecentProperty = VarProperty;
@@ -2291,7 +2279,7 @@ DEFINE_VM_FUNCTION(UObject::execDefaultVariable)
 }
 IMPLEMENT_VM_FUNCTION( EX_DefaultVariable, execDefaultVariable );
 
-DEFINE_VM_FUNCTION(UObject::execLocalOutVariable)
+DEFINE_FUNCTION(UObject::execLocalOutVariable)
 {
 	checkSlow(Stack.Object == P_THIS);
 
@@ -2316,7 +2304,7 @@ DEFINE_VM_FUNCTION(UObject::execLocalOutVariable)
 }
 IMPLEMENT_VM_FUNCTION(EX_LocalOutVariable, execLocalOutVariable);
 
-DEFINE_VM_FUNCTION(UObject::execInterfaceContext)
+DEFINE_FUNCTION(UObject::execInterfaceContext)
 {
 	// get the value of the interface variable
 	FScriptInterface InterfaceValue;
@@ -2330,7 +2318,7 @@ DEFINE_VM_FUNCTION(UObject::execInterfaceContext)
 }
 IMPLEMENT_VM_FUNCTION( EX_InterfaceContext, execInterfaceContext );
 
-DEFINE_VM_FUNCTION(UObject::execClassContext)
+DEFINE_FUNCTION(UObject::execClassContext)
 {
 	// Get class expression.
 	UClass* ClassContext = NULL;
@@ -2381,7 +2369,7 @@ DEFINE_VM_FUNCTION(UObject::execClassContext)
 }
 IMPLEMENT_VM_FUNCTION( EX_ClassContext, execClassContext );
 
-DEFINE_VM_FUNCTION(UObject::execEndOfScript)
+DEFINE_FUNCTION(UObject::execEndOfScript)
 {
 #if WITH_EDITOR
 	if (GIsEditor)
@@ -2399,19 +2387,19 @@ DEFINE_VM_FUNCTION(UObject::execEndOfScript)
 }
 IMPLEMENT_VM_FUNCTION( EX_EndOfScript, execEndOfScript );
 
-DEFINE_VM_FUNCTION(UObject::execNothing)
+DEFINE_FUNCTION(UObject::execNothing)
 {
 	// Do nothing.
 }
 IMPLEMENT_VM_FUNCTION( EX_Nothing, execNothing );
 
-DEFINE_VM_FUNCTION(UObject::execNothingOp4a)
+DEFINE_FUNCTION(UObject::execNothingOp4a)
 {
 	// Do nothing.
 }
 IMPLEMENT_VM_FUNCTION( EX_DeprecatedOp4A, execNothingOp4a );
 
-DEFINE_VM_FUNCTION(UObject::execBreakpoint)
+DEFINE_FUNCTION(UObject::execBreakpoint)
 {
 	if (FBlueprintCoreDelegates::IsDebuggingEnabled())
 	{
@@ -2421,7 +2409,7 @@ DEFINE_VM_FUNCTION(UObject::execBreakpoint)
 }
 IMPLEMENT_VM_FUNCTION( EX_Breakpoint, execBreakpoint );
 
-DEFINE_VM_FUNCTION(UObject::execTracepoint)
+DEFINE_FUNCTION(UObject::execTracepoint)
 {
 	if (FBlueprintCoreDelegates::IsDebuggingEnabled())
 	{
@@ -2431,7 +2419,7 @@ DEFINE_VM_FUNCTION(UObject::execTracepoint)
 }
 IMPLEMENT_VM_FUNCTION( EX_Tracepoint, execTracepoint );
 
-DEFINE_VM_FUNCTION(UObject::execWireTracepoint)
+DEFINE_FUNCTION(UObject::execWireTracepoint)
 {
 	if (FBlueprintCoreDelegates::IsDebuggingEnabled())
 	{
@@ -2441,7 +2429,7 @@ DEFINE_VM_FUNCTION(UObject::execWireTracepoint)
 }
 IMPLEMENT_VM_FUNCTION( EX_WireTracepoint, execWireTracepoint );
 
-DEFINE_VM_FUNCTION(UObject::execInstrumentation)
+DEFINE_FUNCTION(UObject::execInstrumentation)
 {
 #if !UE_BUILD_SHIPPING
 	const EScriptInstrumentation::Type EventType = static_cast<EScriptInstrumentation::Type>(Stack.PeekCode());
@@ -2482,7 +2470,7 @@ DEFINE_VM_FUNCTION(UObject::execInstrumentation)
 }
 IMPLEMENT_VM_FUNCTION( EX_InstrumentationEvent, execInstrumentation );
 
-DEFINE_VM_FUNCTION(UObject::execEndFunctionParms)
+DEFINE_FUNCTION(UObject::execEndFunctionParms)
 {
 	// For skipping over optional function parms without values specified.
 	Stack.Code--;
@@ -2490,7 +2478,7 @@ DEFINE_VM_FUNCTION(UObject::execEndFunctionParms)
 IMPLEMENT_VM_FUNCTION( EX_EndFunctionParms, execEndFunctionParms );
 
 
-DEFINE_VM_FUNCTION(UObject::execJump)
+DEFINE_FUNCTION(UObject::execJump)
 {
 	CHECK_RUNAWAY;
 
@@ -2500,7 +2488,7 @@ DEFINE_VM_FUNCTION(UObject::execJump)
 }
 IMPLEMENT_VM_FUNCTION( EX_Jump, execJump );
 
-DEFINE_VM_FUNCTION(UObject::execComputedJump)
+DEFINE_FUNCTION(UObject::execComputedJump)
 {
 	CHECK_RUNAWAY;
 
@@ -2514,7 +2502,7 @@ DEFINE_VM_FUNCTION(UObject::execComputedJump)
 }
 IMPLEMENT_VM_FUNCTION( EX_ComputedJump, execComputedJump );
 	
-DEFINE_VM_FUNCTION(UObject::execJumpIfNot)
+DEFINE_FUNCTION(UObject::execJumpIfNot)
 {
 	CHECK_RUNAWAY;
 
@@ -2533,7 +2521,7 @@ DEFINE_VM_FUNCTION(UObject::execJumpIfNot)
 }
 IMPLEMENT_VM_FUNCTION( EX_JumpIfNot, execJumpIfNot );
 
-DEFINE_VM_FUNCTION(UObject::execAssert)
+DEFINE_FUNCTION(UObject::execAssert)
 {
 	// Get line number.
 	int32 wLine = Stack.ReadWord();
@@ -2561,7 +2549,7 @@ DEFINE_VM_FUNCTION(UObject::execAssert)
 }
 IMPLEMENT_VM_FUNCTION( EX_Assert, execAssert );
 
-DEFINE_VM_FUNCTION(UObject::execPushExecutionFlow)
+DEFINE_FUNCTION(UObject::execPushExecutionFlow)
 {
 	// Read a code offset and push it onto the flow stack
 	CodeSkipSizeType Offset = Stack.ReadCodeSkipCount();
@@ -2569,7 +2557,7 @@ DEFINE_VM_FUNCTION(UObject::execPushExecutionFlow)
 }
 IMPLEMENT_VM_FUNCTION( EX_PushExecutionFlow, execPushExecutionFlow );
 
-DEFINE_VM_FUNCTION(UObject::execPopExecutionFlow)
+DEFINE_FUNCTION(UObject::execPopExecutionFlow)
 {
 	// Since this is a branch function, check for runaway script execution
 	CHECK_RUNAWAY;
@@ -2588,7 +2576,7 @@ DEFINE_VM_FUNCTION(UObject::execPopExecutionFlow)
 }
 IMPLEMENT_VM_FUNCTION( EX_PopExecutionFlow, execPopExecutionFlow );
 
-DEFINE_VM_FUNCTION(UObject::execPopExecutionFlowIfNot)
+DEFINE_FUNCTION(UObject::execPopExecutionFlowIfNot)
 {
 	// Since this is a branch function, check for runaway script execution
 	CHECK_RUNAWAY;
@@ -2614,7 +2602,7 @@ DEFINE_VM_FUNCTION(UObject::execPopExecutionFlowIfNot)
 }
 IMPLEMENT_VM_FUNCTION( EX_PopExecutionFlowIfNot, execPopExecutionFlowIfNot );
 
-DEFINE_VM_FUNCTION(UObject::execLetValueOnPersistentFrame)
+DEFINE_FUNCTION(UObject::execLetValueOnPersistentFrame)
 {
 #if USE_UBER_GRAPH_PERSISTENT_FRAME
 	Stack.MostRecentProperty = nullptr;
@@ -2636,7 +2624,7 @@ DEFINE_VM_FUNCTION(UObject::execLetValueOnPersistentFrame)
 }
 IMPLEMENT_VM_FUNCTION(EX_LetValueOnPersistentFrame, execLetValueOnPersistentFrame);
 
-DEFINE_VM_FUNCTION(UObject::execSwitchValue)
+DEFINE_FUNCTION(UObject::execSwitchValue)
 {
 	const int32 NumCases = Stack.ReadWord();
 	const CodeSkipSizeType OffsetToEnd = Stack.ReadCodeSkipCount();
@@ -2705,7 +2693,7 @@ DEFINE_VM_FUNCTION(UObject::execSwitchValue)
 }
 IMPLEMENT_VM_FUNCTION(EX_SwitchValue, execSwitchValue);
 
-DEFINE_VM_FUNCTION(UObject::execArrayGetByRef)
+DEFINE_FUNCTION(UObject::execArrayGetByRef)
 {
 	// Get variable address.
 	Stack.MostRecentPropertyAddress = nullptr;
@@ -2768,7 +2756,7 @@ DEFINE_VM_FUNCTION(UObject::execArrayGetByRef)
 }
 IMPLEMENT_VM_FUNCTION(EX_ArrayGetByRef, execArrayGetByRef);
 
-DEFINE_VM_FUNCTION(UObject::execLet)
+DEFINE_FUNCTION(UObject::execLet)
 {
 	Stack.MostRecentProperty = nullptr;
 	FProperty* LocallyKnownProperty = Stack.ReadPropertyUnchecked();
@@ -2831,7 +2819,7 @@ DEFINE_VM_FUNCTION(UObject::execLet)
 }
 IMPLEMENT_VM_FUNCTION( EX_Let, execLet );
 
-DEFINE_VM_FUNCTION(UObject::execLetObj)
+DEFINE_FUNCTION(UObject::execLetObj)
 {
 	// Get variable address.
 	Stack.MostRecentPropertyAddress = nullptr;
@@ -2878,7 +2866,7 @@ DEFINE_VM_FUNCTION(UObject::execLetObj)
 }
 IMPLEMENT_VM_FUNCTION( EX_LetObj, execLetObj );
 
-DEFINE_VM_FUNCTION(UObject::execLetWeakObjPtr)
+DEFINE_FUNCTION(UObject::execLetWeakObjPtr)
 {
 	// Get variable address.
 	Stack.MostRecentPropertyAddress = nullptr;
@@ -2926,7 +2914,7 @@ DEFINE_VM_FUNCTION(UObject::execLetWeakObjPtr)
 }
 IMPLEMENT_VM_FUNCTION( EX_LetWeakObjPtr, execLetWeakObjPtr );
 
-DEFINE_VM_FUNCTION(UObject::execLetBool)
+DEFINE_FUNCTION(UObject::execLetBool)
 {
 	Stack.MostRecentPropertyAddress = nullptr;
 	Stack.MostRecentPropertyContainer = nullptr;
@@ -2981,7 +2969,7 @@ DEFINE_VM_FUNCTION(UObject::execLetBool)
 IMPLEMENT_VM_FUNCTION( EX_LetBool, execLetBool );
 
 
-DEFINE_VM_FUNCTION(UObject::execLetDelegate)
+DEFINE_FUNCTION(UObject::execLetDelegate)
 {
 	// Get variable address.
 	Stack.MostRecentPropertyAddress = nullptr;
@@ -3001,7 +2989,7 @@ DEFINE_VM_FUNCTION(UObject::execLetDelegate)
 IMPLEMENT_VM_FUNCTION( EX_LetDelegate, execLetDelegate );
 
 
-DEFINE_VM_FUNCTION(UObject::execLetMulticastDelegate)
+DEFINE_FUNCTION(UObject::execLetMulticastDelegate)
 {
 	// Get variable address.
 	Stack.MostRecentPropertyAddress = nullptr;
@@ -3022,7 +3010,7 @@ DEFINE_VM_FUNCTION(UObject::execLetMulticastDelegate)
 IMPLEMENT_VM_FUNCTION( EX_LetMulticastDelegate, execLetMulticastDelegate );
 
 
-DEFINE_VM_FUNCTION(UObject::execSelf)
+DEFINE_FUNCTION(UObject::execSelf)
 {
 	// Get Self actor for this context.
 	if (RESULT_PARAM != nullptr)
@@ -3042,13 +3030,13 @@ DEFINE_VM_FUNCTION(UObject::execSelf)
 }
 IMPLEMENT_VM_FUNCTION( EX_Self, execSelf );
 
-DEFINE_VM_FUNCTION(UObject::execContext)
+DEFINE_FUNCTION(UObject::execContext)
 {
 	P_THIS->ProcessContextOpcode(Stack, RESULT_PARAM, /*bCanFailSilently=*/ false);
 }
 IMPLEMENT_VM_FUNCTION( EX_Context, execContext );
 
-DEFINE_VM_FUNCTION(UObject::execContext_FailSilent)
+DEFINE_FUNCTION(UObject::execContext_FailSilent)
 {
 	P_THIS->ProcessContextOpcode(Stack, RESULT_PARAM, /*bCanFailSilently=*/ true);
 }
@@ -3134,7 +3122,7 @@ void UObject::ProcessContextOpcode( FFrame& Stack, RESULT_DECL, bool bCanFailSil
 	}
 }
 
-DEFINE_VM_FUNCTION(UObject::execStructMemberContext)
+DEFINE_FUNCTION(UObject::execStructMemberContext)
 {
 	// Get the structure element we care about
 	FProperty* StructProperty = Stack.ReadProperty();
@@ -3178,28 +3166,28 @@ DEFINE_VM_FUNCTION(UObject::execStructMemberContext)
 }
 IMPLEMENT_VM_FUNCTION( EX_StructMemberContext, execStructMemberContext );
 
-DEFINE_VM_FUNCTION(UObject::execVirtualFunction)
+DEFINE_FUNCTION(UObject::execVirtualFunction)
 {
 	// Call the virtual function.
 	P_THIS->CallFunction( Stack, RESULT_PARAM, P_THIS->FindFunctionChecked(Stack.ReadName()) );
 }
 IMPLEMENT_VM_FUNCTION( EX_VirtualFunction, execVirtualFunction );
 
-DEFINE_VM_FUNCTION(UObject::execFinalFunction)
+DEFINE_FUNCTION(UObject::execFinalFunction)
 {
 	// Call the final function.
 	P_THIS->CallFunction( Stack, RESULT_PARAM, (UFunction*)Stack.ReadObject() );
 }
 IMPLEMENT_VM_FUNCTION( EX_FinalFunction, execFinalFunction );
 
-DEFINE_VM_FUNCTION(UObject::execLocalVirtualFunction)
+DEFINE_FUNCTION(UObject::execLocalVirtualFunction)
 {
 	// Call the virtual function.
 	ProcessLocalFunction(Context, P_THIS->FindFunctionChecked(Stack.ReadName()), Stack, RESULT_PARAM);
 }
 IMPLEMENT_VM_FUNCTION( EX_LocalVirtualFunction, execLocalVirtualFunction );
 
-DEFINE_VM_FUNCTION(UObject::execLocalFinalFunction)
+DEFINE_FUNCTION(UObject::execLocalFinalFunction)
 {
 	// Call the final function.
 	ProcessLocalFunction(Context, (UFunction*)Stack.ReadObject(), Stack, RESULT_PARAM);
@@ -3261,13 +3249,13 @@ public:
 	}
 };
 
-DEFINE_VM_FUNCTION(UObject::execCallMulticastDelegate)
+DEFINE_FUNCTION(UObject::execCallMulticastDelegate)
 {
 	FCallDelegateHelper::CallMulticastDelegate(Stack);
 }
 IMPLEMENT_VM_FUNCTION( EX_CallMulticastDelegate, execCallMulticastDelegate );
 
-DEFINE_VM_FUNCTION(UObject::execAddMulticastDelegate)
+DEFINE_FUNCTION(UObject::execAddMulticastDelegate)
 {
 	// Get variable address.
 	Stack.MostRecentPropertyAddress = nullptr;
@@ -3288,7 +3276,7 @@ DEFINE_VM_FUNCTION(UObject::execAddMulticastDelegate)
 }
 IMPLEMENT_VM_FUNCTION( EX_AddMulticastDelegate, execAddMulticastDelegate );
 
-DEFINE_VM_FUNCTION(UObject::execRemoveMulticastDelegate)
+DEFINE_FUNCTION(UObject::execRemoveMulticastDelegate)
 {
 	// Get variable address.
 	Stack.MostRecentPropertyAddress = nullptr;
@@ -3309,7 +3297,7 @@ DEFINE_VM_FUNCTION(UObject::execRemoveMulticastDelegate)
 }
 IMPLEMENT_VM_FUNCTION( EX_RemoveMulticastDelegate, execRemoveMulticastDelegate );
 
-DEFINE_VM_FUNCTION(UObject::execClearMulticastDelegate)
+DEFINE_FUNCTION(UObject::execClearMulticastDelegate)
 {
 	// Get the delegate address
 	Stack.MostRecentPropertyAddress = nullptr;
@@ -3327,44 +3315,47 @@ DEFINE_VM_FUNCTION(UObject::execClearMulticastDelegate)
 }
 IMPLEMENT_VM_FUNCTION( EX_ClearMulticastDelegate, execClearMulticastDelegate );
 
-DEFINE_VM_FUNCTION(UObject::execIntConst)
+DEFINE_FUNCTION(UObject::execIntConst)
 {
 	*(int32*)RESULT_PARAM = Stack.ReadInt<int32>();
 }
 IMPLEMENT_VM_FUNCTION( EX_IntConst, execIntConst );
 
-DEFINE_VM_FUNCTION(UObject::execInt64Const)
+DEFINE_FUNCTION(UObject::execInt64Const)
 {
 	*(int64*)RESULT_PARAM = Stack.ReadInt<int64>();
 }
 IMPLEMENT_VM_FUNCTION(EX_Int64Const, execInt64Const);
 
-DEFINE_VM_FUNCTION(UObject::execUInt64Const)
+DEFINE_FUNCTION(UObject::execUInt64Const)
 {
 	*(uint64*)RESULT_PARAM = Stack.ReadInt<uint64>();
 }
 IMPLEMENT_VM_FUNCTION(EX_UInt64Const, execUInt64Const);
 
-DEFINE_VM_FUNCTION(UObject::execSkipOffsetConst)
+DEFINE_FUNCTION(UObject::execSkipOffsetConst)
 {
 	CodeSkipSizeType Literal = Stack.ReadCodeSkipCount();
 	*(int32*)RESULT_PARAM = Literal;
 }
 IMPLEMENT_VM_FUNCTION( EX_SkipOffsetConst, execSkipOffsetConst );
 
-DEFINE_VM_FUNCTION(UObject::execFloatConst)
+DEFINE_FUNCTION(UObject::execFloatConst)
 {
 	*(float*)RESULT_PARAM = Stack.ReadFloat();
 }
 IMPLEMENT_VM_FUNCTION( EX_FloatConst, execFloatConst );
 
-DEFINE_VM_FUNCTION(UObject::execDoubleConst)
+// Disable false positive buffer overrun warning during pgoprofile linking step
+PRAGMA_DISABLE_BUFFER_OVERRUN_WARNING
+DEFINE_FUNCTION(UObject::execDoubleConst)
 {
 	*(double*)RESULT_PARAM = Stack.ReadInt<double>();
 }
 IMPLEMENT_VM_FUNCTION( EX_DoubleConst, execDoubleConst );
+PRAGMA_ENABLE_BUFFER_OVERRUN_WARNING
 
-DEFINE_VM_FUNCTION(UObject::execStringConst)
+DEFINE_FUNCTION(UObject::execStringConst)
 {
 	*(FString*)RESULT_PARAM = (ANSICHAR*)Stack.Code;
 	while( *Stack.Code )
@@ -3373,7 +3364,7 @@ DEFINE_VM_FUNCTION(UObject::execStringConst)
 }
 IMPLEMENT_VM_FUNCTION( EX_StringConst, execStringConst );
 
-DEFINE_VM_FUNCTION(UObject::execUnicodeStringConst)
+DEFINE_FUNCTION(UObject::execUnicodeStringConst)
 {
 	FString& ResultStr = *(FString*)RESULT_PARAM;
 	ResultStr = FString((UCS2CHAR*)Stack.Code);
@@ -3389,7 +3380,7 @@ DEFINE_VM_FUNCTION(UObject::execUnicodeStringConst)
 }
 IMPLEMENT_VM_FUNCTION( EX_UnicodeStringConst, execUnicodeStringConst );
 
-DEFINE_VM_FUNCTION(UObject::execTextConst)
+DEFINE_FUNCTION(UObject::execTextConst)
 {
 	// What kind of text are we dealing with?
 	const EBlueprintTextLiteralType TextLiteralType = (EBlueprintTextLiteralType)*Stack.Code++;
@@ -3456,19 +3447,19 @@ DEFINE_VM_FUNCTION(UObject::execTextConst)
 }
 IMPLEMENT_VM_FUNCTION( EX_TextConst, execTextConst );
 
-DEFINE_VM_FUNCTION(UObject::execPropertyConst)
+DEFINE_FUNCTION(UObject::execPropertyConst)
 {
 	*(FProperty**)RESULT_PARAM = (FProperty*)Stack.ReadObject();
 }
 IMPLEMENT_VM_FUNCTION(EX_PropertyConst, execPropertyConst);
 
-DEFINE_VM_FUNCTION(UObject::execObjectConst)
+DEFINE_FUNCTION(UObject::execObjectConst)
 {
 	*(UObject**)RESULT_PARAM = (UObject*)Stack.ReadObject();
 }
 IMPLEMENT_VM_FUNCTION( EX_ObjectConst, execObjectConst );
 
-DEFINE_VM_FUNCTION(UObject::execSoftObjectConst)
+DEFINE_FUNCTION(UObject::execSoftObjectConst)
 {
 	FString LongPath;
 	Stack.Step(Stack.Object, &LongPath);
@@ -3476,7 +3467,7 @@ DEFINE_VM_FUNCTION(UObject::execSoftObjectConst)
 }
 IMPLEMENT_VM_FUNCTION( EX_SoftObjectConst, execSoftObjectConst);
 
-DEFINE_VM_FUNCTION(UObject::execFieldPathConst)
+DEFINE_FUNCTION(UObject::execFieldPathConst)
 {
 	FString StringPath;
 	Stack.Step(Stack.Object, &StringPath);
@@ -3486,14 +3477,14 @@ DEFINE_VM_FUNCTION(UObject::execFieldPathConst)
 }
 IMPLEMENT_VM_FUNCTION(EX_FieldPathConst, execFieldPathConst);
 
-DEFINE_VM_FUNCTION(UObject::execInstanceDelegate)
+DEFINE_FUNCTION(UObject::execInstanceDelegate)
 {
 	FName FunctionName = Stack.ReadName();
 	((FScriptDelegate*)RESULT_PARAM)->BindUFunction( (FunctionName == NAME_None) ? NULL : P_THIS, FunctionName );
 }
 IMPLEMENT_VM_FUNCTION( EX_InstanceDelegate, execInstanceDelegate );
 
-DEFINE_VM_FUNCTION(UObject::execBindDelegate)
+DEFINE_FUNCTION(UObject::execBindDelegate)
 {
 	FName FunctionName = Stack.ReadName();
 
@@ -3515,19 +3506,19 @@ DEFINE_VM_FUNCTION(UObject::execBindDelegate)
 }
 IMPLEMENT_VM_FUNCTION( EX_BindDelegate, execBindDelegate );
 
-DEFINE_VM_FUNCTION(UObject::execNameConst)
+DEFINE_FUNCTION(UObject::execNameConst)
 {
 	*(FName*)RESULT_PARAM = Stack.ReadName();
 }
 IMPLEMENT_VM_FUNCTION( EX_NameConst, execNameConst );
 
-DEFINE_VM_FUNCTION(UObject::execByteConst)
+DEFINE_FUNCTION(UObject::execByteConst)
 {
 	*(uint8*)RESULT_PARAM = *Stack.Code++;
 }
 IMPLEMENT_VM_FUNCTION( EX_ByteConst, execByteConst );
 
-DEFINE_VM_FUNCTION(UObject::execRotationConst)
+DEFINE_FUNCTION(UObject::execRotationConst)
 {
 	((FRotator*)RESULT_PARAM)->Pitch = Stack.ReadDouble();
 	((FRotator*)RESULT_PARAM)->Yaw   = Stack.ReadDouble();
@@ -3535,7 +3526,7 @@ DEFINE_VM_FUNCTION(UObject::execRotationConst)
 }
 IMPLEMENT_VM_FUNCTION( EX_RotationConst, execRotationConst );
 
-DEFINE_VM_FUNCTION(UObject::execVectorConst)
+DEFINE_FUNCTION(UObject::execVectorConst)
 {
 	((FVector*)RESULT_PARAM)->X = Stack.ReadDouble();
 	((FVector*)RESULT_PARAM)->Y = Stack.ReadDouble();
@@ -3543,7 +3534,7 @@ DEFINE_VM_FUNCTION(UObject::execVectorConst)
 }
 IMPLEMENT_VM_FUNCTION( EX_VectorConst, execVectorConst );
 
-DEFINE_VM_FUNCTION(UObject::execVector3fConst)
+DEFINE_FUNCTION(UObject::execVector3fConst)
 {
 	((FVector3f*)RESULT_PARAM)->X = Stack.ReadFloat();
 	((FVector3f*)RESULT_PARAM)->Y = Stack.ReadFloat();
@@ -3551,7 +3542,7 @@ DEFINE_VM_FUNCTION(UObject::execVector3fConst)
 }
 IMPLEMENT_VM_FUNCTION(EX_Vector3fConst, execVector3fConst);
 
-DEFINE_VM_FUNCTION(UObject::execTransformConst)
+DEFINE_FUNCTION(UObject::execTransformConst)
 {
 	// Rotation
 	FQuat TmpRotation;
@@ -3576,7 +3567,7 @@ DEFINE_VM_FUNCTION(UObject::execTransformConst)
 }
 IMPLEMENT_VM_FUNCTION( EX_TransformConst, execTransformConst );
 
-DEFINE_VM_FUNCTION(UObject::execStructConst)
+DEFINE_FUNCTION(UObject::execStructConst)
 {
 	UScriptStruct* ScriptStruct = CastChecked<UScriptStruct>(Stack.ReadObject());
 	int32 SerializedSize = Stack.ReadInt<int32>();
@@ -3609,7 +3600,7 @@ DEFINE_VM_FUNCTION(UObject::execStructConst)
 }
 IMPLEMENT_VM_FUNCTION( EX_StructConst, execStructConst );
 
-DEFINE_VM_FUNCTION(UObject::execSetArray)
+DEFINE_FUNCTION(UObject::execSetArray)
 {
 	// Get the array address
 	Stack.MostRecentPropertyAddress = nullptr;
@@ -3633,7 +3624,7 @@ DEFINE_VM_FUNCTION(UObject::execSetArray)
 }
 IMPLEMENT_VM_FUNCTION( EX_SetArray, execSetArray );
 
-DEFINE_VM_FUNCTION(UObject::execSetSet)
+DEFINE_FUNCTION(UObject::execSetSet)
 {
 	// Get the set address
 	Stack.MostRecentPropertyAddress = nullptr;
@@ -3667,7 +3658,7 @@ DEFINE_VM_FUNCTION(UObject::execSetSet)
 }
 IMPLEMENT_VM_FUNCTION( EX_SetSet, execSetSet );
 
-DEFINE_VM_FUNCTION(UObject::execSetMap)
+DEFINE_FUNCTION(UObject::execSetMap)
 {
 	// Get the map address
 	Stack.MostRecentPropertyAddress = nullptr;
@@ -3702,7 +3693,7 @@ DEFINE_VM_FUNCTION(UObject::execSetMap)
 }
 IMPLEMENT_VM_FUNCTION( EX_SetMap, execSetMap );
 
-DEFINE_VM_FUNCTION(UObject::execArrayConst)
+DEFINE_FUNCTION(UObject::execArrayConst)
 {
 	FProperty* InnerProperty = CastFieldChecked<FProperty>((FField*)Stack.ReadObject());
 	int32 Num = Stack.ReadInt<int32>();
@@ -3722,7 +3713,7 @@ DEFINE_VM_FUNCTION(UObject::execArrayConst)
 }
 IMPLEMENT_VM_FUNCTION(EX_ArrayConst, execArrayConst);
 
-DEFINE_VM_FUNCTION(UObject::execSetConst)
+DEFINE_FUNCTION(UObject::execSetConst)
 {
 	FProperty* InnerProperty = CastFieldChecked<FProperty>((FField*)Stack.ReadObject());
 	int32 Num = Stack.ReadInt<int32>();
@@ -3742,7 +3733,7 @@ DEFINE_VM_FUNCTION(UObject::execSetConst)
 }
 IMPLEMENT_VM_FUNCTION(EX_SetConst, execSetConst);
 
-DEFINE_VM_FUNCTION(UObject::execMapConst)
+DEFINE_FUNCTION(UObject::execMapConst)
 {
 	FProperty* KeyProperty = CastFieldChecked<FProperty>((FField*)Stack.ReadObject());
 	FProperty* ValProperty = CastFieldChecked<FProperty>((FField*)Stack.ReadObject());
@@ -3764,51 +3755,51 @@ DEFINE_VM_FUNCTION(UObject::execMapConst)
 }
 IMPLEMENT_VM_FUNCTION(EX_MapConst, execMapConst);
 
-DEFINE_VM_FUNCTION(UObject::execIntZero)
+DEFINE_FUNCTION(UObject::execIntZero)
 {
 	*(int32*)RESULT_PARAM = 0;
 }
 IMPLEMENT_VM_FUNCTION( EX_IntZero, execIntZero );
 
-DEFINE_VM_FUNCTION(UObject::execIntOne)
+DEFINE_FUNCTION(UObject::execIntOne)
 {
 	*(int32*)RESULT_PARAM = 1;
 }
 IMPLEMENT_VM_FUNCTION( EX_IntOne, execIntOne );
 
-DEFINE_VM_FUNCTION(UObject::execTrue)
+DEFINE_FUNCTION(UObject::execTrue)
 {
 	*(bool*)RESULT_PARAM = true;
 }
 IMPLEMENT_VM_FUNCTION( EX_True, execTrue );
 
-DEFINE_VM_FUNCTION(UObject::execFalse)
+DEFINE_FUNCTION(UObject::execFalse)
 {
 	*(bool*)RESULT_PARAM = false;
 }
 IMPLEMENT_VM_FUNCTION( EX_False, execFalse );
 
-DEFINE_VM_FUNCTION(UObject::execNoObject)
+DEFINE_FUNCTION(UObject::execNoObject)
 {
 	*(UObject**)RESULT_PARAM = NULL;
 }
 IMPLEMENT_VM_FUNCTION( EX_NoObject, execNoObject );
 
-DEFINE_VM_FUNCTION(UObject::execNullInterface)
+DEFINE_FUNCTION(UObject::execNullInterface)
 {
 	FScriptInterface& InterfaceValue = *(FScriptInterface*)RESULT_PARAM;
 	InterfaceValue.SetObject(nullptr);
 }
 IMPLEMENT_VM_FUNCTION( EX_NoInterface, execNullInterface );
 
-DEFINE_VM_FUNCTION(UObject::execIntConstByte)
+DEFINE_FUNCTION(UObject::execIntConstByte)
 {
 	*(int32*)RESULT_PARAM = *Stack.Code++;
 }
 IMPLEMENT_VM_FUNCTION( EX_IntConstByte, execIntConstByte );
 
 
-DEFINE_VM_FUNCTION(UObject::execDynamicCast)
+DEFINE_FUNCTION(UObject::execDynamicCast)
 {
 	// Get "to cast to" class for the dynamic actor class
 	UClass* ClassPtr = (UClass *)Stack.ReadObject();
@@ -3855,7 +3846,7 @@ DEFINE_VM_FUNCTION(UObject::execDynamicCast)
 }
 IMPLEMENT_VM_FUNCTION( EX_DynamicCast, execDynamicCast );
 
-DEFINE_VM_FUNCTION(UObject::execMetaCast)
+DEFINE_FUNCTION(UObject::execMetaCast)
 {
 	UClass* MetaClass = (UClass*)Stack.ReadObject();
 
@@ -3867,20 +3858,20 @@ DEFINE_VM_FUNCTION(UObject::execMetaCast)
 }
 IMPLEMENT_VM_FUNCTION( EX_MetaCast, execMetaCast );
 
-DEFINE_VM_FUNCTION(UObject::execCast)
+DEFINE_FUNCTION(UObject::execCast)
 {
 	int32 B = *(Stack.Code)++;
 	(*GCasts[B])( Stack.Object, Stack, RESULT_PARAM );
 }
 IMPLEMENT_VM_FUNCTION( EX_Cast, execCast );
 
-DEFINE_VM_FUNCTION(UObject::execInterfaceCast)
+DEFINE_FUNCTION(UObject::execInterfaceCast)
 {
 	(*GCasts[CST_ObjectToInterface])(Stack.Object, Stack, RESULT_PARAM);
 }
 IMPLEMENT_VM_FUNCTION( EX_ObjToInterfaceCast, execInterfaceCast );
 
-DEFINE_VM_FUNCTION(UObject::execDoubleToFloatCast)
+DEFINE_FUNCTION(UObject::execDoubleToFloatCast)
 {
 	if (Stack.StepAndCheckMostRecentProperty(Stack.Object, nullptr))
 	{
@@ -3896,7 +3887,7 @@ DEFINE_VM_FUNCTION(UObject::execDoubleToFloatCast)
 }
 IMPLEMENT_CAST_FUNCTION( CST_DoubleToFloat, execDoubleToFloatCast )
 
-DEFINE_VM_FUNCTION(UObject::execFloatToDoubleCast)
+DEFINE_FUNCTION(UObject::execFloatToDoubleCast)
 {
 	if (Stack.StepAndCheckMostRecentProperty(Stack.Object, nullptr))
 	{
@@ -3912,7 +3903,7 @@ DEFINE_VM_FUNCTION(UObject::execFloatToDoubleCast)
 }
 IMPLEMENT_CAST_FUNCTION( CST_FloatToDouble, execFloatToDoubleCast )
 
-DEFINE_VM_FUNCTION(UObject::execObjectToBool)
+DEFINE_FUNCTION(UObject::execObjectToBool)
 {
 	UObject* Obj=NULL;
 	Stack.Step( Stack.Object, &Obj );
@@ -3920,7 +3911,7 @@ DEFINE_VM_FUNCTION(UObject::execObjectToBool)
 }
 IMPLEMENT_CAST_FUNCTION( CST_ObjectToBool, execObjectToBool );
 
-DEFINE_VM_FUNCTION(UObject::execInterfaceToBool)
+DEFINE_FUNCTION(UObject::execInterfaceToBool)
 {
 	FScriptInterface Interface;
 	Stack.Step( Stack.Object, &Interface);
@@ -3928,7 +3919,7 @@ DEFINE_VM_FUNCTION(UObject::execInterfaceToBool)
 }
 IMPLEMENT_CAST_FUNCTION( CST_InterfaceToBool, execInterfaceToBool );
 
-DEFINE_VM_FUNCTION(UObject::execObjectToInterface)
+DEFINE_FUNCTION(UObject::execObjectToInterface)
 {
 	FScriptInterface& InterfaceValue = *(FScriptInterface*)RESULT_PARAM;
 
@@ -3954,7 +3945,7 @@ DEFINE_VM_FUNCTION(UObject::execObjectToInterface)
 }
 IMPLEMENT_CAST_FUNCTION( CST_ObjectToInterface, execObjectToInterface );
 
-DEFINE_VM_FUNCTION(UObject::execInterfaceToInterface)
+DEFINE_FUNCTION(UObject::execInterfaceToInterface)
 {
 	FScriptInterface& CastResult = *(FScriptInterface*)RESULT_PARAM;
 
@@ -3982,7 +3973,7 @@ DEFINE_VM_FUNCTION(UObject::execInterfaceToInterface)
 }
 IMPLEMENT_VM_FUNCTION( EX_CrossInterfaceCast, execInterfaceToInterface );
 
-DEFINE_VM_FUNCTION(UObject::execInterfaceToObject)
+DEFINE_FUNCTION(UObject::execInterfaceToObject)
 {
 	// read the interface class off the stack
 	UClass* ObjClassToCastTo = dynamic_cast<UClass*>(Stack.ReadObject());
