Added additional information to determine whether or not validation is done during cooking.

diff --git a/Engine/Source/Runtime/CoreUObject/Private/UObject/SavePackage2.cpp b/Engine/Source/Runtime/CoreUObject/Private/UObject/SavePackage2.cpp
--- a/Engine/Source/Runtime/CoreUObject/Private/UObject/SavePackage2.cpp
+++ b/Engine/Source/Runtime/CoreUObject/Private/UObject/SavePackage2.cpp
@@ -566,9 +566,13 @@ ESavePackageResult ValidateExports(FSaveContext& SaveContext)
 		TSet<UObject*> Exports;
 		Algo::Transform(SaveContext.GetExports(), Exports, [](const FTaggedExport& InExport) { return InExport.Obj; });
 
+		FExportsValidationContext::EFlags Flags = SaveContext.IsCooking() 
+			? FExportsValidationContext::EFlags::IsCooking
+			: FExportsValidationContext::EFlags::None;
+		FOutputDevice* OutputDevice = SaveContext.IsGenerateSaveError() ? SaveContext.GetError() : nullptr;
 		for (const TFunction<FSavePackageSettings::ExternalExportValidationFunc>& ValidateExport : SaveContext.GetExternalExportValidations())
 		{
-			SaveContext.Result = ValidateExport({ SaveContext.GetPackage(), Exports, SaveContext.IsGenerateSaveError() ? SaveContext.GetError() : nullptr });
+			SaveContext.Result = ValidateExport({ SaveContext.GetPackage(), Exports, Flags, OutputDevice});
 			if (SaveContext.Result != ESavePackageResult::Success)
 			{
 				return SaveContext.Result;
diff --git a/Engine/Source/Runtime/CoreUObject/Public/UObject/SavePackage.h b/Engine/Source/Runtime/CoreUObject/Public/UObject/SavePackage.h
--- a/Engine/Source/Runtime/CoreUObject/Public/UObject/SavePackage.h
+++ b/Engine/Source/Runtime/CoreUObject/Public/UObject/SavePackage.h
@@ -11,6 +11,7 @@
 #include "Delegates/Delegate.h"
 #include "Logging/LogMacros.h"
 #include "Misc/DateTime.h"
+#include "Misc/EnumClassFlags.h"
 #include "Misc/Optional.h"
 #include "Misc/OutputDeviceError.h"
 #include "ObjectMacros.h"
@@ -157,16 +158,26 @@ struct FImportsValidationContext
 /** Param struct for external export validation functions */
 struct FExportsValidationContext
 {
-	FExportsValidationContext(const UPackage* InPackage, const TSet<UObject*>& InExports, FOutputDevice* InOutputDevice)
+	enum class EFlags
+	{
+		None		= 0,
+		IsCooking	= 1 << 0,
+	};
+
+	FExportsValidationContext(const UPackage* InPackage, const TSet<UObject*>& InExports, EFlags InFlags,
+		FOutputDevice* InOutputDevice)
 		: Package(InPackage)
 		, Exports(InExports)
+		, Flags(InFlags)
 		, OutputDevice(InOutputDevice)
 	{}
 
 	const UPackage* Package;
 	const TSet<UObject*>& Exports;
+	const EFlags Flags;
 	FOutputDevice* OutputDevice;
 };
+ENUM_CLASS_FLAGS(FExportsValidationContext::EFlags)
 
 /** struct persistent settings used by all save unless overridden. @see FSavePackageContext */
 class FSavePackageSettings
