Added support for ArrayProperty types when applying save overrides during FPackageHarvester::ResolveOverrides.

diff --git a/Engine/Source/Runtime/CoreUObject/Private/UObject/SavePackage/PackageHarvester.cpp b/Engine/Source/Runtime/CoreUObject/Private/UObject/SavePackage/PackageHarvester.cpp
--- a/Engine/Source/Runtime/CoreUObject/Private/UObject/SavePackage/PackageHarvester.cpp
+++ b/Engine/Source/Runtime/CoreUObject/Private/UObject/SavePackage/PackageHarvester.cpp
@@ -594,13 +594,30 @@ void FPackageHarvester::ResolveOverrides()
 		{
 			if (FProperty* Prop = Override.bMarkTransient ? CastField<FProperty>(Override.PropertyPath.GetTyped(FProperty::StaticClass())) : nullptr)
 			{
+				FProperty* InnerProp = Prop;
+				if (FArrayProperty* ArrayProp = CastField<FArrayProperty>(Prop))
+				{
+					InnerProp = ArrayProp->Inner;
+				}
+
 				// We currently only support object property
-				if (ensureAlwaysMsgf(Prop->IsA<FObjectProperty>(), TEXT("Save Overrides supports only object property at the moment. Name: %s, Type: %s"), *Prop->GetName(), *Prop->GetClass()->GetName()))
+				if (ensureAlwaysMsgf(InnerProp && InnerProp->IsA<FObjectProperty>(), TEXT("Save Overrides supports only object properties at the moment. Name: %s, Type: %s"), *Prop->GetName(), *Prop->GetClass()->GetName()))
 				{
 					// Harvest the name of the property, since it is only made transient for the purpose of the package harvest, it will be needed for the LinkerSave
 					HarvestExportDataName(Prop->GetFName());
 
 					Props.Add(Prop);
+
+					if (Prop != InnerProp)
+					{
+						if (Prop->GetFName() != InnerProp->GetFName())
+						{
+							// Harvest the name of the inner property as well, since it is only made transient for the purpose of the package harvest, it will be needed for the LinkerSave
+							HarvestExportDataName(InnerProp->GetFName());
+						}
+
+						Props.Add(InnerProp);
+					}
 				}
 			}
 		}
