#UE Add a little more info to an assert

diff --git a/Engine/Source/Runtime/CoreUObject/Private/UObject/GarbageCollection.cpp b/Engine/Source/Runtime/CoreUObject/Private/UObject/GarbageCollection.cpp
--- a/Engine/Source/Runtime/CoreUObject/Private/UObject/GarbageCollection.cpp
+++ b/Engine/Source/Runtime/CoreUObject/Private/UObject/GarbageCollection.cpp
@@ -5335,7 +5335,7 @@ void UClass::AssembleReferenceTokenStreamInternal(bool bForce)
 			FTokenStreamBuilder::Merge(/* out */ ReferenceTokens.Strong, StrongTokens, SuperStrongTokens);
 		}
 
-		check(!HasAnyClassFlags(CLASS_TokenStreamAssembled)); // recursion here is probably bad
+		checkf(!HasAnyClassFlags(CLASS_TokenStreamAssembled), TEXT("Token stream already assembled for class '%s'"), *GetPathName()); // recursion here is probably bad
 		ClassFlags |= CLASS_TokenStreamAssembled;
 	}
 }
