Object - Giving the GetExtendedAssetRegistryTagsForSave function the same extensibility as the GetAssetRegistryTags function and providing a general delegate for anyone to include extra tags just for save time.  You can now hook UObject::FAssetRegistryTag::OnGetExtendedAssetRegistryTagsForSave, to provide external asset registry tags just for save time, which are difficult and expensive to compute on the fly.

diff --git a/Engine/Source/Runtime/CoreUObject/Private/UObject/Obj.cpp b/Engine/Source/Runtime/CoreUObject/Private/UObject/Obj.cpp
--- a/Engine/Source/Runtime/CoreUObject/Private/UObject/Obj.cpp
+++ b/Engine/Source/Runtime/CoreUObject/Private/UObject/Obj.cpp
@@ -97,6 +97,7 @@ static UPackage*			GObjTransientPkg								= NULL;
 
 #if WITH_EDITOR
 	UObject::FAssetRegistryTag::FOnGetObjectAssetRegistryTags UObject::FAssetRegistryTag::OnGetExtraObjectTags;
+	UObject::FAssetRegistryTag::FOnGetExtendedAssetRegistryTagsForSave UObject::FAssetRegistryTag::OnGetExtendedAssetRegistryTagsForSave;
 #endif // WITH_EDITOR
 
 UObject::UObject( EStaticConstructor, EObjectFlags InFlags )
@@ -2052,6 +2053,9 @@ void UObject::GetAssetRegistryTags(TArray<FAssetRegistryTag>& OutTags) const
 #if WITH_EDITOR
 void UObject::GetExtendedAssetRegistryTagsForSave(const ITargetPlatform* TargetPlatform, TArray<FAssetRegistryTag>& OutTags) const
 {
+	// Notify external sources that we need tags for save.
+	FAssetRegistryTag::OnGetExtendedAssetRegistryTagsForSave.Broadcast(this, TargetPlatform, OutTags);
+
 	PRAGMA_DISABLE_DEPRECATION_WARNINGS
 	GetExternalActorExtendedAssetRegistryTags(OutTags);
 	PRAGMA_ENABLE_DEPRECATION_WARNINGS
diff --git a/Engine/Source/Runtime/CoreUObject/Public/UObject/Object.h b/Engine/Source/Runtime/CoreUObject/Public/UObject/Object.h
--- a/Engine/Source/Runtime/CoreUObject/Public/UObject/Object.h
+++ b/Engine/Source/Runtime/CoreUObject/Public/UObject/Object.h
@@ -841,6 +841,10 @@ public:
 		/** Callback  */
 		DECLARE_MULTICAST_DELEGATE_TwoParams(FOnGetObjectAssetRegistryTags, const UObject* /*Object*/, TArray<FAssetRegistryTag>& /*InOutTags*/);
 		COREUOBJECT_API static FOnGetObjectAssetRegistryTags OnGetExtraObjectTags;
+
+		/** Callback for GetExtendedAssetRegistryTagsForSave */
+		DECLARE_MULTICAST_DELEGATE_ThreeParams(FOnGetExtendedAssetRegistryTagsForSave, const UObject* /*Object*/, const ITargetPlatform* TargetPlatform, TArray<FAssetRegistryTag>& /*InOutTags*/);
+		COREUOBJECT_API static FOnGetExtendedAssetRegistryTagsForSave OnGetExtendedAssetRegistryTagsForSave;
 #endif // WITH_EDITOR
 	};
 
