Fix various int32 overflows when using int32 multiplication to get bytes

diff --git a/Engine/Source/Runtime/CoreUObject/Private/UObject/PropertyArray.cpp b/Engine/Source/Runtime/CoreUObject/Private/UObject/PropertyArray.cpp
--- a/Engine/Source/Runtime/CoreUObject/Private/UObject/PropertyArray.cpp
+++ b/Engine/Source/Runtime/CoreUObject/Private/UObject/PropertyArray.cpp
@@ -729,7 +729,7 @@ void FArrayProperty::CopyValuesInternal( void* Dest, void const* Src, int32 Coun
 	}
 	if (Num)
 	{
-		int32 Size = Inner->ElementSize;
+		size_t Size = Inner->ElementSize;
 		uint8* SrcData = (uint8*)SrcArrayHelper.GetRawPtr();
 		uint8* DestData = (uint8*)DestArrayHelper.GetRawPtr();
 		if( !(Inner->PropertyFlags & CPF_IsPlainOldData) )
diff --git a/Engine/Source/Runtime/CoreUObject/Public/UObject/UnrealType.h b/Engine/Source/Runtime/CoreUObject/Public/UObject/UnrealType.h
--- a/Engine/Source/Runtime/CoreUObject/Public/UObject/UnrealType.h
+++ b/Engine/Source/Runtime/CoreUObject/Public/UObject/UnrealType.h
@@ -507,7 +507,7 @@ public:
 		{
 			for (int32 Idx = 0; Idx < ArrayDim; Idx++)
 			{
-				void *Target = ContainerPtrToValuePtr<void>(Data, Idx);
+				void* Target = ContainerPtrToValuePtr<void>(Data, Idx);
 				void const* Default = ContainerPtrToValuePtrForDefaults<void>(DefaultStruct, DefaultData, Idx);
 				if ( !Identical(Target, Default, UnderlyingArchive.GetPortFlags()) )
 				{
@@ -691,7 +691,7 @@ private:
 			check(!GetOwner<UClass>()); // Check we are _not_ calling this on a direct child property of a UClass, you should pass in a UObject* in that case
 		}
 
-		return (uint8*)ContainerPtr + Offset_Internal + ElementSize * ArrayIndex;
+		return (uint8*)ContainerPtr + Offset_Internal + static_cast<size_t>(ElementSize) * ArrayIndex;
 	}
 
 	FORCEINLINE void* ContainerUObjectPtrToValuePtrInternal(UObject* ContainerPtr, int32 ArrayIndex) const
@@ -718,7 +718,7 @@ private:
 			check(!GetOwner<UClass>()); // Check we are _not_ calling this on a direct child property of a UClass, you should pass in a UObject* in that case
 		}
 
-		return (uint8*)ContainerPtr + Offset_Internal + ElementSize * ArrayIndex;
+		return (uint8*)ContainerPtr + Offset_Internal + static_cast<size_t>(ElementSize) * ArrayIndex;
 	}
 
 protected:
@@ -727,7 +727,7 @@ protected:
 	{
 		if (PropertyPointerType == EPropertyPointerType::Container)
 		{
-			return (uint8*)ContainerOrPropertyPtr + Offset_Internal + ElementSize * ArrayIndex;
+			return (uint8*)ContainerOrPropertyPtr + Offset_Internal + static_cast<size_t>(ElementSize) * ArrayIndex;
 		}
 		else
 		{
@@ -865,7 +865,7 @@ public:
 		{
 			if (PropertyFlags & CPF_IsPlainOldData)
 			{
-				FMemory::Memcpy( Dest, Src, ElementSize * ArrayDim );
+				FMemory::Memcpy( Dest, Src, static_cast<size_t>(ElementSize) * ArrayDim );
 			}
 			else
 			{
@@ -1006,7 +1006,7 @@ public:
 	{
 		if (PropertyFlags & CPF_ZeroConstructor)
 		{
-			FMemory::Memzero(Dest,ElementSize * ArrayDim);
+			FMemory::Memzero(Dest, static_cast<size_t>(ElementSize) * ArrayDim);
 		}
 		else
 		{
@@ -1024,7 +1024,7 @@ public:
 	{
 		if (PropertyFlags & CPF_ZeroConstructor)
 		{
-			FMemory::Memzero(ContainerPtrToValuePtr<void>(Dest),ElementSize * ArrayDim);
+			FMemory::Memzero(ContainerPtrToValuePtr<void>(Dest), static_cast<size_t>(ElementSize) * ArrayDim);
 		}
 		else
 		{
@@ -1098,6 +1098,8 @@ public:
 	 */
 	virtual void EmitReferenceInfo(UE::GC::FTokenStreamBuilder& TokenStream, int32 BaseOffset, TArray<const FStructProperty*>& EncounteredStructProps, FGCStackSizeHelper& StackSizeHelper);
 
+    // @TODO: Surely this can have an int32 overflow. This should probably return size_t. Just
+    // need to audit all callers to make such a change.
 	FORCEINLINE int32 GetSize() const
 	{
 		return ArrayDim * ElementSize;
@@ -1444,14 +1446,14 @@ public:
 	{
 		for (int32 i = 0; i < this->ArrayDim; ++i)
 		{
-			TTypeFundamentals::InitializePropertyValue((uint8*)Dest + i * this->ElementSize);
+			TTypeFundamentals::InitializePropertyValue((uint8*)Dest + i * static_cast<size_t>(this->ElementSize));
 		}
 	}
 	virtual void DestroyValueInternal( void* Dest ) const override
 	{
 		for (int32 i = 0; i < this->ArrayDim; ++i)
 		{
-			TTypeFundamentals::DestroyPropertyValue((uint8*)Dest + i * this->ElementSize);
+			TTypeFundamentals::DestroyPropertyValue((uint8*)Dest + i * static_cast<size_t>(this->ElementSize));
 		}
 	}
 
@@ -3768,14 +3770,14 @@ public:
 
 			for (int32 i = 0; i < this->ArrayDim; ++i)
 			{
-				new ((uint8*)Dest + i * this->ElementSize) FFreezableScriptArray;
+				new ((uint8*)Dest + i * static_cast<size_t>(this->ElementSize)) FFreezableScriptArray;
 			}
 		}
 		else
 		{
 			for (int32 i = 0; i < this->ArrayDim; ++i)
 			{
-				new ((uint8*)Dest + i * this->ElementSize) FScriptArray;
+				new ((uint8*)Dest + i * static_cast<size_t>(this->ElementSize)) FScriptArray;
 			}
 		}
 	}
@@ -3904,14 +3906,14 @@ public:
 
 			for (int32 i = 0; i < this->ArrayDim; ++i)
 			{
-				new ((uint8*)Dest + i * this->ElementSize) FFreezableScriptMap;
+				new ((uint8*)Dest + i * static_cast<size_t>(this->ElementSize)) FFreezableScriptMap;
 			}
 		}
 		else
 		{
 			for (int32 i = 0; i < this->ArrayDim; ++i)
 			{
-				new ((uint8*)Dest + i * this->ElementSize) FScriptMap;
+				new ((uint8*)Dest + i * static_cast<size_t>(this->ElementSize)) FScriptMap;
 			}
 		}
 	}
@@ -4170,7 +4172,7 @@ public:
 			return NULL;
 		}
 		checkSlow(IsValidIndex(Index)); 
-		return (uint8*)WithScriptArray([](auto* Array) { return Array->GetData(); }) + Index * ElementSize;
+		return (uint8*)WithScriptArray([](auto* Array) { return Array->GetData(); }) + Index * static_cast<size_t>(ElementSize);
 	}
 	/**
 	 *	Returns a uint8 pointer to an element in the array. This call is identical to GetRawPtr and is
@@ -4411,7 +4413,7 @@ private:
 			uint8* Dest = GetRawPtr(Index);
 			if (InnerProperty->PropertyFlags & CPF_ZeroConstructor)
 			{
-				FMemory::Memzero(Dest, Count * ElementSize);
+				FMemory::Memzero(Dest, Count * static_cast<size_t>(ElementSize));
 			}
 			else
 			{
@@ -4461,7 +4463,7 @@ private:
 			uint8* Dest = GetRawPtr(Index);
 			if ((InnerProperty->PropertyFlags & (CPF_ZeroConstructor | CPF_NoDestructor)) == (CPF_ZeroConstructor | CPF_NoDestructor))
 			{
-				FMemory::Memzero(Dest, Count * ElementSize);
+				FMemory::Memzero(Dest, Count * static_cast<size_t>(ElementSize));
 			}
 			else
 			{
