Log GC speed and number of objects and references

Also refactored GC rerun to track garbage references and some minor driveby cleanup

#rnx

diff --git a/Engine/Source/Runtime/CoreUObject/Private/UObject/GarbageCollection.cpp b/Engine/Source/Runtime/CoreUObject/Private/UObject/GarbageCollection.cpp
--- a/Engine/Source/Runtime/CoreUObject/Private/UObject/GarbageCollection.cpp
+++ b/Engine/Source/Runtime/CoreUObject/Private/UObject/GarbageCollection.cpp
@@ -240,10 +240,8 @@ static FAutoConsoleVariableRef CMultithreadedDestructionEnabled(
 );
 
 #if UE_BUILD_SHIPPING
-static constexpr bool GRedoReachabilityToTrackGarbage = false;
 static constexpr int32 GGarbageReferenceTrackingEnabled = 0;
 #else
-static bool GRedoReachabilityToTrackGarbage;
 int32 GGarbageReferenceTrackingEnabled = 0;
 static FAutoConsoleVariableRef CGarbageReferenceTrackingEnabled(
 	TEXT("gc.GarbageReferenceTrackingEnabled"),
@@ -776,6 +774,9 @@ bool IsInGarbageCollectorThread()
 
 namespace UE::GC {
 
+	
+
+
 using AROFunc = void (*)(UObject*, FReferenceCollector&);
 
 #if UE_WITH_GC
@@ -844,10 +845,7 @@ void FContextPoolScope::ReturnToPool(FWorkerContext* Context)
 	check(Pool.NumAllocated >= 1);
 	--Pool.NumAllocated;
 	Context->FreeWorkerIndex();
-#if !UE_BUILD_SHIPPING
-	GRedoReachabilityToTrackGarbage |= Context->bDetectedGarbageReference;
-	Context->bDetectedGarbageReference = false;
-#endif
+	Context->Stats = {};
 	Pool.Reusable.Emplace(Context);
 }
 
@@ -1310,6 +1308,7 @@ private:
 	FORCEINLINE_DEBUGGABLE void DrainUnvalidated(const uint32 Num)
 	{
 		check(Num <= UnvalidatedBatchSize);
+		Context.Stats.AddReferences(Num);
 
 		FPermanentObjectPoolExtents Permanent(PermanentPool);
 		FValidatedBitmask ValidsA, ValidsB;
@@ -2450,9 +2449,7 @@ public:
 
 	static FORCEINLINE void DetectGarbageReference(FWorkerContext& Context, FReferenceMetadata Metadata)
 	{
-#if !UE_BUILD_SHIPPING
-		Context.bDetectedGarbageReference |= !IsWithPendingKill() && Metadata.Has(KillFlag);
-#endif
+		Context.Stats.TrackPotentialGarbageReference(!IsWithPendingKill() && Metadata.Has(KillFlag));
 	}
 
 	/**
@@ -2591,6 +2588,7 @@ struct TBatchDispatcher
 	FORCENOINLINE void HandleReferenceDirectly(const UObject* ReferencingObject, UObject*& Object, FTokenId TokenId, EKillable Killable)
 	{
 		ProcessorType::ProcessReferenceDirectly(Context, PermanentPool, ReferencingObject, Object, TokenId, Killable);
+		Context.Stats.AddReferences(1);
 	}
 
 	FORCEINLINE void HandleReferenceDirectly(const UObject* ReferencingObject, UObject*& Object, FTokenId TokenId, EGCTokenType TokenType, bool bAllowReferenceElimination)
@@ -2743,6 +2741,7 @@ public:
 
 	virtual void AddStableReferenceSet(TSet<UObject*>* Objects) override
 	{
+
 		Dispatcher.QueueSet(Dispatcher.Context.GetReferencingObject(), *Objects, ETokenlessId::Collector, MayKill());
 	}
 	
@@ -2864,6 +2863,7 @@ private:
 	}
 };
 
+
 template <EGCOptions Options>
 class TDebugReachabilityCollector final : public TReachabilityCollectorBase<Options>
 {
@@ -2890,6 +2890,7 @@ public:
 #endif // ENABLE_GC_OBJECT_CHECKS
 
 		Processor.HandleTokenStreamObjectReference(Context, ReferencingObject, Object, ETokenlessId::Collector, CurrentTokenType, bAllowEliminatingReferences);
+		Context.Stats.AddReferences(1);
 	}
 
 	virtual void HandleObjectReferences(UObject** Objects, int32 Num, const UObject* ReferencingObject, const FProperty* ReferencingProperty) override
@@ -2900,6 +2901,8 @@ public:
 		}
 	}
 
+
+
 	virtual void AddStableReference(UObject** Object) override
 	{
 		HandleObjectReference(*Object, Context.GetReferencingObject(), nullptr);
@@ -2970,6 +2973,7 @@ void TReachabilityCollector<Options>::HandleObjectReferences(UObject** InObjects
 	{
 		Dispatcher.HandleReferenceDirectly(ReferencingObject, Object, ETokenlessId::Collector, Killable);
 	}
+
 }
 
 #else // !UE_WITH_GC
@@ -3129,13 +3133,13 @@ class FRealtimeGC : public FGarbageCollectionTracer
 	{
 		TRACE_CPUPROFILER_EVENT_SCOPE(PerformReachabilityAnalysisOnObjectsInternal);
 
-#if !UE_BUILD_SHIPPING		
+#if !UE_BUILD_SHIPPING
 		TDebugReachabilityProcessor<Options> DebugProcessor;
 		if (DebugProcessor.TracksHistory() | 
-			DebugProcessor.TracksGarbage() & GRedoReachabilityToTrackGarbage)
+			DebugProcessor.TracksGarbage() & Stats.bFoundGarbageRef)
 		{
 			CollectReferences<TDebugReachabilityCollector<Options>>(DebugProcessor, Context);
-			GRedoReachabilityToTrackGarbage = false;
+			UE_CLOG(Stats.bFoundGarbageRef && !Context.Stats.bFoundGarbageRef, LogGarbage, Warning, TEXT("GC rerun didn't find any garbage references but first run did"));
 			return;
 		}
 #endif
@@ -3182,7 +3186,6 @@ public:
 	void MarkObjectsAsUnreachable(const EObjectFlags KeepFlags)
 	{
 		TRACE_CPUPROFILER_EVENT_SCOPE(MarkObjectsAsUnreachable);
-		const EInternalObjectFlags FastKeepFlags = EInternalObjectFlags::GarbageCollectionKeepFlags;
 		const int32 MaxNumberOfObjects = GUObjectArray.GetObjectArrayNum() - GUObjectArray.GetFirstGCIndex();
 		const int32 NumThreads = FMath::Max(1, FTaskGraphInterface::Get().GetNumWorkerThreads());
 		const int32 NumberOfObjectsPerThread = (MaxNumberOfObjects / NumThreads) + 1;		
@@ -3195,9 +3198,12 @@ public:
 
 		// Iterate over all objects. Note that we iterate over the UObjectArray and usually check only internal flags which
 		// are part of the array so we don't suffer from cache misses as much as we would if we were to check ObjectFlags.
-		ParallelFor( TEXT("GC.MarkUnreachable"),NumThreads,1, [&ObjectsToSerializeArrays, &ClustersToDissolveList, &KeepClusterRefsList, FastKeepFlags, KeepFlags, NumberOfObjectsPerThread, NumThreads, MaxNumberOfObjects](int32 ThreadIndex)
+		ParallelFor( TEXT("GC.MarkUnreachable"),NumThreads,1,
+			[&ObjectsToSerializeArrays, &ClustersToDissolveList, &KeepClusterRefsList,
+			 KeepFlags, NumberOfObjectsPerThread, NumThreads, MaxNumberOfObjects, bIsRerun = Stats.bFoundGarbageRef] (int32 ThreadIndex)
 		{
 			TRACE_CPUPROFILER_EVENT_SCOPE(MarkObjectsAsUnreachableTask);
+			constexpr EInternalObjectFlags FastKeepFlags = EInternalObjectFlags::GarbageCollectionKeepFlags;
 			int32 FirstObjectIndex = ThreadIndex * NumberOfObjectsPerThread + GUObjectArray.GetFirstGCIndex();
 			int32 NumObjects = (ThreadIndex < (NumThreads - 1)) ? NumberOfObjectsPerThread : (MaxNumberOfObjects - (NumThreads - 1) * NumberOfObjectsPerThread);
 			int32 LastObjectIndex = FMath::Min(GUObjectArray.GetObjectArrayNum() - 1, FirstObjectIndex + NumObjects - 1);
@@ -3212,7 +3218,7 @@ public:
 					UObject* Object = (UObject*)ObjectItem->Object;
 
 					// We can't collect garbage during an async load operation and by now all unreachable objects should've been purged.
-					checkf(	GRedoReachabilityToTrackGarbage ||
+					checkf(	bIsRerun ||
 							!ObjectItem->HasAnyFlags(EInternalObjectFlags::Unreachable|EInternalObjectFlags::PendingConstruction),
 							TEXT("Object: '%s' with ObjectFlags=0x%08x and InternalObjectFlags=0x%08x. ")
 							TEXT("State: IsEngineExitRequested=%d, GIsCriticalError=%d, GExitPurge=%d, GObjPurgeIsRequired=%d, GObjIncrementalPurgeIsInProgress=%d, GObjFinishDestroyHasBeenRoutedToAllObjects=%d, GGCObjectsPendingDestructionCount=%d"),
@@ -3416,6 +3422,7 @@ public:
 
 			PerformReachabilityAnalysisOnObjects(Context, Options);
 
+			Stats = Context->Stats;
 			Context->ResetInitialObjects();
 			Context->InitialNativeReferences = TConstArrayView<UObject**>();
 			Pool.ReturnToPool(Context);
@@ -3434,6 +3441,8 @@ PRAGMA_ENABLE_DEPRECATION_WARNINGS
 	{
 		(this->*ReachabilityAnalysisFunctions[GetGCFunctionIndex(Options)])(*Context);
 	}
+
+	FProcessorStats Stats;
 };
 
 } // namespace UE::GC
@@ -4292,21 +4301,28 @@ void CollectGarbageImpl(EObjectFlags KeepFlags)
 
 			// Perform reachability analysis.
 			{
-				SCOPED_NAMED_EVENT(FRealtimeGC_PerformReachabilityAnalysis, FColor::Red);
-				DECLARE_SCOPE_CYCLE_COUNTER(TEXT("FRealtimeGC::PerformReachabilityAnalysis"), STAT_FArchiveRealtimeGC_PerformReachabilityAnalysis, STATGROUP_GC);
-				const double StartTime = FPlatformTime::Seconds();
-				FRealtimeGC().PerformReachabilityAnalysis(KeepFlags, Options);
-				UE_LOG(LogGarbage, Log, TEXT("%f ms for GC with %d clusters"), (FPlatformTime::Seconds() - StartTime) * 1000, GUObjectClusters.GetNumAllocatedClusters());
-			}
+				FRealtimeGC GC;
+				{
+					SCOPED_NAMED_EVENT(FRealtimeGC_PerformReachabilityAnalysis, FColor::Red);
+					DECLARE_SCOPE_CYCLE_COUNTER(TEXT("FRealtimeGC::PerformReachabilityAnalysis"), STAT_FArchiveRealtimeGC_PerformReachabilityAnalysis, STATGROUP_GC);
+					const double StartTime = FPlatformTime::Seconds();
+					GC.PerformReachabilityAnalysis(KeepFlags, Options);
+					const double Ms = (FPlatformTime::Seconds() - StartTime) * 1000;
+					UE_LOG(LogGarbage, Log, TEXT("%.2f ms for GC - %d refs/ms while processing %d references from %d objects with %d clusters"),
+							Ms, (int32)(GC.Stats.NumReferences / Ms), GC.Stats.NumReferences, GC.Stats.NumObjects, GUObjectClusters.GetNumAllocatedClusters());
+				}
+
+				if (GC.Stats.bFoundGarbageRef && GGarbageReferenceTrackingEnabled > 0)
+				{
+					
+					CSV_SCOPED_TIMING_STAT_EXCLUSIVE(GarbageCollectionDebug);
+					SCOPED_NAMED_EVENT(FRealtimeGC_PerformReachabilityAnalysisRerun, FColor::Orange);
+					DECLARE_SCOPE_CYCLE_COUNTER(TEXT("FRealtimeGC::PerformReachabilityAnalysisRerun"), STAT_FArchiveRealtimeGC_PerformReachabilityAnalysisRerun, STATGROUP_GC);
+					const double StartTime = FPlatformTime::Seconds();
+					GC.PerformReachabilityAnalysis(KeepFlags, Options);
+					UE_LOG(LogGarbage, Log, TEXT("%.2f ms for GC rerun to track garbage references (gc.GarbageReferenceTrackingEnabled=%d)"), (FPlatformTime::Seconds() - StartTime) * 1000, GGarbageReferenceTrackingEnabled);
+				}
 
-			if (GRedoReachabilityToTrackGarbage && GGarbageReferenceTrackingEnabled > 0)
-			{
-				CSV_SCOPED_TIMING_STAT_EXCLUSIVE(GarbageCollectionDebug);
-				SCOPED_NAMED_EVENT(FRealtimeGC_PerformReachabilityAnalysisRerun, FColor::Orange);
-				DECLARE_SCOPE_CYCLE_COUNTER(TEXT("FRealtimeGC::PerformReachabilityAnalysisRerun"), STAT_FArchiveRealtimeGC_PerformReachabilityAnalysisRerun, STATGROUP_GC);
-				const double StartTime = FPlatformTime::Seconds();
-				FRealtimeGC().PerformReachabilityAnalysis(KeepFlags, Options);
-				UE_LOG(LogGarbage, Log, TEXT("%f ms for GC rerun to track garbage references (gc.GarbageReferenceTrackingEnabled=%d)"), (FPlatformTime::Seconds() - StartTime) * 1000, GGarbageReferenceTrackingEnabled);
 			}
 
 
@@ -6000,6 +6016,7 @@ ELoot StealWork(FWorkerContext& Context, FReferenceCollector& Collector, FWorkBl
 void ProcessAsync(void (*ProcessSync)(void*, FWorkerContext&), void* Processor, FWorkerContext& InContext)
 {
 	checkf(InContext.ObjectsToSerialize.IsUnused(), TEXT("Use InitialObjects instead, ObjectsToSerialize.Add() may only be called during reference processing"));
+	check(InContext.Stats.NumObjects == 0 && InContext.Stats.NumReferences == 0 && InContext.Stats.bFoundGarbageRef == false);
 
 	TConstArrayView<UObject*> InitialObjects = InContext.GetInitialObjects();
 	TConstArrayView<UObject**> InitialReferences = InContext.InitialNativeReferences;
@@ -6071,6 +6088,7 @@ void ProcessAsync(void (*ProcessSync)(void*, FWorkerContext&), void* Processor,
 
 	for (FWorkerContext* Context : Contexts.RightChop(1))
 	{
+		InContext.Stats.AddStats(Context->Stats);
 		ContextPool.ReturnToPool(Context);
 	}
 }
diff --git a/Engine/Source/Runtime/CoreUObject/Public/UObject/FastReferenceCollector.h b/Engine/Source/Runtime/CoreUObject/Public/UObject/FastReferenceCollector.h
--- a/Engine/Source/Runtime/CoreUObject/Public/UObject/FastReferenceCollector.h
+++ b/Engine/Source/Runtime/CoreUObject/Public/UObject/FastReferenceCollector.h
@@ -216,6 +216,33 @@ struct FSlowARO
 
 //////////////////////////////////////////////////////////////////////////
 
+struct FProcessorStats
+{
+#if UE_BUILD_SHIPPING
+	static constexpr uint32 NumObjects = 0;
+	static constexpr uint32 NumReferences = 0;
+	static constexpr bool bFoundGarbageRef = false;
+	FORCEINLINE constexpr void AddObjects(uint32) {}
+	FORCEINLINE constexpr void AddReferences(uint32) {}
+	FORCEINLINE constexpr void TrackPotentialGarbageReference(bool) {}
+#else
+	uint32 NumObjects = 0;
+	uint32 NumReferences = 0;
+	bool bFoundGarbageRef = false;
+	FORCEINLINE void AddObjects(uint32 Num) { NumObjects += Num; }
+	FORCEINLINE void AddReferences(uint32 Num) { NumReferences += Num; }
+	FORCEINLINE void TrackPotentialGarbageReference(bool bDetectedGarbage) { bFoundGarbageRef |= bDetectedGarbage; }
+
+#endif
+
+	void AddStats(FProcessorStats Stats)
+	{
+		AddObjects(Stats.NumObjects);
+		AddReferences(Stats.NumReferences);
+		TrackPotentialGarbageReference(Stats.bFoundGarbageRef);
+	}
+};
+
 /** Thread-local context containing initial objects and references to collect */
 struct alignas(PLATFORM_CACHE_LINE_SIZE) FWorkerContext
 {
@@ -236,9 +263,10 @@ public:
 	TConstArrayView<UObject**> InitialNativeReferences;
 	FWorkCoordinator* Coordinator = nullptr;
 	TArray<UObject**> WeakReferences;
+	FProcessorStats Stats;
+
 #if !UE_BUILD_SHIPPING
 	TArray<FGarbageReferenceInfo> GarbageReferences;
-	bool bDetectedGarbageReference = false;
 #endif
 #if ENABLE_GC_HISTORY
 	TMap<const UObject*, TArray<FGCDirectReference>*> History;
@@ -421,11 +449,11 @@ public:
 //////////////////////////////////////////////////////////////////////////
 
 /** Forwards references directly to ProcessorType::HandleTokenStreamObjectReference(), unlike TBatchDispatcher */
-template<class ProcessorType, class ContextType>
+template<class ProcessorType>
 struct TDirectDispatcher
 {
 	ProcessorType& Processor;
-	ContextType& Context;
+	FWorkerContext& Context;
 
 	FORCEINLINE void HandleReferenceDirectly(UObject* ReferencingObject, UObject*& Object, FTokenId TokenId, EGCTokenType TokenType, bool bAllowReferenceElimination) const
 	{
@@ -433,6 +461,7 @@ struct TDirectDispatcher
 		{
 			Processor.HandleTokenStreamObjectReference(Context, ReferencingObject, Object, TokenId, TokenType, bAllowReferenceElimination);
 		}
+		Context.Stats.AddReferences(1);
 	}
 	
 	FORCEINLINE void HandleKillableReference(UObject*& Object, FTokenId TokenId, EGCTokenType TokenType) const
@@ -471,8 +500,8 @@ struct TDirectDispatcher
 };
 
 // Default implementation is to create new direct dispatcher
-template<class CollectorType, class ProcessorType, class ContextType>
-TDirectDispatcher<ProcessorType, ContextType> GetDispatcher(CollectorType&, ProcessorType& Processor, ContextType& Context)
+template<class CollectorType, class ProcessorType>
+TDirectDispatcher<ProcessorType> GetDispatcher(CollectorType&, ProcessorType& Processor, FWorkerContext& Context)
 {
 	return { Processor, Context };
 }
@@ -547,6 +576,7 @@ StoleContext:
 		TConstArrayView<UObject*> CurrentObjects = Context.InitialObjects;
 		while (true)
 		{
+			Context.Stats.AddObjects(CurrentObjects.Num());
 			for (FPrefetchingObjectIterator It(CurrentObjects); It.HasMore(); It.Advance())
 			{
 				UObject* CurrentObject = It.GetCurrentObject();
