Remove newly added warning that triggered during cook

#horde 263001
#rnx

diff --git a/Engine/Source/Runtime/CoreUObject/Private/UObject/GarbageCollection.cpp b/Engine/Source/Runtime/CoreUObject/Private/UObject/GarbageCollection.cpp
--- a/Engine/Source/Runtime/CoreUObject/Private/UObject/GarbageCollection.cpp
+++ b/Engine/Source/Runtime/CoreUObject/Private/UObject/GarbageCollection.cpp
@@ -3139,7 +3139,6 @@ class FRealtimeGC : public FGarbageCollectionTracer
 			DebugProcessor.TracksGarbage() & Stats.bFoundGarbageRef)
 		{
 			CollectReferences<TDebugReachabilityCollector<Options>>(DebugProcessor, Context);
-			UE_CLOG(Stats.bFoundGarbageRef && !Context.Stats.bFoundGarbageRef, LogGarbage, Warning, TEXT("GC rerun didn't find any garbage references but first run did"));
 			return;
 		}
 #endif
