add FNullArchive which does nothing
added because FArchive by default asserts on FObjectPtr

diff --git a/Engine/Source/Runtime/CoreUObject/Private/Serialization/NullArchive.h b/Engine/Source/Runtime/CoreUObject/Private/Serialization/NullArchive.h
new file mode 100644
--- /dev/null
+++ b/Engine/Source/Runtime/CoreUObject/Private/Serialization/NullArchive.h
@@ -0,0 +1,13 @@
+// Copyright Epic Games, Inc. All Rights Reserved.
+
+#pragma once
+#include "Serialization/Archive.h"
+
+class FNullArchive : public FArchive
+{
+	//override FObjectPtr to do nothing
+	virtual FArchive& operator<<(struct FObjectPtr& Value) override
+	{
+		return *this;
+	}
+};
\ No newline at end of file
diff --git a/Engine/Source/Runtime/CoreUObject/Private/Serialization/PropertyLocalizationDataGathering.cpp b/Engine/Source/Runtime/CoreUObject/Private/Serialization/PropertyLocalizationDataGathering.cpp
--- a/Engine/Source/Runtime/CoreUObject/Private/Serialization/PropertyLocalizationDataGathering.cpp
+++ b/Engine/Source/Runtime/CoreUObject/Private/Serialization/PropertyLocalizationDataGathering.cpp
@@ -12,6 +12,7 @@
 #include "HAL/UnrealMemory.h"
 #include "Internationalization/TextNamespaceUtil.h"
 #include "Internationalization/TextPackageNamespaceUtil.h"
+#include "Serialization/NullArchive.h"
 
 FPropertyLocalizationDataGatherer::FPropertyLocalizationDataGatherer(TArray<FGatherableTextData>& InOutGatherableTextDataArray, const UPackage* const InPackage, EPropertyLocalizationGathererResultFlags& OutResultFlags)
 	: GatherableTextDataArray(InOutGatherableTextDataArray)
@@ -765,7 +766,7 @@ private:
 	FPropertyLocalizationDataGatherer& PropertyLocalizationDataGatherer;
 	bool bTreatAsEditorOnlyData;
 
-	FArchive DummyArchive;
+	FNullArchive DummyArchive;
 	bool bIsParsingText;
 	FString LastParsedString;
 };
diff --git a/Engine/Source/Runtime/CoreUObject/Private/UObject/Class.cpp b/Engine/Source/Runtime/CoreUObject/Private/UObject/Class.cpp
--- a/Engine/Source/Runtime/CoreUObject/Private/UObject/Class.cpp
+++ b/Engine/Source/Runtime/CoreUObject/Private/UObject/Class.cpp
@@ -39,6 +39,7 @@
 #include "Internationalization/PolyglotTextData.h"
 #include "Serialization/ArchiveScriptReferenceCollector.h"
 #include "Serialization/ArchiveUObjectFromStructuredArchive.h"
+#include "Serialization/NullArchive.h"
 #include "UObject/FrameworkObjectVersion.h"
 #include "UObject/GarbageCollection.h"
 #include "UObject/UObjectThreadContext.h"
@@ -680,7 +681,7 @@ void UStruct::AddCppProperty(FProperty* Property)
 
 void UStruct::StaticLink(bool bRelinkExistingProperties)
 {
-	FArchive ArDummy;
+	FNullArchive ArDummy;
 	Link(ArDummy, bRelinkExistingProperties);
 }
 
diff --git a/Engine/Source/Runtime/CoreUObject/Private/UObject/ScriptCore.cpp b/Engine/Source/Runtime/CoreUObject/Private/UObject/ScriptCore.cpp
--- a/Engine/Source/Runtime/CoreUObject/Private/UObject/ScriptCore.cpp
+++ b/Engine/Source/Runtime/CoreUObject/Private/UObject/ScriptCore.cpp
@@ -21,6 +21,7 @@
 #include "UObject/CoreNative.h"
 #include "UObject/Class.h"
 #include "Templates/Casts.h"
+#include "Serialization/NullArchive.h"
 #include "UObject/SoftObjectPtr.h"
 #include "UObject/PropertyPortFlags.h"
 #include "UObject/UnrealType.h"
@@ -1661,7 +1662,7 @@ static void OutputMostFrequentlyCalledFunctions(FOutputDevice& OutputAr, int32 N
 
 			int32 iCode = 0;
 			const int32 ScriptSizeBytes = Script.Num();
-			FArchive DummyArchive;
+			FNullArchive DummyArchive;
 
 			while (iCode < ScriptSizeBytes)
 			{
@@ -1794,7 +1795,7 @@ static void OutputMostFrequentlyUsedInstructions(FOutputDevice& OutputAr, int32
 
 			int32 iCode = 0;
 			const int32 ScriptSizeBytes = Script.Num();
-			FArchive DummyArchive;
+			FNullArchive DummyArchive;
 
 			while (iCode < ScriptSizeBytes)
 			{
diff --git a/Engine/Source/Runtime/CoreUObject/Public/UObject/ScriptSerialization.h b/Engine/Source/Runtime/CoreUObject/Public/UObject/ScriptSerialization.h
--- a/Engine/Source/Runtime/CoreUObject/Public/UObject/ScriptSerialization.h
+++ b/Engine/Source/Runtime/CoreUObject/Public/UObject/ScriptSerialization.h
@@ -11,7 +11,7 @@
 #ifdef VISUAL_ASSIST_HACK
 	EExprToken SerializeExpr( int32&, FArchive& );
 	int32 iCode=0;
-	FArchive Ar;
+	FNullArchive Ar;
 	TArray<uint8> Script;
 	EExprToken Expr = (EExprToken)0;
 #endif
