Fixed non unity compile errors

diff --git a/Engine/Plugins/Experimental/Chooser/Source/ChooserEditor/Private/ChooserTableEditorCommands.h b/Engine/Plugins/Experimental/Chooser/Source/ChooserEditor/Private/ChooserTableEditorCommands.h
--- a/Engine/Plugins/Experimental/Chooser/Source/ChooserEditor/Private/ChooserTableEditorCommands.h
+++ b/Engine/Plugins/Experimental/Chooser/Source/ChooserEditor/Private/ChooserTableEditorCommands.h
@@ -4,6 +4,7 @@
 
 #include "CoreMinimal.h"
 #include "Framework/Commands/Commands.h"
+#include "Styling/AppStyle.h"
 
 #define LOCTEXT_NAMESPACE "ChooserTableEditorCommands"
 	
diff --git a/Engine/Plugins/Experimental/Chooser/Source/ProxyTableEditor/Private/ProxyTableEditorCommands.h b/Engine/Plugins/Experimental/Chooser/Source/ProxyTableEditor/Private/ProxyTableEditorCommands.h
--- a/Engine/Plugins/Experimental/Chooser/Source/ProxyTableEditor/Private/ProxyTableEditorCommands.h
+++ b/Engine/Plugins/Experimental/Chooser/Source/ProxyTableEditor/Private/ProxyTableEditorCommands.h
@@ -4,6 +4,7 @@
 
 #include "CoreMinimal.h"
 #include "Framework/Commands/Commands.h"
+#include "Styling/AppStyle.h"
 
 #define LOCTEXT_NAMESPACE "ProxyTableEditorCommands"
 	
diff --git a/Engine/Plugins/Online/OnlineServices/Source/OnlineServicesCommon/Private/Online/OnlineServicesCommon.cpp b/Engine/Plugins/Online/OnlineServices/Source/OnlineServicesCommon/Private/Online/OnlineServicesCommon.cpp
--- a/Engine/Plugins/Online/OnlineServices/Source/OnlineServicesCommon/Private/Online/OnlineServicesCommon.cpp
+++ b/Engine/Plugins/Online/OnlineServices/Source/OnlineServicesCommon/Private/Online/OnlineServicesCommon.cpp
@@ -22,6 +22,8 @@ FOnlineServicesCommon::FOnlineServicesCommon(const FString& InConfigName, FName
 {
 }
 
+FOnlineServicesCommon::~FOnlineServicesCommon() = default;
+
 void FOnlineServicesCommon::Init()
 {
 	OpCache.SetLoadConfigFn(
diff --git a/Engine/Plugins/Online/OnlineServices/Source/OnlineServicesCommon/Public/Online/OnlineServicesCommon.h b/Engine/Plugins/Online/OnlineServices/Source/OnlineServicesCommon/Public/Online/OnlineServicesCommon.h
--- a/Engine/Plugins/Online/OnlineServices/Source/OnlineServicesCommon/Public/Online/OnlineServicesCommon.h
+++ b/Engine/Plugins/Online/OnlineServices/Source/OnlineServicesCommon/Public/Online/OnlineServicesCommon.h
@@ -27,7 +27,7 @@ public:
 	FOnlineServicesCommon(const FString& InConfigName, FName InInstanceName);
 	FOnlineServicesCommon(const FOnlineServicesCommon&) = delete;
 	FOnlineServicesCommon(FOnlineServicesCommon&&) = delete;
-	virtual ~FOnlineServicesCommon() {}
+	virtual ~FOnlineServicesCommon();
 
 	// IOnlineServices
 	virtual void Init() override;
diff --git a/Engine/Source/Runtime/Core/Public/Misc/RangeAllocator.h b/Engine/Source/Runtime/Core/Public/Misc/RangeAllocator.h
--- a/Engine/Source/Runtime/Core/Public/Misc/RangeAllocator.h
+++ b/Engine/Source/Runtime/Core/Public/Misc/RangeAllocator.h
@@ -6,6 +6,10 @@
 
 #pragma once
 
+#include "Containers/Array.h"
+#include "HAL/CriticalSection.h"
+#include "Misc/ScopeLock.h"
+
 /*
  * This allocator has a list of chunks (fixed size committed memory regions) and each
  * chunk keeps a list of free ranges. Allocations are served by finding a large enough
diff --git a/Engine/Source/Runtime/CoreUObject/Public/UObject/WeakObjectPtrFwd.h b/Engine/Source/Runtime/CoreUObject/Public/UObject/WeakObjectPtrFwd.h
--- a/Engine/Source/Runtime/CoreUObject/Public/UObject/WeakObjectPtrFwd.h
+++ b/Engine/Source/Runtime/CoreUObject/Public/UObject/WeakObjectPtrFwd.h
@@ -18,6 +18,9 @@ template<> struct TIsPODType<FWeakObjectPtr> { enum { Value = true }; };
 template<> struct TIsZeroConstructType<FWeakObjectPtr> { enum { Value = true }; };
 template<> struct TIsWeakPointerType<FWeakObjectPtr> { enum { Value = true }; };
 
+template <typename T> class TScriptDelegate;
+template <typename T> class TMulticastScriptDelegate;
+
 // Typedef script delegates for convenience.
-typedef TScriptDelegate<> FScriptDelegate;
-typedef TMulticastScriptDelegate<> FMulticastScriptDelegate;
+typedef TScriptDelegate<FWeakObjectPtr> FScriptDelegate;
+typedef TMulticastScriptDelegate<FWeakObjectPtr> FMulticastScriptDelegate;
diff --git a/Engine/Source/Runtime/Renderer/Public/PostProcess/PostProcessMaterialInputs.h b/Engine/Source/Runtime/Renderer/Public/PostProcess/PostProcessMaterialInputs.h
--- a/Engine/Source/Runtime/Renderer/Public/PostProcess/PostProcessMaterialInputs.h
+++ b/Engine/Source/Runtime/Renderer/Public/PostProcess/PostProcessMaterialInputs.h
@@ -2,6 +2,7 @@
 
 #pragma once
 
+#include "SceneTexturesConfig.h"
 #include "ScreenPass.h"
 
 struct FSceneWithoutWaterTextures;
