Fix some Core/AssetRegistry/LaunchEngineLoop code that was invoking undefined behavior:
- Calling memcpy/memset/memcmp with null arguments.
- Binding a reference to a null value.

diff --git a/Engine/Source/Runtime/AssetRegistry/Private/AssetDataGatherer.cpp b/Engine/Source/Runtime/AssetRegistry/Private/AssetDataGatherer.cpp
--- a/Engine/Source/Runtime/AssetRegistry/Private/AssetDataGatherer.cpp
+++ b/Engine/Source/Runtime/AssetRegistry/Private/AssetDataGatherer.cpp
@@ -2597,7 +2597,10 @@ public:
 	void Load(FMutableMemoryView Out)
 	{
 		FMemoryView In = Load(Out.GetSize());
-		FMemory::Memcpy(Out.GetData(), In.GetData(), In.GetSize());
+		if (In.GetSize())
+		{
+			FMemory::Memcpy(Out.GetData(), In.GetData(), In.GetSize());
+		}
 	}
 
 	template<typename T>
@@ -2689,14 +2692,20 @@ class FChecksumArchiveBase : public FArchiveProxy
 		void Write(FMemoryView In)
 		{
 			check(GetRemainingSize() >= In.GetSize());
-			FMemory::Memcpy(Cursor, In.GetData(), In.GetSize());
+			if (In.GetSize())
+			{
+				FMemory::Memcpy(Cursor, In.GetData(), In.GetSize());
+			}
 			Cursor += In.GetSize();	
 		}
 
 		void Read(FMutableMemoryView Out)
 		{
 			check(GetRemainingSize() >= Out.GetSize());
-			FMemory::Memcpy(Out.GetData(), Cursor, Out.GetSize());
+			if (Out.GetSize())
+			{
+				FMemory::Memcpy(Out.GetData(), Cursor, Out.GetSize());
+			}
 			Cursor += Out.GetSize();	
 		}
 	};
diff --git a/Engine/Source/Runtime/Core/Public/Containers/SparseArray.h b/Engine/Source/Runtime/Core/Public/Containers/SparseArray.h
--- a/Engine/Source/Runtime/Core/Public/Containers/SparseArray.h
+++ b/Engine/Source/Runtime/Core/Public/Containers/SparseArray.h
@@ -794,8 +794,11 @@ public:
 			}
 			else
 			{
-				// Use the much faster path for types that allow it
-				FMemory::Memcpy(DestData, SrcData, sizeof(FElementOrFreeListLink) * SrcMax);
+				if (SrcMax)
+				{
+					// Use the much faster path for types that allow it
+					FMemory::Memcpy(DestData, SrcData, sizeof(FElementOrFreeListLink) * SrcMax);
+				}
 			}
 		}
 		return *this;
diff --git a/Engine/Source/Runtime/Core/Public/Misc/StringBuilder.h b/Engine/Source/Runtime/Core/Public/Misc/StringBuilder.h
--- a/Engine/Source/Runtime/Core/Public/Misc/StringBuilder.h
+++ b/Engine/Source/Runtime/Core/Public/Misc/StringBuilder.h
@@ -157,7 +157,10 @@ public:
 	{
 		int32 ConvertedLength = FPlatformString::ConvertedLength<CharType>(String, Length);
 		EnsureAdditionalCapacity(ConvertedLength);
-		CurPos = FPlatformString::Convert(CurPos, ConvertedLength, String, Length);
+		if (Length)
+		{
+			CurPos = FPlatformString::Convert(CurPos, ConvertedLength, String, Length);
+		}
 		return *this;
 	}
 
@@ -233,7 +236,10 @@ public:
 			}
 		}
 		
-		FMemory::Memcpy(Base + Pos, Str.GetData(), Str.Len() * sizeof(CharType));
+		if (Str.Len())
+		{
+			FMemory::Memcpy(Base + Pos, Str.GetData(), Str.Len() * sizeof(CharType));
+		}
 	}
 
 	/** Insert substring at given position */
diff --git a/Engine/Source/Runtime/Core/Public/Templates/MemoryOps.h b/Engine/Source/Runtime/Core/Public/Templates/MemoryOps.h
--- a/Engine/Source/Runtime/Core/Public/Templates/MemoryOps.h
+++ b/Engine/Source/Runtime/Core/Public/Templates/MemoryOps.h
@@ -110,7 +110,10 @@ FORCEINLINE void ConstructItems(void* Dest, const SourceElementType* Source, Siz
 {
 	if constexpr (TIsBitwiseConstructible<DestinationElementType, SourceElementType>::Value)
 	{
-		FMemory::Memcpy(Dest, Source, sizeof(SourceElementType) * Count);
+		if (Count)
+		{
+			FMemory::Memcpy(Dest, Source, sizeof(SourceElementType) * Count);
+		}
 	}
 	else
 	{
diff --git a/Engine/Source/Runtime/CoreUObject/Private/UObject/GarbageCollection.cpp b/Engine/Source/Runtime/CoreUObject/Private/UObject/GarbageCollection.cpp
--- a/Engine/Source/Runtime/CoreUObject/Private/UObject/GarbageCollection.cpp
+++ b/Engine/Source/Runtime/CoreUObject/Private/UObject/GarbageCollection.cpp
@@ -5452,14 +5452,20 @@ void FTokenStreamBuilder::Merge(FTokenStreamView& Out, const FTokenStreamBuilder
 		FTokenStreamHeader* Header = new (FMemory::Malloc(NumBytes)) FTokenStreamHeader;
 		uint32* Tokens = reinterpret_cast<uint32*>(Header + 1);
 
-		FMemory::Memcpy(Tokens,						Super.Tokens,					sizeof(uint32) * Super.Num);
-		FMemory::Memcpy(Tokens + Super.Num,			Class.Tokens.GetData(),			sizeof(uint32) * Class.Num());
+		if (Super.Num)
+		{
+			FMemory::Memcpy(Tokens,             Super.Tokens,           sizeof(uint32) * Super.Num  );
+		}
+		if (Class.Num())
+		{
+			FMemory::Memcpy(Tokens + Super.Num, Class.Tokens.GetData(), sizeof(uint32) * Class.Num());
+		}
 
 	#if ENABLE_GC_TOKEN_DEBUG_INFO
 		check(Class.DebugNames.Num() == Class.Tokens.Num());
 		FName* DebugNames = reinterpret_cast<FName*>(Tokens + Num);
-		FMemory::Memcpy(DebugNames,					Super.GetDebugNames(),			sizeof(FName) * Super.Num);
-		FMemory::Memcpy(DebugNames + Super.Num,		Class.DebugNames.GetData(),		sizeof(FName) * Class.Num());
+		if (Super.Num)   { FMemory::Memcpy(DebugNames,             Super.GetDebugNames(),      sizeof(FName) * Super.Num); }
+		if (Class.Num()) { FMemory::Memcpy(DebugNames + Super.Num, Class.DebugNames.GetData(), sizeof(FName) * Class.Num()); }
 	#endif
 
 		UpdateGlobalTokenCount(int32(Num));
diff --git a/Engine/Source/Runtime/CoreUObject/Private/UObject/ScriptCore.cpp b/Engine/Source/Runtime/CoreUObject/Private/UObject/ScriptCore.cpp
--- a/Engine/Source/Runtime/CoreUObject/Private/UObject/ScriptCore.cpp
+++ b/Engine/Source/Runtime/CoreUObject/Private/UObject/ScriptCore.cpp
@@ -921,7 +921,10 @@ void ProcessScriptFunction(UObject* Context, UFunction* Function, FFrame& Stack,
 	if (!bUsePersistentFrame)
 	{
 		FrameMemory = (uint8*)UE_VSTACK_ALLOC_ALIGNED(NewStack.CachedThreadVirtualStackAllocator, Function->PropertiesSize, Function->GetMinAlignment());
-		FMemory::Memzero(FrameMemory, Function->PropertiesSize);
+		if (Function->PropertiesSize)
+		{
+			FMemory::Memzero(FrameMemory, Function->PropertiesSize);
+		}
 	}
 
 	/* 
@@ -2054,11 +2057,18 @@ void UObject::ProcessEvent( UFunction* Function, void* Parms )
 		{
 			Frame = (uint8*)UE_VSTACK_ALLOC_ALIGNED(VirtualStackAllocator, Function->PropertiesSize, Function->GetMinAlignment());
 			// zero the local property memory
-			FMemory::Memzero(Frame + Function->ParmsSize, Function->PropertiesSize - Function->ParmsSize);
+			const int32 NonParmsPropertiesSize = Function->PropertiesSize - Function->ParmsSize;
+			if (NonParmsPropertiesSize)
+			{
+				FMemory::Memzero(Frame + Function->ParmsSize, NonParmsPropertiesSize);
+			}
 		}
 
 		// initialize the parameter properties
-		FMemory::Memcpy(Frame, Parms, Function->ParmsSize);
+		if (Function->ParmsSize)
+		{
+			FMemory::Memcpy(Frame, Parms, Function->ParmsSize);
+		}
 
 		// Create a new local execution stack.
 		FFrame NewStack(this, Function, Frame, NULL, Function->ChildProperties);
diff --git a/Engine/Source/Runtime/Launch/Private/LaunchEngineLoop.cpp b/Engine/Source/Runtime/Launch/Private/LaunchEngineLoop.cpp
--- a/Engine/Source/Runtime/Launch/Private/LaunchEngineLoop.cpp
+++ b/Engine/Source/Runtime/Launch/Private/LaunchEngineLoop.cpp
@@ -2989,6 +2989,9 @@ PRAGMA_ENABLE_DEPRECATION_WARNINGS
 
 	FEmbeddedCommunication::ForceTick(2);
 
+	PreInitContext.SlowTaskPtr = new FScopedSlowTask(100, NSLOCTEXT("EngineLoop", "EngineLoop_Initializing", "Initializing..."));
+	FScopedSlowTask& SlowTask = *PreInitContext.SlowTaskPtr;
+
 #if WITH_ENGINE
 	// allow for game explorer processing (including parental controls) and firewalls installation
 	if (!FPlatformMisc::CommandLineCommands())
@@ -3176,9 +3179,6 @@ PRAGMA_ENABLE_DEPRECATION_WARNINGS
 
 	FEmbeddedCommunication::ForceTick(3);
 
-	PreInitContext.SlowTaskPtr = new FScopedSlowTask(100, NSLOCTEXT("EngineLoop", "EngineLoop_Initializing", "Initializing..."));
-	FScopedSlowTask& SlowTask = *PreInitContext.SlowTaskPtr;
-
 	SlowTask.EnterProgressFrame(10);
 
 #if USE_LOCALIZED_PACKAGE_CACHE
@@ -3492,6 +3492,8 @@ int32 FEngineLoop::PreInitPostStartupScreen(const TCHAR* CmdLine)
 #if UE_EDITOR || WITH_ENGINE
 	const TCHAR* CommandletCommandLine = *PreInitContext.CommandletCommandLine;
 #endif // UE_EDITOR || WITH_ENGINE
+
+	check(PreInitContext.SlowTaskPtr);
 	FScopedSlowTask& SlowTask = *PreInitContext.SlowTaskPtr;
 
 #if WITH_ENGINE
