Split the package mapping logic out of FLinkerInstancingContext so that it can be re-used

FLinkerInstancedPackageMap contains the instanced package mapping logic of FLinkerInstancingContext, in a way that supports mapping from instance->original as well as original->instance. This avoids exposing any internal functions of FLinkerInstancingContext itself, as that could cause it to be misused for its primary use-case.

#rnx

diff --git a/Engine/Source/Runtime/CoreUObject/Private/UObject/LinkerInstancingContext.cpp b/Engine/Source/Runtime/CoreUObject/Private/UObject/LinkerInstancingContext.cpp
--- a/Engine/Source/Runtime/CoreUObject/Private/UObject/LinkerInstancingContext.cpp
+++ b/Engine/Source/Runtime/CoreUObject/Private/UObject/LinkerInstancingContext.cpp
@@ -17,45 +17,34 @@ void FLinkerInstancingContext::FixupSoftObjectPath(FSoftObjectPath& InOutSoftObj
 		{
 			InOutSoftObjectPath = RemappedAssetPath;
 		}
-		else if (FName LongPackageName = InOutSoftObjectPath.GetLongPackageFName(), RemappedPackage = RemapPackage(LongPackageName); RemappedPackage != LongPackageName)
+		else
 		{
-			InOutSoftObjectPath = FSoftObjectPath(RemappedPackage, InOutSoftObjectPath.GetAssetFName(), InOutSoftObjectPath.GetSubPathString());
+			InstancedPackageMap.FixupSoftObjectPath(InOutSoftObjectPath);
 		}
-#if !WITH_EDITOR
-		else if (!GeneratedPackagesFolder.IsEmpty())
-		{
-			check(!InstancedPackageSuffix.IsEmpty());
-
-			FNameBuilder TmpSoftObjectPathBuilder;
-			InOutSoftObjectPath.ToString(TmpSoftObjectPathBuilder);
+	}
+}
 
-			// Does this package path start with the generated folder path?
-			FStringView TmpSoftObjectPathView = TmpSoftObjectPathBuilder.ToView();
-			if (TmpSoftObjectPathView.StartsWith(GeneratedPackagesFolder))
-			{
-				// ... and is that generated folder path immediately preceding the package name?
-				if (const int32 ExtraSlashIndex = TmpSoftObjectPathView.Find(TEXTVIEW("/"), GeneratedPackagesFolder.Len()); ExtraSlashIndex == INDEX_NONE)
-				{
-					FNameBuilder PackageNameBuilder;
-					PackageNameBuilder.Append(InOutSoftObjectPath.GetLongPackageName());
-					PackageNameBuilder.Append(InstancedPackageSuffix);
-					FTopLevelAssetPath SuffixTopLevelAsset(FName(PackageNameBuilder.ToString()), InOutSoftObjectPath.GetAssetFName());
-					InOutSoftObjectPath = FSoftObjectPath(SuffixTopLevelAsset, InOutSoftObjectPath.GetSubPathString());
-				}
-			}
-		}
-#endif
+void FLinkerInstancedPackageMap::AddPackageMapping(FName Original, FName Instanced)
+{
+	if (InstanceMappingDirection == EInstanceMappingDirection::OriginalToInstanced)
+	{
+		InstancedPackageMapping.Add(Original, Instanced);
+	}
+	else
+	{
+		check(InstanceMappingDirection == EInstanceMappingDirection::InstancedToOriginal);
+		InstancedPackageMapping.Add(Instanced, Original);
 	}
 }
 
-void FLinkerInstancingContext::BuildPackageMapping(FName Original, FName Instanced)
+void FLinkerInstancedPackageMap::BuildPackageMapping(FName Original, FName Instanced, const bool bBuildWorldPartitionCellMapping)
 {
 	check(GeneratedPackagesFolder.IsEmpty() && InstancedPackageSuffix.IsEmpty());
 
 	AddPackageMapping(Original, Instanced);
 
 #if !WITH_EDITOR
-	if (bSoftObjectPathRemappingEnabled)
+	if (bBuildWorldPartitionCellMapping)
 	{
 		FNameBuilder TmpOriginal(Original);
 		FNameBuilder TmpInstanced(Instanced);
@@ -82,12 +71,13 @@ void FLinkerInstancingContext::BuildPackageMapping(FName Original, FName Instanc
 				{
 					GeneratedPackagesFolder = OriginalView.Left(GeneratedFolderEndIndex);
 
-					FNameBuilder PackageNameBuilder(FName(OriginalView.Left(GeneratedFolderStartIndex)));
+					FNameBuilder PackageNameBuilder;
+					PackageNameBuilder.Append(OriginalView.Left(GeneratedFolderStartIndex));
 					FName PersistentPackageName(PackageNameBuilder);
 					PackageNameBuilder.Append(InstancedPackageSuffix);
 					FName PersistentPackageInstanceName(PackageNameBuilder);
 										
-					PackageMapping.Add(MakeTuple(PersistentPackageName, PersistentPackageInstanceName));
+					AddPackageMapping(PersistentPackageName, PersistentPackageInstanceName);
 				}
 			}
 
@@ -101,3 +91,54 @@ void FLinkerInstancingContext::BuildPackageMapping(FName Original, FName Instanc
 #endif
 }
 
+bool FLinkerInstancedPackageMap::FixupSoftObjectPath(FSoftObjectPath& InOutSoftObjectPath) const
+{
+	if (IsInstanced())
+	{
+		if (FName LongPackageName = InOutSoftObjectPath.GetLongPackageFName(), RemappedPackage = RemapPackage(LongPackageName); RemappedPackage != LongPackageName)
+		{
+			InOutSoftObjectPath = FSoftObjectPath(RemappedPackage, InOutSoftObjectPath.GetAssetFName(), InOutSoftObjectPath.GetSubPathString());
+			return true;
+		}
+#if !WITH_EDITOR
+		else if (!GeneratedPackagesFolder.IsEmpty())
+		{
+			check(!InstancedPackageSuffix.IsEmpty());
+
+			FNameBuilder TmpSoftObjectPathBuilder;
+			InOutSoftObjectPath.ToString(TmpSoftObjectPathBuilder);
+
+			// Does this package path start with the generated folder path?
+			FStringView TmpSoftObjectPathView = TmpSoftObjectPathBuilder.ToView();
+			if (TmpSoftObjectPathView.StartsWith(GeneratedPackagesFolder))
+			{
+				// ... and is that generated folder path immediately preceding the package name?
+				if (const int32 ExtraSlashIndex = TmpSoftObjectPathView.Find(TEXTVIEW("/"), GeneratedPackagesFolder.Len()); ExtraSlashIndex == INDEX_NONE)
+				{
+					FNameBuilder PackageNameBuilder;
+					PackageNameBuilder.Append(InOutSoftObjectPath.GetLongPackageName());
+					if (InstanceMappingDirection == EInstanceMappingDirection::OriginalToInstanced)
+					{
+						if (!PackageNameBuilder.ToView().EndsWith(InstancedPackageSuffix))
+						{
+							PackageNameBuilder.Append(InstancedPackageSuffix);
+						}
+					}
+					else
+					{
+						check(InstanceMappingDirection == EInstanceMappingDirection::InstancedToOriginal);
+						if (PackageNameBuilder.ToView().EndsWith(InstancedPackageSuffix))
+						{
+							PackageNameBuilder.RemoveSuffix(InstancedPackageSuffix.Len());
+						}
+					}
+					FTopLevelAssetPath SuffixTopLevelAsset(FName(PackageNameBuilder), InOutSoftObjectPath.GetAssetFName());
+					InOutSoftObjectPath = FSoftObjectPath(SuffixTopLevelAsset, InOutSoftObjectPath.GetSubPathString());
+					return true;
+				}
+			}
+		}
+#endif
+	}
+	return false;
+}
diff --git a/Engine/Source/Runtime/CoreUObject/Private/UObject/LinkerLoad.cpp b/Engine/Source/Runtime/CoreUObject/Private/UObject/LinkerLoad.cpp
--- a/Engine/Source/Runtime/CoreUObject/Private/UObject/LinkerLoad.cpp
+++ b/Engine/Source/Runtime/CoreUObject/Private/UObject/LinkerLoad.cpp
@@ -2075,7 +2075,7 @@ FLinkerLoad::ELinkerStatus FLinkerLoad::PopulateInstancingContext()
 			// add remapping for all the packages that should be instantiated along with this one
 			for (const FName& InstancingName : InstancingPackageName)
 			{
-				FName& InstancedName = InstancingContext.PackageMapping.FindOrAdd(InstancingName);
+				FName& InstancedName = InstancingContext.FindOrAddPackageMapping(InstancingName);
 				// if there's isn't already a remapping for that package, create one
 				if (InstancedName.IsNone())
 				{
diff --git a/Engine/Source/Runtime/CoreUObject/Public/UObject/LinkerInstancingContext.h b/Engine/Source/Runtime/CoreUObject/Public/UObject/LinkerInstancingContext.h
--- a/Engine/Source/Runtime/CoreUObject/Public/UObject/LinkerInstancingContext.h
+++ b/Engine/Source/Runtime/CoreUObject/Public/UObject/LinkerInstancingContext.h
@@ -4,6 +4,67 @@
 
 #include "UObject/SoftObjectPath.h"
 
+/**
+ * Helper class to map between an original package and an instance of it (including world partition cells).
+ */
+class COREUOBJECT_API FLinkerInstancedPackageMap
+{
+public:
+	enum class EInstanceMappingDirection : uint8
+	{
+		OriginalToInstanced,
+		InstancedToOriginal,
+	};
+
+	FLinkerInstancedPackageMap() = default;
+
+	explicit FLinkerInstancedPackageMap(EInstanceMappingDirection MappingDirection)
+		: InstanceMappingDirection(MappingDirection)
+	{
+	}
+
+	bool IsInstanced() const
+	{
+		return InstancedPackageMapping.Num() > 0;
+	}
+
+	/** Remap the package name from the import table to its instanced counterpart, otherwise return the name unmodified. */
+	FName RemapPackage(const FName& PackageName) const
+	{
+		if (const FName* RemappedName = InstancedPackageMapping.Find(PackageName))
+		{
+			return *RemappedName;
+		}
+		return PackageName;
+	}
+
+	/** Add a mapping from a package name to a new package name. There should be no separators (. or :) in these strings. */
+	void AddPackageMapping(FName Original, FName Instanced);
+
+	void BuildPackageMapping(FName Original, FName Instanced, const bool bBuildWorldPartitionCellMapping = true);
+
+	bool FixupSoftObjectPath(FSoftObjectPath& InOutSoftObjectPath) const;
+
+private:
+	friend class FLinkerInstancingContext;
+
+	/**
+	 * Map between the original package name and its instance counterpart.
+	 * Key=Original and Value=Instanced when InstanceMappingDirection==OriginalToInstanced
+	 * Key=Instanced and Value=Original when InstanceMappingDirection==InstancedToOriginal
+	 */
+	TMap<FName, FName> InstancedPackageMapping;
+
+	/**
+	 * In which direction has this mapping been built?
+	 */
+	EInstanceMappingDirection InstanceMappingDirection = EInstanceMappingDirection::OriginalToInstanced;
+
+	/** Data needed to re-map world partition cells */
+	FString GeneratedPackagesFolder;
+	FString InstancedPackageSuffix;
+};
+
 /**
  * Helper class to remap package imports during loading.
  * This is usually when objects in a package are outer-ed to object in another package or vice versa.
@@ -25,17 +86,13 @@ public:
 
 	bool IsInstanced() const
 	{
-		return PackageMapping.Num() > 0 || PathMapping.Num() > 0;
+		return InstancedPackageMap.IsInstanced() || PathMapping.Num() > 0;
 	}
 
 	/** Remap the package name from the import table to its instanced counterpart, otherwise return the name unmodified. */
 	FName RemapPackage(const FName& PackageName) const
 	{
-		if (const FName* RemappedName = PackageMapping.Find(PackageName))
-		{
-			return *RemappedName;
-		}
-		return PackageName;
+		return InstancedPackageMap.RemapPackage(PackageName);
 	}
 
 	/**
@@ -54,7 +111,7 @@ public:
 	/** Add a mapping from a package name to a new package name. There should be no separators (. or :) in these strings. */
 	void AddPackageMapping(FName Original, FName Instanced)
 	{
-		PackageMapping.Add(Original, Instanced);
+		InstancedPackageMap.AddPackageMapping(Original, Instanced);
 	}
 
 	/** Add a mapping from a top level asset path (/Path/To/Package.AssetName) to another. */
@@ -108,7 +165,15 @@ public:
 	void FixupSoftObjectPath(FSoftObjectPath& InOutSoftObjectPath) const;
 
 private:
-	void BuildPackageMapping(FName Original, FName Instanced);
+	void BuildPackageMapping(FName Original, FName Instanced)
+	{
+		InstancedPackageMap.BuildPackageMapping(Original, Instanced, GetSoftObjectPathRemappingEnabled());
+	}
+
+	FName& FindOrAddPackageMapping(FName Original)
+	{
+		return InstancedPackageMap.InstancedPackageMapping.FindOrAdd(Original);
+	}
 
 	/** Used internally by the linker to try to fix references on relocated packages. */
 	FName RelocatePackage(const FName& PackageName) const
@@ -124,7 +189,7 @@ private:
 	friend struct FAsyncPackage2;
 
 	/** Map of original package name to their instance counterpart. */
-	TMap<FName, FName> PackageMapping;
+	FLinkerInstancedPackageMap InstancedPackageMap;
 	/** Map of original top level asset path to their instance counterpart. */
 	TMap<FTopLevelAssetPath, FTopLevelAssetPath> PathMapping;
 	/** Map of original package name to their potential relocated counterpart. */
@@ -134,7 +199,4 @@ private:
 	TSet<FName> Tags;
 	/** Remap soft object paths */
 	bool bSoftObjectPathRemappingEnabled = true;
-
-	FString GeneratedPackagesFolder;
-	FString InstancedPackageSuffix;
 };
