Temp fix for null FGCObject referencer when printing garbage object refs.

diff --git a/Engine/Source/Runtime/CoreUObject/Private/UObject/GarbageCollection.cpp b/Engine/Source/Runtime/CoreUObject/Private/UObject/GarbageCollection.cpp
--- a/Engine/Source/Runtime/CoreUObject/Private/UObject/GarbageCollection.cpp
+++ b/Engine/Source/Runtime/CoreUObject/Private/UObject/GarbageCollection.cpp
@@ -4423,7 +4423,7 @@ FString FGarbageReferenceInfo::GetReferencingObjectInfo() const
 	else
 	{
 		return FString::Printf(TEXT("%s->AddRereferencedObjects"),
-			*Referencer.GCObject->GetReferencerName()
+			Referencer.GCObject ? *Referencer.GCObject->GetReferencerName() : TEXT("Unknown")
 			);
 	}
 }
