ZenLoader: Fixed some potential integer overflows

diff --git a/Engine/Source/Runtime/CoreUObject/Private/Serialization/AsyncLoading2.cpp b/Engine/Source/Runtime/CoreUObject/Private/Serialization/AsyncLoading2.cpp
--- a/Engine/Source/Runtime/CoreUObject/Private/Serialization/AsyncLoading2.cpp
+++ b/Engine/Source/Runtime/CoreUObject/Private/Serialization/AsyncLoading2.cpp
@@ -3367,6 +3367,13 @@ private:
 			ShaderMapHashesMemSize +
 			ExportToBundleMappingMemSize;
 
+#if PLATFORM_32BITS
+		if (MemoryBufferSize > MAX_Int32)
+		{
+			UE_LOG(LogStreaming, Fatal, TEXT("Memory buffer size overflow"));
+			return;
+		}
+#endif
 		Data.MemoryBuffer = reinterpret_cast<uint8*>(FMemory::Malloc(MemoryBufferSize));
 
 		uint8* DataPtr = Data.MemoryBuffer;
@@ -4191,6 +4198,11 @@ void FAsyncPackage2::SetupSerializedArcs(const FAsyncPackageHeaderData& Header,
 		ArcsArchive << FromExportBundleIndex;
 		int32 ToExportBundleIndex;
 		ArcsArchive << ToExportBundleIndex;
+		if (FromExportBundleIndex >= Data.ExportBundleCount || ToExportBundleIndex >= Data.ExportBundleCount)
+		{
+			UE_LOG(LogStreaming, Fatal, TEXT("Corrupt package data when serializing arcs"));
+			return;
+		}
 		uint32 FromNodeIndexBase = FromExportBundleIndex * EEventLoadNode2::ExportBundle_NumPhases;
 		uint32 ToNodeIndexBase = ToExportBundleIndex * EEventLoadNode2::ExportBundle_NumPhases;
 		for (int32 Phase = 0; Phase < EEventLoadNode2::ExportBundle_NumPhases; ++Phase)
@@ -4218,10 +4230,13 @@ void FAsyncPackage2::SetupSerializedArcs(const FAsyncPackageHeaderData& Header,
 			ArcsArchive << ToExportBundleIndex;
 			if (ImportedPackage)
 			{
-				check(FromImportIndex < Header.ImportMap.Num());
-				check(FromCommandType < FExportBundleEntry::ExportCommandType_Count);
-				check(ToExportBundleIndex < ExportBundleNodes.Num());
-				
+				if (FromImportIndex >= Header.ImportMap.Num() ||
+					FromCommandType >= FExportBundleEntry::ExportCommandType_Count ||
+					ToExportBundleIndex >= Data.ExportBundleCount)
+				{
+					UE_LOG(LogStreaming, Fatal, TEXT("Corrupt package data when serializing arcs"));
+					return;
+				}
 				FPackageObjectIndex GlobalImportIndex = Header.ImportMap[FromImportIndex];
 				FPackageImportReference PackageImportRef = GlobalImportIndex.ToPackageImportRef();
 				const uint64 ImportedPublicExportHash = Header.ImportedPublicExportHashes[PackageImportRef.GetImportedPublicExportHashIndex()];
