Move PropertyHelper.h implementations into a .cpp file.

diff --git a/Engine/Source/Runtime/CoreUObject/Private/UObject/PropertyHelper.cpp b/Engine/Source/Runtime/CoreUObject/Private/UObject/PropertyHelper.cpp
new file mode 100644
--- /dev/null
+++ b/Engine/Source/Runtime/CoreUObject/Private/UObject/PropertyHelper.cpp
@@ -0,0 +1,247 @@
+// Copyright Epic Games, Inc. All Rights Reserved.
+
+#include "UObject/PropertyHelper.h"
+
+//#include "CoreMinimal.h"
+//#include "UObject/ObjectMacros.h"
+//#include "UObject/Object.h"
+//#include "UObject/Class.h"
+//#include "UObject/CoreNet.h"
+//#include "UObject/UnrealType.h"
+//#include "UObject/Package.h"
+//#include "UObject/PropertyPortFlags.h"
+
+void SkipWhitespace(const TCHAR*& Str)
+{
+	while(FChar::IsWhitespace(*Str))
+	{
+		Str++;
+	}
+}
+
+bool AreInstancedObjectsIdentical( UObject* ObjectA, UObject* ObjectB, uint32 PortFlags )
+{
+	check(ObjectA && ObjectB);
+	bool Result = true;
+
+	// we check for recursive comparison here...if we have encountered this pair before on the call stack, the graphs are isomorphic
+	struct FRecursionCheck
+	{
+		UObject* ObjectA;
+		UObject* ObjectB;
+		uint32 PortFlags;
+
+		bool operator==(const FRecursionCheck& Other) const
+		{
+			return ObjectA == Other.ObjectA && ObjectB == Other.ObjectB && PortFlags == Other.PortFlags;
+		}
+	};
+	static TArray<FRecursionCheck> RecursionCheck;
+	FRecursionCheck Test;
+	Test.ObjectA = ObjectA;
+	Test.ObjectB = ObjectB;
+	Test.PortFlags = PortFlags;
+	if (!RecursionCheck.Contains(Test))
+	{
+		RecursionCheck.Push(Test);
+		for ( FProperty* Prop = ObjectA->GetClass()->PropertyLink; Prop && Result; Prop = Prop->PropertyLinkNext )
+		{
+			// only the properties that could have been modified in the editor should be compared
+			// (skipping the name and archetype properties, since name will almost always be different)
+			bool bConsiderProperty;
+			if ( (PortFlags&PPF_Copy) == 0 )
+			{
+				bConsiderProperty = Prop->ShouldDuplicateValue();
+			}
+			else
+			{
+				bConsiderProperty = (Prop->PropertyFlags&CPF_Edit) != 0;
+			}
+
+			if ( bConsiderProperty )
+			{
+				for ( int32 i = 0; i < Prop->ArrayDim && Result; i++ )
+				{
+					if ( !Prop->Identical_InContainer(ObjectA, ObjectB, i, PortFlags) )
+					{
+						Result = false;
+					}
+				}
+			}
+		}
+		if (Result)
+		{
+			// Allow the component to compare its native/ intrinsic properties.
+			Result = ObjectA->AreNativePropertiesIdenticalTo( ObjectB );
+		}
+		RecursionCheck.Pop();
+	}
+	return Result;
+}
+
+bool FDeltaIndexHelper::SerializeNext(FArchive &Ar, int32& Index)
+{
+	NET_CHECKSUM(Ar);
+	if (Ar.IsSaving())
+	{
+		uint32 Delta = Index - LastIndex;
+		Ar.SerializeIntPacked(Delta);
+		LastIndex = Index;
+		LastIndexFull = Index;
+	}
+	else
+	{
+		uint32 Delta = 0;
+		Ar.SerializeIntPacked(Delta);
+		Index = (Delta == 0 ? INDEX_NONE : LastIndex + Delta);
+		LastIndex = Index;
+		LastIndexFull = Index;
+	}
+
+	return (Index != INDEX_NONE && !Ar.IsError());
+}
+
+void FDeltaIndexHelper::SerializeNext(FArchive &OutBunch, FArchive &OutFull, int32 Index)
+{
+	NET_CHECKSUM(OutBunch);
+	NET_CHECKSUM(OutFull);
+		
+	// Full state
+	uint32 DeltaFull = Index - LastIndexFull;
+	OutFull.SerializeIntPacked(DeltaFull);
+	LastIndexFull = Index;
+
+	// Delta State
+	uint32 Delta = Index - LastIndex;
+	OutBunch.SerializeIntPacked(Delta);
+}
+
+void FDeltaIndexHelper::SerializeEarlyEnd(FArchive &Ar)
+{
+	NET_CHECKSUM(Ar);
+
+	uint32 End = 0;
+	Ar.SerializeIntPacked(End);
+}
+
+const TCHAR* DelegatePropertyTools::ImportDelegateFromText( FScriptDelegate& Delegate, const UFunction* SignatureFunction, const TCHAR* Buffer, UObject* Parent, FOutputDevice* ErrorText )
+{
+	SkipWhitespace(Buffer);
+
+	TCHAR ObjName[NAME_SIZE];
+	TCHAR FuncName[NAME_SIZE];
+	// Get object name
+	int32 i;
+	while (*Buffer == TEXT('('))
+	{
+		++Buffer;
+	}
+	for( i=0; *Buffer && *Buffer != TCHAR('.') && *Buffer != TCHAR(')') && *Buffer != TCHAR(','); Buffer++ )
+	{
+		ObjName[i++] = *Buffer;
+	}
+	ObjName[i] = TCHAR('\0');
+	UClass* Cls = nullptr;
+	UObject* Object = nullptr;
+	// Get function name
+	if (*Buffer == TCHAR('.'))
+	{
+		Buffer++;
+		for( i=0; *Buffer && *Buffer != TCHAR(')') && *Buffer != TCHAR(','); Buffer++ )
+		{
+			FuncName[i++] = *Buffer;
+		}
+		FuncName[i] = TEXT('\0');
+	}
+	else
+	{
+		// if there's no dot, then a function name was specified without any object qualifier
+		// if we're importing for a subobject, assume the parent object as the source
+		// otherwise, assume Parent for the source
+		// if neither are valid, then the importing fails
+		if (Parent == nullptr)
+		{
+			FString DelegateString = Delegate.GetFunctionName().ToString();
+			FString SignatureFunctionString = GetNameSafe(SignatureFunction);
+			ErrorText->Logf(ELogVerbosity::Warning, TEXT("Cannot import unqualified delegate name; no object to search. Delegate=%s SignatureFunction=%s"),
+				*DelegateString, *SignatureFunctionString);
+			return nullptr;
+		}
+		// since we don't support nested components, we only need to check one level deep
+		else if (!Parent->HasAnyFlags(RF_ClassDefaultObject) && Parent->GetOuter()->HasAnyFlags(RF_ClassDefaultObject))
+		{
+			Object = Parent->GetOuter();
+			Cls = Parent->GetOuter()->GetClass();
+		}
+		else
+		{
+			Object = Parent;
+			Cls = Parent->GetClass();
+		}
+		FCString::Strncpy(FuncName, ObjName, i + 1);
+	}
+	if (Cls == nullptr)
+	{
+		Cls = UClass::TryFindTypeSlow<UClass>(ObjName);
+		if (Cls != nullptr)
+		{
+			Object = Cls->GetDefaultObject();
+		}
+		else
+		{
+			// Check outer chain before checking all packages
+			UObject* OuterToCheck = Parent;
+			while (Object == nullptr && OuterToCheck)
+			{
+				Object = StaticFindObject(UObject::StaticClass(), OuterToCheck, ObjName);
+				OuterToCheck = OuterToCheck->GetOuter();
+			}
+			
+			if (Object == nullptr)
+			{
+				Object = StaticFindFirstObject(UObject::StaticClass(), ObjName, EFindFirstObjectOptions::NativeFirst | EFindFirstObjectOptions::EnsureIfAmbiguous);
+			}
+			if (Object != nullptr)
+			{
+				Cls = Object->GetClass();
+			}
+		}
+	}
+	UFunction *Func = FindUField<UFunction>( Cls, FuncName );
+	// Check function params.
+	if(Func != nullptr)
+	{
+		// Find the delegate UFunction to check params
+		check(SignatureFunction != nullptr && "Invalid delegate property");
+
+		// check return type and params
+		if(	Func->NumParms == SignatureFunction->NumParms )
+		{
+			int32 Count=0;
+			for( TFieldIterator<FProperty> It1(Func),It2(SignatureFunction); Count<SignatureFunction->NumParms; ++It1,++It2,++Count )
+			{
+				if( It1->GetClass()!=It2->GetClass() || (It1->PropertyFlags&CPF_OutParm)!=(It2->PropertyFlags&CPF_OutParm) )
+				{
+					ErrorText->Logf(ELogVerbosity::Warning,TEXT("Function %s does not match param types with delegate signature %s"), *Func->GetName(), *SignatureFunction->GetName());
+					Func = nullptr;
+					break;
+				}
+			}
+		}
+		else
+		{
+			ErrorText->Logf(ELogVerbosity::Warning,TEXT("Function %s does not match number of params with delegate signature %s"),*Func->GetName(), *SignatureFunction->GetName());
+			Func = nullptr;
+		}
+	}
+	else 
+	{
+		ErrorText->Logf(ELogVerbosity::Warning,TEXT("Unable to find function %s in object %s for delegate (found class: %s)"),FuncName,ObjName,*GetNameSafe(Cls));
+	}
+
+	//UE_LOG(LogProperty, Log, TEXT("... importing delegate FunctionName:'%s'(%s)   Object:'%s'(%s)"),Func != nullptr ? *Func->GetName() : TEXT("nullptr"), FuncName, Object != nullptr ? *Object->GetFullName() : TEXT("nullptr"), ObjName);
+
+	Delegate.BindUFunction( Func ? Object : nullptr, Func ? Func->GetFName() : NAME_None );
+
+	return ( Func != nullptr && Object != nullptr ) ? Buffer : nullptr;
+}
diff --git a/Engine/Source/Runtime/CoreUObject/Private/UObject/PropertyHelper.h b/Engine/Source/Runtime/CoreUObject/Private/UObject/PropertyHelper.h
--- a/Engine/Source/Runtime/CoreUObject/Private/UObject/PropertyHelper.h
+++ b/Engine/Source/Runtime/CoreUObject/Private/UObject/PropertyHelper.h
@@ -1,14 +1,12 @@
 // Copyright Epic Games, Inc. All Rights Reserved.
 #pragma once
 
-#include "CoreMinimal.h"
-#include "UObject/ObjectMacros.h"
-#include "UObject/Object.h"
-#include "UObject/Class.h"
-#include "UObject/CoreNet.h"
-#include "UObject/UnrealType.h"
-#include "UObject/Package.h"
-#include "UObject/PropertyPortFlags.h"
+#include "CoreTypes.h"
+#include "Logging/LogMacros.h"
+#include "UObject/WeakObjectPtrFwd.h"
+
+class UFunction;
+class UObject;
 
 DECLARE_LOG_CATEGORY_EXTERN(LogProperty, Log, All);
 
@@ -17,13 +15,7 @@ DECLARE_LOG_CATEGORY_EXTERN(LogProperty, Log, All);
  * 
  * @param	Str		the buffer to remove whitespace from
  */
-static void SkipWhitespace(const TCHAR*& Str)
-{
-	while(FChar::IsWhitespace(*Str))
-	{
-		Str++;
-	}
-}
+COREUOBJECT_API void SkipWhitespace(const TCHAR*& Str);
 
 /**
  * Determine whether the editable properties of CompA and CompB are identical. Used
@@ -35,65 +27,7 @@ static void SkipWhitespace(const TCHAR*& Str)
  *
  * @return	true if the values of all of the editable properties of CompA match the values in CompB
  */
-static bool AreInstancedObjectsIdentical( UObject* ObjectA, UObject* ObjectB, uint32 PortFlags )
-{
-	check(ObjectA && ObjectB);
-	bool Result = true;
-
-	// we check for recursive comparison here...if we have encountered this pair before on the call stack, the graphs are isomorphic
-	struct FRecursionCheck
-	{
-		UObject* ObjectA;
-		UObject* ObjectB;
-		uint32 PortFlags;
-
-		bool operator==(const FRecursionCheck& Other) const
-		{
-			return ObjectA == Other.ObjectA && ObjectB == Other.ObjectB && PortFlags == Other.PortFlags;
-		}
-	};
-	static TArray<FRecursionCheck> RecursionCheck;
-	FRecursionCheck Test;
-	Test.ObjectA = ObjectA;
-	Test.ObjectB = ObjectB;
-	Test.PortFlags = PortFlags;
-	if (!RecursionCheck.Contains(Test))
-	{
-		RecursionCheck.Push(Test);
-		for ( FProperty* Prop = ObjectA->GetClass()->PropertyLink; Prop && Result; Prop = Prop->PropertyLinkNext )
-		{
-			// only the properties that could have been modified in the editor should be compared
-			// (skipping the name and archetype properties, since name will almost always be different)
-			bool bConsiderProperty;
-			if ( (PortFlags&PPF_Copy) == 0 )
-			{
-				bConsiderProperty = Prop->ShouldDuplicateValue();
-			}
-			else
-			{
-				bConsiderProperty = (Prop->PropertyFlags&CPF_Edit) != 0;
-			}
-
-			if ( bConsiderProperty )
-			{
-				for ( int32 i = 0; i < Prop->ArrayDim && Result; i++ )
-				{
-					if ( !Prop->Identical_InContainer(ObjectA, ObjectB, i, PortFlags) )
-					{
-						Result = false;
-					}
-				}
-			}
-		}
-		if (Result)
-		{
-			// Allow the component to compare its native/ intrinsic properties.
-			Result = ObjectA->AreNativePropertiesIdenticalTo( ObjectB );
-		}
-		RecursionCheck.Pop();
-	}
-	return Result;
-}
+COREUOBJECT_API bool AreInstancedObjectsIdentical( UObject* ObjectA, UObject* ObjectB, uint32 PortFlags );
 
 /** Helper struct for serializing index deltas
  *		-Serialize delta index as a packed int (Hope to get 1 byte per index)
@@ -101,50 +35,13 @@ static bool AreInstancedObjectsIdentical( UObject* ObjectA, UObject* ObjectB, ui
  */
 struct FDeltaIndexHelper
 {
-	FDeltaIndexHelper()
-	{
-		LastIndex = -1; // Start at -1 so index 0 can be serialized as delta=1 (so that 0 can be reserved for 'no more')
-		LastIndexFull = -1; // Separate index for full state since it will never be rolled back
-	}
+	FDeltaIndexHelper() = default;
 
 	/** Serialize Index as delta from previous index. Return false if we should stop */
-	bool SerializeNext(FArchive &Ar, int32& Index)
-	{
-		NET_CHECKSUM(Ar);
-		if (Ar.IsSaving())
-		{
-			uint32 Delta = Index - LastIndex;
-			Ar.SerializeIntPacked(Delta);
-			LastIndex = Index;
-			LastIndexFull = Index;
-		}
-		else
-		{
-			uint32 Delta = 0;
-			Ar.SerializeIntPacked(Delta);
-			Index = (Delta == 0 ? INDEX_NONE : LastIndex + Delta);
-			LastIndex = Index;
-			LastIndexFull = Index;
-		}
-
-		return (Index != INDEX_NONE && !Ar.IsError());
-	}
+	COREUOBJECT_API bool SerializeNext(FArchive &Ar, int32& Index);
 
 	/** Helper for NetSerializeItemDelta which has full/partial/old archives. Wont auto-advance LastIndex, must call Increment after this */
-	void SerializeNext(FArchive &OutBunch, FArchive &OutFull, int32 Index)
-	{
-		NET_CHECKSUM(OutBunch);
-		NET_CHECKSUM(OutFull);
-		
-		// Full state
-		uint32 DeltaFull = Index - LastIndexFull;
-		OutFull.SerializeIntPacked(DeltaFull);
-		LastIndexFull = Index;
-
-		// Delta State
-		uint32 Delta = Index - LastIndex;
-		OutBunch.SerializeIntPacked(Delta);
-	}
+	COREUOBJECT_API void SerializeNext(FArchive &OutBunch, FArchive &OutFull, int32 Index);
 
 	/** Sets LastIndex for delta state. Must be called if using SerializeNet(OutBunch, Outfull, Old, Index) */
 	void Increment(int32 NewIndex)
@@ -153,21 +50,14 @@ struct FDeltaIndexHelper
 	}
 
 	/** Serialize early end (0) */
-	void SerializeEarlyEnd(FArchive &Ar)
-	{
-		NET_CHECKSUM(Ar);
+	COREUOBJECT_API void SerializeEarlyEnd(FArchive &Ar);
 
-		uint32 End = 0;
-		Ar.SerializeIntPacked(End);
-	}
-
-	int32 LastIndex;
-	int32 LastIndexFull;
+	int32 LastIndex     = -1; // Start at -1 so index 0 can be serialized as delta=1 (so that 0 can be reserved for 'no more')
+	int32 LastIndexFull = -1; // Separate index for full state since it will never be rolled back
 };
 
 namespace DelegatePropertyTools
 {
-
 	/**
 	 * Imports a single-cast delegate as "object.function", or "function" (self is object) from a text buffer
 	 *
@@ -179,156 +69,5 @@ namespace DelegatePropertyTools
 	 *
 	 * @return	The adjusted text buffer pointer on success, or NULL on failure
 	 */
-	static const TCHAR* ImportDelegateFromText( FScriptDelegate& Delegate, const UFunction* SignatureFunction, const TCHAR* Buffer, UObject* Parent, FOutputDevice* ErrorText )
-	{
-		SkipWhitespace(Buffer);
-
-		TCHAR ObjName[NAME_SIZE];
-		TCHAR FuncName[NAME_SIZE];
-
-		// Get object name
-		int32 i;
-		while (*Buffer == TEXT('('))
-		{
-			++Buffer;
-		}
-		for( i=0; *Buffer && *Buffer != TCHAR('.') && *Buffer != TCHAR(')') && *Buffer != TCHAR(','); Buffer++ )
-		{
-			ObjName[i++] = *Buffer;
-		}
-		ObjName[i] = TCHAR('\0');
-		UClass* Cls = nullptr;
-		UObject* Object = nullptr;
-		bool bIsNullClass = false;
-		
-		// Handle (null).None syntax used for empty delegates in ExportText
-		if (*Buffer == TCHAR(')') && FCString::Strcmp(ObjName, TEXT("null")) == 0)
-		{
-			++Buffer;
-			bIsNullClass = true;
-		}
-
-		// Get function name
-		if (*Buffer == TCHAR('.'))
-		{
-			Buffer++;
-			for( i=0; *Buffer && *Buffer != TCHAR(')') && *Buffer != TCHAR(','); Buffer++ )
-			{
-				FuncName[i++] = *Buffer;
-			}
-			FuncName[i] = TEXT('\0');
-		}
-		else
-		{
-			// if there's no dot, then a function name was specified without any object qualifier
-			// if we're importing for a subobject, assume the parent object as the source
-			// otherwise, assume Parent for the source
-			// if neither are valid, then the importing fails
-			if (Parent == nullptr)
-			{
-				FString DelegateString = Delegate.GetFunctionName().ToString();
-				FString SignatureFunctionString = GetNameSafe(SignatureFunction);
-				ErrorText->Logf(ELogVerbosity::Warning, TEXT("Cannot import unqualified delegate name; no object to search. Delegate=%s SignatureFunction=%s"),
-					*DelegateString, *SignatureFunctionString);
-				return nullptr;
-			}
-			// since we don't support nested components, we only need to check one level deep
-			else if (!Parent->HasAnyFlags(RF_ClassDefaultObject) && Parent->GetOuter()->HasAnyFlags(RF_ClassDefaultObject))
-			{
-				Object = Parent->GetOuter();
-				Cls = Parent->GetOuter()->GetClass();
-			}
-			else
-			{
-				Object = Parent;
-				Cls = Parent->GetClass();
-			}
-			FCString::Strncpy(FuncName, ObjName, i + 1);
-		}
-
-		if (bIsNullClass)
-		{
-			FName ParsedName(FuncName);
-			if (ParsedName == NAME_None)
-			{
-				// Deliberately null
-				Delegate.Unbind();
-				return Buffer;
-			}
-			else
-			{
-				FString DelegateString = Delegate.GetFunctionName().ToString();
-				FString SignatureFunctionString = GetNameSafe(SignatureFunction);
-				ErrorText->Logf(ELogVerbosity::Warning, TEXT("Cannot import delegate with function name %s but null class. Delegate=%s SignatureFunction=%s"),
-					*FuncName, *DelegateString, *SignatureFunctionString);
-				return nullptr;
-			}
-		}
-
-		if (Cls == nullptr)
-		{
-			Cls = UClass::TryFindTypeSlow<UClass>(ObjName);
-			if (Cls != nullptr)
-			{
-				Object = Cls->GetDefaultObject();
-			}
-			else
-			{
-				// Check outer chain before checking all packages
-				UObject* OuterToCheck = Parent;
-				while (Object == nullptr && OuterToCheck)
-				{
-					Object = StaticFindObject(UObject::StaticClass(), OuterToCheck, ObjName);
-					OuterToCheck = OuterToCheck->GetOuter();
-				}
-				
-				if (Object == nullptr)
-				{
-					Object = StaticFindFirstObject(UObject::StaticClass(), ObjName, EFindFirstObjectOptions::NativeFirst | EFindFirstObjectOptions::EnsureIfAmbiguous);
-				}
-				if (Object != nullptr)
-				{
-					Cls = Object->GetClass();
-				}
-			}
-		}
-		UFunction *Func = FindUField<UFunction>( Cls, FuncName );
-		// Check function params.
-		if(Func != nullptr)
-		{
-			// Find the delegate UFunction to check params
-			check(SignatureFunction != nullptr && "Invalid delegate property");
-
-			// check return type and params
-			if(	Func->NumParms == SignatureFunction->NumParms )
-			{
-				int32 Count=0;
-				for( TFieldIterator<FProperty> It1(Func),It2(SignatureFunction); Count<SignatureFunction->NumParms; ++It1,++It2,++Count )
-				{
-					if( It1->GetClass()!=It2->GetClass() || (It1->PropertyFlags&CPF_OutParm)!=(It2->PropertyFlags&CPF_OutParm) )
-					{
-						ErrorText->Logf(ELogVerbosity::Warning,TEXT("Function %s does not match param types with delegate signature %s"), *Func->GetName(), *SignatureFunction->GetName());
-						Func = nullptr;
-						break;
-					}
-				}
-			}
-			else
-			{
-				ErrorText->Logf(ELogVerbosity::Warning,TEXT("Function %s does not match number of params with delegate signature %s"),*Func->GetName(), *SignatureFunction->GetName());
-				Func = nullptr;
-			}
-		}
-		else 
-		{
-			ErrorText->Logf(ELogVerbosity::Warning,TEXT("Unable to find function %s in object %s for delegate (found class: %s)"),FuncName,ObjName,*GetNameSafe(Cls));
-		}
-
-		//UE_LOG(LogProperty, Log, TEXT("... importing delegate FunctionName:'%s'(%s)   Object:'%s'(%s)"),Func != nullptr ? *Func->GetName() : TEXT("nullptr"), FuncName, Object != nullptr ? *Object->GetFullName() : TEXT("nullptr"), ObjName);
-	
-		Delegate.BindUFunction( Func ? Object : nullptr, Func ? Func->GetFName() : NAME_None );
-
-		return ( Func != nullptr && Object != nullptr ) ? Buffer : nullptr;
-	}
+	COREUOBJECT_API const TCHAR* ImportDelegateFromText( FScriptDelegate& Delegate, const UFunction* SignatureFunction, const TCHAR* Buffer, UObject* Parent, FOutputDevice* ErrorText );
 }
-
