Reapplied changes that were overwritten by the merge in CL# 24255812.

diff --git a/Engine/Source/Runtime/CoreUObject/Private/UObject/PropertyHelper.cpp b/Engine/Source/Runtime/CoreUObject/Private/UObject/PropertyHelper.cpp
--- a/Engine/Source/Runtime/CoreUObject/Private/UObject/PropertyHelper.cpp
+++ b/Engine/Source/Runtime/CoreUObject/Private/UObject/PropertyHelper.cpp
@@ -130,6 +130,7 @@ const TCHAR* DelegatePropertyTools::ImportDelegateFromText( FScriptDelegate& Del
 
 	TCHAR ObjName[NAME_SIZE];
 	TCHAR FuncName[NAME_SIZE];
+
 	// Get object name
 	int32 i;
 	while (*Buffer == TEXT('('))
@@ -143,6 +144,15 @@ const TCHAR* DelegatePropertyTools::ImportDelegateFromText( FScriptDelegate& Del
 	ObjName[i] = TCHAR('\0');
 	UClass* Cls = nullptr;
 	UObject* Object = nullptr;
+	bool bIsNullClass = false;
+
+	// Handle (null).None syntax used for empty delegates in ExportText
+	if (*Buffer == TCHAR(')') && FCString::Strcmp(ObjName, TEXT("null")) == 0)
+	{
+		++Buffer;
+		bIsNullClass = true;
+	}
+
 	// Get function name
 	if (*Buffer == TCHAR('.'))
 	{
@@ -180,6 +190,26 @@ const TCHAR* DelegatePropertyTools::ImportDelegateFromText( FScriptDelegate& Del
 		}
 		FCString::Strncpy(FuncName, ObjName, i + 1);
 	}
+
+	if (bIsNullClass)
+	{
+		FName ParsedName(FuncName);
+		if (ParsedName == NAME_None)
+		{
+			// Deliberately null
+			Delegate.Unbind();
+			return Buffer;
+		}
+		else
+		{
+			FString DelegateString = Delegate.GetFunctionName().ToString();
+			FString SignatureFunctionString = GetNameSafe(SignatureFunction);
+			ErrorText->Logf(ELogVerbosity::Warning, TEXT("Cannot import delegate with function name %s but null class. Delegate=%s SignatureFunction=%s"),
+				*FuncName, *DelegateString, *SignatureFunctionString);
+			return nullptr;
+		}
+	}
+
 	if (Cls == nullptr)
 	{
 		Cls = UClass::TryFindTypeSlow<UClass>(ObjName);
