Re-added commented out includes.

diff --git a/Engine/Source/Runtime/CoreUObject/Private/UObject/PropertyHelper.cpp b/Engine/Source/Runtime/CoreUObject/Private/UObject/PropertyHelper.cpp
--- a/Engine/Source/Runtime/CoreUObject/Private/UObject/PropertyHelper.cpp
+++ b/Engine/Source/Runtime/CoreUObject/Private/UObject/PropertyHelper.cpp
@@ -2,14 +2,14 @@
 
 #include "UObject/PropertyHelper.h"
 
-//#include "CoreMinimal.h"
-//#include "UObject/ObjectMacros.h"
-//#include "UObject/Object.h"
-//#include "UObject/Class.h"
-//#include "UObject/CoreNet.h"
-//#include "UObject/UnrealType.h"
-//#include "UObject/Package.h"
-//#include "UObject/PropertyPortFlags.h"
+#include "CoreMinimal.h"
+#include "UObject/ObjectMacros.h"
+#include "UObject/Object.h"
+#include "UObject/Class.h"
+#include "UObject/CoreNet.h"
+#include "UObject/UnrealType.h"
+#include "UObject/Package.h"
+#include "UObject/PropertyPortFlags.h"
 
 void SkipWhitespace(const TCHAR*& Str)
 {
