Added missing includes for unity build, caused by CL# 24255189

diff --git a/Engine/Source/Runtime/CoreUObject/Private/UObject/Class.cpp b/Engine/Source/Runtime/CoreUObject/Private/UObject/Class.cpp
--- a/Engine/Source/Runtime/CoreUObject/Private/UObject/Class.cpp
+++ b/Engine/Source/Runtime/CoreUObject/Private/UObject/Class.cpp
@@ -14,6 +14,7 @@
 #include "Misc/AutomationTest.h"
 #include "Misc/EnumClassFlags.h"
 #include "Misc/StringBuilder.h"
+#include "UObject/CoreNet.h"
 #include "Modules/ModuleManager.h"
 #include "UObject/UObjectAllocator.h"
 #include "UObject/UObjectHash.h"
@@ -4852,7 +4853,7 @@ void UClass::Serialize( FArchive& Ar )
 	//@TODO: UCREMOVAL: Should we just regenerate the FuncMap post load, instead of serializing it?
 	{
 		FWriteScopeLock ScopeLock(FuncMapLock);
-		Ar << FuncMap;
+	Ar << FuncMap;
 	}
 
 	// Class flags first.
@@ -5273,7 +5274,7 @@ void UClass::PurgeClass(bool bRecompilingOnLoad)
 
 	{
 		FWriteScopeLock ScopeLock(FuncMapLock);
-		FuncMap.Empty();
+	FuncMap.Empty();
 	}
 	ClearFunctionMapsCaches();
 	PropertyLink = nullptr;
diff --git a/Engine/Source/Runtime/CoreUObject/Private/UObject/CoreRedirects.cpp b/Engine/Source/Runtime/CoreUObject/Private/UObject/CoreRedirects.cpp
--- a/Engine/Source/Runtime/CoreUObject/Private/UObject/CoreRedirects.cpp
+++ b/Engine/Source/Runtime/CoreUObject/Private/UObject/CoreRedirects.cpp
@@ -6,6 +6,7 @@
 #include "UObject/Package.h"
 #include "UObject/PropertyHelper.h"
 #include "UObject/TopLevelAssetPath.h"
+#include "UObject/UnrealType.h"
 
 #include "GenericPlatform/GenericPlatformFile.h"
 #include "Misc/AutomationTest.h"
diff --git a/Engine/Source/Runtime/CoreUObject/Private/UObject/Property.cpp b/Engine/Source/Runtime/CoreUObject/Private/UObject/Property.cpp
--- a/Engine/Source/Runtime/CoreUObject/Private/UObject/Property.cpp
+++ b/Engine/Source/Runtime/CoreUObject/Private/UObject/Property.cpp
@@ -24,6 +24,7 @@
 #include "Math/Ray.h"
 #include "Math/Sphere.h"
 #include "Math/InterpCurvePoint.h"
+#include "UObject/Package.h"
 #include "UObject/ReleaseObjectVersion.h"
 
 DEFINE_LOG_CATEGORY(LogProperty);
diff --git a/Engine/Source/Runtime/CoreUObject/Private/UObject/PropertyBaseObject.cpp b/Engine/Source/Runtime/CoreUObject/Private/UObject/PropertyBaseObject.cpp
--- a/Engine/Source/Runtime/CoreUObject/Private/UObject/PropertyBaseObject.cpp
+++ b/Engine/Source/Runtime/CoreUObject/Private/UObject/PropertyBaseObject.cpp
@@ -3,10 +3,12 @@
 #include "CoreMinimal.h"
 #include "UObject/ObjectMacros.h"
 #include "Templates/Casts.h"
+#include "UObject/Package.h"
 #include "UObject/UnrealType.h"
 #include "UObject/UnrealTypePrivate.h"
 #include "UObject/UObjectHash.h"
 #include "Blueprint/BlueprintSupport.h"
+#include "UObject/CoreNet.h"
 #include "UObject/PropertyHelper.h"
 #include "UObject/LinkerPlaceholderClass.h"
 #include "UObject/LinkerPlaceholderExportObject.h"
