Always allow FObjectProperty::ConvertFromType to syncload an object if it's called on the game thread

diff --git a/Engine/Source/Runtime/CoreUObject/Private/UObject/PropertyObject.cpp b/Engine/Source/Runtime/CoreUObject/Private/UObject/PropertyObject.cpp
--- a/Engine/Source/Runtime/CoreUObject/Private/UObject/PropertyObject.cpp
+++ b/Engine/Source/Runtime/CoreUObject/Private/UObject/PropertyObject.cpp
@@ -50,7 +50,7 @@ EConvertFromTypeResult FObjectProperty::ConvertFromType(const FPropertyTag& Tag,
 		UObject* PreviousValueObj = nullptr;
 
 		// If we're async loading it's not safe to do a sync load because it may crash or fail to set the variable, so throw an error if it's not already in memory
-		if (IsInAsyncLoadingThread())
+		if (!IsInGameThread())
 		{
 			PreviousValueObj = PreviousValue.Get();
 
