No PCH build fix

#rnx

diff --git a/Engine/Source/Runtime/CoreUObject/Public/UObject/Class.h b/Engine/Source/Runtime/CoreUObject/Public/UObject/Class.h
--- a/Engine/Source/Runtime/CoreUObject/Public/UObject/Class.h
+++ b/Engine/Source/Runtime/CoreUObject/Public/UObject/Class.h
@@ -4064,9 +4064,9 @@ template< class T > struct TVariantStructure
 };
 
 #define UE_DECLARE_CORE_VARIANT_TYPE(VARIANT, CORE)														\
-template<> struct TBaseStructure<F##CORE> { COREUOBJECT_API static UScriptStruct* Get(); };				\
-template<> struct TVariantStructure<F##VARIANT##f> { COREUOBJECT_API static UScriptStruct* Get(); };	\
-template<> struct TVariantStructure<F##VARIANT##d> { COREUOBJECT_API static UScriptStruct* Get(); };
+template<> struct TBaseStructure<::F##CORE> { COREUOBJECT_API static UScriptStruct* Get(); };			\
+template<> struct TVariantStructure<::F##VARIANT##f> { COREUOBJECT_API static UScriptStruct* Get(); };	\
+template<> struct TVariantStructure<::F##VARIANT##d> { COREUOBJECT_API static UScriptStruct* Get(); };
 
 UE_DECLARE_CORE_VARIANT_TYPE(Vector2,	Vector2D);
 UE_DECLARE_CORE_VARIANT_TYPE(Vector3,	Vector);
