Implemented FExportArchive::GetArchiveName() to get more information in logs.

diff --git a/Engine/Source/Runtime/CoreUObject/Private/Serialization/AsyncLoading2.cpp b/Engine/Source/Runtime/CoreUObject/Private/Serialization/AsyncLoading2.cpp
--- a/Engine/Source/Runtime/CoreUObject/Private/Serialization/AsyncLoading2.cpp
+++ b/Engine/Source/Runtime/CoreUObject/Private/Serialization/AsyncLoading2.cpp
@@ -1612,6 +1612,11 @@ public:
 		ActiveFPLB->StartFastPathLoadBuffer += Length;
 	}
 
+	virtual FString GetArchiveName() const override
+	{
+		return PackageDesc ? PackageDesc->UPackageName.ToString() : TEXT("FExportArchive");
+	}
+
 	void UsingCustomVersion(const FGuid& Key) override {};
 	using FArchive::operator<<; // For visibility of the overloads we don't override
 
