Crash fix when collecting references in GC with gc.historysize 1

diff --git a/Engine/Source/Runtime/CoreUObject/Private/UObject/GarbageCollection.cpp b/Engine/Source/Runtime/CoreUObject/Private/UObject/GarbageCollection.cpp
--- a/Engine/Source/Runtime/CoreUObject/Private/UObject/GarbageCollection.cpp
+++ b/Engine/Source/Runtime/CoreUObject/Private/UObject/GarbageCollection.cpp
@@ -2851,7 +2851,15 @@ private:
 		}
 		else if (FGCObject::GGCObjectReferencer == Referencer)
 		{
-			Ref.ReferencerName = *FGCObject::GGCObjectReferencer->GetCurrentlySerializingObject()->GetReferencerName();
+			FGCObject* SerializingObj = FGCObject::GGCObjectReferencer->GetCurrentlySerializingObject();
+			if(SerializingObj)
+			{
+				Ref.ReferencerName = *SerializingObj->GetReferencerName();
+			}			
+			else
+			{
+				Ref.ReferencerName = TEXT("Unknown");
+			}
 		}
 
 		TArray<FGCDirectReference>*& DirectReferences = Context.History.FindOrAdd(Referencer);
