Deprecate unused EGCOptions::ProcessWeakReferences

diff --git a/Engine/Source/Runtime/CoreUObject/Public/UObject/FastReferenceCollector.h b/Engine/Source/Runtime/CoreUObject/Public/UObject/FastReferenceCollector.h
--- a/Engine/Source/Runtime/CoreUObject/Public/UObject/FastReferenceCollector.h
+++ b/Engine/Source/Runtime/CoreUObject/Public/UObject/FastReferenceCollector.h
@@ -32,14 +32,17 @@ enum class EGCOptions : uint32
 	Parallel = 1 << 0,					// Use all task workers to collect references, must be started on main thread
 	AutogenerateTokenStream = 1 << 1,	// Lazily generate token streams for new UClasses
 	WithClusters = 1 << 2,				// Use clusters, see FUObjectCluster
-	ProcessWeakReferences = 1 << 3,		// Collect both weak and strong references instead of just strong
+	ProcessWeakReferences UE_DEPRECATED(5.2, "Weak reference collection deprecated to reduce complexity. Use FArchive instead with ArIsObjectReferenceCollector and ArIsModifyingWeakAndStrongReferences")
+					= 1 << 3,			// Collect both weak and strong references instead of just strong
 	WithPendingKill = 1 << 4,			// Internal flag used by reachability analysis
 };
 ENUM_CLASS_FLAGS(EGCOptions);
 
 inline constexpr bool IsParallel(EGCOptions Options) { return !!(Options & EGCOptions::Parallel); }
 inline constexpr bool IsPendingKill(EGCOptions Options) { return !!(Options & EGCOptions::WithPendingKill); }
+PRAGMA_DISABLE_DEPRECATION_WARNINGS
 inline constexpr bool ShouldProcessWeakReferences(EGCOptions Options) { return !!(Options & EGCOptions::ProcessWeakReferences); }
+PRAGMA_ENABLE_DEPRECATION_WARNINGS
 
 using EFastReferenceCollectorOptions UE_DEPRECATED(5.2, "Use EGCOptions instead") = EGCOptions;
 
@@ -531,7 +534,9 @@ class TFastReferenceCollector : public FGCInternals
 	static constexpr EGCOptions Options = ProcessorType::Options;
 
 	static constexpr FORCEINLINE bool IsParallel()					{ return !!(Options & EGCOptions::Parallel); }
+PRAGMA_DISABLE_DEPRECATION_WARNINGS
 	static constexpr FORCEINLINE bool ShouldProcessWeakReferences()	{ return !!(Options & EGCOptions::ProcessWeakReferences); }
+PRAGMA_ENABLE_DEPRECATION_WARNINGS
 
 	static FORCEINLINE FTokenId MakeTokenId(uint32 Index)
 	{
