Don't set an invalid size on FBulkData if the serialized element count is too large.

#rnx

- The main issue we are trying to solve is that we serialize the number of elements when saving/loading the package. On load we multiple the number of elements (int64) by the size of each element (int64) and if either of those values is large enough we could overflow and end up with an incorrect size. We do check for this when saving the package so this is only likely to come up if the package has been corrupted in some way.
- Change FBulkMetaData::FromSerialized to do the actual serialization as well. This lets us set the error flag on the archive if we do detect problems and keep all the error handling in a single place.
-- Note that setting the error flag on an archive while a package is being loaded doesn't seem to do much at the moment, so this is more of an attempt to future proof.
- If we detect an error with serializing FBulkMetaResource we will log an error and immediately return so that FBulkMetaData will remain set to defaults.
- We now calculate the largest value of ElementCount that we can support, given the ElementCount size and if the count is exceeded we will log an error, set the error flag on the archive and return leaving FBulkMetaData as default.
- Note that both before and after this change the package is effectively broken and won't work properly, but this change will stop us from potentially stomping memory.
- Note that at runtime we generally serialize the bulkdata info from the BulkDataMap in the package header, so this is mostly an editor thing.

diff --git a/Engine/Source/Runtime/CoreUObject/Private/Serialization/AsyncLoading2.cpp b/Engine/Source/Runtime/CoreUObject/Private/Serialization/AsyncLoading2.cpp
--- a/Engine/Source/Runtime/CoreUObject/Private/Serialization/AsyncLoading2.cpp
+++ b/Engine/Source/Runtime/CoreUObject/Private/Serialization/AsyncLoading2.cpp
@@ -1828,10 +1828,7 @@ private:
 
 		if (UNLIKELY(HeaderData->BulkDataMap.IsEmpty()))
 		{
-			FBulkMetaResource SerializedMeta;
-			Ar << SerializedMeta;
-			Meta = FBulkMetaData::FromSerialized(SerializedMeta, ElementSize);
-			DuplicateSerialOffset = SerializedMeta.DuplicateOffset;
+			FBulkMetaData::FromSerialized(Ar, ElementSize,  Meta, DuplicateSerialOffset);
 		}
 		else
 		{
diff --git a/Engine/Source/Runtime/CoreUObject/Private/Serialization/BulkData.cpp b/Engine/Source/Runtime/CoreUObject/Private/Serialization/BulkData.cpp
--- a/Engine/Source/Runtime/CoreUObject/Private/Serialization/BulkData.cpp
+++ b/Engine/Source/Runtime/CoreUObject/Private/Serialization/BulkData.cpp
@@ -195,10 +195,85 @@ FArchive& operator<<(FArchive& Ar, FBulkMetaResource& BulkMeta)
 	return Ar;
 }
 
+/** 
+ * Attempt to return the best debug name we can from a given archive.
+ * First we check to see the archive has a FLinkerLoad from which we
+ * can find the name of the package being loaded from. If not we 
+ * return an unknown string as the FArchive name is often not very 
+ * useful for identification purposes.
+ */
+FString GetDebugNameFromArchive(FArchive& Ar)
+{
+	if (FLinkerLoad* LinkerLoad = Cast<FLinkerLoad>(Ar.GetLinker()))
+	{
+		return LinkerLoad->GetDebugName();
+	}
+	else
+	{
+		return TEXT("Unknown");
+	}
+}
+
+bool FBulkMetaData::FromSerialized(FArchive& Ar, int64 ElementSize, FBulkMetaData& OutMetaData, int64& OutDuplicateOffset)
+{
+	FBulkMetaResource Resource;
+	Ar << Resource;
+
+	if (Ar.IsError())
+	{
+		// Note that setting the error flag on the archive is not enough to stop the package from being loaded so for now we 
+		// need to fatal error to prevent the process from continuing to use the corrupted package.
+		UE_LOG(LogSerialization, Fatal, TEXT("Bulkdata error when serializing '%s', could not serialize FBulkMetaResource correctly"), *GetDebugNameFromArchive(Ar));
+		OutMetaData = FBulkMetaData();
+		return false;
+	}
+
+	if (Resource.ElementCount > 0)
+	{
+		// TODO: This would be a good use case for FGuardedInt64 once it is moved to core
+		if (Resource.ElementCount < (MAX_int64 / ElementSize))
+		{
+			OutMetaData.SetSize(Resource.ElementCount * ElementSize);
+		}
+		else
+		{
+			// We should only get here if the package is severely corrupted (which should get detected earlier in package loading)
+			// Note that setting the error flag on the archive is not enough to stop the package from being loaded so for now we 
+			// need to fatal error to prevent the process from continuing to use the corrupted package.
+			Ar.SetError();
+
+			UE_LOG(LogSerialization, Fatal,
+				TEXT("Bulkdata error when serializing '%s', ElementCount (%" INT64_FMT ") and ElementSize (%" INT64_FMT ") would cause an int64 overflow"),
+				*GetDebugNameFromArchive(Ar),
+				Resource.ElementCount,
+				ElementSize);
+
+			OutMetaData = FBulkMetaData();
+			return false;
+		}
+	}
+
+	OutMetaData.SetSizeOnDisk(Resource.SizeOnDisk);
+	OutMetaData.SetOffset(Resource.Offset);
+	OutMetaData.SetFlags(Resource.Flags);
+
+	check(Resource.ElementCount <= 0 || OutMetaData.GetSize() == Resource.ElementCount * ElementSize);
+	check(OutMetaData.GetOffset() == Resource.Offset);
+	check(OutMetaData.GetFlags() == Resource.Flags);
+
+#if !USE_RUNTIME_BULKDATA
+	check(Resource.ElementCount <= 0 || OutMetaData.GetSizeOnDisk() == Resource.SizeOnDisk);
+#endif
+
+	OutDuplicateOffset = Resource.DuplicateOffset;
+
+	return true;
+}
+
 } // namespace UE::BulkData::Private
 
 /*-----------------------------------------------------------------------------
-	Memory managament
+	Memory management
 -----------------------------------------------------------------------------*/
 
 void FBulkData::FAllocatedPtr::Free(FBulkData* Owner)
@@ -1198,9 +1273,7 @@ void FBulkData::Serialize(FArchive& Ar, UObject* Owner, bool bAttemptFileMapping
 		{
 			checkf(IsUnlocked(), TEXT("Serialize bulk data FAILED, bulk data is locked"));
 
-			FBulkMetaResource MetaResource;
-			Ar << MetaResource;
-			BulkMeta = FBulkMetaData::FromSerialized(MetaResource, ElementSize);
+			FBulkMetaData::FromSerialized(Ar, ElementSize, BulkMeta);
 			BulkChunkId = FIoChunkId::InvalidChunkId;
 
 			const int64 BulkDataSize = GetBulkDataSize();
diff --git a/Engine/Source/Runtime/CoreUObject/Private/UObject/LinkerLoad.cpp b/Engine/Source/Runtime/CoreUObject/Private/UObject/LinkerLoad.cpp
--- a/Engine/Source/Runtime/CoreUObject/Private/UObject/LinkerLoad.cpp
+++ b/Engine/Source/Runtime/CoreUObject/Private/UObject/LinkerLoad.cpp
@@ -6836,10 +6836,7 @@ void FLinkerLoad::SerializeBulkMeta(UE::BulkData::Private::FBulkMetaData& Meta,
 
 	if (DataResourceMap.IsEmpty())
 	{
-		FBulkMetaResource SerializedMeta;
-		Ar << SerializedMeta;
-		Meta = FBulkMetaData::FromSerialized(SerializedMeta, ElementSize);
-		DuplicateSerialOffset = SerializedMeta.DuplicateOffset;
+		FBulkMetaData::FromSerialized(Ar, ElementSize, Meta, DuplicateSerialOffset);
 	}
 	else
 	{
diff --git a/Engine/Source/Runtime/CoreUObject/Public/Serialization/BulkData.h b/Engine/Source/Runtime/CoreUObject/Public/Serialization/BulkData.h
--- a/Engine/Source/Runtime/CoreUObject/Public/Serialization/BulkData.h
+++ b/Engine/Source/Runtime/CoreUObject/Public/Serialization/BulkData.h
@@ -446,26 +446,17 @@ public:
 		return (GetFlags() & Flags) == Flags;
 	}
 
-	static FBulkMetaData FromSerialized(const FBulkMetaResource& MetaResource, int64 ElementSize)
-	{
-		FBulkMetaData Meta;
-		
-		if (MetaResource.ElementCount > 0)
-		{
-			Meta.SetSize(MetaResource.ElementCount * ElementSize);
-		}
-
-		Meta.SetSizeOnDisk(MetaResource.SizeOnDisk);
-		Meta.SetOffset(MetaResource.Offset);
-		Meta.SetFlags(MetaResource.Flags);
+	/** 
+	 * Serializes FBulkMetaResource from the given archive and builds the returned FBulkMetaData from it. 
+	 * The offset for duplicated data will also be returned
+	 */
+	COREUOBJECT_API static bool FromSerialized(FArchive& Ar, int64 ElementSize, FBulkMetaData& OutMetaData, int64& OutDuplicateOffset);
 
-		check(MetaResource.ElementCount <= 0 || Meta.GetSize() == MetaResource.ElementCount * ElementSize);
-		check(Meta.GetOffset() == MetaResource.Offset);
-		check(Meta.GetFlags() == MetaResource.Flags);
-#if !USE_RUNTIME_BULKDATA
-		check(MetaResource.ElementCount <= 0 || Meta.GetSizeOnDisk() == MetaResource.SizeOnDisk);
-#endif
-		return Meta;
+	/** Serializes FBulkMetaResource from the given archive and builds the returned FBulkMetaData from it. */
+	static bool FromSerialized(FArchive& Ar, int64 ElementSize, FBulkMetaData& OutMetaData)
+	{
+		int64 DuplicateOffset = INDEX_NONE;
+		return FromSerialized(Ar, ElementSize, OutMetaData, DuplicateOffset);
 	}
 
 private:
