Treating loading / import package errors as warnings (as the ModerationArtifactsCommandlet aims to generate artifacts as a primary goal rather than validate the input)

Ditto for any failed actor loads from world partitions

diff --git a/Engine/Source/Runtime/CoreUObject/Private/UObject/LinkerLoad.cpp b/Engine/Source/Runtime/CoreUObject/Private/UObject/LinkerLoad.cpp
--- a/Engine/Source/Runtime/CoreUObject/Private/UObject/LinkerLoad.cpp
+++ b/Engine/Source/Runtime/CoreUObject/Private/UObject/LinkerLoad.cpp
@@ -143,7 +143,7 @@ bool FLinkerLoad::ShouldCreateThrottledSlowTask() const
 	return ShouldReportProgress();
 }
 
-int32 GTreatVerifyImportErrorsAsWarnings = 0;
+COREUOBJECT_API int32 GTreatVerifyImportErrorsAsWarnings = 0;
 static FAutoConsoleVariableRef CVarTreatVerifyImportErrorsAsWarnings(
 	TEXT("linker.TreatVerifyImportErrorsAsWarnings"),
 	GTreatVerifyImportErrorsAsWarnings,
