Memory leak creating a 3D array

Add missing DestructItems FScriptArrayHelper::MoveAssign and  FScriptMapHelper::MoveAssign to ensure that destructors of elements are called.

Will be submitted in two parts, one for UnrealTypes.h, and one for the rest of the files.

diff --git a/Engine/Source/Runtime/CoreUObject/Public/UObject/UnrealType.h b/Engine/Source/Runtime/CoreUObject/Public/UObject/UnrealType.h
--- a/Engine/Source/Runtime/CoreUObject/Public/UObject/UnrealType.h
+++ b/Engine/Source/Runtime/CoreUObject/Public/UObject/UnrealType.h
@@ -4351,6 +4351,8 @@ public:
 	void MoveAssign(void* InOtherArray)
 	{
 		checkSlow(InOtherArray);
+		// FScriptArray::MoveAssign does not call destructors for our elements, so do that before calling it.
+		DestructItems(0, Num());
 		WithScriptArray([this, InOtherArray](auto* Array) { Array->MoveAssign(*static_cast<decltype(Array)>(InOtherArray), ElementSize, ElementAlignment); });
 	}
 
@@ -4678,6 +4680,8 @@ public:
 	void MoveAssign(void* InOtherMap)
 	{
 		checkSlow(InOtherMap);
+		// FScriptArray::MoveAssign does not call destructors for our elements, so do that before calling it.
+		DestructItems(0, Num());
 		return WithScriptMap([this, InOtherMap](auto* Map)
 		{
 			Map->MoveAssign(*(decltype(Map))InOtherMap, MapLayout);
