Introducing the relocation workflow.

When on, at load the packages will detect if they where moved and if this is the case they will modify their ImportMap and SoftObjectPathList. The package reference in those that were in the same MountPoint as the asset when it was saved will behave as if they were relative references.

diff --git a/Engine/Source/Runtime/AssetRegistry/Private/PackageReader.cpp b/Engine/Source/Runtime/AssetRegistry/Private/PackageReader.cpp
--- a/Engine/Source/Runtime/AssetRegistry/Private/PackageReader.cpp
+++ b/Engine/Source/Runtime/AssetRegistry/Private/PackageReader.cpp
@@ -13,6 +13,7 @@
 #include "Misc/PackagePath.h"
 #include "UObject/Class.h"
 #include "UObject/Linker.h"
+#include "UObject/PackageRelocation.h"
 #include "UObject/PackageTrailer.h"
 
 
@@ -323,6 +324,19 @@ bool FPackageReader::SerializePackageTrailer(FAssetPackageData& PackageData)
 	return true;
 }
 
+void FPackageReader::ApplyRelocationToImportMapAndSoftPackageReferenceList(FStringView LoadedPackageName, TArray<FName>& OutSoftPackageReferenceList)
+{
+#if WITH_EDITOR
+	UE::Package::Relocation::Private::FPackageRelocationContext RelocationArgs;
+	if (UE::Package::Relocation::Private::ShouldApplyRelocation(PackageFileSummary, LoadedPackageName, RelocationArgs))
+	{
+		UE_LOG(LogPackageRelocation, Verbose, TEXT("Detected relocated package (%.*s). The package was saved as (%s)."), LoadedPackageName.Len(), LoadedPackageName.GetData(), *PackageFileSummary.PackageName);
+		UE::Package::Relocation::Private::ApplyRelocationToObjectImportMap(RelocationArgs, ImportMap);
+		UE::Package::Relocation::Private::ApplyRelocationToNameArray(RelocationArgs, OutSoftPackageReferenceList);
+	}
+#endif
+}
+
 bool FPackageReader::ReadAssetDataFromThumbnailCache(TArray<FAssetData*>& AssetDataList)
 {
 	if (!StartSerializeSection(PackageFileSummary.ThumbnailTableOffset))
@@ -520,6 +534,8 @@ bool FPackageReader::ReadDependencyData(FPackageDependencyData& OutDependencyDat
 			return false;
 		}
 
+		ApplyRelocationToImportMapAndSoftPackageReferenceList(PackageNameString, SoftPackageReferenceList);
+
 		OutDependencyData.LoadDependenciesFromPackageHeader(OutDependencyData.PackageName, ImportMap, SoftPackageReferenceList,
 			SearchableNames.SearchableNamesMap, ImportUsedInGame, SoftPackageUsedInGame);
 	}
diff --git a/Engine/Source/Runtime/AssetRegistry/Private/PackageReader.h b/Engine/Source/Runtime/AssetRegistry/Private/PackageReader.h
--- a/Engine/Source/Runtime/AssetRegistry/Private/PackageReader.h
+++ b/Engine/Source/Runtime/AssetRegistry/Private/PackageReader.h
@@ -93,6 +93,8 @@ private:
 		const TArray<FObjectImport>& InImportMap, const TArray<FName>& SoftPackageReferenceList);
 	bool SerializePackageTrailer(FAssetPackageData& PackageData);
 
+	void ApplyRelocationToImportMapAndSoftPackageReferenceList(FStringView LoadedPackageName, TArray<FName>& OutSoftPackageReferenceList);
+
 	/** Returns flags the asset package was saved with */
 	uint32 GetPackageFlags() const;
 
diff --git a/Engine/Source/Runtime/CoreUObject/Internal/UObject/PackageRelocation.h b/Engine/Source/Runtime/CoreUObject/Internal/UObject/PackageRelocation.h
new file mode 100644
--- /dev/null
+++ b/Engine/Source/Runtime/CoreUObject/Internal/UObject/PackageRelocation.h
@@ -0,0 +1,65 @@
+// Copyright Epic Games, Inc.All Rights Reserved.
+
+#pragma once
+
+#include "Containers/StringFwd.h"
+#include "Containers/StringView.h"
+#include "Logging/LogMacros.h"
+#include "UObject/UnrealNames.h"
+
+struct FObjectImport;
+struct FPackageFileSummary;
+
+// Category to use for the log related to the relocation
+COREUOBJECT_API DECLARE_LOG_CATEGORY_EXTERN(LogPackageRelocation, Log, All);
+
+/**
+ * Object relocation is a process done on load to make the references in ImportMap and SoftReferencePathList that point to
+ * packages in the same mount point behave as if they were relative paths.
+ * This namespace contains the utilities to calculate relocations and is called from code that loads packages.
+ * 
+ * The API from this namespace is experimental and will be changed without deprecation
+ */
+namespace UE::Package::Relocation::Private
+{
+	/** A struct that contains the most used arguments for the relocations functions */
+	struct FPackageRelocationContext
+	{
+		FStringView OriginalPackagePath;
+		FStringView CurrentPackagePath;
+		FStringView OriginalPackageMount;
+	};
+
+
+	/**
+	 * Determine if we should relocate the package or not.
+	 * Populate the context for the relocation functions if the package should be relocated
+	 * 
+	 * Note: Keep the package summary and the source of LoadedPackageName StringView in memory as the OutPackageRelocationContext use a view on the data of these arguments
+	 */
+	COREUOBJECT_API bool ShouldApplyRelocation(const FPackageFileSummary& PackageSummary, FStringView LoadedPackageName, FPackageRelocationContext& OutPackageRelocationContext);
+
+	/**
+	 * Try to relocate a package name with the same relative pathing from the current package path as from the original package path
+	 * i.e.
+	 * OriginalPackagePath: /Game/Folder/
+	 * CurrentPackagePath: /Game/OtherFolder/SubFolder/
+	 * Example 1:
+	 * InPackageNameToRelocate: /Game/Folder/Sub/Asset
+	 * OutNewLocation: /Game/OtherFolder/SubFolder/Sub/Asset
+	 * Example 2:
+	 * InPackageNameToRelocate: /Game/OtherAsset
+	 * OutNewLocation: /Game/OtherFolder/OtherAsset
+	 * @param InPackageRelocationContext The context used for the relocation of the loaded asset
+	 * @param InPackageNameToRelocate The package name we want to remap relative to the current location
+	 * @param OutNewLocation The relocated package name.
+	 * @return True If the package need to be relocated.
+	 */
+	COREUOBJECT_API bool TryRelocateReference(const FPackageRelocationContext& InPackageRelocationContext, FStringView InPackageNameToRelocate, FStringBuilderBase& OutNewLocation);
+
+	COREUOBJECT_API void ApplyRelocationToObjectImportMap(const FPackageRelocationContext& InPackageRelocationContext, TArrayView<FObjectImport> ImportMapView);
+
+	COREUOBJECT_API void ApplyRelocationToSoftObjectArray(const FPackageRelocationContext& InPackageRelocationContext, TArrayView<FSoftObjectPath> SoftObjectPaths);
+
+	COREUOBJECT_API void ApplyRelocationToNameArray(const FPackageRelocationContext& InPackageRelocationContext, TArrayView<FName> PackageNames);
+}
\ No newline at end of file
diff --git a/Engine/Source/Runtime/CoreUObject/Private/UObject/LinkerLoad.cpp b/Engine/Source/Runtime/CoreUObject/Private/UObject/LinkerLoad.cpp
--- a/Engine/Source/Runtime/CoreUObject/Private/UObject/LinkerLoad.cpp
+++ b/Engine/Source/Runtime/CoreUObject/Private/UObject/LinkerLoad.cpp
@@ -53,6 +53,7 @@
 #include "Serialization/EditorBulkData.h"
 #include "HAL/FileManager.h"
 #include "UObject/CoreRedirects.h"
+#include "UObject/PackageRelocation.h"
 #include "Misc/AssetRegistryInterface.h"
 #include "Misc/StringBuilder.h"
 #include "Misc/EngineBuildSettings.h"
@@ -798,11 +799,18 @@ FLinkerLoad::ELinkerStatus FLinkerLoad::ProcessPackageSummary(TMap<TPair<FName,
 		Status = PopulateInstancingContext();
 	}
 
-	// Populate the linker instancing context for relocation loading if needed.
+	// Modify the ImportMap and SoftObjectPathList to account for the potential relocation of the packages
 	if (Status == LINKER_Loaded)
 	{
-		SCOPED_LOADTIMER(LinkerLoad_PopulateRelocationContext);
-		Status = PopulateRelocationContext();
+		SCOPED_LOADTIMER(LinkerLoad_ApplyRelocationToImportMapAndSoftObjectPathList);
+		Status = RelocateReferences();
+	}
+
+	// Modify the SoftObjectPathList for the instancing context
+	if (Status == LINKER_Loaded)
+	{
+		SCOPED_LOADTIMER(LinkerLoad_ApplyInstancingContextToSoftObjectPathList);
+		Status = ApplyInstancingContext();
 	}
 
 	// Fix up export map for object class conversion 
@@ -966,7 +974,8 @@ FLinkerLoad::FLinkerLoad(UPackage* InParent, const FPackagePath& InPackagePath,
 , bHasSerializedPreloadDependencies(false)
 , bHasFixedUpImportMap(false)
 , bHasPopulatedInstancingContext(false)
-, bHasPopulatedRelocationContext(false)
+, bHasRelocatedReferences(false)
+, bHasAppliedInstancingContext(false)
 , bFixupExportMapDone(false)
 , bHasFoundExistingExports(false)
 , bHasFinishedInitialization(false)
@@ -1772,8 +1781,6 @@ FLinkerLoad::ELinkerStatus FLinkerLoad::SerializeSoftObjectPathList()
 		FSoftObjectPath& SoftObjectPath = SoftObjectPathList.AddDefaulted_GetRef();
 		FStructuredArchive::FSlot Slot = Stream.EnterElement();
 		SoftObjectPath.SerializePath(Slot.GetUnderlyingArchive());
-
-		FixupSoftObjectPathForInstancedPackage(SoftObjectPath);
 		++SoftObjectPathListIndex;
 	}
 
@@ -2090,120 +2097,48 @@ FLinkerLoad::ELinkerStatus FLinkerLoad::PopulateInstancingContext()
 	return IsTimeLimitExceeded(TEXT("populating instancing context")) ? LINKER_TimedOut : LINKER_Loaded;
 }
 
-/**
- * Generate a package name with the same relative pathing from the current package path as from the original package path
- * i.e.
- * OriginalPath: /Game/Folder/
- * CurrentPath: /Game/OtherFolder/SubFolder/
- * Example 1:
- * PackageName: /Game/Folder/Sub/Asset
- * Result: /Game/OtherFolder/SubFolder/Sub/Asset
- * Example 2:
- * PackageName: /Game/OtherAsset
- * Result: /Game/OtherFolder/OtherAsset
- * @param InOriginalPackagePath The original package path where this package was located
- * @param InCurrentPackagePath The current location of the package
- * @param InPackageName The package name we want to remap relative to the current location
- * @return FName new package name keeping the same relative location to the current path as the original
- */
-FName GenerateCurrentRelativePackageName(FStringView InOriginalPackagePath, FStringView InCurrentPackagePath, FStringView InPackageName)
+FLinkerLoad::ELinkerStatus FLinkerLoad::RelocateReferences()
 {
-	FName Result = NAME_None;
-	FString RelPackageName(InPackageName);
-	//@todo 5.2: Create and use a string view version of MakePathRelativeTo
-	FString OriginalPackagePath(InOriginalPackagePath);
-	FPaths::MakePathRelativeTo(RelPackageName, *OriginalPackagePath);
-
-	// Strip the mount point from the current path prior to collapsing as we do not want to allow it to be collapsed
-	FStringView CurrentMount, MountlessCurrentPath;
-	FPathViews::SplitFirstComponent(FStringView(InCurrentPackagePath.RightChop(1)), CurrentMount, MountlessCurrentPath);
-	TStringBuilder<64> MountBuilder;
-	MountBuilder << '/' << CurrentMount << '/';
-
-	// Build the new path to collapse against the current path
-	TStringBuilder<256> Builder;
-	if (MountlessCurrentPath.Len() > 0)
-	{
-		Builder << MountlessCurrentPath << '/' << RelPackageName;
-	}
-	else
-	{
-		Builder << RelPackageName;
-	}
-	bool bSuccess = FPathViews::CollapseRelativeDirectories(Builder);
+	DECLARE_SCOPE_CYCLE_COUNTER(TEXT("FLinkerLoad::RelocateReferences"), STAT_LinkerLoad_RelocateReferences, STATGROUP_LinkerLoad);
 
-	// if we successfully collapsed the path prepend the mount point back
-	if (bSuccess)
-	{
-		Builder.Prepend(MountBuilder);
-		Result = FName(*Builder);
-	}
-	// if we couldn't, at least fix up the mount point if needed
-	else if (!InPackageName.StartsWith(MountBuilder.ToView()))
-	{
-		FStringView PackageMount, MountlessPackagePath;
-		FPathViews::SplitFirstComponent(FStringView(InPackageName.RightChop(1)), PackageMount, MountlessPackagePath);
-		Builder.Reset();
-		Builder << MountBuilder << MountlessPackagePath;
-		Result = FName(*Builder);
-	}
-	return Result;
-}
-
-FLinkerLoad::ELinkerStatus FLinkerLoad::PopulateRelocationContext()
-{
-	DECLARE_SCOPE_CYCLE_COUNTER(TEXT("FLinkerLoad::PopulateRelocationContext"), STAT_LinkerLoad_PopulateRelocationContext, STATGROUP_LinkerLoad);
-
-	if (!bHasPopulatedRelocationContext)
+	if (!bHasRelocatedReferences)
 	{
 #if WITH_EDITOR
 		// Validate if the package was moved and we want to generate fix up for references
 		FString PackageNameToLoad = GetPackagePath().GetPackageName();
-		// Check the summmary versioning alongside, PackageName used to be FolderName which wasn't actively used and got deprecated
-		bool bRelocated = false;// PackageNameToLoad != Summary.PackageName && Summary.GetFileVersionUE() >= EUnrealEngineObjectUE5Version::ADD_SOFTOBJECTPATH_LIST;
+
+		UE::Package::Relocation::Private::FPackageRelocationContext RelocationArgs;
+
+		bool bRelocated = UE::Package::Relocation::Private::ShouldApplyRelocation(Summary, PackageNameToLoad, RelocationArgs);
 		// Do not consider a package relocated if it's being loaded for Diff
 		if (bRelocated && (LoadFlags & LOAD_ForDiff) == 0)
 		{
-			// Only Generate relative fix up if the file was actually moved not just renamed
-			FStringView PackagePathToLoad = FPathViews::GetPath(PackageNameToLoad);
-			FStringView OriginalPackagePath = FPathViews::GetPath(Summary.PackageName);
-
-			if (PackagePathToLoad != OriginalPackagePath)
-			{
-				//@todo 5.2: switch to use the string view version
-				//FStringView PackageMount = FPathViews::GetMountPointNameFromPath(OriginalPackagePath, nullptr, false);
-				FString PackageMount = FPackageName::GetPackageMountPoint(FString(OriginalPackagePath.GetData(), OriginalPackagePath.Len()), false).ToString();
-
-				for (const FObjectImport& Import : ImportMap)
-				{
-					if (!Import.OuterIndex.IsNull())
-					{
-						continue;
-					}
-
-					// Generate relative fix-up for package path from the same original package mount point of the package we are loading
-					FString ImportPackageName = Import.ObjectName.ToString();
-					if (FPathViews::IsParentPathOf(PackageMount, ImportPackageName))
-					{
-						FName GeneratedName = GenerateCurrentRelativePackageName(OriginalPackagePath, PackagePathToLoad, ImportPackageName);
-						if (!GeneratedName.IsNone())
-						{
-							FName& RelocatedName = InstancingContext.RelocatedPackageMapping.FindOrAdd(Import.ObjectName);
-							// if there isn't already a remapping for that package, create one
-							if (RelocatedName.IsNone())
-							{
-								RelocatedName = GeneratedName;
-							}
-						}
-					}
-				}
-			}
+			UE_LOG(LogPackageRelocation, Verbose, TEXT("Loading relocated package (%s). The package was saved as (%s)."), *PackageNameToLoad, *Summary.PackageName);
+			UE::Package::Relocation::Private::ApplyRelocationToObjectImportMap(RelocationArgs, ImportMap);
+			UE::Package::Relocation::Private::ApplyRelocationToSoftObjectArray(RelocationArgs, SoftObjectPathList);
 		}
 #endif
 		// Avoid duplicate work in async case.
-		bHasPopulatedRelocationContext = true;
+		bHasRelocatedReferences = true;
 	}
-	return IsTimeLimitExceeded(TEXT("populating relocation context")) ? LINKER_TimedOut : LINKER_Loaded;
+	return IsTimeLimitExceeded(TEXT("relocating the ImportMap and SoftObjectPathList")) ? LINKER_TimedOut : LINKER_Loaded;
+}
+
+FLinkerLoad::ELinkerStatus FLinkerLoad::ApplyInstancingContext()
+{
+	DECLARE_SCOPE_CYCLE_COUNTER(TEXT("FLinkerLoad::ApplyInstancingContext"), STAT_LinkerLoad_ApplyInstancingContext, STATGROUP_LinkerLoad);
+	if (!bHasAppliedInstancingContext)
+	{ 
+		for (FSoftObjectPath& SoftObjectPath : SoftObjectPathList)
+		{
+			FixupSoftObjectPathForInstancedPackage(SoftObjectPath);
+		}
+
+		// Avoid duplicate work in async case.
+		bHasAppliedInstancingContext = true;
+	}
+
+	return IsTimeLimitExceeded(TEXT("applying the instancing context to the SoftObjectPathList")) ? LINKER_TimedOut : LINKER_Loaded;;
 }
 
 /**
@@ -3351,9 +3286,7 @@ bool FLinkerLoad::VerifyImportInner(const int32 ImportIndex, FString& WarningSuf
 		FUObjectSerializeContext* SerializeContext = GetSerializeContext();
 
 		// Resolve the package name for the import, potentially remapping it, if instancing
-		FName OriginalPackageToLoad = !Import.HasPackageName() ? Import.ObjectName : Import.GetPackageName();
-		FName PackageToLoad = InstancingContext.RelocatePackage(OriginalPackageToLoad);
-
+		FName PackageToLoad = !Import.HasPackageName() ? Import.ObjectName : Import.GetPackageName();
 		FName PackageToLoadInto = InstancingContext.RemapPackage(PackageToLoad);
 #if WITH_EDITOR
 		if (SlowTask)
@@ -3447,13 +3380,6 @@ bool FLinkerLoad::VerifyImportInner(const int32 ImportIndex, FString& WarningSuf
 		// to be linked to any other package's ImportMaps
 		if (!Package || Package->HasAnyPackageFlags(PKG_Compiling))
 		{
-			FName RelocatedName;
-			if (InstancingContext.RelocatedPackageMapping.RemoveAndCopyValue(OriginalPackageToLoad, RelocatedName))
-			{
-				UE_ASSET_LOG(LogLinker, Warning, PackagePath, TEXT("VerifyImport: Failed to load package for import object '%s' through using relative path '%s'. Using orginal location as fallback."), *GetImportFullName(ImportIndex), *RelocatedName.ToString());
-				return nullptr;
-			}
-
 			if (!FLinkerLoad::IsKnownMissingPackage(PackageToLoad))
 			{
 				FLinkerLoad::AddKnownMissingPackage(PackageToLoad);
diff --git a/Engine/Source/Runtime/CoreUObject/Private/UObject/PackageRelocation.cpp b/Engine/Source/Runtime/CoreUObject/Private/UObject/PackageRelocation.cpp
new file mode 100644
--- /dev/null
+++ b/Engine/Source/Runtime/CoreUObject/Private/UObject/PackageRelocation.cpp
@@ -0,0 +1,182 @@
+// Copyright Epic Games, Inc.All Rights Reserved.
+
+#include "UObject/PackageRelocation.h"
+
+#include "HAL/IConsoleManager.h"
+#include "Misc/PathViews.h"
+#include "String/ParseTokens.h"
+#include "UObject/ObjectResource.h"
+#include "UObject/PackageFileSummary.h"
+#include "UObject/SoftObjectPath.h"
+#include "UObject/TopLevelAssetPath.h"
+
+DEFINE_LOG_CATEGORY(LogPackageRelocation);
+
+namespace UE::Package::Relocation::Private
+{
+	static TAutoConsoleVariable<int32> CVarRelocationMode(
+			TEXT("Package.Relocation"),
+			0,
+			TEXT("Define when we should run the relocation logic for the dependencies of a package. Note changing this value at runtime won't update the cached depencencies in the asset registry.\n")
+			TEXT("  0: Off (never apply relocation. References to other relocated packages will give errors and fail to resolve.)\n")
+			TEXT("  1: Relocate any asset saved after EUnrealEngineObjectUE5Version::ADD_SOFTOBJECTPATH_LIST (5.1). Default value for the new projects.\n")
+		);
+
+	bool ShouldApplyRelocation(const FPackageFileSummary& PackageSummary, FStringView LoadedPackageName, FPackageRelocationContext& OutPackageRelocationContext)
+	{
+		OutPackageRelocationContext.CurrentPackagePath = FPathViews::GetPath(LoadedPackageName);
+		OutPackageRelocationContext.OriginalPackagePath = FPathViews::GetPath(PackageSummary.PackageName);
+
+		if (OutPackageRelocationContext.CurrentPackagePath == OutPackageRelocationContext.OriginalPackagePath)
+		{
+			// Don't apply relocation to packages that have renamed their leaf filename but remain in the same parent directory
+			return false;
+		}
+
+		switch (CVarRelocationMode.GetValueOnAnyThread())
+		{
+		case 1:
+			if (PackageSummary.GetFileVersionUE() < EUnrealEngineObjectUE5Version::ADD_SOFTOBJECTPATH_LIST)
+			{
+				return false;
+			}
+			break;
+
+		default:
+			return false;
+		}
+
+		bool bHasClassesPrefix = false;
+		FStringView MountPointName = FPathViews::GetMountPointNameFromPath(OutPackageRelocationContext.OriginalPackagePath, &bHasClassesPrefix);
+		if (MountPointName.IsEmpty() || bHasClassesPrefix)
+		{
+			return false;
+		}
+
+		/*
+		 * The function GetMountPointNameFromPath return the mount point name without the '/' at the front and at the end.
+		 * So we use the original package path and the MountPointName length + 2 to reintroduce the '/' to the cached mount point.
+		 */
+		OutPackageRelocationContext.OriginalPackageMount = FStringView(OutPackageRelocationContext.OriginalPackagePath.GetData(), MountPointName.Len() + 2);
+
+		return true;
+	}
+
+	bool TryRelocateReference(const FPackageRelocationContext& InPackageRelocationContext, FStringView InPackageNameToRelocate, FStringBuilderBase& OutNewLocation)
+	{
+		// Validate that the package to relocate is the same mount point as the original package name
+		if (!InPackageNameToRelocate.StartsWith(InPackageRelocationContext.OriginalPackageMount))
+		{
+			return false;
+		}
+
+		// Compute the info to generate a relative path of the InPackageName against the OriginalPackagePath. (this is faster and work better for this scenario that the one from FPaths)
+		const UE::String::EParseTokensOptions ParseOptions = UE::String::EParseTokensOptions::SkipEmpty;
+		 
+		int32 FolderUpCount = 0;
+		TArray<FStringView, TInlineAllocator<16>> ToAppend;
+		{
+			const FStringView PackagePathToRelocateWithoutMountPoint = InPackageNameToRelocate.RightChop(InPackageRelocationContext.OriginalPackageMount.Len());
+			TArray<FStringView, TInlineAllocator<16>> PackageNameArrayBuffer;
+			UE::String::ParseTokens(PackagePathToRelocateWithoutMountPoint, TEXTVIEW("/"), PackageNameArrayBuffer, ParseOptions);
+
+			const FStringView OriginalPackagePathWithoutMountPoint = InPackageRelocationContext.OriginalPackagePath.RightChop(InPackageRelocationContext.OriginalPackageMount.Len());
+			TArray<FStringView, TInlineAllocator<16>> OriginalPackagePathArrayBuffer;
+			UE::String::ParseTokens(OriginalPackagePathWithoutMountPoint, TEXTVIEW("/"), OriginalPackagePathArrayBuffer, ParseOptions);
+
+
+			int32 IdenticalDepth = 0;
+			while (IdenticalDepth < PackageNameArrayBuffer.Num() && IdenticalDepth < OriginalPackagePathArrayBuffer.Num() && PackageNameArrayBuffer[IdenticalDepth] == OriginalPackagePathArrayBuffer[IdenticalDepth])
+			{
+				++IdenticalDepth;
+			}
+
+			FolderUpCount = OriginalPackagePathArrayBuffer.Num() - IdenticalDepth;
+
+			ToAppend.Reserve(PackageNameArrayBuffer.Num() - IdenticalDepth);
+			for (int32 Index = IdenticalDepth; Index < PackageNameArrayBuffer.Num(); ++Index)
+			{
+				ToAppend.Add(MoveTemp(PackageNameArrayBuffer[Index]));
+			}
+		}
+
+		TArray<FStringView, TInlineAllocator<16>> CurrentPackageNameArrayBuffer;
+		UE::String::ParseTokens(InPackageRelocationContext.CurrentPackagePath, TEXTVIEW("/"), CurrentPackageNameArrayBuffer, ParseOptions);
+
+		OutNewLocation.Reset();
+
+		if (CurrentPackageNameArrayBuffer.Num() <= FolderUpCount)
+		{
+			// Bail out. This path doesn't exist.
+			return true;
+		}
+
+
+		for (int32 Index = 0; Index < CurrentPackageNameArrayBuffer.Num() - FolderUpCount; ++Index)
+		{
+			OutNewLocation.AppendChar(TEXT('/'));
+			OutNewLocation.Append(CurrentPackageNameArrayBuffer[Index]);
+		}
+
+		for (int32 Index = 0; Index < ToAppend.Num(); ++Index)
+		{
+			OutNewLocation.AppendChar(TEXT('/'));
+			OutNewLocation.Append(ToAppend[Index]);
+		}
+
+		return true;
+	}
+
+	void ApplyRelocationToObjectImportMap(const FPackageRelocationContext& InPackageRelocationContext, TArrayView<FObjectImport> ImportMapView)
+	{
+		for (FObjectImport& Import : ImportMapView)
+		{
+			if (!Import.OuterIndex.IsNull())
+			{
+				continue;
+			}
+
+			// Generate relative fix-up for package path from the same original package mount point of the package we are loading
+			FNameBuilder ImportPackageName(Import.ObjectName);
+			FNameBuilder RelocatedPackageName;
+			if (TryRelocateReference(InPackageRelocationContext, ImportPackageName.ToView(), RelocatedPackageName))
+			{
+				Import.ObjectName = *RelocatedPackageName;
+			}
+		}
+	}
+
+	void ApplyRelocationToSoftObjectArray(const FPackageRelocationContext& InPackageRelocationContext, TArrayView<FSoftObjectPath> SoftObjectPaths)
+	{
+		for (FSoftObjectPath& SoftObjectPath : SoftObjectPaths)
+		{
+			FNameBuilder PackageName(SoftObjectPath.GetAssetPath().GetPackageName());
+			FNameBuilder RelocatedPackageName;
+			if (TryRelocateReference(InPackageRelocationContext, PackageName.ToView(), RelocatedPackageName))
+			{
+				if (RelocatedPackageName.Len() != 0)
+				{
+					FTopLevelAssetPath AssetPath = SoftObjectPath.GetAssetPath();
+					SoftObjectPath.SetPath(FTopLevelAssetPath(FName(*RelocatedPackageName), AssetPath.GetAssetName()));
+				}
+				else
+				{
+					SoftObjectPath.Reset();
+				}
+			}
+		}
+	}
+
+	void ApplyRelocationToNameArray(const FPackageRelocationContext& InPackageRelocationContext, TArrayView<FName> PackageNames)
+	{
+		for (FName& PackageName : PackageNames)
+		{
+			FNameBuilder Package(PackageName);
+			FNameBuilder RelocatedPackageName;
+			if (TryRelocateReference(InPackageRelocationContext, Package.ToView(), RelocatedPackageName))
+			{
+				PackageName = *RelocatedPackageName;
+			}
+		}
+	}
+}
\ No newline at end of file
diff --git a/Engine/Source/Runtime/CoreUObject/Public/UObject/LinkerInstancingContext.h b/Engine/Source/Runtime/CoreUObject/Public/UObject/LinkerInstancingContext.h
--- a/Engine/Source/Runtime/CoreUObject/Public/UObject/LinkerInstancingContext.h
+++ b/Engine/Source/Runtime/CoreUObject/Public/UObject/LinkerInstancingContext.h
@@ -175,16 +175,6 @@ private:
 		return InstancedPackageMap.InstancedPackageMapping.FindOrAdd(Original);
 	}
 
-	/** Used internally by the linker to try to fix references on relocated packages. */
-	FName RelocatePackage(const FName& PackageName) const
-	{
-		if (const FName* RemappedName = RelocatedPackageMapping.Find(PackageName))
-		{
-			return *RemappedName;
-		}
-		return PackageName;
-	}
-
 	friend class FLinkerLoad;
 	friend struct FAsyncPackage2;
 
@@ -192,9 +182,6 @@ private:
 	FLinkerInstancedPackageMap InstancedPackageMap;
 	/** Map of original top level asset path to their instance counterpart. */
 	TMap<FTopLevelAssetPath, FTopLevelAssetPath> PathMapping;
-	/** Map of original package name to their potential relocated counterpart. */
-	TMap<FName, FName> RelocatedPackageMapping;
-
 	/** Tags can be used to determine some loading behavior. */
 	TSet<FName> Tags;
 	/** Remap soft object paths */
diff --git a/Engine/Source/Runtime/CoreUObject/Public/UObject/LinkerLoad.h b/Engine/Source/Runtime/CoreUObject/Public/UObject/LinkerLoad.h
--- a/Engine/Source/Runtime/CoreUObject/Public/UObject/LinkerLoad.h
+++ b/Engine/Source/Runtime/CoreUObject/Public/UObject/LinkerLoad.h
@@ -430,8 +430,10 @@ private:
 	bool					bHasFixedUpImportMap:1;
 	/** Whether we already populated the instancing context.																*/
 	bool					bHasPopulatedInstancingContext:1;
-	/** Whether we already populated the relocation context.																*/
-	bool					bHasPopulatedRelocationContext:1;
+	/** Whether we already relocated references.																			*/
+	bool					bHasRelocatedReferences:1;
+	/** Whether we already applied the instancing context to the soft object list.											*/
+	bool					bHasAppliedInstancingContext : 1;
 	/** Used for ActiveClassRedirects functionality */
 	bool					bFixupExportMapDone:1;
 	/** Whether we already matched up existing exports.																		*/
@@ -1156,10 +1158,9 @@ private:
 	 */
 	ELinkerStatus PopulateInstancingContext();
 
-	/**
-	 * Generate remapping for the relocation context if this is a relocated package.
-	 */
-	ELinkerStatus PopulateRelocationContext();
+	ELinkerStatus RelocateReferences();
+
+	ELinkerStatus ApplyInstancingContext();
 
 	/**
 	 * Serializes the export map.
