Fix potential memory corruption issue when allocating bulkdata payloads on 32bit systems.

#rnx

- The problem is that our allocation functions take allocation sizes as type SIZE_T which on older 32bit platforms will be int32, but bulkdata stores sizes as int64. If the requested allocation is large enough the size will be truncated and the allocation could still return a valid memory allocation but one much smaller than requested. Later when we try to load into that allocation we will overwrite the end of the buffer. Given our supported platforms and how bulkdata is generally used this is an unlikely scenario but one we still should protect against.
- Added two functions ::SafeMalloc and ::SafeRealloc to replace direct calls to malloc/realloc. These utility functions will check to make sure that we can make allocations of the requested size by making sure that SIZE_T (the type that our allocators take for allocation sizes) is a large enough type.
- If the platform cannot service the request we will issue a fatal log error and halt the process. Technically we could error and return nullptr but in practice very little code checks to see if bulkdata allocations are valid or not and will end up crashing anyway. Given that this problem is only going to show up on 32bit machines and only for very large assets (most likely invalid or broken ones at that) we might as well halt execution as close to the problem as possible to make debugging easier.
- The changes are only applied for platforms where PLATFORM_32BITS is true, should be no cost to 64 bit platforms.

diff --git a/Engine/Source/Runtime/CoreUObject/Private/Serialization/BulkData.cpp b/Engine/Source/Runtime/CoreUObject/Private/Serialization/BulkData.cpp
--- a/Engine/Source/Runtime/CoreUObject/Private/Serialization/BulkData.cpp
+++ b/Engine/Source/Runtime/CoreUObject/Private/Serialization/BulkData.cpp
@@ -270,6 +270,50 @@ bool FBulkMetaData::FromSerialized(FArchive& Ar, int64 ElementSize, FBulkMetaDat
 	return true;
 }
 
+/**
+ * Bulkdata payload lengths are stored as type int64 but FMemory::Malloc takes SIZE_T which could potentially
+ * be 32bit. This utility helps convert between the two and makes sure that if SizeInBytes exceeds the max size
+ * that FMemory::Malloc accepts that we return nullptr rather than a smaller buffer than was requested.
+*/
+inline void* SafeMalloc(int64 SizeInBytes, uint32 Alignment)
+{
+#if !PLATFORM_32BITS
+	return FMemory::Malloc(SizeInBytes, Alignment);
+#else
+	if (SizeInBytes >= 0 && (uint64)SizeInBytes <= TNumericLimits<SIZE_T>::Max())
+	{
+		return FMemory::Malloc(SizeInBytes, Alignment);
+	}
+	else
+	{
+		UE_LOG(LogSerialization, Fatal, TEXT("Bulkdata payload allocation is an invalid length (%" INT64_FMT ") max size allowed (%" SIZE_T_FMT ")"), SizeInBytes, TNumericLimits<SIZE_T>::Max());
+		return nullptr;
+	}
+#endif //!PLATFORM_32BITS
+}
+
+/**
+ * Bulkdata payload lengths are stored as type int64 but FMemory::Malloc takes SIZE_T which could potentially
+ * be 32bit. This utility helps convert between the two and makes sure that if SizeInBytes exceeds the max size
+ * that FMemory::Realloc accepts that we return nullptr rather than a smaller buffer than was requested.
+*/
+inline void* SafeRealloc(void* Original, int64 SizeInBytes, uint32 Alignment)
+{
+#if !PLATFORM_32BITS
+	return FMemory::Realloc(Original, SizeInBytes, DEFAULT_ALIGNMENT);
+#else
+	if (SizeInBytes >= 0 && (uint64)SizeInBytes <= TNumericLimits<SIZE_T>::Max())
+	{
+		return FMemory::Realloc(Original, SizeInBytes, DEFAULT_ALIGNMENT);
+	}
+	else
+	{
+		UE_LOG(LogSerialization, Fatal, TEXT("Bulkdata payload allocation is an invalid length (%" INT64_FMT ") max size allowed (%" SIZE_T_FMT ")"), SizeInBytes, TNumericLimits<SIZE_T>::Max());
+		return nullptr;
+	}
+#endif //!PLATFORM_32BITS
+}
+
 } // namespace UE::BulkData::Private
 
 /*-----------------------------------------------------------------------------
@@ -290,11 +334,11 @@ void FBulkData::FAllocatedPtr::Free(FBulkData* Owner)
 	}
 }
 
-void* FBulkData::FAllocatedPtr::ReallocateData(FBulkData* Owner, SIZE_T SizeInBytes)
+void* FBulkData::FAllocatedPtr::ReallocateData(FBulkData* Owner, int64 SizeInBytes)
 {
 	checkf(!Owner->IsDataMemoryMapped(),  TEXT("Trying to reallocate a memory mapped BulkData object without freeing it first!"));
 
-	Allocation.RawData = FMemory::Realloc(Allocation.RawData, SizeInBytes, DEFAULT_ALIGNMENT);
+	Allocation.RawData = UE::BulkData::Private::SafeRealloc(Allocation.RawData, SizeInBytes, DEFAULT_ALIGNMENT);
 
 	return Allocation.RawData;
 }
@@ -370,7 +414,7 @@ void FBulkData::FAllocatedPtr::Swap(FBulkData* Owner, void** DstBuffer)
 	{
 		const int64 DataSize = Owner->GetBulkDataSize();
 
-		*DstBuffer = FMemory::Malloc(DataSize, DEFAULT_ALIGNMENT);
+		*DstBuffer = UE::BulkData::Private::SafeMalloc(DataSize, DEFAULT_ALIGNMENT);
 		FMemory::Memcpy(*DstBuffer, Allocation.MemoryMappedData->GetPointer(), DataSize);
 
 		delete Allocation.MemoryMappedData;
@@ -795,12 +839,12 @@ void FBulkData::GetCopy( void** Dest, bool bDiscardInternalCopy )
 	else
 	{
 		// The data is already loaded so we can simply use a mempcy.
-		if( IsBulkDataLoaded() )
+		if (IsBulkDataLoaded())
 		{
 			// If the internal copy should be discarded and we are still attached to an archive we can
 			// simply "return" the already existing copy and NULL out the internal reference. We can
 			// also do this if the data is single use like e.g. when uploading texture data.
-			if( bDiscardInternalCopy && CanDiscardInternalData() )
+			if (bDiscardInternalCopy && CanDiscardInternalData())
 			{
 				DataAllocation.Swap(this, Dest);
 			}
@@ -810,7 +854,7 @@ void FBulkData::GetCopy( void** Dest, bool bDiscardInternalCopy )
 				if (BulkDataSize != 0)
 				{
 					// Allocate enough memory for data...
-					*Dest = FMemory::Malloc( BulkDataSize, GetBulkDataAlignment() );
+					*Dest = UE::BulkData::Private::SafeMalloc(BulkDataSize, GetBulkDataAlignment());
 
 					// ... and copy it into memory now pointed to by out parameter.
 					FMemory::Memcpy( *Dest, GetDataBufferReadOnly(), BulkDataSize );
diff --git a/Engine/Source/Runtime/CoreUObject/Public/Serialization/BulkData.h b/Engine/Source/Runtime/CoreUObject/Public/Serialization/BulkData.h
--- a/Engine/Source/Runtime/CoreUObject/Public/Serialization/BulkData.h
+++ b/Engine/Source/Runtime/CoreUObject/Public/Serialization/BulkData.h
@@ -509,7 +509,7 @@ class COREUOBJECT_API FBulkData
 
 		void Free(FBulkData* Owner);
 
-		void* ReallocateData(FBulkData* Owner, SIZE_T SizeInBytes);
+		void* ReallocateData(FBulkData* Owner, int64 SizeInBytes);
 
 		void SetData(FBulkData* Owner, void* Buffer);
 		void SetMemoryMappedData(FBulkData* Owner, IMappedFileHandle* MappedHandle, IMappedFileRegion* MappedRegion);
@@ -1076,7 +1076,7 @@ private:
 	bool CanDiscardInternalData() const;
 	
 	/** Reallocate bulk data */
-	inline void* ReallocateData(SIZE_T SizeInBytes) { return DataAllocation.ReallocateData(this, SizeInBytes); }
+	inline void* ReallocateData(int64 SizeInBytes) { return DataAllocation.ReallocateData(this, SizeInBytes); }
 
 	/** Free bulk data */
 	inline void  FreeData() { DataAllocation.Free(this); }
