Update comments on BlueprintType and NotBlueprintType to include the fact that they are inherited by subclasses. You can see that this is true in UEdGraphSchema_K2::IsAllowableBlueprintVariableType, where it walks the class heirarchy to check for the metadata.

diff --git a/Engine/Source/Runtime/CoreUObject/Public/UObject/ObjectMacros.h b/Engine/Source/Runtime/CoreUObject/Public/UObject/ObjectMacros.h
--- a/Engine/Source/Runtime/CoreUObject/Public/UObject/ObjectMacros.h
+++ b/Engine/Source/Runtime/CoreUObject/Public/UObject/ObjectMacros.h
@@ -727,10 +727,10 @@ namespace UC
 		/// Declares that instances of this class should always have an outer of the specified class.  This is inherited by subclasses unless overridden.
 		Within, /* =OuterClassName */
 
-		/// Exposes this class as a type that can be used for variables in blueprints
+		/// Exposes this class as a type that can be used for variables in blueprints. This is inherited by subclasses unless overridden.
 		BlueprintType,
 
-		/// Prevents this class from being used for variables in blueprints
+		/// Prevents this class from being used for variables in blueprints. This is inherited by subclasses unless overridden.
 		NotBlueprintType,
 
 		/// Exposes this class as an acceptable base class for creating blueprints. The default is NotBlueprintable, unless inherited otherwise. This is inherited by subclasses.
