ObjectPtr: Changed `FObjectPtrProperty::Identical` to avoid unnecessary reads.
#rnx

diff --git a/Engine/Source/Runtime/CoreUObject/Private/UObject/PropertyObjectPtr.cpp b/Engine/Source/Runtime/CoreUObject/Private/UObject/PropertyObjectPtr.cpp
--- a/Engine/Source/Runtime/CoreUObject/Private/UObject/PropertyObjectPtr.cpp
+++ b/Engine/Source/Runtime/CoreUObject/Private/UObject/PropertyObjectPtr.cpp
@@ -135,10 +135,27 @@ bool FObjectPtrProperty::Identical(const void* A, const void* B, uint32 PortFlag
 		return true;
 	}
 
-	// Resolve the object handles and run the deep comparison logic 
+	if (IsObjectHandleNull(ObjectAHandle) || IsObjectHandleNull(ObjectBHandle))
+	{
+		return false;
+	}
+
+	// If a deep comparison is required, resolve the object handles and run the deep comparison logic
+	// If a deep comparison is not required, avoid resolving the object handles because resolving declares
+	// a cook dependency.
 	if ((PortFlags & (PPF_DeepCompareInstances | PPF_DeepComparison)) != 0)
 	{
-		return FObjectPropertyBase::StaticIdentical(ObjectA.Get(), ObjectB.Get(), PortFlags);
+		bool bPerformDeepComparison = true;
+		if ((PortFlags & PPF_DeepCompareDSOsOnly) != 0)
+		{
+			UClass* ClassA = ObjectA.GetClass();
+			UObject* DSO = ClassA ? ClassA->GetDefaultSubobjectByName(ObjectA.GetFName()) : nullptr;
+			bPerformDeepComparison = DSO != nullptr;
+		}
+		if (bPerformDeepComparison)
+		{
+			return FObjectPropertyBase::StaticIdentical(ObjectA.Get(), ObjectB.Get(), PortFlags);
+		}
 	}
 
 	return false;
