Added a check for invalid counts in FScriptArrayHelper::Resize().

diff --git a/Engine/Source/Runtime/CoreUObject/Public/UObject/UnrealType.h b/Engine/Source/Runtime/CoreUObject/Public/UObject/UnrealType.h
--- a/Engine/Source/Runtime/CoreUObject/Public/UObject/UnrealType.h
+++ b/Engine/Source/Runtime/CoreUObject/Public/UObject/UnrealType.h
@@ -4148,17 +4148,6 @@ public:
 		checkSlow(Result >= 0);
 		return Result;
 	}
-	/**
-	 *	Static version of Num() used when you don't need to bother to construct a FScriptArrayHelper. Returns the number of elements in the array.
-	 *	@param	Target: pointer to the raw memory associated with a FScriptArray
-	 *	@return The number of elements in the array.
-	**/
-	UE_DEPRECATED(4.25, "This shortcut is no longer valid - the Num() should be read from a proper array helper")
-	static FORCEINLINE int32 Num(const void *Target)
-	{
-		checkSlow(((const FScriptArray*)Target)->Num() >= 0); 
-		return ((const FScriptArray*)Target)->Num();
-	}
 	/**
 	 *	Returns a uint8 pointer to an element in the array
 	 *	@param	Index: index of the item to return a pointer to.
@@ -4228,8 +4217,12 @@ public:
 	*	@param	Count: the number of items the array will have on completion.
 	**/
 	void Resize(int32 Count)
-	{ 
-		check(Count>=0);
+	{
+		if (Count < 0)
+		{
+			UE::Core::Private::OnInvalidArrayNum(Count);
+		}
+
 		int32 OldNum = Num();
 		if (Count > OldNum)
 		{
