SavePackage TransientPropertyOverrides: In addition to skipping the properties when harvesting, we need to skip the properties when writing exports to disk. If we serialized the skipped properties for export A, but the skipped properties did exist as an import or export because they were included in the save by another export B, then we would still record them in the save of the export, and would give error  "not listed as a dependency" because they were skipped for during the harvest of export A.

diff --git a/Engine/Source/Runtime/CoreUObject/Private/UObject/LinkerSave.cpp b/Engine/Source/Runtime/CoreUObject/Private/UObject/LinkerSave.cpp
--- a/Engine/Source/Runtime/CoreUObject/Private/UObject/LinkerSave.cpp
+++ b/Engine/Source/Runtime/CoreUObject/Private/UObject/LinkerSave.cpp
@@ -389,9 +389,22 @@ FArchive& FLinkerSave::operator<<( FName& InName )
 FArchive& FLinkerSave::operator<<( UObject*& Obj )
 {
 	FPackageIndex Save;
+
 	if (Obj)
 	{
-		Save = MapObject(Obj);
+		bool bKeepObject = true;
+		if (TransientPropertyOverrides && !TransientPropertyOverrides->IsEmpty())
+		{
+			const TSet<FProperty*>* Props = TransientPropertyOverrides->Find(CurrentlySavingExportObject);
+			if (Props && Props->Contains(this->GetSerializedProperty()))
+			{
+				bKeepObject = false;
+			}
+		}
+		if (bKeepObject)
+		{
+			Save = MapObject(Obj);
+		}
 	}
 	return *this << Save;
 }
diff --git a/Engine/Source/Runtime/CoreUObject/Private/UObject/SavePackage/PackageHarvester.cpp b/Engine/Source/Runtime/CoreUObject/Private/UObject/SavePackage/PackageHarvester.cpp
--- a/Engine/Source/Runtime/CoreUObject/Private/UObject/SavePackage/PackageHarvester.cpp
+++ b/Engine/Source/Runtime/CoreUObject/Private/UObject/SavePackage/PackageHarvester.cpp
@@ -628,6 +628,11 @@ void FPackageHarvester::ResolveOverrides()
 	}
 }
 
+TMap<UObject*, TSet<FProperty*>> FPackageHarvester::ReleaseTransientPropertyOverrides()
+{
+	return MoveTemp(TransientPropertyOverrides);
+}
+
 void FPackageHarvester::HarvestDependency(UObject* InObj, bool bIsNative)
 {
 	// if we aren't currently processing an export or the referenced object is a package, do not harvest the dependency
diff --git a/Engine/Source/Runtime/CoreUObject/Private/UObject/SavePackage/PackageHarvester.h b/Engine/Source/Runtime/CoreUObject/Private/UObject/SavePackage/PackageHarvester.h
--- a/Engine/Source/Runtime/CoreUObject/Private/UObject/SavePackage/PackageHarvester.h
+++ b/Engine/Source/Runtime/CoreUObject/Private/UObject/SavePackage/PackageHarvester.h
@@ -81,6 +81,7 @@ public:
 	void HarvestDependency(UObject* InObj, bool bIsNative);
 
 	bool CurrentExportHasDependency(UObject* InObj) const;
+	TMap<UObject*, TSet<FProperty*>> ReleaseTransientPropertyOverrides();
 
 	// FArchiveUObject implementation
 	virtual FString GetArchiveName() const override;
diff --git a/Engine/Source/Runtime/CoreUObject/Private/UObject/SavePackage/SaveContext.h b/Engine/Source/Runtime/CoreUObject/Private/UObject/SavePackage/SaveContext.h
--- a/Engine/Source/Runtime/CoreUObject/Private/UObject/SavePackage/SaveContext.h
+++ b/Engine/Source/Runtime/CoreUObject/Private/UObject/SavePackage/SaveContext.h
@@ -1031,6 +1031,16 @@ public:
 		return SavedAssets;
 	}
 
+	const TMap<UObject*, TSet<FProperty*>>& GetTransientPropertyOverrides()
+	{
+		return TransientPropertyOverrides;
+	}
+
+	void SetTransientPropertyOverrides(TMap<UObject*, TSet<FProperty*>>&& InTransientPropertyOverrides)
+	{
+		TransientPropertyOverrides = MoveTemp(InTransientPropertyOverrides);
+	}
+
 public:
 	ESavePackageResult Result;
 
@@ -1101,6 +1111,8 @@ private:
 	// Set of AssetDatas created for the Assets saved into the package
 	TArray<FAssetData> SavedAssets;
 
+	// Overrided properties for each export that should be treated as transient, and nulled out when serializing
+	TMap<UObject*, TSet<FProperty*>> TransientPropertyOverrides;
 };
 
 const TCHAR* LexToString(FSaveContext::ESaveableStatus Status);
diff --git a/Engine/Source/Runtime/CoreUObject/Private/UObject/SavePackage2.cpp b/Engine/Source/Runtime/CoreUObject/Private/UObject/SavePackage2.cpp
--- a/Engine/Source/Runtime/CoreUObject/Private/UObject/SavePackage2.cpp
+++ b/Engine/Source/Runtime/CoreUObject/Private/UObject/SavePackage2.cpp
@@ -420,6 +420,7 @@ ESavePackageResult HarvestPackage(FSaveContext& SaveContext)
 		Harvester.UsingCustomVersion(FEditorObjectVersion::GUID);
 	}
 	SaveContext.SetCustomVersions(Harvester.GetCustomVersions());
+	SaveContext.SetTransientPropertyOverrides(Harvester.ReleaseTransientPropertyOverrides());
 
 	return ReturnSuccessOrCancel();
 }
@@ -936,6 +937,7 @@ ESavePackageResult CreateLinker(FSaveContext& SaveContext)
 		{
 			SaveContext.GetLinker()->SetOutputDevice(SaveContext.GetError());
 		}
+		SaveContext.GetLinker()->SetTransientPropertyOverrides(SaveContext.GetTransientPropertyOverrides());
 		SaveContext.GetLinker()->bUpdatingLoadedPath = SaveContext.IsUpdatingLoadedPath();
 		SaveContext.GetLinker()->bProceduralSave = SaveContext.IsProceduralSave();
 
@@ -1953,6 +1955,7 @@ ESavePackageResult WritePackageTextHeader(FStructuredArchive::FRecord& Structure
 
 		Export.SerialOffset = ExportsArchive.Tell();
 		Linker.CurrentlySavingExport = FPackageIndex::FromExport(ExportIndex);
+		Linker.CurrentlySavingExportObject = Export.Object;
 
 		FExportProxyArchive Ar(ExportsArchive);
 		TGuardValue<FArchive*> GuardSaver(Linker.Saver, &Ar);
@@ -1976,6 +1979,7 @@ ESavePackageResult WritePackageTextHeader(FStructuredArchive::FRecord& Structure
 		}
 
 		Linker.CurrentlySavingExport = FPackageIndex();
+		Linker.CurrentlySavingExportObject = nullptr;
 		Export.SerialSize = ExportsArchive.Tell() - Export.SerialOffset;
 	}
 
@@ -2009,6 +2013,7 @@ ESavePackageResult WriteExports(FStructuredArchive::FRecord& StructuredArchiveRo
 			// Save the object data.
 			Export.SerialOffset = Linker->Tell();
 			Linker->CurrentlySavingExport = FPackageIndex::FromExport(i);
+			Linker->CurrentlySavingExportObject = Export.Object;
 
 			FString ObjectName = Export.Object->GetPathName(SaveContext.GetPackage());
 			FStructuredArchive::FSlot ExportSlot = ExportsRecord.EnterField(*ObjectName);
@@ -2062,6 +2067,7 @@ ESavePackageResult WriteExports(FStructuredArchive::FRecord& StructuredArchiveRo
 #endif
 			}
 			Linker->CurrentlySavingExport = FPackageIndex();
+			Linker->CurrentlySavingExportObject = nullptr;
 			Export.SerialSize = Linker->Tell() - Export.SerialOffset;
 		}
 	}
diff --git a/Engine/Source/Runtime/CoreUObject/Public/UObject/LinkerSave.h b/Engine/Source/Runtime/CoreUObject/Public/UObject/LinkerSave.h
--- a/Engine/Source/Runtime/CoreUObject/Public/UObject/LinkerSave.h
+++ b/Engine/Source/Runtime/CoreUObject/Public/UObject/LinkerSave.h
@@ -5,6 +5,7 @@
 #include "Compression/CompressedBuffer.h"
 #include "Containers/Array.h"
 #include "Containers/Map.h"
+#include "Containers/Set.h"
 #include "Containers/StringFwd.h"
 #include "Containers/UnrealString.h"
 #include "CoreMinimal.h"
@@ -57,6 +58,7 @@ public:
 	FArchive* Saver;
 
 	FPackageIndex CurrentlySavingExport;
+	UObject* CurrentlySavingExportObject = nullptr;
 	TArray<FPackageIndex> DepListForErrorChecking;
 
 	/** Index array - location of the resource for a UObject is stored in the ObjectIndices array using the UObject's Index */
@@ -173,6 +175,12 @@ public:
 	 */
 	virtual void SetFilterEditorOnly(bool bInFilterEditorOnly) override;
 
+	/** Sets the map of overrided properties for each export that should be treated as transient, and nulled out when serializing */
+	void SetTransientPropertyOverrides(const TMap<UObject*, TSet<FProperty*>>& InTransientPropertyOverrides)
+	{
+		TransientPropertyOverrides = &InTransientPropertyOverrides;
+	}
+
 	/** Set target platform memory map alignment. A negative value disables memory mapped bulk data. */
 	void SetMemoryMapAlignment(int64 InAlignment)
 	{
@@ -304,6 +312,7 @@ private:
 	FFileRegionMemoryWriter OptionalBulkDataAr;
 	/** Memory mapped bulk data archive. */
 	FFileRegionMemoryWriter MemoryMappedBulkDataAr;
+	const TMap<UObject*, TSet<FProperty*>>* TransientPropertyOverrides = nullptr;
 	/** Alignment for memory mapped data .*/ 
 	int64 MemoryMappingAlignment = -1;
 	/** Whether file regions are enabled. */
