Unshelved from pending changelist '24294393':

Added extra checks for importing delegate definitions from strings.

diff --git a/Engine/Source/Runtime/CoreUObject/Private/UObject/PropertyHelper.cpp b/Engine/Source/Runtime/CoreUObject/Private/UObject/PropertyHelper.cpp
--- a/Engine/Source/Runtime/CoreUObject/Private/UObject/PropertyHelper.cpp
+++ b/Engine/Source/Runtime/CoreUObject/Private/UObject/PropertyHelper.cpp
@@ -137,10 +137,14 @@ const TCHAR* DelegatePropertyTools::ImportDelegateFromText( FScriptDelegate& Del
 	{
 		++Buffer;
 	}
-	for( i=0; *Buffer && *Buffer != TCHAR('.') && *Buffer != TCHAR(')') && *Buffer != TCHAR(','); Buffer++ )
+	for( i=0; *Buffer && *Buffer != TCHAR('.') && *Buffer != TCHAR(')') && *Buffer != TCHAR(',') && i < NAME_SIZE; Buffer++ )
 	{
 		ObjName[i++] = *Buffer;
 	}
+	if (i == NAME_SIZE)
+	{
+		return nullptr;
+	}
 	ObjName[i] = TCHAR('\0');
 	UClass* Cls = nullptr;
 	UObject* Object = nullptr;
@@ -157,10 +161,14 @@ const TCHAR* DelegatePropertyTools::ImportDelegateFromText( FScriptDelegate& Del
 	if (*Buffer == TCHAR('.'))
 	{
 		Buffer++;
-		for( i=0; *Buffer && *Buffer != TCHAR(')') && *Buffer != TCHAR(','); Buffer++ )
+		for( i=0; *Buffer && *Buffer != TCHAR(')') && *Buffer != TCHAR(',') && i < NAME_SIZE; Buffer++ )
 		{
 			FuncName[i++] = *Buffer;
 		}
+		if (i == NAME_SIZE)
+		{
+			return nullptr;
+		}
 		FuncName[i] = TEXT('\0');
 	}
 	else
