Fixed assert firing in Debug mode when calling FScriptArrayHelper::MoveAssign(0, 0).  Regressed in CL# 24328679.

diff --git a/Engine/Source/Runtime/CoreUObject/Public/UObject/UnrealType.h b/Engine/Source/Runtime/CoreUObject/Public/UObject/UnrealType.h
--- a/Engine/Source/Runtime/CoreUObject/Public/UObject/UnrealType.h
+++ b/Engine/Source/Runtime/CoreUObject/Public/UObject/UnrealType.h
@@ -4430,7 +4430,6 @@ private:
 		{
 			checkSlow(Count >= 0);
 			checkSlow(Index >= 0); 
-			checkSlow(Index < Num());
 			checkSlow(Index + Count <= Num());
 			if (Count > 0)
 			{
