Fix incorrect comment for the ObjectMacros NoClear tag that indicates it will also remove the browse button

diff --git a/Engine/Source/Runtime/CoreUObject/Public/UObject/ObjectMacros.h b/Engine/Source/Runtime/CoreUObject/Public/UObject/ObjectMacros.h
--- a/Engine/Source/Runtime/CoreUObject/Public/UObject/ObjectMacros.h
+++ b/Engine/Source/Runtime/CoreUObject/Public/UObject/ObjectMacros.h
@@ -418,7 +418,7 @@ enum EPropertyFlags : uint64
 	//CPF_								= 0x0000000000400000,	///< 
 	//CPF_    							= 0x0000000000800000,	///< 
 	CPF_SaveGame						= 0x0000000001000000,	///< Property should be serialized for save games, this is only checked for game-specific archives with ArIsSaveGame
-	CPF_NoClear							= 0x0000000002000000,	///< Hide clear (and browse) button.
+	CPF_NoClear							= 0x0000000002000000,	///< Hide clear button.
 	//CPF_  							= 0x0000000004000000,	///<
 	CPF_ReferenceParm					= 0x0000000008000000,	///< Value is passed by reference; CPF_OutParam and CPF_Param should also be set.
 	CPF_BlueprintAssignable				= 0x0000000010000000,	///< MC Delegates only.  Property should be exposed for assigning in blueprint code
@@ -1005,7 +1005,7 @@ namespace UP
 		/// Object property can be exported with it's owner.
 		Export,
 
-		/// Hide clear (and browse) button in the editor.
+		/// Hide clear button in the editor.
 		NoClear,
 
 		/// Indicates that elements of an array can be modified, but its size cannot be changed.
