Add static assert to Cast to catch cases where the type passed in is not a UObjec,nor an Interface

diff --git a/Engine/Source/Runtime/CoreUObject/Public/Templates/Casts.h b/Engine/Source/Runtime/CoreUObject/Public/Templates/Casts.h
--- a/Engine/Source/Runtime/CoreUObject/Public/Templates/Casts.h
+++ b/Engine/Source/Runtime/CoreUObject/Public/Templates/Casts.h
@@ -116,6 +116,8 @@ FORCEINLINE To* Cast(From* Src)
 		}
 		else
 		{
+			static_assert(std::is_base_of_v<UObject, From>, "Attempting to use Cast<> on a type that is not a UObject or an Interface");
+			
 			if (((const UObject*)Src)->IsA<To>())
 			{
 				return (To*)Src;
