Fixed FUObjectThreadContext to once again be an exported thread-local singleton

#rnx

diff --git a/Engine/Source/Runtime/CoreUObject/Private/UObject/UObjectThreadContext.cpp b/Engine/Source/Runtime/CoreUObject/Private/UObject/UObjectThreadContext.cpp
--- a/Engine/Source/Runtime/CoreUObject/Private/UObject/UObjectThreadContext.cpp
+++ b/Engine/Source/Runtime/CoreUObject/Private/UObject/UObjectThreadContext.cpp
@@ -25,6 +25,13 @@ FUObjectThreadContext::~FUObjectThreadContext()
 {
 }
 
+FObjectInitializer& FUObjectThreadContext::ReportNull()
+{
+	FObjectInitializer* ObjectInitializerPtr = TopInitializer();
+	UE_CLOG(!ObjectInitializerPtr, LogUObjectThreadContext, Fatal, TEXT("Tried to get the current ObjectInitializer, but none is set. Please use NewObject to construct new UObject-derived classes."));
+	return *ObjectInitializerPtr;
+}
+
 FUObjectSerializeContext::FUObjectSerializeContext()
 	: RefCount(0)
 	, ImportCount(0)
diff --git a/Engine/Source/Runtime/CoreUObject/Public/UObject/UObjectThreadContext.h b/Engine/Source/Runtime/CoreUObject/Public/UObject/UObjectThreadContext.h
--- a/Engine/Source/Runtime/CoreUObject/Public/UObject/UObjectThreadContext.h
+++ b/Engine/Source/Runtime/CoreUObject/Public/UObject/UObjectThreadContext.h
@@ -31,12 +31,12 @@ struct FUObjectSerializeContext;
 
 COREUOBJECT_API DECLARE_LOG_CATEGORY_EXTERN(LogUObjectThreadContext, Log, All);
 
-class FUObjectThreadContext : public TThreadSingleton<FUObjectThreadContext>
+class COREUOBJECT_API FUObjectThreadContext : public TThreadSingleton<FUObjectThreadContext>
 {
 	friend TThreadSingleton<FUObjectThreadContext>;
 
-	COREUOBJECT_API FUObjectThreadContext();
-	COREUOBJECT_API virtual ~FUObjectThreadContext();
+	FUObjectThreadContext();
+	virtual ~FUObjectThreadContext();
 
 	/** Stack of currently used FObjectInitializers for this thread */
 	TArray<FObjectInitializer*> InitializerStack;
@@ -75,9 +75,11 @@ public:
 	*/
 	FObjectInitializer& TopInitializerChecked()
 	{
-		FObjectInitializer* ObjectInitializerPtr = TopInitializer();
-		UE_CLOG(!ObjectInitializerPtr, LogUObjectThreadContext, Fatal, TEXT("Tried to get the current ObjectInitializer, but none is set. Please use NewObject to construct new UObject-derived classes."));
-		return *ObjectInitializerPtr;
+		if (FObjectInitializer* ObjectInitializerPtr = TopInitializer())
+		{
+			return *ObjectInitializerPtr;
+		}
+		return ReportNull();
 	}
 
 	/** true when we are routing ConditionalPostLoad/PostLoad to objects										*/
@@ -113,6 +115,9 @@ public:
 	}
 
 private:
+	/** Report that the current ObjectInitializer is null. */
+	FObjectInitializer& ReportNull();
+
 	/** Current serialization context */
 	TRefCountPtr<FUObjectSerializeContext> SerializeContext;
 };
