Added a few missing definitions that were preventing "Debug Editor" builds of projects that disable GC from linking.

diff --git a/Engine/Source/Runtime/CoreUObject/Private/UObject/GarbageCollection.cpp b/Engine/Source/Runtime/CoreUObject/Private/UObject/GarbageCollection.cpp
--- a/Engine/Source/Runtime/CoreUObject/Private/UObject/GarbageCollection.cpp
+++ b/Engine/Source/Runtime/CoreUObject/Private/UObject/GarbageCollection.cpp
@@ -2985,21 +2985,22 @@ void TReachabilityCollector<Options>::HandleObjectReferences(UObject** InObjects
 }
 
 #else // !UE_WITH_GC
-
 static constexpr int32 MaxWorkers = 1;
-static void CollectGarbageInternal(EObjectFlags, bool)				{ unimplemented(); }
-void PadObjectArray(TArray<UObject*>&)								{ unimplemented(); }
-FWorkerContext::FWorkerContext()									{ unimplemented(); }
-FWorkerContext::~FWorkerContext()									{ unimplemented(); }
-FWorkBlockifier::~FWorkBlockifier()									{ unimplemented(); }
-void FWorkBlockifier::FreeOwningBlock(UObject*const*)				{ unimplemented(); }
-void FWorkBlockifier::PushFullBlockSync()							{ unimplemented(); }
-FWorkBlock* FWorkBlockifier::PopFullBlockSync()						{ unimplemented(); return nullptr; }
-FWorkBlock* FWorkBlockifier::PopWipBlock()							{ unimplemented(); return nullptr; }
-void FSlowARO::CallSync(uint32, UObject*, FReferenceCollector&)		{ unimplemented(); }
-void RegisterSlowImplementation(AROFunc, EAROFlags)					{ return; }
-int32 FindSlowImplementation(AROFunc)								{ return INDEX_NONE; }
-
+static void CollectGarbageInternal(EObjectFlags, bool)											{ unimplemented(); }
+void PadObjectArray(TArray<UObject*>&)															{ unimplemented(); }
+FWorkerContext::FWorkerContext()																{ unimplemented(); }
+FWorkerContext::~FWorkerContext()																{ unimplemented(); }
+FWorkBlockifier::~FWorkBlockifier()																{ unimplemented(); }
+void FWorkBlockifier::FreeOwningBlock(UObject*const*)											{ unimplemented(); }
+void FWorkBlockifier::PushFullBlockSync()														{ unimplemented(); }
+FWorkBlock* FWorkBlockifier::PopFullBlockSync()													{ unimplemented(); return nullptr; }
+FWorkBlock* FWorkBlockifier::PopFullBlockAsync()												{ unimplemented(); return nullptr; }
+FWorkBlock* FWorkBlockifier::PopWipBlock()														{ unimplemented(); return nullptr; }
+void FSlowARO::CallSync(uint32, UObject*, FReferenceCollector&)									{ unimplemented(); }
+bool FSlowARO::TryQueueCall(uint32 SlowAROIndex, UObject* Object, FWorkerContext& Context)		{ unimplemented(); return false; }
+void RegisterSlowImplementation(AROFunc, EAROFlags)												{ return; }
+int32 FindSlowImplementation(AROFunc)															{ return INDEX_NONE; }
+ELoot StealWork(FWorkerContext& Context, FReferenceCollector& Collector, FWorkBlock*& OutBlock)	{ return ELoot::Nothing; }
 #endif // !UE_WITH_GC
 
 } // namespace UE::GC
