Bulkdata will no longer fatal error if the archive is in an error state before serialization starts.

- CL 24316937 made an effort to prevent serialization of bad data from corrupting memory, however the checks were a little too zealous. In some cases the archive provided for serialization is already in an error state in which case we should trust that the calling code has processes the problem correctly. Handling the error at this point from within bulkdata can give misleading errors.
- We now check if the archive is in the error state or not BEFORE we start serialization and if it is in error we just return early with default values.
- This should ensure that the problem is reported by the system creating it rather than always by bulkdata and make things easier to diagnose/fix/

diff --git a/Engine/Source/Runtime/CoreUObject/Private/Serialization/BulkData.cpp b/Engine/Source/Runtime/CoreUObject/Private/Serialization/BulkData.cpp
--- a/Engine/Source/Runtime/CoreUObject/Private/Serialization/BulkData.cpp
+++ b/Engine/Source/Runtime/CoreUObject/Private/Serialization/BulkData.cpp
@@ -216,6 +216,12 @@ FString GetDebugNameFromArchive(FArchive& Ar)
 
 bool FBulkMetaData::FromSerialized(FArchive& Ar, int64 ElementSize, FBulkMetaData& OutMetaData, int64& OutDuplicateOffset)
 {
+	if (Ar.IsError())
+	{
+		OutMetaData = FBulkMetaData();
+		return false;
+	}
+
 	FBulkMetaResource Resource;
 	Ar << Resource;
 
