Fix string table entry comparison in FTextProperty. Prevent comparison by display string when two entries use the placeholder text.

diff --git a/Engine/Source/Runtime/CoreUObject/Private/UObject/TextProperty.cpp b/Engine/Source/Runtime/CoreUObject/Private/UObject/TextProperty.cpp
--- a/Engine/Source/Runtime/CoreUObject/Private/UObject/TextProperty.cpp
+++ b/Engine/Source/Runtime/CoreUObject/Private/UObject/TextProperty.cpp
@@ -59,7 +59,8 @@ bool FTextProperty::Identical_Implementation(const FText& ValueA, const FText& V
 	// A culture variant text is never equal to a culture invariant text
 	// A transient text is never equal to a non-transient text
 	// An empty text is never equal to a non-empty text
-	if (ValueA.IsCultureInvariant() != ValueB.IsCultureInvariant() || ValueA.IsTransient() != ValueB.IsTransient() || ValueA.IsEmpty() != ValueB.IsEmpty())
+	// Text from a string table is never equal to text not from a string table
+	if (ValueA.IsCultureInvariant() != ValueB.IsCultureInvariant() || ValueA.IsTransient() != ValueB.IsTransient() || ValueA.IsEmpty() != ValueB.IsEmpty() || ValueA.IsFromStringTable() != ValueB.IsFromStringTable())
 	{
 		return false;
 	}
@@ -70,24 +71,23 @@ bool FTextProperty::Identical_Implementation(const FText& ValueA, const FText& V
 		return true;
 	}
 
-	// If both texts share the same pointer, then they must be equal
-	if (ValueA.IdenticalTo(ValueB))
+	// String table entries should only be considered equal if they're using the same string table and key.
+	if (ValueA.IsFromStringTable())
 	{
-		// Placeholder string table entries will have the same pointer, but should only be considered equal if they're using the same string table and key
-		if (ValueA.IsFromStringTable() && FTextInspector::GetSourceString(ValueA) == &FStringTableEntry::GetPlaceholderSourceString())
-		{
-			FName ValueAStringTableId;
-			FString ValueAStringTableEntryKey;
-			FTextInspector::GetTableIdAndKey(ValueA, ValueAStringTableId, ValueAStringTableEntryKey);
+		FName ValueAStringTableId;
+		FTextKey ValueAStringTableEntryKey;
+		FTextInspector::GetTableIdAndKey(ValueA, ValueAStringTableId, ValueAStringTableEntryKey);
 
-			FName ValueBStringTableId;
-			FString ValueBStringTableEntryKey;
-			FTextInspector::GetTableIdAndKey(ValueB, ValueBStringTableId, ValueBStringTableEntryKey);
+		FName ValueBStringTableId;
+		FTextKey ValueBStringTableEntryKey;
+		FTextInspector::GetTableIdAndKey(ValueB, ValueBStringTableId, ValueBStringTableEntryKey);
 
-			return ValueAStringTableId == ValueBStringTableId && ValueAStringTableEntryKey.Equals(ValueBStringTableEntryKey, ESearchCase::CaseSensitive);
-		}
+		return ValueAStringTableId == ValueBStringTableId && ValueAStringTableEntryKey == ValueBStringTableEntryKey;
+	}
 
-		// Otherwise they're equal
+	// If both texts share the same pointer, then they must be equal
+	if (ValueA.IdenticalTo(ValueB))
+	{
 		return true;
 	}
 
