SavePackage: Extend the warning given for "PreSave() has modified %s" to handle the case when PreSave renames the object out of the package.
#rnx

diff --git a/Engine/Source/Runtime/CoreUObject/Private/UObject/SavePackage2.cpp b/Engine/Source/Runtime/CoreUObject/Private/UObject/SavePackage2.cpp
--- a/Engine/Source/Runtime/CoreUObject/Private/UObject/SavePackage2.cpp
+++ b/Engine/Source/Runtime/CoreUObject/Private/UObject/SavePackage2.cpp
@@ -277,20 +277,38 @@ ESavePackageResult RoutePresave(FSaveContext& SaveContext)
 			{
 				FArchiveObjectCrc32 CrcArchive;
 				CrcArchive.ArIsFilterEditorOnly = true;
+				FString PathNameBefore = Object->GetPathName();
 				int32 Before = CrcArchive.Crc32(Object);
 				UE::SavePackageUtilities::CallPreSave(Object, SaveContext.GetObjectSaveContext());
 				int32 After = CrcArchive.Crc32(Object);
 
 				if (Before != After)
 				{
-					UE_ASSET_LOG(
-						LogSavePackage,
-						Warning,
-						Object,
-						TEXT("Non-deterministic cook warning - PreSave() has modified %s '%s' - a resave may be required"),
-						Object->HasAnyFlags(RF_ClassDefaultObject) ? TEXT("CDO") : TEXT("archetype"),
-						*Object->GetName()
-					);
+					FString PathNameAfter = Object->GetPathName();
+					if (PathNameBefore == PathNameAfter)
+					{
+						UE_ASSET_LOG(
+							LogSavePackage,
+							Warning,
+							Object,
+							TEXT("Non-deterministic cook warning - PreSave() has modified %s '%s' - a resave may be required"),
+							Object->HasAnyFlags(RF_ClassDefaultObject) ? TEXT("CDO") : TEXT("archetype"),
+							*Object->GetName()
+						);
+					}
+					else
+					{
+						UObject* ObjectUsedToReport = Object->IsInPackage(SaveContext.GetPackage()) ? Object : SaveContext.GetPackage();
+						UE_ASSET_LOG(
+							LogSavePackage,
+							Warning,
+							ObjectUsedToReport,
+							TEXT("Non-deterministic cook warning - PreSave() has renamed %s '%s' to %s - a resave may be required"),
+							Object->HasAnyFlags(RF_ClassDefaultObject) ? TEXT("CDO") : TEXT("archetype"),
+							*PathNameBefore, *PathNameAfter
+						);
+					}
+
 				}
 			}
 			else
