fix for performance problem when converting arrays of TObjectPtr to UObject*
changed the access tracking to a bulk operation api to allow for single calls instead of O(n) calls

diff --git a/Engine/Source/Editor/UnrealEd/Private/Cooker/PackageBuildDependencyTracker.cpp b/Engine/Source/Editor/UnrealEd/Private/Cooker/PackageBuildDependencyTracker.cpp
--- a/Engine/Source/Editor/UnrealEd/Private/Cooker/PackageBuildDependencyTracker.cpp
+++ b/Engine/Source/Editor/UnrealEd/Private/Cooker/PackageBuildDependencyTracker.cpp
@@ -64,9 +64,9 @@ FPackageBuildDependencyTracker::~FPackageBuildDependencyTracker()
 	UE::CoreUObject::RemoveObjectHandleReadCallback(ObjectHandleReadHandle);
 }
 
-void FPackageBuildDependencyTracker::StaticOnObjectHandleRead(UObject* ReadObject)
+void FPackageBuildDependencyTracker::StaticOnObjectHandleRead(const UObject* const* Objects, uint32 Count)
 {
-	if (!ReadObject)
+	if(Count == 0 || (Count == 1 && (!Objects[0] || !Objects[0]->HasAnyFlags(RF_Public))))
 	{
 		return;
 	}
@@ -82,43 +82,52 @@ void FPackageBuildDependencyTracker::StaticOnObjectHandleRead(UObject* ReadObjec
 		return;
 	}
 
-	if (!ReadObject->HasAnyFlags(RF_Public))
+	for (uint32 i = 0; i < Count; ++i)
 	{
-		return;
-	}
+		const UObject* ReadObject = Objects[i];
+		if (!ReadObject)
+		{
+			continue;
+		}
 
-	FName Referencer = AccumulatedScopeData->PackageName;
-	FName Referenced = ReadObject->GetOutermost()->GetFName();
-	if (Referencer == Referenced)
-	{
-		return;
-	}
+		if (!ReadObject->HasAnyFlags(RF_Public))
+		{
+			continue;
+		}
 
-	if (AccumulatedScopeData->OpName == PackageAccessTrackingOps::NAME_NoAccessExpected)
-	{
-		UE_LOG(LogPackageBuildDependencyTracker, Warning, TEXT("Object %s is referencing object %s inside of a NAME_NoAccessExpected scope. Programmer should narrow the scope or debug the reference."),
-			*Referencer.ToString(), *Referenced.ToString());
-	}
+		FName Referencer = AccumulatedScopeData->PackageName;
+		FName Referenced = ReadObject->GetOutermost()->GetFName();
+		if (Referencer == Referenced)
+		{
+			continue;
+		}
+
+		if (AccumulatedScopeData->OpName == PackageAccessTrackingOps::NAME_NoAccessExpected)
+		{
+			UE_LOG(LogPackageBuildDependencyTracker, Warning, TEXT("Object %s is referencing object %s inside of a NAME_NoAccessExpected scope. Programmer should narrow the scope or debug the reference."),
+				*Referencer.ToString(), *Referenced.ToString());
+		}
 
-	LLM_SCOPE_BYNAME(TEXTVIEW("PackageBuildDependencyTracker"));
+		LLM_SCOPE_BYNAME(TEXTVIEW("PackageBuildDependencyTracker"));
 
-	FBuildDependencyAccessData AccessData{Referenced, AccumulatedScopeData->TargetPlatform};
-	FScopeLock RecordsScopeLock(&Singleton.RecordsLock);
-	if (Referencer == Singleton.LastReferencer)
-	{
-		if (AccessData != Singleton.LastAccessData)
+		FBuildDependencyAccessData AccessData{ Referenced, AccumulatedScopeData->TargetPlatform };
+		FScopeLock RecordsScopeLock(&Singleton.RecordsLock);
+		if (Referencer == Singleton.LastReferencer)
+		{
+			if (AccessData != Singleton.LastAccessData)
+			{
+				Singleton.LastAccessData = AccessData;
+				Singleton.LastReferencerSet->Add(AccessData);
+			}
+		}
+		else
 		{
 			Singleton.LastAccessData = AccessData;
+			Singleton.LastReferencer = Referencer;
+			Singleton.LastReferencerSet = &Singleton.Records.FindOrAdd(Referencer);
 			Singleton.LastReferencerSet->Add(AccessData);
 		}
 	}
-	else
-	{
-		Singleton.LastAccessData = AccessData;
-		Singleton.LastReferencer = Referencer;
-		Singleton.LastReferencerSet = &Singleton.Records.FindOrAdd(Referencer);
-		Singleton.LastReferencerSet->Add(AccessData);
-	}
 }
 
 void DumpBuildDependencyTrackerStats()
diff --git a/Engine/Source/Editor/UnrealEd/Private/Cooker/PackageBuildDependencyTracker.h b/Engine/Source/Editor/UnrealEd/Private/Cooker/PackageBuildDependencyTracker.h
--- a/Engine/Source/Editor/UnrealEd/Private/Cooker/PackageBuildDependencyTracker.h
+++ b/Engine/Source/Editor/UnrealEd/Private/Cooker/PackageBuildDependencyTracker.h
@@ -40,7 +40,7 @@ private:
 	virtual ~FPackageBuildDependencyTracker();
 
 	/** Track object reference reads */
-	static void StaticOnObjectHandleRead(UObject* ReadObject);
+	static void StaticOnObjectHandleRead(const UObject* const* ReadObject, uint32 Count);
 
 	mutable FCriticalSection RecordsLock;
 	TMap<FName, TSet<FBuildDependencyAccessData>> Records;
diff --git a/Engine/Source/Runtime/Core/Public/Containers/Array.h b/Engine/Source/Runtime/Core/Public/Containers/Array.h
--- a/Engine/Source/Runtime/Core/Public/Containers/Array.h
+++ b/Engine/Source/Runtime/Core/Public/Containers/Array.h
@@ -2285,7 +2285,7 @@ public:
 	operator TArray<typename TContainerElementTypeCompatibility<AliasElementType>::ReinterpretType, AllocatorType>& ()
 	{
 		using ElementCompat = TContainerElementTypeCompatibility<ElementType>;
-		ElementCompat::ReinterpretRange(begin(), end());
+		ElementCompat::ReinterpretRangeContiguous(begin(), end(), Num());
 		return *reinterpret_cast<TArray<typename ElementCompat::ReinterpretType>*>(this);
 	}
 
@@ -2297,7 +2297,7 @@ public:
 	operator const TArray<typename TContainerElementTypeCompatibility<AliasElementType>::ReinterpretType, AllocatorType>& () const
 	{
 		using ElementCompat = TContainerElementTypeCompatibility<ElementType>;
-		ElementCompat::ReinterpretRange(begin(), end());
+		ElementCompat::ReinterpretRangeContiguous(begin(), end(), Num());
 		return *reinterpret_cast<const TArray<typename ElementCompat::ReinterpretType>*>(this);
 	}
 
diff --git a/Engine/Source/Runtime/Core/Public/Containers/ArrayView.h b/Engine/Source/Runtime/Core/Public/Containers/ArrayView.h
--- a/Engine/Source/Runtime/Core/Public/Containers/ArrayView.h
+++ b/Engine/Source/Runtime/Core/Public/Containers/ArrayView.h
@@ -44,8 +44,9 @@ namespace ArrayViewPrivate
 		auto NaturalPtr = GetData(Forward<T>(Arg));
 		using NaturalElementType = typename TRemovePointer<decltype(NaturalPtr)>::Type;
 
-		auto EndPtr = NaturalPtr + GetNum(Arg);
-		TContainerElementTypeCompatibility<NaturalElementType>::ReinterpretRange(NaturalPtr, EndPtr);
+		auto Size = GetNum(Arg);
+		auto EndPtr = NaturalPtr + Size;
+		TContainerElementTypeCompatibility<NaturalElementType>::ReinterpretRangeContiguous(NaturalPtr, EndPtr, Size);
 
 		return reinterpret_cast<typename TContainerElementTypeCompatibility<NaturalElementType>::ReinterpretType*>(NaturalPtr);
 	}
diff --git a/Engine/Source/Runtime/Core/Public/Containers/ContainerElementTypeCompatibility.h b/Engine/Source/Runtime/Core/Public/Containers/ContainerElementTypeCompatibility.h
--- a/Engine/Source/Runtime/Core/Public/Containers/ContainerElementTypeCompatibility.h
+++ b/Engine/Source/Runtime/Core/Public/Containers/ContainerElementTypeCompatibility.h
@@ -14,13 +14,17 @@ template <typename InElementType>
 struct TContainerElementTypeCompatibility
 {
 	typedef InElementType ReinterpretType;
+	typedef InElementType CopyFromOtherType;
 
 	template <typename IterBeginType, typename IterEndType, typename OperatorType = InElementType&(*)(IterBeginType&)>
 	static void ReinterpretRange(IterBeginType Iter, IterEndType IterEnd, OperatorType Operator = [](IterBeginType& InIt) -> InElementType& { return *InIt; })
 	{
 	}
 
-	typedef InElementType CopyFromOtherType;
+	template <typename IterBeginType, typename IterEndType, typename SizeType, typename OperatorType = InElementType & (*)(IterBeginType&)>
+	static void ReinterpretRangeContiguous(IterBeginType Iter, IterEndType IterEnd, SizeType Size, OperatorType Operator = [](IterBeginType& InIt) -> InElementType& { return *InIt; })
+	{
+	}
 
 	static constexpr void CopyingFromOtherType() {}
 };
diff --git a/Engine/Source/Runtime/CoreUObject/Private/UObject/Class.cpp b/Engine/Source/Runtime/CoreUObject/Private/UObject/Class.cpp
--- a/Engine/Source/Runtime/CoreUObject/Private/UObject/Class.cpp
+++ b/Engine/Source/Runtime/CoreUObject/Private/UObject/Class.cpp
@@ -3443,7 +3443,17 @@ static void FindUninitializedScriptStructMembers(UScriptStruct* ScriptStruct, ES
 		}
 #endif // WITH_EDITORONLY_DATA
 
-		if (const FObjectPropertyBase* ObjectProperty = CastField<const FObjectPropertyBase>(Property))
+		if (const FObjectProperty* ObjectPtrProperty = CastField<const FObjectProperty>(Property))
+		{
+			// Check any reflected pointer properties to make sure they got initialized
+			TObjectPtr<UObject> PropValue = ObjectPtrProperty->GetObjectPtrPropertyValue(WrapperFE.GetData());
+			const void** RawValue = reinterpret_cast<const void**>(&PropValue);
+			if (RawValue == BadPointer)
+			{
+				OutUninitializedProperties.Add(Property);
+			}
+		}
+		else if (const FObjectPropertyBase* ObjectProperty = CastField<const FObjectPropertyBase>(Property))
 		{
 			// Check any reflected pointer properties to make sure they got initialized
 			const UObject* PropValue = ObjectProperty->GetObjectPropertyValue_InContainer(WrapperFE.GetData());
@@ -3580,8 +3590,9 @@ int32 FStructUtils::AttemptToFindUninitializedScriptStructMembers()
 		const FObjectPropertyBase* UninitializedProperty = CastFieldChecked<const FObjectPropertyBase>(TestUninitializedScriptStructMembersTestStruct->FindPropertyByName(TEXT("UninitializedObjectReference")));
 		const FObjectPropertyBase* InitializedProperty = CastFieldChecked<const FObjectPropertyBase>(TestUninitializedScriptStructMembersTestStruct->FindPropertyByName(TEXT("InitializedObjectReference")));
 		
-		const UObject* UninitializedPropValue = UninitializedProperty->GetObjectPropertyValue_InContainer(WrapperFE.GetData());
-		if (UninitializedPropValue != BadPointer)
+		TObjectPtr<UObject> UninitializedPropValue = UninitializedProperty->GetObjectPtrPropertyValue(WrapperFE.GetData());
+		const void** RawValue = reinterpret_cast<const void**>(&UninitializedPropValue);
+		if (*RawValue != BadPointer)
 		{
 			UE_LOG(LogClass, Warning, TEXT("ObjectProperty %s%s::%s seems to be initialized properly but it shouldn't be. Verify that AttemptToFindUninitializedScriptStructMembers() is working properly"), 
 				TestUninitializedScriptStructMembersTestStruct->GetPrefixCPP(), *TestUninitializedScriptStructMembersTestStruct->GetName(), *UninitializedProperty->GetNameCPP());
diff --git a/Engine/Source/Runtime/CoreUObject/Private/UObject/ObjectHandleTracking.cpp b/Engine/Source/Runtime/CoreUObject/Private/UObject/ObjectHandleTracking.cpp
--- a/Engine/Source/Runtime/CoreUObject/Private/UObject/ObjectHandleTracking.cpp
+++ b/Engine/Source/Runtime/CoreUObject/Private/UObject/ObjectHandleTracking.cpp
@@ -17,12 +17,12 @@ namespace UE::CoreUObject
 				return Callbacks;
 			}
 
-			void OnHandleRead(UObject* Object)
+			void OnHandleRead(const UObject* const* Object, uint32 Count)
 			{
 				FReadScopeLock _(HandleLock);
 				for (auto&& Pair : ReadHandleCallbacks)
 				{
-					Pair.Value(Object);
+					Pair.Value(Object, Count);
 				}
 			}
 
@@ -148,9 +148,9 @@ namespace UE::CoreUObject
 			FRWLock HandleLock;
 		};
 
-		void OnHandleRead(UObject* Object)
+		void OnHandleRead(const UObject* const* Object, uint32 Count)
 		{
-			ObjectHandleCallbacks::Get().OnHandleRead(Object);
+			ObjectHandleCallbacks::Get().OnHandleRead(Object, Count);
 		}
 
 		void OnClassReferenceResolved(const FObjectRef& ObjectRef, UPackage* Package, UClass* Class)
diff --git a/Engine/Source/Runtime/CoreUObject/Private/UObject/ObjectSerializeAccessScope.cpp b/Engine/Source/Runtime/CoreUObject/Private/UObject/ObjectSerializeAccessScope.cpp
--- a/Engine/Source/Runtime/CoreUObject/Private/UObject/ObjectSerializeAccessScope.cpp
+++ b/Engine/Source/Runtime/CoreUObject/Private/UObject/ObjectSerializeAccessScope.cpp
@@ -41,7 +41,7 @@ namespace UE::CoreUObject::Private
 		{
 			if (bEnabled)
 			{
-				FObjectHandleReadFunc Func = [this](UObject* ReadObject)
+				FObjectHandleReadFunc Func = [this](const UObject* const* ReadObject, int Count)
 				{
 					FLocalData& Data = GetThreadLocalData();
 					if (!IsEnableSerializeAccessWarnings() || (Data.CurrentScope == nullptr || Data.CurrentScope->bSuspended))
@@ -60,7 +60,14 @@ namespace UE::CoreUObject::Private
 					ANSICHAR StackTrace[StackTraceSize];
 					StackTrace[0] = 0;
 					FPlatformStackWalk::StackWalkAndDump(StackTrace, StackTraceSize, 3);
-					UE_LOG(LogLinker, Warning, TEXT("%s trying to access referenced object %s during serialize\n%s"), *CurrentObject.GetClass()->GetFullName(), (ReadObject ? *ReadObject->GetFullName() : TEXT("None")), ANSI_TO_TCHAR(StackTrace));
+					if (Count == 1)
+					{
+						UE_LOG(LogLinker, Warning, TEXT("%s trying to access referenced object %s during serialize\n%s"), *CurrentObject.GetClass()->GetFullName(), (ReadObject[0] ? *ReadObject[0]->GetFullName() : TEXT("None")), ANSI_TO_TCHAR(StackTrace));
+					}
+					else
+					{
+						UE_LOG(LogLinker, Warning, TEXT("%s trying to access an array of referenced objects during serialize\n%s"), *CurrentObject.GetClass()->GetFullName(), ANSI_TO_TCHAR(StackTrace));
+					}
 				};
 
 				CallbackHandle = AddObjectHandleReadCallback(Func);
diff --git a/Engine/Source/Runtime/CoreUObject/Public/UObject/ObjectHandleTracking.h b/Engine/Source/Runtime/CoreUObject/Public/UObject/ObjectHandleTracking.h
--- a/Engine/Source/Runtime/CoreUObject/Public/UObject/ObjectHandleTracking.h
+++ b/Engine/Source/Runtime/CoreUObject/Public/UObject/ObjectHandleTracking.h
@@ -25,7 +25,7 @@ namespace UE::CoreUObject
 	  *
 	  * @param ReadObject	The object that was read from a handle.
 	  */
-	using FObjectHandleReadFunc = TFunction<void(UObject* ReadObject)>;
+	using FObjectHandleReadFunc = TFunction<void(const UObject* const* ReadObject, uint32 Count)>;
 
 	/**
 	 * Callback notifying when a class is resolved from an object handle or object reference.
@@ -102,7 +102,12 @@ namespace UE::CoreUObject
 #if UE_WITH_OBJECT_HANDLE_TRACKING
 namespace UE::CoreUObject::Private
 {
-	COREUOBJECT_API void OnHandleRead(UObject* Object);
+	
+	COREUOBJECT_API void OnHandleRead(const UObject* const* Object, uint32 Count);
+	inline void OnHandleRead(const UObject* Object)
+	{
+		OnHandleRead(&Object, 1u);
+	}
 	COREUOBJECT_API void OnClassReferenceResolved(const FObjectRef& ObjectRef, UPackage* ClassPackage, UClass* Class);
 	COREUOBJECT_API void OnReferenceResolved(const FObjectRef& ObjectRef, UPackage* ObjectPackage, UObject* Object);
 	COREUOBJECT_API void OnReferenceLoaded(const FObjectRef& ObjectRef, UPackage* ObjectPackage, UObject* Object);
@@ -112,6 +117,7 @@ namespace UE::CoreUObject::Private
 namespace UE::CoreUObject::Private
 {
 	inline void OnHandleRead(UObject* Object) { }
+	inline void OnHandleRead(UObject** Object, int Count) { }
 	inline void OnClassReferenceResolved(const FObjectRef& ObjectRef, UPackage* Package, UClass* Class) { }
 	inline void OnReferenceResolved(const FObjectRef& ObjectRef, UPackage* Package, UClass* Class) { }
 	inline void OnReferenceLoaded(const FObjectRef& ObjectRef, UPackage* Package, UObject* Object) { }
diff --git a/Engine/Source/Runtime/CoreUObject/Public/UObject/ObjectPtr.h b/Engine/Source/Runtime/CoreUObject/Public/UObject/ObjectPtr.h
--- a/Engine/Source/Runtime/CoreUObject/Public/UObject/ObjectPtr.h
+++ b/Engine/Source/Runtime/CoreUObject/Public/UObject/ObjectPtr.h
@@ -742,7 +742,8 @@ template <typename T>
 struct TContainerElementTypeCompatibility<TObjectPtr<T>>
 {
 	typedef T* ReinterpretType;
-	
+	typedef T* CopyFromOtherType;
+
 	template <typename IterBeginType, typename IterEndType, typename OperatorType = std::remove_reference_t<decltype(*std::declval<IterBeginType>())>& (*)(IterBeginType&)>
 	UE_OBJPTR_DEPRECATED(5.0, "Reinterpretation between ranges of one type to another type is deprecated.")
 	static void ReinterpretRange(IterBeginType Iter, IterEndType IterEnd, OperatorType Operator = [](IterBeginType& InIt) -> decltype(auto) { return *InIt; })
@@ -756,8 +757,24 @@ struct TContainerElementTypeCompatibility<TObjectPtr<T>>
 #endif
 	}
 
-	typedef T* CopyFromOtherType;
-
+	template <typename IterBeginType, typename IterEndType, typename SizeType, typename OperatorType = std::remove_reference_t<decltype(*std::declval<IterBeginType>())>& (*)(IterBeginType&)>
+	UE_OBJPTR_DEPRECATED(5.0, "Reinterpretation between ranges of one type to another type is deprecated.")
+	static void ReinterpretRangeContiguous(IterBeginType Iter, IterEndType IterEnd, SizeType Size, OperatorType Operator = [](IterBeginType& InIt) -> decltype(auto) { return *InIt; })
+	{
+#if UE_WITH_OBJECT_HANDLE_LATE_RESOLVE || UE_WITH_OBJECT_HANDLE_TRACKING
+		const TObjectPtr<T>* Begin = &*Iter;
+		while (Iter != IterEnd)
+		{
+			auto& Ptr = Operator(Iter);
+			const FObjectPtr& ObjPtr = reinterpret_cast<const FObjectPtr&>(Ptr);
+			UE::CoreUObject::Private::ResolveObjectHandleNoRead(ObjPtr.GetHandleRef());
+			++Iter;
+		}
+		
+		UE::CoreUObject::Private::OnHandleRead(reinterpret_cast<const UObject* const*>(&*Begin), Size);
+#endif
+	}
+	
 	UE_OBJPTR_DEPRECATED(5.0, "Copying ranges of one type to another type is deprecated.")
 	static constexpr void CopyingFromOtherType() {}
 };
@@ -766,6 +783,7 @@ template <typename T>
 struct TContainerElementTypeCompatibility<const TObjectPtr<T>>
 {
 	typedef T* const ReinterpretType;
+	typedef T* const CopyFromOtherType;
 
 	template <typename IterBeginType, typename IterEndType, typename OperatorType = const std::remove_cv_t<std::remove_reference_t<decltype(*std::declval<IterBeginType>())>>&(*)(IterBeginType&)>
 	UE_OBJPTR_DEPRECATED(5.0, "Reinterpretation between ranges of one type to another type is deprecated.")
@@ -780,13 +798,27 @@ struct TContainerElementTypeCompatibility<const TObjectPtr<T>>
 #endif
 	}
 
-	typedef T* const CopyFromOtherType;
+	template <typename IterBeginType, typename IterEndType, typename SizeType, typename OperatorType = std::remove_reference_t<decltype(*std::declval<IterBeginType>())>& (*)(IterBeginType&)>
+	UE_OBJPTR_DEPRECATED(5.0, "Reinterpretation between ranges of one type to another type is deprecated.")
+	static void ReinterpretRangeContiguous(IterBeginType Iter, IterEndType IterEnd, SizeType Size, OperatorType Operator = [](IterBeginType& InIt) -> decltype(auto) { return *InIt; })
+	{
+#if UE_WITH_OBJECT_HANDLE_LATE_RESOLVE || UE_WITH_OBJECT_HANDLE_TRACKING
+		const TObjectPtr<T>* Begin = &*Iter;
+		while (Iter != IterEnd)
+		{
+			auto& Ptr = Operator(Iter);
+			const FObjectPtr& ObjPtr = reinterpret_cast<const FObjectPtr&>(Ptr);
+			UE::CoreUObject::Private::ResolveObjectHandleNoRead(ObjPtr.GetHandleRef());
+			++Iter;
+		}
+ 		UE::CoreUObject::Private::OnHandleRead(reinterpret_cast<const UObject* const*>(&Begin), Size);
+#endif
+	}
 
 	UE_OBJPTR_DEPRECATED(5.0, "Copying ranges of one type to another type is deprecated.")
 	static constexpr void CopyingFromOtherType() {}
 };
 
-
 // Trait which allows TObjectPtr to be default constructed by memsetting to zero.
 template <typename T>
 struct TIsZeroConstructType<TObjectPtr<T>>
diff --git a/Engine/Source/Runtime/CoreUObject/Tests/ObjectPropertyTest.cpp b/Engine/Source/Runtime/CoreUObject/Tests/ObjectPropertyTest.cpp
--- a/Engine/Source/Runtime/CoreUObject/Tests/ObjectPropertyTest.cpp
+++ b/Engine/Source/Runtime/CoreUObject/Tests/ObjectPropertyTest.cpp
@@ -17,7 +17,7 @@ TEST_CASE("UE::CoreUObject::FObjectProperty::CheckValidAddress")
 {
 	bool bAllowRead = false;
 #if UE_WITH_OBJECT_HANDLE_TRACKING
-	auto CallbackHandle = UE::CoreUObject::AddObjectHandleReadCallback([&bAllowRead](UObject* ReadObject)
+	auto CallbackHandle = UE::CoreUObject::AddObjectHandleReadCallback([&bAllowRead](const UObject* const* ReadObject, int Count)
 		{
 			if (!bAllowRead)
 				FAIL("Unexpected read during CheckValidObject");
@@ -71,7 +71,7 @@ TEST_CASE("UE::CoreUObject::FObjectProperty::CheckValidAddressNonNullable")
 {
 	bool bAllowRead = false;
 #if UE_WITH_OBJECT_HANDLE_TRACKING
-	auto CallbackHandle = UE::CoreUObject::AddObjectHandleReadCallback([&bAllowRead](UObject* ReadObject)
+	auto CallbackHandle = UE::CoreUObject::AddObjectHandleReadCallback([&bAllowRead](const UObject* const* ReadObject, int Count)
 		{
 			if (!bAllowRead)
 			{
@@ -163,7 +163,7 @@ static void TestSerializeItem(FName ObjectName)
 
 	int ReadCount = 0;
 #if UE_WITH_OBJECT_HANDLE_TRACKING
-	auto CallbackId = UE::CoreUObject::AddObjectHandleReadCallback([&ReadCount](UObject* ReadObject)
+	auto CallbackId = UE::CoreUObject::AddObjectHandleReadCallback([&ReadCount](const UObject* const* ReadObject, int Count)
 		{
 			++ReadCount;
 		});
diff --git a/Engine/Source/Runtime/CoreUObject/Tests/ObjectPtrTest.cpp b/Engine/Source/Runtime/CoreUObject/Tests/ObjectPtrTest.cpp
--- a/Engine/Source/Runtime/CoreUObject/Tests/ObjectPtrTest.cpp
+++ b/Engine/Source/Runtime/CoreUObject/Tests/ObjectPtrTest.cpp
@@ -500,6 +500,72 @@ TEST_CASE("CoreUObject::TObjectPtr::InnerRename")
 	TestPackage->RemoveFromRoot();
 }
 
+TEST_CASE("CoreUObject::TObjectPtr::Swap")
+{
+	const FName TestPackageName(TEXT("/Engine/TestPackage"));
+	UPackage* TestPackage = NewObject<UPackage>(nullptr, TestPackageName, RF_Transient);
+	TestPackage->AddToRoot();
+	UObject* Obj1 = NewObject<UObjectPtrTestClass>(TestPackage, TEXT("Obj1"));
+	UObject* Inner1 = NewObject<UObjectPtrTestClass>(Obj1, TEXT("Inner1"));
+	UObject* RawPtrA = Obj1;
+	UObject* RawPtrB = Inner1;
+
+	FObjectPtr ObjPtr1(MakeUnresolvedHandle(RawPtrA));
+	FObjectPtr InnerPtr1(MakeUnresolvedHandle(RawPtrB));
+
+	TObjectPtr<UObject> PtrA = *reinterpret_cast<TObjectPtr<UObject>*>(&ObjPtr1);
+	TObjectPtr<UObject> PtrB = *reinterpret_cast<TObjectPtr<UObject>*>(&InnerPtr1);
+	CHECK(!PtrA.IsResolved());
+	CHECK(!PtrB.IsResolved());
+
+	Swap(PtrA, PtrB);
+	
+	CHECK(!PtrA.IsResolved());
+	CHECK(!PtrB.IsResolved());
+	CHECK(PtrA != PtrB);
+	CHECK(PtrA != RawPtrA);
+	
+	Swap(PtrA, RawPtrA);
+	CHECK(PtrA.IsResolved());
+	CHECK(RawPtrA == Inner1);
+	CHECK(PtrA == Obj1);
+
+	Swap(PtrB, PtrA);
+	CHECK(!PtrA.IsResolved());
+	CHECK(PtrB.IsResolved());
+
+}
+
+TEST_CASE("CoreUObject::TObjectPtr::SwapArray")
+{
+	const FName TestPackageName(TEXT("/Engine/TestPackage"));
+	UPackage* TestPackage = NewObject<UPackage>(nullptr, TestPackageName, RF_Transient);
+	TestPackage->AddToRoot();
+	UObject* Obj1 = NewObject<UObjectPtrTestClass>(TestPackage, TEXT("Obj1"));
+	UObject* Inner1 = NewObject<UObjectPtrTestClass>(Obj1, TEXT("Inner1"));
+	UObject* RawPtrA = Obj1;
+	UObject* RawPtrB = Inner1;
+
+	FObjectPtr ObjPtr1(MakeUnresolvedHandle(RawPtrA));
+	FObjectPtr InnerPtr1(MakeUnresolvedHandle(RawPtrB));
+	TObjectPtr<UObject> PtrA = *reinterpret_cast<TObjectPtr<UObject>*>(&ObjPtr1);
+	TObjectPtr<UObject> PtrB = *reinterpret_cast<TObjectPtr<UObject>*>(&InnerPtr1);
+
+	
+	TArray<TObjectPtr<UObject>> ArrayPtr;
+	auto t = ArrayPtr.begin();
+	ArrayPtr.Add(PtrA);
+
+	TArray<UObject*> ArrayRaw;
+	ArrayRaw.Add(RawPtrB);
+
+	Swap(ArrayRaw, ArrayPtr);
+
+	CHECK(ArrayRaw[0] == RawPtrA);
+	CHECK(ArrayPtr[0] == RawPtrB);
+}
+
+
 TEST_CASE("CoreUObject::TObjectPtr::Move")
 {
 	UPackage* TestPackageA = NewObject<UPackage>(nullptr, TEXT("/Engine/PackageA"), RF_Transient);
@@ -699,8 +765,242 @@ TEST_CASE("CoreUObject::TObjectPtr::GetTypeHash")
 	TestPackageA->RemoveFromRoot();
 	TestPackageB->RemoveFromRoot();
 }
+
+
+#endif
+
+template<typename TObj>
+void TestArrayConversion()
+{
+	const FName TestPackageName(TEXT("/Engine/TestPackage"));
+	UPackage* TestPackage = NewObject<UPackage>(nullptr, TestPackageName, RF_Transient);
+	TestPackage->AddToRoot();
+	UObject* Obj1 = NewObject<UObjectPtrTestClass>(TestPackage, TEXT("Obj1"));
+
+#if UE_WITH_OBJECT_HANDLE_TRACKING
+	int ResolveCount = 0;
+	uint32 ObjCount = 0;
+	auto ResolveDelegate = [&ResolveCount, &ObjCount](const UObject* const* Obj, int Count)
+	{
+		++ResolveCount;
+		ObjCount = Count;
+	};
+	auto Handle = UE::CoreUObject::AddObjectHandleReadCallback(ResolveDelegate);
+	ON_SCOPE_EXIT
+	{
+		UE::CoreUObject::RemoveObjectHandleReadCallback(Handle);
+	};
+#endif
+	TArray<TObjectPtr<TObj>> PtrArray;
+	uint32  NumObjs = 1000000;
+	for (uint32 i = 0; i < NumObjs; ++i)
+	{
+		PtrArray.Add(Obj1);
+	}
+
+	{
+		TArray<TObj*> RawArray;
+		RawArray = PtrArray;
+
+#if UE_WITH_OBJECT_HANDLE_TRACKING
+		CHECK(ResolveCount == 1);
+		CHECK(ObjCount == NumObjs);
+		ResolveCount = 0;
+		ObjCount = 0;
+#endif
+		CHECK(RawArray.Num() == NumObjs);
+		for (uint32 i = 0; i < NumObjs; ++i)
+		{
+			CHECK(RawArray[i] == Obj1);
+		}
+	}
+	{
+		TArray<TObjectPtr<TObj>> EmptyArray;
+		TArray<TObj*> RawArray;
+		RawArray = EmptyArray;
+
+#if UE_WITH_OBJECT_HANDLE_TRACKING
+		CHECK(ResolveCount == 1);
+		CHECK(ObjCount == 0);
+		ResolveCount = 0;
+		ObjCount = 0;
 #endif
+		CHECK(RawArray.Num() == 0);
+	}
+	{
+		const TArray<TObjectPtr<TObj>>& ConstPtrArray = PtrArray;
+		TArray<TObj*> RawArray;
+		RawArray = ConstPtrArray;
+
+#if UE_WITH_OBJECT_HANDLE_TRACKING
+		CHECK(ResolveCount == 1);
+		CHECK(ObjCount == NumObjs);
+		ResolveCount = 0;
+		ObjCount = 0;
+#endif
+		CHECK(RawArray.Num() == NumObjs);
+		for (uint32 i = 0; i < NumObjs; ++i)
+		{
+			CHECK(RawArray[i] == Obj1);
+		}
+	}
+	{
+		TArray<TObjectPtr<TObj>> EmptyArray;
+		const TArray<TObjectPtr<TObj>>& ConstPtrArray = EmptyArray;
+		TArray<TObj*> RawArray;
+		RawArray = ConstPtrArray;
 
+#if UE_WITH_OBJECT_HANDLE_TRACKING
+		CHECK(ResolveCount == 1);
+		CHECK(ObjCount == 0);
+		ResolveCount = 0;
+		ObjCount = 0;
+#endif
+		CHECK(RawArray.Num() == 0);
+	}
+	{
+		const TArray<TObjectPtr<TObj>>& ConstPtrArray = PtrArray;
+		const TArray<TObj*>& RawArray = ConstPtrArray;
+
+#if UE_WITH_OBJECT_HANDLE_TRACKING
+		CHECK(ResolveCount == 1);
+		CHECK(ObjCount == NumObjs);
+		ResolveCount = 0;
+		ObjCount = 0;
+#endif
+		CHECK(RawArray.Num() == NumObjs);
+		for (uint32 i = 0; i < NumObjs; ++i)
+		{
+			CHECK(RawArray[i] == Obj1);
+		}
+	}
+	{
+		TArray<TObjectPtr<TObj>> EmptyArray;
+		const TArray<TObjectPtr<TObj>>& ConstPtrArray = EmptyArray;
+		const TArray<TObj*>& RawArray = ConstPtrArray;
+
+#if UE_WITH_OBJECT_HANDLE_TRACKING
+		CHECK(ResolveCount == 1);
+		CHECK(ObjCount == 0);
+		ResolveCount = 0;
+		ObjCount = 0;
+#endif
+		CHECK(RawArray.Num() == 0);
+	}
+
+	//ArrayView
+	{
+		TArrayView<TObj*> RawArray = PtrArray;
+#if UE_WITH_OBJECT_HANDLE_TRACKING
+		CHECK(ResolveCount == 1);
+		CHECK(ObjCount == NumObjs);
+		ResolveCount = 0;
+		ObjCount = 0;
+#endif
+		CHECK(RawArray.Num() == NumObjs);
+		for (uint32 i = 0; i < NumObjs; ++i)
+		{
+			CHECK(RawArray[i] == Obj1);
+		}
+	}
+	{
+		TArray<TObjectPtr<TObj>> EmptyArray;
+		TArrayView<TObj*> RawArray = EmptyArray;
+#if UE_WITH_OBJECT_HANDLE_TRACKING
+		CHECK(ResolveCount == 1);
+		CHECK(ObjCount == 0);
+		ResolveCount = 0;
+		ObjCount = 0;
+#endif
+		CHECK(RawArray.Num() == 0);
+	}
+	{
+		const TArray<TObjectPtr<TObj>>& ConstPtrArray = PtrArray;
+		const TArrayView<TObj* const> RawArray = ConstPtrArray;
+
+#if UE_WITH_OBJECT_HANDLE_TRACKING
+		CHECK(ResolveCount == 1);
+		CHECK(ObjCount == NumObjs);
+		ResolveCount = 0;
+		ObjCount = 0;
+#endif
+		CHECK(RawArray.Num() == NumObjs);
+		for (uint32 i = 0; i < NumObjs; ++i)
+		{
+			CHECK(RawArray[i] == Obj1);
+		}
+	}
+	{
+		TArray<TObjectPtr<TObj>> EmptyArray;
+		const TArray<TObjectPtr<TObj>>& ConstPtrArray = EmptyArray;
+		const TArrayView<TObj* const> RawArray = ConstPtrArray;
+
+#if UE_WITH_OBJECT_HANDLE_TRACKING
+		CHECK(ResolveCount == 1);
+		CHECK(ObjCount == 0);
+		ResolveCount = 0;
+		ObjCount = 0;
+#endif
+		CHECK(RawArray.Num() == 0);
+	}
+}
+
+TEST_CASE("CoreUObject::TObjectPtr::ArrayConversion")
+{
+	TestArrayConversion<UObject>();
+	TestArrayConversion<const UObject>();
+}
+
+TEST_CASE("CoreUObject::TObjectPtr::ArrayConversionReferenceForSet")
+{
+	const FName TestPackageName(TEXT("/Engine/TestPackage"));
+	UPackage* TestPackage = NewObject<UPackage>(nullptr, TestPackageName, RF_Transient);
+	TestPackage->AddToRoot();
+	UObject* Obj1 = NewObject<UObjectPtrTestClass>(TestPackage, TEXT("Obj1"));
+	UObject* Obj2 = NewObject<UObjectPtrTestClass>(TestPackage, TEXT("Obj2"));
+
+#if UE_WITH_OBJECT_HANDLE_TRACKING
+	int ResolveCount = 0;
+	uint32 ObjCount = 0;
+	auto ResolveDelegate = [&ResolveCount, &ObjCount](const UObject* const* Obj, int Count)
+	{
+		++ResolveCount;
+		ObjCount = Count;
+	};
+	auto Handle = UE::CoreUObject::AddObjectHandleReadCallback(ResolveDelegate);
+	ON_SCOPE_EXIT
+	{
+		UE::CoreUObject::RemoveObjectHandleReadCallback(Handle);
+	};
+#endif
+	{
+		TArray<TObjectPtr<UObject>> Array;
+		Array.Add(Obj1);
+		Array.Add(Obj2);
+
+		TSet<UObject*> RawSet(Array);
+
+#if UE_WITH_OBJECT_HANDLE_TRACKING
+		CHECK(ResolveCount == 1);
+		CHECK(ObjCount == 2);
+		ResolveCount = 0;
+		ObjCount = 0;
+#endif
+		CHECK(RawSet.Num() == 2);
+	}
+	{
+		TArray<TObjectPtr<UObject>> Array;
+		TSet<UObject*> RawSet(Array);
+
+#if UE_WITH_OBJECT_HANDLE_TRACKING
+		CHECK(ResolveCount == 1);
+		CHECK(ObjCount == 0);
+		ResolveCount = 0;
+		ObjCount = 0;
+#endif
+		CHECK(RawSet.Num() == 0);
+	}
+}
 // @TODO: OBJPTR: We should have a test that ensures that lazy loading of an object with an external package is handled correctly.
 //				  This should also include external packages in the outer chain of the target object.
 // IMPLEMENT_CUSTOM_SIMPLE_AUTOMATION_TEST(FObjectPtrTestExternalPackages, FObjectPtrTestBase, TEXT(TEST_NAME_ROOT ".ExternalPackages"), ObjectPtrTestFlags)
diff --git a/Engine/Source/Runtime/CoreUObject/Tests/ObjectRefTrackingTestBase.cpp b/Engine/Source/Runtime/CoreUObject/Tests/ObjectRefTrackingTestBase.cpp
--- a/Engine/Source/Runtime/CoreUObject/Tests/ObjectRefTrackingTestBase.cpp
+++ b/Engine/Source/Runtime/CoreUObject/Tests/ObjectRefTrackingTestBase.cpp
@@ -30,11 +30,11 @@ TEST_CASE("UE::CoreUObject::ObjectHandleTracking::Callbacks")
 	int32 CallbackCount7 = 0;
 	int32 CallbackCount8 = 0;
 
-	FObjectHandleReadFunc ReadCallback1 = [&](UObject* Obj)
+	FObjectHandleReadFunc ReadCallback1 = [&](const UObject* const* Obj, int Count)
 	{
 		++CallbackCount1;
 	};
-	FObjectHandleReadFunc ReadCallback2 = [&](UObject* Obj)
+	FObjectHandleReadFunc ReadCallback2 = [&](const UObject* const* Obj, int Count)
 	{
 		++CallbackCount2;
 	};
@@ -205,10 +205,10 @@ TEST_CASE("UE::CoreUObject::ObjectHandleTracking::Callbacks::Verify")
 
 	TObjectPtr<UObject> ObjPtr = Obj1;
 	int32 ReadCount = 0;
-	auto ReadCallbackHandle1 = AddObjectHandleReadCallback([&](UObject* ReadObject)
+	auto ReadCallbackHandle1 = AddObjectHandleReadCallback([&](const UObject* const* ReadObject, int Count)
 		{
 			++ReadCount;
-			CHECK(Obj1 == ReadObject);
+			CHECK(Obj1 == ReadObject[0]);
 		});
 	
 	CHECK(ObjPtr == Obj1);
diff --git a/Engine/Source/Runtime/CoreUObject/Tests/ObjectRefTrackingTestBase.h b/Engine/Source/Runtime/CoreUObject/Tests/ObjectRefTrackingTestBase.h
--- a/Engine/Source/Runtime/CoreUObject/Tests/ObjectRefTrackingTestBase.h
+++ b/Engine/Source/Runtime/CoreUObject/Tests/ObjectRefTrackingTestBase.h
@@ -35,7 +35,12 @@ public:
 		, OriginalNumFailedResolves(Test.GetNumFailedResolves())
 		, OriginalNumReads(Test.GetNumReads())
 		{
-			Test.ConditionalInstallCallbacks();
+			Test.InstallCallbacks();
+		}
+
+		~FSnapshotObjectRefMetrics()
+		{
+			Test.RemoveCallbacks();
 		}
 
 		void TestNumResolves(const TCHAR* What, uint32 ExpectedDelta)
@@ -85,33 +90,26 @@ private:
 			NumFailedResolves++;
 		}
 	}
-	static void OnRefRead(UObject* Obj)
+	static void OnRefRead(const UObject* const* Obj, uint32 Count)
 	{
 		NumReads++;
 	}
 #endif
 	
-	static void ConditionalInstallCallbacks()
+	static void InstallCallbacks()
 	{
-		static bool bCallbacksInstalled = false;
-		static FCriticalSection CallbackInstallationLock;
-
-		if (bCallbacksInstalled)
-		{
-			return;
-		}
+#if UE_WITH_OBJECT_HANDLE_TRACKING
+		ResolvedCallbackHandle = UE::CoreUObject::AddObjectHandleReferenceResolvedCallback(OnRefResolved);
+		HandleReadCallbackHandle = UE::CoreUObject::AddObjectHandleReadCallback(OnRefRead);
+#endif
+	}
 
-		FScopeLock ScopeLock(&CallbackInstallationLock);
-		if (!bCallbacksInstalled)
-		{
+	static void RemoveCallbacks()
+	{
 #if UE_WITH_OBJECT_HANDLE_TRACKING
-			ResolvedCallbackHandle = UE::CoreUObject::AddObjectHandleReferenceResolvedCallback(OnRefResolved);
-			HandleReadCallbackHandle = UE::CoreUObject::AddObjectHandleReadCallback(OnRefRead);
-			// TODO We should unhook these handles somewhere, but i don't want to refactor the test, it's not as if they were
-			// being unhooked before.  So...
+		UE::CoreUObject::RemoveObjectHandleReferenceResolvedCallback(ResolvedCallbackHandle);
+		UE::CoreUObject::RemoveObjectHandleReadCallback(HandleReadCallbackHandle);
 #endif
-			bCallbacksInstalled = true;
-		}
 	}
 
 #if UE_WITH_OBJECT_HANDLE_TRACKING
diff --git a/Engine/Source/Runtime/CoreUObject/Tests/PropertyClassTest.cpp b/Engine/Source/Runtime/CoreUObject/Tests/PropertyClassTest.cpp
--- a/Engine/Source/Runtime/CoreUObject/Tests/PropertyClassTest.cpp
+++ b/Engine/Source/Runtime/CoreUObject/Tests/PropertyClassTest.cpp
@@ -17,7 +17,7 @@ TEST_CASE("UE::CoreUObject::FClassProperty::Identical")
 {
 	bool bAllowRead = false;
 #if UE_WITH_OBJECT_HANDLE_TRACKING
-	auto CallbackHandle = UE::CoreUObject::AddObjectHandleReadCallback([&bAllowRead](UObject* ReadObject)
+	auto CallbackHandle = UE::CoreUObject::AddObjectHandleReadCallback([&bAllowRead](const UObject* const* ReadObject, int Count)
 		{
 			if (!bAllowRead)
 				FAIL("Unexpected read during CheckValidObject");
