Static assert when casting a non-UObject to an interface.

diff --git a/Engine/Source/Runtime/CoreUObject/Public/Templates/Casts.h b/Engine/Source/Runtime/CoreUObject/Public/Templates/Casts.h
--- a/Engine/Source/Runtime/CoreUObject/Public/Templates/Casts.h
+++ b/Engine/Source/Runtime/CoreUObject/Public/Templates/Casts.h
@@ -110,17 +110,20 @@ FORCEINLINE To* Cast(From* Src)
 				return (To*)Src;
 			}
 		}
-		else if constexpr (TIsIInterface<To>::Value)
-		{
-			return (To*)((UObject*)Src)->GetInterfaceAddress(To::UClassType::StaticClass());
-		}
 		else
 		{
 			static_assert(std::is_base_of_v<UObject, From>, "Attempting to use Cast<> on a type that is not a UObject or an Interface");
-			
-			if (((const UObject*)Src)->IsA<To>())
+
+			if constexpr (TIsIInterface<To>::Value)
 			{
-				return (To*)Src;
+				return (To*)((UObject*)Src)->GetInterfaceAddress(To::UClassType::StaticClass());
+			}
+			else
+			{
+				if (((const UObject*)Src)->IsA<To>())
+				{
+					return (To*)Src;
+				}
 			}
 		}
 	}
