Fix potential for data loss when loading hierarchical Blueprint class assets w/ circular dependencies asynchronously in the editor.

diff --git a/Engine/Source/Runtime/CoreUObject/Private/Blueprint/BlueprintSupport.cpp b/Engine/Source/Runtime/CoreUObject/Private/Blueprint/BlueprintSupport.cpp
--- a/Engine/Source/Runtime/CoreUObject/Private/Blueprint/BlueprintSupport.cpp
+++ b/Engine/Source/Runtime/CoreUObject/Private/Blueprint/BlueprintSupport.cpp
@@ -149,7 +149,7 @@ void FBlueprintSupport::SetClassReparentingFPtr(FClassReparentingFPtr Ptr)
 	ClassReparentingFPtr = Ptr;
 }
 
-bool FBlueprintSupport::IsDeferredDependencyPlaceholder(UObject* LoadedObj)
+bool FBlueprintSupport::IsDeferredDependencyPlaceholder(const UObject* LoadedObj)
 {
 	return LoadedObj && ( LoadedObj->IsA<ULinkerPlaceholderClass>() ||
 		LoadedObj->IsA<ULinkerPlaceholderFunction>() ||
@@ -328,11 +328,11 @@ bool FBlueprintSupport::ShouldSuppressWarning(FName WarningIdentifier)
 	return BlueprintWarningsToSuppress.Find(WarningIdentifier) != nullptr;
 }
 
-bool FBlueprintSupport::IsClassPlaceholder(UClass* Class)
+bool FBlueprintSupport::IsClassPlaceholder(const UClass* Class)
 {
 	while (Class)
 	{
-		if (Cast<ULinkerPlaceholderClass>(Class))
+		if (Cast<const ULinkerPlaceholderClass>(Class))
 		{
 			return true;
 		}
@@ -1244,8 +1244,40 @@ bool FLinkerLoad::DeferExportCreation(const int32 Index, UObject* Outer)
 		return false;
 	}
 
+	const bool bIsCDOExport = (Export.ObjectFlags & RF_ClassDefaultObject) != 0;
+	if (bIsCDOExport)
+	{
+		// Check for any load dependencies that may have been deferred.
+		bool bHasDeferredDependencies = false;
+		TArray<UObject*> CDOPreloadDependencies;
+		LoadClass->GetDefaultObjectPreloadDependencies(CDOPreloadDependencies);
+		for (const UObject* PreloadDependency : CDOPreloadDependencies)
+		{
+			if (FBlueprintSupport::IsDeferredDependencyPlaceholder(PreloadDependency))
+			{
+				bHasDeferredDependencies = true;
+			}
+		}
+
+		// Defer the CDO export only if we're preloading its class and it has a deferred dependency. For
+		// example, we may need to resolve a non-native subobject type override before we can construct
+		// the actual CDO and execute its native ctor/initializer.
+		if (((LoadFlags & LOAD_DeferDependencyLoads) != 0) && bHasDeferredDependencies)
+		{
+			// This will cause IsBlueprintFinalizationPending() to fail below (which is what we want).
+			// We'll then fall through and create a placeholder object for the CDO in order to defer its
+			// actual construction (along with serialization) until after we've resolved its dependencies.
+			DEFERRED_DEPENDENCY_CHECK(DeferredCDOIndex == INDEX_NONE);
+			DeferredCDOIndex = Index;
+		}
+		else
+		{
+			return false;
+		}
+	}
+
 	ULinkerPlaceholderClass* AsPlaceholderClass = Cast<ULinkerPlaceholderClass>(LoadClass);
-	bool const bIsPlaceholderClass = (AsPlaceholderClass != nullptr);
+	const bool bIsPlaceholderClass = (AsPlaceholderClass != nullptr);
 
 	FLinkerLoad* ClassLinker = LoadClass->GetLinker();
 	if (!bIsPlaceholderClass
@@ -1255,7 +1287,7 @@ bool FLinkerLoad::DeferExportCreation(const int32 Index, UObject* Outer)
 		return false;
 	}
 
-	bool const bIsLoadingExportClass = (LoadFlags & LOAD_DeferDependencyLoads) ||
+	const bool bIsLoadingExportClass = (LoadFlags & LOAD_DeferDependencyLoads) ||
 		IsBlueprintFinalizationPending();
 	// if we're not in the process of "loading/finalizing" this package's 
 	// Blueprint class, then we're either running this before the linker has got 
@@ -1289,6 +1321,12 @@ bool FLinkerLoad::DeferExportCreation(const int32 Index, UObject* Outer)
 	FResolvingExportTracker::Get().AddLinkerPlaceholderObject(LoadClass, Placeholder);
 
 	Export.Object = Placeholder;
+
+	if (bIsCDOExport)
+	{
+		DEFERRED_DEPENDENCY_CHECK(LoadClass->ClassDefaultObject == nullptr);
+		LoadClass->ClassDefaultObject = Placeholder;
+	}
 #endif // USE_CIRCULAR_DEPENDENCY_LOAD_DEFERRING
 
 	return true;
@@ -1955,12 +1993,21 @@ void FLinkerLoad::ResolveDeferredExports(UClass* LoadClass)
 
 	// Handle deferred construction of the CDO and patch it into the export table. Any deferred ctor
 	// initializer dependencies (e.g. subobject class overrides) should now be resolved at this point.
-	if (DeferredCDOIndex != INDEX_NONE && !ExportMap[DeferredCDOIndex].Object)
+	if (DeferredCDOIndex != INDEX_NONE)
 	{
-		// Note: We could just call GetDefaultObject() here, but then we'd also need to set object
-		// flags on it to ensure that it gets serialized later. For consistency/safety, this routes
-		// through CreateExport() so that we don't have to worry about keeping object flags in sync.
-		CreateExport(DeferredCDOIndex);
+		if (ULinkerPlaceholderExportObject* PlaceholderExport = Cast<ULinkerPlaceholderExportObject>(ExportMap[DeferredCDOIndex].Object))
+		{
+			LoadClass->ClassDefaultObject = nullptr;
+			PlaceholderExport->SetLinker(nullptr, INDEX_NONE);
+			ExportMap[DeferredCDOIndex].ResetObject();
+			UObject* ExportObj = CreateExport(DeferredCDOIndex);
+
+			PlaceholderExport->ResolveAllPlaceholderReferences(ExportObj);
+			ResolvedDeferredSubobjects(PlaceholderExport);
+			PlaceholderExport->MarkAsGarbage();
+
+			DEFERRED_DEPENDENCY_CHECK(LoadClass->ClassDefaultObject == ExportObj);
+		}
 	}
 
 	UObject* BlueprintCDO = DeferredCDOIndex != INDEX_NONE ? ExportMap[DeferredCDOIndex].Object : LoadClass->ClassDefaultObject;
@@ -2139,7 +2186,7 @@ void FLinkerLoad::ResolveDeferredExports(UClass* LoadClass)
 			{
 				FObjectExport& Export = ExportMap[ExportIndex];
 				ULinkerPlaceholderExportObject* PlaceholderExport = Cast<ULinkerPlaceholderExportObject>(Export.Object);
-				if (ensure(PlaceholderExport))
+				if (ensure(PlaceholderExport) && !PlaceholderExport->IsMarkedResolved())
 				{
 					// replace the placeholder with the proper object instance
 					PlaceholderExport->SetLinker(nullptr, INDEX_NONE);
diff --git a/Engine/Source/Runtime/CoreUObject/Private/UObject/LinkerLoad.cpp b/Engine/Source/Runtime/CoreUObject/Private/UObject/LinkerLoad.cpp
--- a/Engine/Source/Runtime/CoreUObject/Private/UObject/LinkerLoad.cpp
+++ b/Engine/Source/Runtime/CoreUObject/Private/UObject/LinkerLoad.cpp
@@ -4882,7 +4882,7 @@ UObject* FLinkerLoad::CreateExport( int32 Index )
 			{
 				Preload(LoadClass);
 			}
-			else if ((Export.Object == nullptr) && !(Export.ObjectFlags & RF_ClassDefaultObject))
+			else if (Export.Object == nullptr)
 			{
 				bool const bExportWasDeferred = DeferExportCreation(Index, ThisParent);
 				if (bExportWasDeferred)
@@ -4891,7 +4891,7 @@ UObject* FLinkerLoad::CreateExport( int32 Index )
 					check(Export.Object != nullptr);
 #endif // USE_DEFERRED_DEPENDENCY_CHECK_VERIFICATION_TESTS
 					return Export.Object;
-				}				
+				}
 			}
 			else if (Cast<ULinkerPlaceholderExportObject>(Export.Object))
 			{
@@ -5164,24 +5164,6 @@ UObject* FLinkerLoad::CreateExport( int32 Index )
 			}
 		}
 
-		const bool bIsBlueprintCDO = ((Export.ObjectFlags & RF_ClassDefaultObject) != 0) && LoadClass->HasAnyClassFlags(CLASS_CompiledFromBlueprint) &&
-			LoadClass->GetClass()->HasAnyClassFlags(CLASS_NeedsDeferredDependencyLoading);
-
-#if USE_CIRCULAR_DEPENDENCY_LOAD_DEFERRING
-		const bool bDeferCDOSerialization = bIsBlueprintCDO && ((LoadFlags & LOAD_DeferDependencyLoads) != 0);
-		if (bDeferCDOSerialization)
-		{
-			// if LOAD_DeferDependencyLoads is set, then we're already
-			// serializing the blueprint's class somewhere up the chain... 
-			// we don't want the class regenerated while it in the middle of
-			// serializing. we also don't want to construct the CDO yet,
-			// as that may depend on other deferred imports which may not
-			// be fully resolved (e.g. default subobject class overrides).
-			DeferredCDOIndex = Index;
-			return Export.Object;
-		}
-#endif // USE_CIRCULAR_DEPENDENCY_LOAD_DEFERRING
-
 		LoadClass->GetDefaultObject();
 
 		FStaticConstructObjectParameters Params(LoadClass);
@@ -5233,8 +5215,21 @@ UObject* FLinkerLoad::CreateExport( int32 Index )
 		
 		if( Export.Object )
 		{
+			const bool bIsBlueprintCDO = ((Export.ObjectFlags & RF_ClassDefaultObject) != 0) && LoadClass->HasAnyClassFlags(CLASS_CompiledFromBlueprint) &&
+				LoadClass->GetClass()->HasAnyClassFlags(CLASS_NeedsDeferredDependencyLoading);
+
 #if USE_CIRCULAR_DEPENDENCY_LOAD_DEFERRING
-			if (bIsBlueprintCDO && IsBlueprintFinalizationPending())
+			const bool bDeferCDOSerialization = bIsBlueprintCDO && ((LoadFlags & LOAD_DeferDependencyLoads) != 0);
+			if (bDeferCDOSerialization)
+			{
+				// if LOAD_DeferDependencyLoads is set, then we're already
+				// serializing the blueprint's class somewhere up the chain... 
+				// we don't want the class regenerated while it in the middle of
+				// serializing
+				DeferredCDOIndex = Index;
+				return Export.Object;
+			}
+			else if (bIsBlueprintCDO && IsBlueprintFinalizationPending())
 			{
 				// Class regeneration is deferred until Blueprint finalization, so just return the CDO.
 				return Export.Object;
diff --git a/Engine/Source/Runtime/CoreUObject/Public/Blueprint/BlueprintSupport.h b/Engine/Source/Runtime/CoreUObject/Public/Blueprint/BlueprintSupport.h
--- a/Engine/Source/Runtime/CoreUObject/Public/Blueprint/BlueprintSupport.h
+++ b/Engine/Source/Runtime/CoreUObject/Public/Blueprint/BlueprintSupport.h
@@ -113,7 +113,7 @@ struct FBlueprintSupport
 	COREUOBJECT_API static void SetClassReparentingFPtr(FClassReparentingFPtr Ptr);
 
 	/** Tells if the specified object is one of the many flavors of FLinkerPlaceholderBase that we have. */
-	COREUOBJECT_API static bool IsDeferredDependencyPlaceholder(UObject* LoadedObj);
+	COREUOBJECT_API static bool IsDeferredDependencyPlaceholder(const UObject* LoadedObj);
 
 	/** Registers any object properties in this struct with the deferred dependency system */
 	COREUOBJECT_API static void RegisterDeferredDependenciesInStruct(const UStruct* Struct, void* StructData);
@@ -130,7 +130,7 @@ struct FBlueprintSupport
 	COREUOBJECT_API static bool ShouldTreatWarningAsError(FName WarningIdentifier);
 	COREUOBJECT_API static bool ShouldSuppressWarning(FName WarningIdentifier);
 
-	COREUOBJECT_API static bool IsClassPlaceholder(UClass* Class);
+	COREUOBJECT_API static bool IsClassPlaceholder(const UClass* Class);
 
 #if WITH_EDITOR
 	/** Function that walks the object graph, ensuring that there are no references to TRASH or REINST classes: */
diff --git a/Engine/Source/Runtime/Engine/Private/BlueprintGeneratedClass.cpp b/Engine/Source/Runtime/Engine/Private/BlueprintGeneratedClass.cpp
--- a/Engine/Source/Runtime/Engine/Private/BlueprintGeneratedClass.cpp
+++ b/Engine/Source/Runtime/Engine/Private/BlueprintGeneratedClass.cpp
@@ -409,6 +409,8 @@ void UBlueprintGeneratedClass::SerializeDefaultObject(UObject* Object, FStructur
 
 	if (UnderlyingArchive.IsLoading() && !UnderlyingArchive.IsObjectReferenceCollector() && Object == ClassDefaultObject)
 	{
+		CreatePersistentUberGraphFrame(Object, true);
+
 		// On load, build the custom property list used in post-construct initialization logic. Note that in the editor, this will be refreshed during compile-on-load.
 		// @TODO - Potentially make this serializable (or cooked data) to eliminate the slight load time cost we'll incur below to generate this list in a cooked build. For now, it's not serialized since the raw FProperty references cannot be saved out.
 		UpdateCustomPropertyListForPostConstruction();
@@ -969,15 +971,6 @@ UInheritableComponentHandler* UBlueprintGeneratedClass::GetInheritableComponentH
 	{
 		return nullptr;
 	}
-	
-	if (InheritableComponentHandler)
-	{
-		if (!GEventDrivenLoaderEnabled || !EVENT_DRIVEN_ASYNC_LOAD_ACTIVE_AT_RUNTIME)
-		{
-			// This preload will not succeed in EDL
-			InheritableComponentHandler->PreloadAll();
-		}	
-	}
 
 	if (!InheritableComponentHandler && bCreateIfNecessary)
 	{
@@ -2130,8 +2123,6 @@ void UBlueprintGeneratedClass::Serialize(FArchive& Ar)
 
 	if (Ar.IsLoading() && 0 == (Ar.GetPortFlags() & PPF_Duplicate))
 	{
-		CreatePersistentUberGraphFrame(ClassDefaultObject, true);
-
 		UPackage* Package = GetOutermost();
 		if (Package && Package->HasAnyPackageFlags(PKG_ForDiffing))
 		{
