Add a CVar to control getting property names and callstacks from FGCObject as it's very expensive.
If we fail/choose not to get property names, get the name of the FGCObject itself.

#localization none
#tests Testing running obj refs command during gameplay in a client. Time reduced from over a minute to ~5s for finding references to a StaticMeshComponent.

diff --git a/Engine/Source/Runtime/CoreUObject/Private/UObject/ReferenceChainSearch.cpp b/Engine/Source/Runtime/CoreUObject/Private/UObject/ReferenceChainSearch.cpp
--- a/Engine/Source/Runtime/CoreUObject/Private/UObject/ReferenceChainSearch.cpp
+++ b/Engine/Source/Runtime/CoreUObject/Private/UObject/ReferenceChainSearch.cpp
@@ -10,6 +10,22 @@
 
 DEFINE_LOG_CATEGORY_STATIC(LogReferenceChain, Log, All);
 
+static bool GFGCObjectPropertyNames = true;
+static FAutoConsoleVariableRef CVarFGCObjectPropertyNames(
+	TEXT("gc.ReferenceChainSearch.FGCObjectPropertyNames"),
+	GFGCObjectPropertyNames,
+	TEXT("Whether or not to extract property names from FGCObject implementors during reference chain search - can be very expensive."),
+	ECVF_Default
+);
+
+static bool GAddReferencedObjectsStackTrace = true;
+static FAutoConsoleVariableRef CVarAddReferencedObjectsStackTrace(
+	TEXT("gc.ReferenceChainSearch.AddReferencedObjectsStackTrace"),
+	GAddReferencedObjectsStackTrace,
+	TEXT("Whether to gather a stack trace for calls to AddReferencedObjects."),
+	ECVF_Default
+);
+
 // Returns true if the object can't be collected by GC
 static FORCEINLINE bool IsNonGCObject(FGCObjectInfo* Object, EReferenceChainSearchMode SearchMode)
 {
@@ -624,15 +640,17 @@ public:
 				else
 				{
 					RefInfo.Type = FReferenceChainSearch::EReferenceType::AddReferencedObjects;
-
-					RefInfo.StackFrames.AddUninitialized(FReferenceChainSearch::FObjectReferenceInfo::MaxStackFrames);
-					int32 NumStackFrames = FPlatformStackWalk::CaptureStackBackTrace(RefInfo.StackFrames.GetData(), RefInfo.StackFrames.Num());
-					RefInfo.StackFrames.SetNum(NumStackFrames);
+					if (GAddReferencedObjectsStackTrace)
+					{
+						RefInfo.StackFrames.AddUninitialized(FReferenceChainSearch::FObjectReferenceInfo::MaxStackFrames);
+						int32 NumStackFrames = FPlatformStackWalk::CaptureStackBackTrace(RefInfo.StackFrames.GetData(), RefInfo.StackFrames.Num());
+						RefInfo.StackFrames.SetNum(NumStackFrames);
+					}
 
 					if (FGCObject::GGCObjectReferencer && (!ReferencingObject || ReferencingObject == FGCObject::GGCObjectReferencer))
 					{
 						FString RefName;
-						if (FGCObject::GGCObjectReferencer->GetReferencerName(Object, RefName, true))
+						if (GFGCObjectPropertyNames && FGCObject::GGCObjectReferencer->GetReferencerName(Object, RefName, true))
 						{
 							// In case FGCObjects misbehave and give us long strings, truncate so we can make them an FName.
 							if (RefName.Len() >= NAME_SIZE)
@@ -642,9 +660,9 @@ public:
 
 							RefInfo.ReferencerName = FName(*RefName);
 						}
-						else if (ReferencingObject)
+						else if (FGCObject* ReferencingFGCObject = FGCObject::GGCObjectReferencer->GetCurrentlySerializingObject())
 						{
-							RefInfo.ReferencerName = *ReferencingObject->GetFullName();
+							RefInfo.ReferencerName = *ReferencingFGCObject->GetReferencerName();
 						}
 					}
 					else if (ReferencingObject)
