Drop weak token support for CollectReferences/TFastReferenceCollector

diff --git a/Engine/Source/Runtime/CoreUObject/Private/UObject/GarbageCollection.cpp b/Engine/Source/Runtime/CoreUObject/Private/UObject/GarbageCollection.cpp
--- a/Engine/Source/Runtime/CoreUObject/Private/UObject/GarbageCollection.cpp
+++ b/Engine/Source/Runtime/CoreUObject/Private/UObject/GarbageCollection.cpp
@@ -4802,39 +4802,13 @@ void FObjectProperty::EmitReferenceInfo(UE::GC::FTokenStreamBuilder& TokenStream
 	TokenStream.EmitObjectReference(BaseOffset + GetOffset_ForGC(), *StackSizeHelper.GetPropertyPath(), GCRT_Object);
 }
 
-void FWeakObjectProperty::EmitReferenceInfo(UE::GC::FTokenStreamBuilder& TokenStream, int32 BaseOffset, TArray<const FStructProperty*>& EncounteredStructProps, FGCStackSizeHelper& StackSizeHelper)
-{
-	UE::GC::FFixedArrayTokenHelper FixedArrayHelper(TokenStream, BaseOffset + GetOffset_ForGC(), ArrayDim, sizeof(FWeakObjectPtr), *this, StackSizeHelper);
-	TokenStream.EmitObjectReference(BaseOffset + GetOffset_ForGC(), *StackSizeHelper.GetPropertyPath(), GCRT_WeakObject);
-}
-void FLazyObjectProperty::EmitReferenceInfo(UE::GC::FTokenStreamBuilder& TokenStream, int32 BaseOffset, TArray<const FStructProperty*>& EncounteredStructProps, FGCStackSizeHelper& StackSizeHelper)
-{
-	UE::GC::FFixedArrayTokenHelper FixedArrayHelper(TokenStream, BaseOffset + GetOffset_ForGC(), ArrayDim, sizeof(FLazyObjectPtr), *this, StackSizeHelper);
-	TokenStream.EmitObjectReference(BaseOffset + GetOffset_ForGC(), *StackSizeHelper.GetPropertyPath(), GCRT_LazyObject);
-}
-void FSoftObjectProperty::EmitReferenceInfo(UE::GC::FTokenStreamBuilder& TokenStream, int32 BaseOffset, TArray<const FStructProperty*>& EncounteredStructProps, FGCStackSizeHelper& StackSizeHelper)
-{
-	UE::GC::FFixedArrayTokenHelper FixedArrayHelper(TokenStream, BaseOffset + GetOffset_ForGC(), ArrayDim, sizeof(FSoftObjectPtr), *this, StackSizeHelper);
-	TokenStream.EmitObjectReference(BaseOffset + GetOffset_ForGC(), *StackSizeHelper.GetPropertyPath(), GCRT_SoftObject);
-}
-void FDelegateProperty::EmitReferenceInfo(UE::GC::FTokenStreamBuilder& TokenStream, int32 BaseOffset, TArray<const FStructProperty*>& EncounteredStructProps, FGCStackSizeHelper& StackSizeHelper)
-{
-	UE::GC::FFixedArrayTokenHelper FixedArrayHelper(TokenStream, BaseOffset + GetOffset_ForGC(), ArrayDim, this->ElementSize, *this, StackSizeHelper);
-	TokenStream.EmitObjectReference(BaseOffset + GetOffset_ForGC(), *StackSizeHelper.GetPropertyPath(), GCRT_Delegate);
-}
-void FMulticastDelegateProperty::EmitReferenceInfo(UE::GC::FTokenStreamBuilder& TokenStream, int32 BaseOffset, TArray<const FStructProperty*>& EncounteredStructProps, FGCStackSizeHelper& StackSizeHelper)
-{
-	UE::GC::FFixedArrayTokenHelper FixedArrayHelper(TokenStream, BaseOffset + GetOffset_ForGC(), ArrayDim, this->ElementSize, *this, StackSizeHelper);
-	TokenStream.EmitObjectReference(BaseOffset + GetOffset_ForGC(), *StackSizeHelper.GetPropertyPath(), GCRT_MulticastDelegate);
-}
-
 /**
  * Emits tokens used by realtime garbage collection code to passed in OwnerClass' ReferenceTokenStream. The offset emitted is relative
  * to the passed in BaseOffset which is used by e.g. arrays of structs.
  */
 void FArrayProperty::EmitReferenceInfo(UE::GC::FTokenStreamBuilder& TokenStream, int32 BaseOffset, TArray<const FStructProperty*>& EncounteredStructProps, FGCStackSizeHelper& StackSizeHelper)
 {
-	if (Inner->ContainsObjectReference(EncounteredStructProps, EPropertyObjectReferenceType::Strong | EPropertyObjectReferenceType::Weak))
+	if (Inner->ContainsObjectReference(EncounteredStructProps, EPropertyObjectReferenceType::Strong))
 	{
 		bool bUsesFreezableAllocator = EnumHasAnyFlags(ArrayFlags, EArrayPropertyFlags::UsesMemoryImageAllocator);
 
@@ -4864,26 +4838,6 @@ void FArrayProperty::EmitReferenceInfo(UE::GC::FTokenStreamBuilder& TokenStream,
 		{
 			TokenStream.EmitObjectReference(BaseOffset + GetOffset_ForGC(), *StackSizeHelper.GetPropertyPath(), GCRT_ArrayAddFieldPathReferencedObject);
 		}
-		else if (Inner->IsA(FWeakObjectProperty::StaticClass()))
-		{
-			TokenStream.EmitObjectReference(BaseOffset + GetOffset_ForGC(), *StackSizeHelper.GetPropertyPath(), GCRT_ArrayWeakObject);
-		}
-		else if (Inner->IsA(FLazyObjectProperty::StaticClass()))
-		{
-			TokenStream.EmitObjectReference(BaseOffset + GetOffset_ForGC(), *StackSizeHelper.GetPropertyPath(), GCRT_ArrayLazyObject);
-		}
-		else if (Inner->IsA(FSoftObjectProperty::StaticClass()))
-		{
-			TokenStream.EmitObjectReference(BaseOffset + GetOffset_ForGC(), *StackSizeHelper.GetPropertyPath(), GCRT_ArraySoftObject);
-		}
-		else if (Inner->IsA(FDelegateProperty::StaticClass()))
-		{
-			TokenStream.EmitObjectReference(BaseOffset + GetOffset_ForGC(), *StackSizeHelper.GetPropertyPath(), GCRT_ArrayDelegate);
-		}
-		else if (Inner->IsA(FMulticastDelegateProperty::StaticClass()))
-		{
-			TokenStream.EmitObjectReference(BaseOffset + GetOffset_ForGC(), *StackSizeHelper.GetPropertyPath(), GCRT_ArrayMulticastDelegate);
-		}
 		else
 		{
 			// In the general case, emit a GCRT_ArrayStruct that can handle any inner property type.
@@ -4908,20 +4862,20 @@ void FArrayProperty::EmitReferenceInfo(UE::GC::FTokenStreamBuilder& TokenStream,
  */
 void FMapProperty::EmitReferenceInfo(UE::GC::FTokenStreamBuilder& TokenStream, int32 BaseOffset, TArray<const FStructProperty*>& EncounteredStructProps, FGCStackSizeHelper& StackSizeHelper)
 {
-	if (ContainsObjectReference(EncounteredStructProps, EPropertyObjectReferenceType::Strong | EPropertyObjectReferenceType::Weak))
+	if (ContainsObjectReference(EncounteredStructProps, EPropertyObjectReferenceType::Strong))
 	{
 		// TMap reference tokens are processed by GC in a similar way to an array of structs
 		TokenStream.EmitObjectReference(BaseOffset + GetOffset_ForGC(), *StackSizeHelper.GetPropertyPath(), GCRT_AddTMapReferencedObjects);
 		TokenStream.EmitPointer((const void*)this);
 		const uint32 SkipIndexIndex = TokenStream.EmitSkipIndexPlaceholder();
 
-		if (KeyProp->ContainsObjectReference(EncounteredStructProps, EPropertyObjectReferenceType::Strong | EPropertyObjectReferenceType::Weak))
+		if (KeyProp->ContainsObjectReference(EncounteredStructProps, EPropertyObjectReferenceType::Strong))
 		{
 			// GCRT_AddTMapReferencedObjects pushes a new stack frame in TFastReferenceCollector
 			FGCStackSizeHelperScope StackSizeScope(StackSizeHelper, KeyProp);
 			KeyProp->EmitReferenceInfo(TokenStream, 0, EncounteredStructProps, StackSizeHelper);
 		}
-		if (ValueProp->ContainsObjectReference(EncounteredStructProps, EPropertyObjectReferenceType::Strong | EPropertyObjectReferenceType::Weak))
+		if (ValueProp->ContainsObjectReference(EncounteredStructProps, EPropertyObjectReferenceType::Strong))
 		{
 			// GCRT_AddTMapReferencedObjects pushes a new stack frame in TFastReferenceCollector
 			FGCStackSizeHelperScope StackSizeScope(StackSizeHelper, ValueProp);
@@ -4939,7 +4893,7 @@ void FMapProperty::EmitReferenceInfo(UE::GC::FTokenStreamBuilder& TokenStream, i
 */
 void FSetProperty::EmitReferenceInfo(UE::GC::FTokenStreamBuilder& TokenStream, int32 BaseOffset, TArray<const FStructProperty*>& EncounteredStructProps, FGCStackSizeHelper& StackSizeHelper)
 {
-	if (ContainsObjectReference(EncounteredStructProps, EPropertyObjectReferenceType::Strong | EPropertyObjectReferenceType::Weak))
+	if (ContainsObjectReference(EncounteredStructProps, EPropertyObjectReferenceType::Strong))
 	{
 		// TSet reference tokens are processed by GC in a similar way to an array of structs
 		TokenStream.EmitObjectReference(BaseOffset + GetOffset_ForGC(), *StackSizeHelper.GetPropertyPath(), GCRT_AddTSetReferencedObjects);
@@ -4983,7 +4937,7 @@ void FStructProperty::EmitReferenceInfo(UE::GC::FTokenStreamBuilder& TokenStream
 		EncounteredStructProps.Add(this);
 		for (FProperty* Property = Struct->PropertyLink; Property && !bHasPropertiesWithObjectReferences; Property = Property->PropertyLinkNext)
 		{
-			bHasPropertiesWithObjectReferences = Property->ContainsObjectReference(EncounteredStructProps, EPropertyObjectReferenceType::Strong | EPropertyObjectReferenceType::Weak);
+			bHasPropertiesWithObjectReferences = Property->ContainsObjectReference(EncounteredStructProps, EPropertyObjectReferenceType::Strong);
 		}
 		EncounteredStructProps.RemoveSingleSwap(this, false /* bAllowShrinking */);
 	}
@@ -5146,90 +5100,6 @@ struct FNamedReferenceInfo
 	FName DebugName;
 };
 
-// Helps iterate over a range of FTokenStreamBuilder tokens and debug names
-class FTokenStreamBuilder::FConstIterator
-{
-public:
-	FConstIterator(const FTokenStreamBuilder& InBuilder)
-	: Builder(InBuilder)
-	, Idx(0)
-	, Num(InBuilder.Tokens.Num())
-	{}
-
-	explicit operator bool() const { return Idx < Num; }
-	
-	uint32 Pos() const { return Idx; }
-	uint32 EndPos() const { return Num; }
-
-	uint32 PeekReturnCount()
-	{
-		return FGCReferenceInfo(GetToken()).ReturnCount;
-	}
-
-	FNamedReferenceInfo ConsumeReferenceInfo()
-	{
-		FNamedReferenceInfo Out = { FGCReferenceInfo(GetToken()), GetDebugName() };
-		++Idx;
-		return Out;
-	}
-
-	// @return Index to skip to
-	uint32 ConsumeSkipInfo()
-	{
-#if ENABLE_GC_TOKEN_DEBUG_INFO
-		static const FName TokenName("SkipIndexPlaceholder");
-		check(GetDebugName() == TokenName);
-#endif
-		FGCSkipInfo SkipInfo(ConsumeToken());
-		return Idx - 1 + SkipInfo.SkipIndex;
-	}
-	
-	uint32 ConsumeStride()
-	{
-#if ENABLE_GC_TOKEN_DEBUG_INFO
-		static const FName TokenName("StrideToken");
-		check(GetDebugName() == TokenName);
-#endif
-		return ConsumeToken();
-	}
-
-	uint32 ConsumeCount()
-	{
-#if ENABLE_GC_TOKEN_DEBUG_INFO
-		static const FName TokenName("CountToken");
-		check(GetDebugName() == TokenName);
-#endif
-		return ConsumeToken();
-	}
-		
-	void* ConsumePointer(uint32& OutReturnCount)
-	{
-		void* Out = Builder.ReadPointer(/* in-out */ Idx);
-#if ENABLE_GC_TOKEN_DEBUG_INFO
-		static const FName TokenName("EndOfPointerToken");
-		check(GetDebugName() == TokenName);
-#endif
-		FGCReferenceInfo EndOfPointer(ConsumeToken());
-		OutReturnCount = EndOfPointer.ReturnCount;
-		return Out;
-	}
-
-	const FTokenStreamBuilder& DebugGetBuilder() { return Builder; }
-private:
-	uint32 GetToken() const { return Builder.Tokens[Idx]; }
-	uint32 ConsumeToken() { return Builder.Tokens[Idx++]; }
-
-#if ENABLE_GC_TOKEN_DEBUG_INFO
-	FName GetDebugName() const { return Builder.DebugNames[Idx]; }
-#else
-	FName GetDebugName() const { return FName(); }
-#endif
-
-	const FTokenStreamBuilder& Builder;
-	uint32 Idx = 0;
-	uint32 Num = 0;
-};
-
 static AROFunc GetAROFunc(UClass* Class)
 {
 	AROFunc ARO = Class->CppClassStaticFunctions.GetAddReferencedObjects();
@@ -5268,8 +5138,6 @@ struct FScopeLockIfNotNative
 
 /** Token stream can assemble code can sometimes be called from two threads throuh a web of async loading calls. */
 static FCriticalSection GReferenceTokenStreamCritical;
-int32 GCreatedMixedTokens;
-int32 GCreatedStrongTokens;
 
 void UClass::AssembleReferenceTokenStream(bool bForce)
 {
@@ -5300,7 +5168,7 @@ void UClass::AssembleReferenceTokenStreamInternal(bool bForce)
 		{
 			IntrinsicBuilder = FIntrinsicClassTokens::ConsumeBuilder(this);
 		}
-		FTokenStreamBuilder& MixedTokens = IntrinsicBuilder ? *IntrinsicBuilder : StackBuilder;
+		FTokenStreamBuilder& Builder = IntrinsicBuilder ? *IntrinsicBuilder : StackBuilder;
 
 		if (bForce)
 		{
@@ -5316,49 +5184,31 @@ void UClass::AssembleReferenceTokenStreamInternal(bool bForce)
 			{
 				FProperty* Property = *It;
 				FGCStackSizeHelperScope StackSizeScope(StackSizeHelper, Property);
-				Property->EmitReferenceInfo(MixedTokens, 0, EncounteredStructProps, StackSizeHelper);
+				Property->EmitReferenceInfo(Builder, 0, EncounteredStructProps, StackSizeHelper);
 			}
 
-			MixedTokens.SetStackSize(StackSizeHelper.GetMaxStackSize());
+			Builder.SetStackSize(StackSizeHelper.GetMaxStackSize());
 		}
 				
-		FTokenStreamView SuperMixedTokens;
-		FTokenStreamView SuperStrongTokens;
+		FTokenStreamView SuperTokens;
 		if (UClass* SuperClass = GetSuperClass())
 		{
 			// Make sure super class has valid token stream.
 			SuperClass->AssembleReferenceTokenStreamInternal();
-			SuperMixedTokens  = FTokenStreamBuilder::DropFinalTokens(SuperClass->ReferenceTokens.Mixed,  GetAROFunc(SuperClass));
-			SuperStrongTokens = FTokenStreamBuilder::DropFinalTokens(SuperClass->ReferenceTokens.Strong, GetAROFunc(SuperClass));
+			SuperTokens = FTokenStreamBuilder::DropFinalTokens(SuperClass->ReferenceTokens.Strong, GetAROFunc(SuperClass));
 		}
 		else
 		{
-			UObjectBase::EmitBaseReferences(MixedTokens);
+			UObjectBase::EmitBaseReferences(Builder);
 		}
 
 		// Make sure all Blueprint properties are marked as non-native
-		MixedTokens.SetTokenType(GetClass()->HasAnyClassFlags(CLASS_NeedsDeferredDependencyLoading) ? EGCTokenType::NonNative : EGCTokenType::Native);
-		MixedTokens.EmitFinalTokens(GetAROFunc(this));
+		Builder.SetTokenType(GetClass()->HasAnyClassFlags(CLASS_NeedsDeferredDependencyLoading) ? EGCTokenType::NonNative : EGCTokenType::Native);
+		Builder.EmitFinalTokens(GetAROFunc(this));
 
 		// Allocate flat token stream with current and super class tokens
-		FTokenStreamBuilder::Merge(/* out */ ReferenceTokens.Mixed, MixedTokens, SuperMixedTokens);
+		FTokenStreamBuilder::Merge(/* out */ ReferenceTokens.Strong, Builder, SuperTokens);
 		
-		// Create subset of strong reference tokens for this class
-		FTokenStreamBuilder StrongTokens(*this);
-		StrongTokens.SetStackSize(MixedTokens.GetStackSize());
-		StrongTokens.SetTokenType(MixedTokens.GetTokenType());
-		FTokenStreamBuilder::CopyStrongTokens(FTokenStreamBuilder::FConstIterator(MixedTokens), /* out */ StrongTokens);
-
-		// Reuse mixed tokens if possible
-		if (MixedTokens.Num() == StrongTokens.Num() && SuperMixedTokens.NumTokens() == SuperStrongTokens.NumTokens())
-		{
-			UpdateStreamReference(/* out */ ReferenceTokens.Strong, ReferenceTokens.Mixed);
-		}
-		else
-		{
-			FTokenStreamBuilder::Merge(/* out */ ReferenceTokens.Strong, StrongTokens, SuperStrongTokens);
-		}
-
 		checkf(!HasAnyClassFlags(CLASS_TokenStreamAssembled), TEXT("Token stream already assembled for class '%s'"), *GetPathName()); // recursion here is probably bad
 		ClassFlags |= CLASS_TokenStreamAssembled;
 	}
@@ -5506,155 +5356,6 @@ static uint32 PickReturnCount(uint32 OuterReturnCount, uint32 InnerReturnCount)
 	return OuterReturnCount;
 }
 
-static bool CopyStrongTokensUntilPos(FTokenStreamBuilder::FConstIterator& /* in-out */ MixedIt, uint32 MixedEndPos, FTokenStreamBuilder& Out, uint32& OutReturnCount)
-{
-	bool bOut = false;
-	uint32 LastReturnCount = 0;
-	while (MixedIt.Pos() < MixedEndPos)
-	{
-		bOut |= FTokenStreamBuilder::CopyNextStrongToken(MixedIt, Out, LastReturnCount);
-	}
-
-	OutReturnCount = PickReturnCount(OutReturnCount, LastReturnCount);
-
-	return bOut;
-}
-
-static bool CopyStrongTokensUntilReturn(FTokenStreamBuilder::FConstIterator& /* in-out */ MixedIt, FTokenStreamBuilder& Out, uint32& OutReturnCount)
-{
-	uint32 FooStartPos = MixedIt.Pos();
-
-	uint32 InnerReturnCount = 0;
-	bool bOut = false;
-	while (InnerReturnCount == 0)
-	{
-		bOut |= FTokenStreamBuilder::CopyNextStrongToken(/* in-out */ MixedIt, Out, InnerReturnCount);
-	}
-	
-	OutReturnCount = PickReturnCount(OutReturnCount, InnerReturnCount);
-
-	return bOut;
-}
-
-void FTokenStreamBuilder::CopyStrongTokens(FConstIterator MixedIt, FTokenStreamBuilder& Out)
-{
-	uint32 ReturnCount = 0;
-	CopyStrongTokensUntilPos(MixedIt, MixedIt.EndPos(), Out, ReturnCount);
-	checkf(Out.Num() < MixedIt.EndPos() || Out == MixedIt.DebugGetBuilder(), 
-		TEXT("No weak tokens stripped, copied tokens should be bitwise identical to the original mixed tokens"));
-}
-
-bool FTokenStreamBuilder::CopyNextStrongToken(FConstIterator& /* in-out */ MixedIt, FTokenStreamBuilder& Out, uint32& OutReturnCount)
-{
-	check(MixedIt);
-	
-	const uint32 RollbackNum = Out.Tokens.Num();
-	FNamedReferenceInfo Token = MixedIt.ConsumeReferenceInfo();
-
-	// Strip old return count and emit new returns since we might strip out entire fixed arrays/struct arrays/maps/sets/etc
-	OutReturnCount = Token.ReferenceInfo.ReturnCount;
-	Token.ReferenceInfo.ReturnCount = 0;
-
-	switch (Token.ReferenceInfo.Type)
-	{
-	case GCRT_Object:
-	case GCRT_ArrayObject:
-	case GCRT_ArrayObjectFreezable:
-	case GCRT_ExternalPackage:
-	case GCRT_AddFieldPathReferencedObject:
-	case GCRT_ArrayAddFieldPathReferencedObject:
-	case GCRT_DynamicallyTypedValue: // Always keep for now
-	case GCRT_SlowAddReferencedObjects:
-	case GCRT_EndOfStream:
-		Out.EmitReferenceInfo(Token.ReferenceInfo, Token.DebugName);
-		return true;
-
-	case GCRT_AddStructReferencedObjects:
-	case GCRT_AddReferencedObjects:
-		Out.EmitReferenceInfo(Token.ReferenceInfo, Token.DebugName);
-		Out.EmitPointer(MixedIt.ConsumePointer(OutReturnCount));
-		return true;
-
-	case GCRT_WeakObject:
-	case GCRT_ArrayWeakObject:
-	case GCRT_LazyObject:
-	case GCRT_ArrayLazyObject:
-	case GCRT_SoftObject:
-	case GCRT_ArraySoftObject:
-	case GCRT_Delegate:
-	case GCRT_ArrayDelegate:
-	case GCRT_MulticastDelegate:
-	case GCRT_ArrayMulticastDelegate:
-		return false;
-
-	// Below property types get their Out tokens rolled back if nested properties are entirely weak
-	case GCRT_Optional:
-	case GCRT_ArrayStructFreezable:
-	case GCRT_ArrayStruct:
-	{
-		Out.EmitReferenceInfo(Token.ReferenceInfo, Token.DebugName);
-		Out.EmitStride(MixedIt.ConsumeStride());
-
-		uint32 MixedSkipPos = MixedIt.ConsumeSkipInfo();
-		uint32 SkipSkipIndex = Out.EmitSkipIndexPlaceholder();
-
-		if (CopyStrongTokensUntilPos(MixedIt, MixedSkipPos, Out, OutReturnCount))
-		{
-			uint32 SkipIndex = Out.EmitReturn();
-			Out.UpdateSkipIndexPlaceholder(SkipSkipIndex, SkipIndex);
-			return true;
-		}
-	}
-	break;
-	case GCRT_FixedArray:
-	{
-		Out.EmitReferenceInfo(Token.ReferenceInfo, Token.DebugName);
-		Out.EmitStride(MixedIt.ConsumeStride());
-		Out.EmitCount(MixedIt.ConsumeCount());
-
-		if (CopyStrongTokensUntilReturn(/* in-out */ MixedIt, Out, OutReturnCount))
-		{
-			Out.EmitReturn();
-			return true;
-		}
-		break;
-	}
-	case GCRT_AddTSetReferencedObjects:
-	case GCRT_AddTMapReferencedObjects:
-	{
-		Out.EmitReferenceInfo(Token.ReferenceInfo, Token.DebugName);
-		uint32 NoReturnCount = 0;
-		FProperty* Property = (FProperty*)MixedIt.ConsumePointer(/* Out */ NoReturnCount);
-		check(NoReturnCount == 0);
-		check(Property->IsA<FMapProperty>() || Property->IsA<FSetProperty>());
-		Out.EmitPointer(Property);
-
-		uint32 MixedSkipPos = MixedIt.ConsumeSkipInfo();
-		uint32 SkipSkipIndex = Out.EmitSkipIndexPlaceholder();
-	
-		if (CopyStrongTokensUntilPos(MixedIt, MixedSkipPos, Out, OutReturnCount))
-		{
-			uint32 SkipIndex = Out.EmitReturn();
-			Out.UpdateSkipIndexPlaceholder(SkipSkipIndex, SkipIndex);
-			return true;
-		}
-		break;
-	}
-	case GCRT_EndOfPointer:
-		checkf(false, TEXT("ConsumePointer() should consume all GCRT_EndOfPointer tokens"));
-	break;
-	default:
-		checkf(false, TEXT("Illegal token type %d %s"), Token.ReferenceInfo.Type, *Token.DebugName.ToString());
-	break;
-	} // switch
-
-	Out.Tokens.SetNum(RollbackNum, /* shrink */ false);		
-#if ENABLE_GC_TOKEN_DEBUG_INFO
-	Out.DebugNames.SetNum(RollbackNum, /* shrink */ false);		
-#endif
-	return false;
-}
-
 int32 FTokenStreamBuilder::EmitReferenceInfo(FGCReferenceInfo ReferenceInfo, FName DebugName)
 {
 	int32 TokenIndex = Tokens.Add(ReferenceInfo);
@@ -5785,14 +5486,13 @@ FTokenInfo FTokenStreamOwner::GetTokenInfo(FTokenId Id) const
 		return {-1, ToName(Id.AsTokenless()) };
 	}
 
-	FTokenStreamView Stream = Id.IsMixed() ? Mixed : Strong;
 	uint32 Idx = Id.GetIndex();
 
-	check(Idx < Stream.Num);
+	check(Idx < Strong.Num);
 	FTokenInfo DebugInfo;
-	DebugInfo.Offset = FGCReferenceInfo(Stream.Tokens[Idx]).Offset;
+	DebugInfo.Offset = FGCReferenceInfo(Strong.Tokens[Idx]).Offset;
 #if ENABLE_GC_TOKEN_DEBUG_INFO
-	DebugInfo.Name = Stream.GetDebugNames()[Idx];
+	DebugInfo.Name = Strong.GetDebugNames()[Idx];
 #else
 	static const FName DefaultTokenName("Token");	
 	DebugInfo.Name = DefaultTokenName; 
@@ -5806,13 +5506,11 @@ FTokenInfo FTokenStreamOwner::GetTokenInfo(FTokenId Id) const
 void FTokenStreamOwner::Reset()
 {
 	UpdateStreamReference(/* out */ Strong, FTokenStreamView());
-	UpdateStreamReference(/* out */ Mixed, FTokenStreamView());
 }
 
 int64 FTokenStreamOwner::GetTokenAllocatedSize() const
 {
-	uint32 NumTokens = Mixed.Tokens == Strong.Tokens ? Mixed.Num : Mixed.Num + Strong.Num;
-	return NumTokens * sizeof(uint32);
+	return Strong.Num * sizeof(uint32);
 }
 
 int64 FTokenStreamOwner::GetDebugInfoAllocatedSize() const
diff --git a/Engine/Source/Runtime/CoreUObject/Public/UObject/FastReferenceCollector.h b/Engine/Source/Runtime/CoreUObject/Public/UObject/FastReferenceCollector.h
--- a/Engine/Source/Runtime/CoreUObject/Public/UObject/FastReferenceCollector.h
+++ b/Engine/Source/Runtime/CoreUObject/Public/UObject/FastReferenceCollector.h
@@ -32,19 +32,12 @@ enum class EGCOptions : uint32
 	Parallel = 1 << 0,					// Use all task workers to collect references, must be started on main thread
 	AutogenerateTokenStream = 1 << 1,	// Lazily generate token streams for new UClasses
 	WithClusters = 1 << 2,				// Use clusters, see FUObjectCluster
-	ProcessWeakReferences UE_DEPRECATED(5.2, "Weak reference collection deprecated to reduce complexity. Use FArchive instead with ArIsObjectReferenceCollector and ArIsModifyingWeakAndStrongReferences")
-					= 1 << 3,			// Collect both weak and strong references instead of just strong
-	WithPendingKill = 1 << 4,			// Internal flag used by reachability analysis
+	WithPendingKill = 1 << 3,			// Internal flag used by reachability analysis
 };
 ENUM_CLASS_FLAGS(EGCOptions);
 
 inline constexpr bool IsParallel(EGCOptions Options) { return !!(Options & EGCOptions::Parallel); }
 inline constexpr bool IsPendingKill(EGCOptions Options) { return !!(Options & EGCOptions::WithPendingKill); }
-PRAGMA_DISABLE_DEPRECATION_WARNINGS
-inline constexpr bool ShouldProcessWeakReferences(EGCOptions Options) { return !!(Options & EGCOptions::ProcessWeakReferences); }
-PRAGMA_ENABLE_DEPRECATION_WARNINGS
-
-using EFastReferenceCollectorOptions UE_DEPRECATED(5.2, "Use EGCOptions instead") = EGCOptions;
 
 /** Helper to give GC internals friend access to certain core classes */
 struct FGCInternals
@@ -323,10 +316,9 @@ enum class ETokenlessId
 struct FTokenId
 {
 public:
-	FORCEINLINE FTokenId(ETokenlessId Tokenless) : Index(0), TokenlessId((uint32)Tokenless), Mixed(0) {}
-	FORCEINLINE FTokenId(uint32 Idx, bool bIsMixed) : Index(Idx), TokenlessId(0), Mixed(bIsMixed) {}
+	FORCEINLINE FTokenId(ETokenlessId Tokenless) : Index(0), TokenlessId((uint32)Tokenless) {}
+	FORCEINLINE explicit FTokenId(uint32 Idx) : Index(Idx), TokenlessId(0) {}
 	
-	bool IsMixed() const { return Mixed; }
 	bool IsTokenless() const { return TokenlessId != 0; }
 	uint32 GetIndex() const { check(!IsTokenless()); return Index; }
 	int32 AsPrintableIndex() const { return IsTokenless() ? -int32(TokenlessId) : Index; }
@@ -337,7 +329,7 @@ public:
 
 private:
 	static constexpr uint32 TokenlessIdBits = 3;
-	static constexpr uint32 IndexBits = 32 - TokenlessIdBits - 1;
+	static constexpr uint32 IndexBits = 32 - TokenlessIdBits;
 	void StaticAssert();
 
 	union
@@ -346,7 +338,6 @@ private:
 		{
 			uint32 Index : IndexBits;
 			uint32 TokenlessId : TokenlessIdBits;
-			uint32 Mixed : 1;
 		};
 		uint32 All;
 	};
@@ -534,14 +525,6 @@ class TFastReferenceCollector : public FGCInternals
 	static constexpr EGCOptions Options = ProcessorType::Options;
 
 	static constexpr FORCEINLINE bool IsParallel()					{ return !!(Options & EGCOptions::Parallel); }
-PRAGMA_DISABLE_DEPRECATION_WARNINGS
-	static constexpr FORCEINLINE bool ShouldProcessWeakReferences()	{ return !!(Options & EGCOptions::ProcessWeakReferences); }
-PRAGMA_ENABLE_DEPRECATION_WARNINGS
-
-	static FORCEINLINE FTokenId MakeTokenId(uint32 Index)
-	{
-		return FTokenId(Index, ShouldProcessWeakReferences());
-	}
 
 	static FORCEINLINE UE::GC::FTokenStreamView& GetTokenStream(UClass* Class)
 	{
@@ -550,7 +533,7 @@ PRAGMA_ENABLE_DEPRECATION_WARNINGS
 			Class->AssembleReferenceTokenStream();
 		}
 
-		return ShouldProcessWeakReferences() ? Class->ReferenceTokens.Mixed : Class->ReferenceTokens.Strong;
+		return Class->ReferenceTokens.Strong;
 	}
 
 	ProcessorType& Processor;
@@ -689,21 +672,21 @@ StoleContext:
 						{								
 							UObject** ObjectPtr = (UObject**)(StackEntryData + ReferenceInfo.Offset);
 							TokenReturnCount = ReferenceInfo.ReturnCount;
-							Dispatcher.HandleKillableReference(*ObjectPtr, MakeTokenId(ReferenceTokenStreamIndex), TokenType);
+							Dispatcher.HandleKillableReference(*ObjectPtr, FTokenId(ReferenceTokenStreamIndex), TokenType);
 						}
 						break;
 						case GCRT_ArrayObject:
 						{
 							TArray<UObject*>& ObjectArray = *((TArray<UObject*>*)(StackEntryData + ReferenceInfo.Offset));
 							TokenReturnCount = ReferenceInfo.ReturnCount;
-							Dispatcher.HandleKillableArray(ObjectArray, MakeTokenId(ReferenceTokenStreamIndex), TokenType);
+							Dispatcher.HandleKillableArray(ObjectArray, FTokenId(ReferenceTokenStreamIndex), TokenType);
 						}
 						break;
 						case GCRT_ArrayObjectFreezable:
 						{
 							TArray<UObject*, FMemoryImageAllocator>& ObjectArray = *((TArray<UObject*, FMemoryImageAllocator>*)(StackEntryData + ReferenceInfo.Offset));
 							TokenReturnCount = ReferenceInfo.ReturnCount;
-							Dispatcher.HandleKillableReferences(ObjectArray, MakeTokenId(ReferenceTokenStreamIndex), TokenType);
+							Dispatcher.HandleKillableReferences(ObjectArray, FTokenId(ReferenceTokenStreamIndex), TokenType);
 						}
 						break;
 						case GCRT_ArrayStruct:
@@ -793,7 +776,7 @@ StoleContext:
 							// Test if the object isn't itself, since currently package are their own external and tracking that reference is pointless
 							UObject* Object = CurrentObject->GetExternalPackageInternal();
 							Object = Object != CurrentObject ? Object : nullptr;
-							Dispatcher.HandleImmutableReference( Object, MakeTokenId(ReferenceTokenStreamIndex), TokenType);
+							Dispatcher.HandleImmutableReference( Object, FTokenId(ReferenceTokenStreamIndex), TokenType);
 						}
 						break;
 						case GCRT_FixedArray:
@@ -916,7 +899,7 @@ StoleContext:
 							{
 								UObject* OwnerObject = static_cast<UObject*>(FieldOwnerItem->Object);
 								UObject* PreviousOwner = OwnerObject;
-								Dispatcher.HandleReferenceDirectly(CurrentObject, OwnerObject, MakeTokenId(ReferenceTokenStreamIndex), TokenType, true);
+								Dispatcher.HandleReferenceDirectly(CurrentObject, OwnerObject, FTokenId(ReferenceTokenStreamIndex), TokenType, true);
 								
 								// Handle reference elimination (PendingKill owner)
 								if (PreviousOwner && !OwnerObject)
@@ -936,7 +919,7 @@ StoleContext:
 								{
 									UObject* OwnerObject = static_cast<UObject*>(FieldOwnerItem->Object);
 									UObject* PreviousOwner = OwnerObject;
-									Dispatcher.HandleReferenceDirectly(CurrentObject, OwnerObject, MakeTokenId(ReferenceTokenStreamIndex), TokenType, true);
+									Dispatcher.HandleReferenceDirectly(CurrentObject, OwnerObject, FTokenId(ReferenceTokenStreamIndex), TokenType, true);
 
 									// Handle reference elimination (PendingKill owner)
 									if (PreviousOwner && !OwnerObject)
@@ -976,136 +959,6 @@ StoleContext:
 							TokenReturnCount = ReferenceInfo.ReturnCount;
 						}
 						break;
-						case GCRT_WeakObject:
-						{
-							TokenReturnCount = ReferenceInfo.ReturnCount;
-							if (ShouldProcessWeakReferences())
-							{
-								FWeakObjectPtr& WeakPtr = *(FWeakObjectPtr*)(StackEntryData + ReferenceInfo.Offset);
-								Dispatcher.HandleWeakReference(WeakPtr, CurrentObject, MakeTokenId(ReferenceTokenStreamIndex), TokenType);
-							}
-						}
-						break;
-						case GCRT_ArrayWeakObject:
-						{
-							TokenReturnCount = ReferenceInfo.ReturnCount;
-							if (ShouldProcessWeakReferences())
-							{
-								TArray<FWeakObjectPtr>& WeakPtrArray = *((TArray<FWeakObjectPtr>*)(StackEntryData + ReferenceInfo.Offset));
-								for (FWeakObjectPtr& WeakPtr : WeakPtrArray)
-								{
-									Dispatcher.HandleWeakReference(WeakPtr, CurrentObject, MakeTokenId(ReferenceTokenStreamIndex), TokenType);
-								}
-							}
-						}
-						break;
-						case GCRT_LazyObject:
-						{
-							TokenReturnCount = ReferenceInfo.ReturnCount;
-							if (ShouldProcessWeakReferences())
-							{
-								FLazyObjectPtr& LazyPtr = *(FLazyObjectPtr*)(StackEntryData + ReferenceInfo.Offset);
-								FWeakObjectPtr& WeakPtr = AccessWeakPtr(LazyPtr);
-								Dispatcher.HandleWeakReference(WeakPtr, CurrentObject, MakeTokenId(ReferenceTokenStreamIndex), TokenType);
-							}
-						}
-						break;
-						case GCRT_ArrayLazyObject:
-						{
-							TokenReturnCount = ReferenceInfo.ReturnCount;
-							if (ShouldProcessWeakReferences())
-							{
-								TArray<FLazyObjectPtr>& LazyPtrArray = *((TArray<FLazyObjectPtr>*)(StackEntryData + ReferenceInfo.Offset));
-								TokenReturnCount = ReferenceInfo.ReturnCount;
-								for (FLazyObjectPtr& LazyPtr : LazyPtrArray)
-								{
-									FWeakObjectPtr& WeakPtr = AccessWeakPtr(LazyPtr);
-									Dispatcher.HandleWeakReference(WeakPtr, CurrentObject, MakeTokenId(ReferenceTokenStreamIndex), TokenType);
-								}
-							}
-						}
-						break;
-						case GCRT_SoftObject:
-						{
-							TokenReturnCount = ReferenceInfo.ReturnCount;
-							if (ShouldProcessWeakReferences())
-							{
-								FSoftObjectPtr& SoftPtr = *(FSoftObjectPtr*)(StackEntryData + ReferenceInfo.Offset);
-								FWeakObjectPtr& WeakPtr = AccessWeakPtr(SoftPtr);
-								Dispatcher.HandleWeakReference(WeakPtr, CurrentObject, MakeTokenId(ReferenceTokenStreamIndex), TokenType);
-							}
-						}
-						break;
-						case GCRT_ArraySoftObject:
-						{
-							TokenReturnCount = ReferenceInfo.ReturnCount;
-							if (ShouldProcessWeakReferences())
-							{
-								TArray<FSoftObjectPtr>& SoftPtrArray = *((TArray<FSoftObjectPtr>*)(StackEntryData + ReferenceInfo.Offset));
-								for (FSoftObjectPtr& SoftPtr : SoftPtrArray)
-								{
-									FWeakObjectPtr& WeakPtr = AccessWeakPtr(SoftPtr);
-									Dispatcher.HandleWeakReference(WeakPtr, CurrentObject, MakeTokenId(ReferenceTokenStreamIndex), TokenType);
-								}
-							}
-						}
-						break;
-						case GCRT_Delegate:
-						{
-							TokenReturnCount = ReferenceInfo.ReturnCount;
-							if (ShouldProcessWeakReferences())
-							{
-								FScriptDelegate& Delegate = *(FScriptDelegate*)(StackEntryData + ReferenceInfo.Offset);
-								UObject* DelegateObject = Delegate.GetUObject();
-								Dispatcher.HandleReferenceDirectly(CurrentObject, DelegateObject, MakeTokenId(ReferenceTokenStreamIndex), TokenType, false);
-							}
-						}
-						break;
-						case GCRT_ArrayDelegate:
-						{
-							TokenReturnCount = ReferenceInfo.ReturnCount;
-							if (ShouldProcessWeakReferences())
-							{
-								TArray<FScriptDelegate>& DelegateArray = *((TArray<FScriptDelegate>*)(StackEntryData + ReferenceInfo.Offset));
-								for (FScriptDelegate& Delegate : DelegateArray)
-								{
-									UObject* DelegateObject = Delegate.GetUObject();
-									Dispatcher.HandleReferenceDirectly(CurrentObject, DelegateObject, MakeTokenId(ReferenceTokenStreamIndex), TokenType, false);
-								}
-							}
-						}
-						break;
-						case GCRT_MulticastDelegate:
-						{
-							TokenReturnCount = ReferenceInfo.ReturnCount;
-							if (ShouldProcessWeakReferences())
-							{
-								FMulticastScriptDelegate& Delegate = *(FMulticastScriptDelegate*)(StackEntryData + ReferenceInfo.Offset);
-								TArray<UObject*> DelegateObjects(Delegate.GetAllObjects());
-								for (UObject* DelegateObject : DelegateObjects)
-								{
-									Dispatcher.HandleReferenceDirectly(CurrentObject, DelegateObject, MakeTokenId(ReferenceTokenStreamIndex), TokenType, false);
-								}
-							}
-						}
-						break;
-						case GCRT_ArrayMulticastDelegate:
-						{
-							TokenReturnCount = ReferenceInfo.ReturnCount;
-							if (ShouldProcessWeakReferences())
-							{
-								TArray<FMulticastScriptDelegate>& DelegateArray = *((TArray<FMulticastScriptDelegate>*)(StackEntryData + ReferenceInfo.Offset));
-								for (FMulticastScriptDelegate& Delegate : DelegateArray)
-								{
-									TArray<UObject*> DelegateObjects(Delegate.GetAllObjects());
-									for (UObject* DelegateObject : DelegateObjects)
-									{
-										Dispatcher.HandleReferenceDirectly(CurrentObject, DelegateObject, MakeTokenId(ReferenceTokenStreamIndex), TokenType, false);
-									}
-								}
-							}
-						}
-						break;
 						case GCRT_DynamicallyTypedValue:
 						{
 							TokenReturnCount = ReferenceInfo.ReturnCount;
diff --git a/Engine/Source/Runtime/CoreUObject/Public/UObject/GarbageCollection.h b/Engine/Source/Runtime/CoreUObject/Public/UObject/GarbageCollection.h
--- a/Engine/Source/Runtime/CoreUObject/Public/UObject/GarbageCollection.h
+++ b/Engine/Source/Runtime/CoreUObject/Public/UObject/GarbageCollection.h
@@ -86,16 +86,6 @@ enum EGCReferenceType
 	GCRT_ArrayObjectFreezable,				// Freezable array of references, e.g. TArray<MyObject*, FMemoryImageAllocator>
 	GCRT_ArrayStructFreezable,				// Freezable array of structs
 	GCRT_Optional,							// TOptional
-	GCRT_WeakObject,						// TWeakObjectPtr
-	GCRT_ArrayWeakObject,					// Array of weak pointers
-	GCRT_LazyObject,						// TLazyObjectPtr
-	GCRT_ArrayLazyObject,					// Array of lazy pointers
-	GCRT_SoftObject,						// TSoftObjectPtr
-	GCRT_ArraySoftObject,					// Array of soft object pointers
-	GCRT_Delegate,							// UobjectDelegate weak reference to target object
-	GCRT_ArrayDelegate,						// Array of delegates
-	GCRT_MulticastDelegate,					// Delegate weak reference to target object
-	GCRT_ArrayMulticastDelegate,			// Array of delegates
 	GCRT_DynamicallyTypedValue,				// FDynamicallyTypedValue
 	GCRT_SlowAddReferencedObjects,			// AddReferencedObjects() function using UE::GC::RegisterSlowImplementation
 };
@@ -266,8 +256,7 @@ struct FTokenStreamOwner
 	FTokenStreamOwner() = default;
 	~FTokenStreamOwner() { Reset(); }
 
-	FTokenStreamView Strong; // Strong references only
-	FTokenStreamView Mixed; // Strong and weak references
+	FTokenStreamView Strong;
 	
 	COREUOBJECT_API void Reset();
 
@@ -338,16 +327,11 @@ public:
 	COREUOBJECT_API void EmitExternalPackageReference();
 
 	void EmitFinalTokens(void(*AddReferencedObjects)(UObject*, FReferenceCollector&));	
-	
-	static FTokenStreamView DropFinalTokens(FTokenStreamView Mixed, void(*DropARO)(UObject*, FReferenceCollector&));
+	static FTokenStreamView DropFinalTokens(FTokenStreamView Tokens, void(*DropARO)(UObject*, FReferenceCollector&));
 
 	/** Allocate merged stream with tokens from super class and current class */
 	static void Merge(FTokenStreamView& Out, const FTokenStreamBuilder& Class, FTokenStreamView Super);
 	
-	class FConstIterator;
-	static void CopyStrongTokens(FConstIterator MixedIt, FTokenStreamBuilder& Out);
-	static bool CopyNextStrongToken(FConstIterator& /* in-out */ MixedIt, FTokenStreamBuilder& Out, uint32& OutReturnCount);
-
 	/** Reads count and advances stream */
 	FORCEINLINE uint32 ReadCount(uint32& CurrentIndex)
 	{
diff --git a/Engine/Source/Runtime/CoreUObject/Public/UObject/UnrealType.h b/Engine/Source/Runtime/CoreUObject/Public/UObject/UnrealType.h
--- a/Engine/Source/Runtime/CoreUObject/Public/UObject/UnrealType.h
+++ b/Engine/Source/Runtime/CoreUObject/Public/UObject/UnrealType.h
@@ -3069,7 +3069,6 @@ PRAGMA_ENABLE_DEPRECATION_WARNINGS
 
 	// FProperty interface
 	virtual void SerializeItem(FStructuredArchive::FSlot Slot, void* Value, void const* Defaults) const override;
-	virtual void EmitReferenceInfo(UE::GC::FTokenStreamBuilder& TokenStream, int32 BaseOffset, TArray<const FStructProperty*>& EncounteredStructProps, FGCStackSizeHelper& StackSizeHelper) override;
 protected:
 	virtual void LinkInternal(FArchive& Ar) override;
 private:
@@ -3138,7 +3137,6 @@ PRAGMA_ENABLE_DEPRECATION_WARNINGS
 	virtual FName GetID() const override;
 	virtual bool Identical( const void* A, const void* B, uint32 PortFlags ) const override;
 	virtual void SerializeItem( FStructuredArchive::FSlot Slot, void* Value, void const* Defaults ) const override;
-	virtual void EmitReferenceInfo(UE::GC::FTokenStreamBuilder& TokenStream, int32 BaseOffset, TArray<const FStructProperty*>& EncounteredStructProps, FGCStackSizeHelper& StackSizeHelper) override;
 	virtual void CopySingleValueToScriptVM( void* Dest, void const* Src ) const override;
 	virtual void CopySingleValueFromScriptVM( void* Dest, void const* Src ) const override;
 	virtual void CopyCompleteValueToScriptVM( void* Dest, void const* Src ) const override;
@@ -3218,7 +3216,6 @@ protected:
 	virtual const TCHAR* ImportText_Internal(const TCHAR* Buffer, void* ContainerOrPropertyPtr, EPropertyPointerType PropertyPointerType, UObject* OwnerObject, int32 PortFlags, FOutputDevice* ErrorText) const override;
 public:
 	virtual EConvertFromTypeResult ConvertFromType(const FPropertyTag& Tag, FStructuredArchive::FSlot Slot, uint8* Data, UStruct* DefaultsStruct) override;
-	virtual void EmitReferenceInfo(UE::GC::FTokenStreamBuilder& TokenStream, int32 BaseOffset, TArray<const FStructProperty*>& EncounteredStructProps, FGCStackSizeHelper& StackSizeHelper) override;
 	// End of FProperty interface
 
 	// FObjectProperty interface
@@ -5969,7 +5966,6 @@ protected:
 	virtual const TCHAR* ImportText_Internal(const TCHAR* Buffer, void* ContainerOrPropertyPtr, EPropertyPointerType PropertyPointerType, UObject* OwnerObject, int32 PortFlags, FOutputDevice* ErrorText) const override;
 public:
 	virtual bool ContainsObjectReference(TArray<const FStructProperty*>& EncounteredStructProps, EPropertyObjectReferenceType InReferenceType = EPropertyObjectReferenceType::Strong) const override;
-	virtual void EmitReferenceInfo(UE::GC::FTokenStreamBuilder& TokenStream, int32 BaseOffset, TArray<const FStructProperty*>& EncounteredStructProps, FGCStackSizeHelper& StackSizeHelper) override;
 	virtual void InstanceSubobjects( void* Data, void const* DefaultData, UObject* Owner, struct FObjectInstancingGraph* InstanceGraph ) override;
 	virtual bool SameType(const FProperty* Other) const override;
 #if WITH_EDITORONLY_DATA
@@ -6039,7 +6035,6 @@ protected:
 	virtual void ExportText_Internal( FString& ValueStr, const void* PropertyValueOrContainer, EPropertyPointerType PropertyPointerType, const void* DefaultValue, UObject* Parent, int32 PortFlags, UObject* ExportRootScope ) const override;
 public:
 	virtual bool ContainsObjectReference(TArray<const FStructProperty*>& EncounteredStructProps, EPropertyObjectReferenceType InReferenceType = EPropertyObjectReferenceType::Strong) const override;
-	virtual void EmitReferenceInfo(UE::GC::FTokenStreamBuilder& TokenStream, int32 BaseOffset, TArray<const FStructProperty*>& EncounteredStructProps, FGCStackSizeHelper& StackSizeHelper) override;
 	virtual void InstanceSubobjects( void* Data, void const* DefaultData, UObject* Owner, struct FObjectInstancingGraph* InstanceGraph ) override;
 	virtual bool SameType(const FProperty* Other) const override;
 	virtual EConvertFromTypeResult ConvertFromType(const FPropertyTag& Tag, FStructuredArchive::FSlot Slot, uint8* Data, UStruct* DefaultsStruct) override;
