Optimize ResolveName() - Change FString PartialName to TStringBuilder

diff --git a/Engine/Source/Runtime/CoreUObject/Private/UObject/UObjectGlobals.cpp b/Engine/Source/Runtime/CoreUObject/Private/UObject/UObjectGlobals.cpp
--- a/Engine/Source/Runtime/CoreUObject/Private/UObject/UObjectGlobals.cpp
+++ b/Engine/Source/Runtime/CoreUObject/Private/UObject/UObjectGlobals.cpp
@@ -1129,7 +1129,9 @@ bool ResolveName(UObject*& InPackage, FString& InOutName, bool Create, bool Thro
 		}
 		
 		int32 DotIndex = DelimiterOrEnd - *InOutName;
-		FString PartialName = InOutName.Left(DotIndex);
+
+		TStringBuilder<FName::StringBufferSize> PartialName;
+		PartialName.Append(*InOutName, DotIndex);
 
 		bool bIsScriptPackage = false;
 		if (!InPackage)
@@ -1140,15 +1142,15 @@ bool ResolveName(UObject*& InPackage, FString& InOutName, bool Create, bool Thro
 				FName* ScriptPackageName = FPackageName::FindScriptPackageName(*PartialName);
 				if (ScriptPackageName)
 				{
-					PartialName = ScriptPackageName->ToString();
+					ScriptPackageName->ToString(PartialName);
 				}
-				bIsScriptPackage = ScriptPackageName || FPackageName::IsScriptPackage(PartialName);
+				bIsScriptPackage = ScriptPackageName || FPackageName::IsScriptPackage(FStringView(PartialName));
 			}
 
 			// Process any package redirects before calling CreatePackage/FindObject
 			{
 				const FCoreRedirectObjectName NewPackageName = FCoreRedirects::GetRedirectedName(ECoreRedirectFlags::Type_Package, FCoreRedirectObjectName(NAME_None, NAME_None, *PartialName));
-				PartialName = NewPackageName.PackageName.ToString();
+				NewPackageName.PackageName.ToString(PartialName);
 			}
 		}
 
@@ -1175,7 +1177,7 @@ bool ResolveName(UObject*& InPackage, FString& InOutName, bool Create, bool Thro
 			}
 			InPackage = NewPackage;
 		}
-		else if (!FPackageName::IsShortPackageName(PartialName))
+		else if (!FPackageName::IsShortPackageName(FStringView(PartialName)))
 		{
 			// Try to find the package in memory first, should be faster than attempting to load or create
 			InPackage = InPackage ? nullptr : StaticFindObjectFast(UPackage::StaticClass(), InPackage, *PartialName);
