Add GC.DumpTokenStats command

#rnx

diff --git a/Engine/Source/Runtime/CoreUObject/Private/UObject/GarbageCollection.cpp b/Engine/Source/Runtime/CoreUObject/Private/UObject/GarbageCollection.cpp
--- a/Engine/Source/Runtime/CoreUObject/Private/UObject/GarbageCollection.cpp
+++ b/Engine/Source/Runtime/CoreUObject/Private/UObject/GarbageCollection.cpp
@@ -5,6 +5,7 @@
 =============================================================================*/
 
 #include "UObject/GarbageCollection.h"
+#include "Algo/Find.h"
 #include "Containers/Deque.h"
 #include "Containers/StaticBitArray.h"
 #include "HAL/ThreadSafeBool.h"
@@ -108,6 +109,8 @@ static FAutoConsoleVariable CVarGCLockBehavior(
 	}),
 	ECVF_Default);
 
+
+
 /** Locks all UObject hash tables for performing GC reachability analysis.
  * Can be queried with IsGarbageCollectingAndLockingUObjectHashTables().
  * */
@@ -253,6 +256,312 @@ static FAutoConsoleVariableRef CGarbageReferenceTrackingEnabled(
 
 namespace UE::GC
 {
+
+namespace TokenDebugNames
+{
+#if ENABLE_GC_TOKEN_DEBUG_INFO
+#define DEFINE_TOKEN_DEBUG_ID(ID) static const FLazyName ID(#ID "Token")
+#else
+#define DEFINE_TOKEN_DEBUG_ID(ID) static constexpr EName ID = EName::None;
+#endif
+	DEFINE_TOKEN_DEBUG_ID(SkipInfo);
+	DEFINE_TOKEN_DEBUG_ID(Stride);
+	DEFINE_TOKEN_DEBUG_ID(Count);
+	DEFINE_TOKEN_DEBUG_ID(Pointer);
+	DEFINE_TOKEN_DEBUG_ID(EndOfPointer);
+	DEFINE_TOKEN_DEBUG_ID(ARO);
+	DEFINE_TOKEN_DEBUG_ID(EndOfStream);
+}
+
+#if !UE_BUILD_SHIPPING
+	static const TCHAR* ToStatName(EGCReferenceType Type)
+	{
+		switch (Type)
+		{
+		case GCRT_Object:								return TEXT("Reference");
+		case GCRT_ExternalPackage:						return TEXT("ExternalPackage");
+		case GCRT_ArrayObject:							return TEXT("Array<Reference>");
+		case GCRT_ArrayStruct:							return TEXT("Array<Struct>");
+		case GCRT_FixedArray:							return TEXT("FixedArray");
+		case GCRT_AddStructReferencedObjects:			return TEXT("StructARO");
+		case GCRT_AddReferencedObjects:					return TEXT("ClassARO");
+		case GCRT_AddTMapReferencedObjects:				return TEXT("Map");
+		case GCRT_AddTSetReferencedObjects:				return TEXT("Set");
+		case GCRT_AddFieldPathReferencedObject:			return TEXT("FieldPath");
+		case GCRT_ArrayAddFieldPathReferencedObject:	return TEXT("Array<FieldPath>");
+		case GCRT_EndOfStream:							return TEXT("EndOfStream");
+		case GCRT_ArrayObjectFreezable:					return TEXT("FrozenArray<Reference>");
+		case GCRT_ArrayStructFreezable:					return TEXT("FrozenArray<Struct>");
+		case GCRT_DynamicallyTypedValue:				return TEXT("DynamicallyTypedValue");
+		case GCRT_SlowAddReferencedObjects:				return TEXT("SlowARO");
+		default:										return TEXT("!!ERROR!!");
+		}
+	}
+
+	class FTokenStreamStats
+	{
+	public:
+		uint32 Counters[GCRT_Count] = {};
+		uint32 NumTokens = 0;
+		uint32 NumPointers = 0;
+		uint32 MinOffset = 0;
+		uint32 MaxOffset = 0;
+		uint32 MaxJump = 0;
+
+		FTokenStreamStats()  {}
+		explicit FTokenStreamStats(const FTokenStreamOwner& Tokens)
+			: NumTokens(Tokens.Strong.NumTokens())
+		{
+			FIterator It{Tokens.Strong, Tokens};
+			CountTokensUntil(It, NumTokens);
+		}
+
+		FTokenStreamStats& operator+=(const FTokenStreamStats& In)
+		{
+			for (uint32 Idx = 0; Idx < GCRT_Count; ++Idx)
+			{
+				Counters[Idx] += In.Counters[Idx];
+			}
+			NumTokens += In.NumTokens;
+			NumPointers += In.NumPointers;
+			MinOffset = MinOffset == 0 ? In.MinOffset : FMath::Min(MinOffset, In.MinOffset);
+			MaxOffset = FMath::Max(MaxOffset, In.MaxOffset);
+			MaxJump = FMath::Max(MaxJump, In.MaxJump);
+
+			return *this;
+		}
+	
+	private:
+		struct FIterator
+		{
+			const FTokenStreamView Tokens;
+			const FTokenStreamOwner& Owner; // Only for debug name access
+			uint32 Idx = 0;
+		
+			FGCReferenceInfo ConsumeReferenceInfo()
+			{
+				return Tokens.AccessReferenceInfo(Idx++);
+			}
+			
+			void Skip(FName DebugName)
+			{
+				check(GetDebugName() == DebugName);
+				++Idx;
+			}
+
+			uint32 ConsumeSkipInfo()
+			{
+				check(GetDebugName() == TokenDebugNames::SkipInfo);
+				return Tokens.ReadSkipInfo(/* in-out */ Idx).SkipIndex;
+			}
+
+			void* ConsumePointer()
+			{
+				check(GetDebugName() == TokenDebugNames::Pointer);
+				void* Ptr = Tokens.ReadPointer(/* in-out */ Idx);
+				Skip(TokenDebugNames::EndOfPointer);
+
+				return Ptr;
+			}
+
+#if ENABLE_GC_TOKEN_DEBUG_INFO
+			FName GetDebugName() const { return Owner.GetTokenInfo(FTokenId(Idx)).Name;}
+#else
+			static constexpr EName GetDebugName() { return EName::None; }
+#endif
+		};
+
+		void CountTokensUntil(FIterator& It, uint32 EndIdx)
+		{
+			static constexpr EGCReferenceType StrongFieldTokens[] = { 
+				GCRT_Object, GCRT_ArrayObject, GCRT_ArrayStruct, GCRT_AddTMapReferencedObjects, GCRT_AddTSetReferencedObjects, GCRT_FixedArray,
+				GCRT_AddFieldPathReferencedObject, GCRT_ArrayAddFieldPathReferencedObject, GCRT_ArrayObjectFreezable, GCRT_ArrayStructFreezable,
+				GCRT_Optional, GCRT_DynamicallyTypedValue };
+
+			const bool bObjectTokens = EndIdx == It.Tokens.NumTokens();
+			uint32 LastOffset = 0;
+			while (It.Idx < EndIdx)
+			{
+				FGCReferenceInfo Token = It.ConsumeReferenceInfo();
+				++Counters[Token.Type];
+
+				if (bObjectTokens && Algo::Find(StrongFieldTokens, Token.Type))
+				{
+					check(Token.Offset);
+					MinOffset = MinOffset == 0 ? Token.Offset : FMath::Min(MinOffset, Token.Offset);
+					MaxOffset = FMath::Max(Token.Offset, MaxOffset);
+					uint32 Jump = (uint32)FMath::Abs((int32)(Token.Offset - LastOffset));
+					LastOffset = Token.Offset;
+					MaxJump = FMath::Max(MaxJump, Jump);
+				}
+
+				// Traverse structural tokens
+				switch (Token.Type)
+				{
+					case GCRT_Optional:
+					case GCRT_ArrayStructFreezable:
+					case GCRT_ArrayStruct:
+						It.Skip(TokenDebugNames::Stride);
+						CountTokensUntil(It, It.ConsumeSkipInfo());
+						break;
+					case GCRT_FixedArray:
+						It.Skip(TokenDebugNames::Stride);
+						It.Skip(TokenDebugNames::Count);
+						break;
+					case GCRT_AddTSetReferencedObjects:
+					case GCRT_AddTMapReferencedObjects:
+					{
+						FProperty* Property = (FProperty*)It.ConsumePointer();
+						check(Property->IsA<FMapProperty>() || Property->IsA<FSetProperty>());
+						++NumPointers;
+						CountTokensUntil(It, It.ConsumeSkipInfo());
+						break;	
+					}
+					case GCRT_EndOfPointer:
+						checkf(false, TEXT("ConsumePointer() should consume all GCRT_EndOfPointer tokens"));
+						break;
+					case GCRT_AddStructReferencedObjects:
+						It.ConsumePointer();
+						++NumPointers;
+						break;
+					case GCRT_AddReferencedObjects:
+						It.ConsumePointer();
+						++NumPointers;
+						[[fallthrough]];
+					case GCRT_SlowAddReferencedObjects:
+						It.Skip(TokenDebugNames::EndOfStream);
+						[[fallthrough]];
+					case GCRT_EndOfStream:
+						check(It.Idx == EndIdx);
+						break;
+					default: // Non-structural leaf field
+						break;
+				}
+			}
+		} 
+	};
+
+	using FLog2Distribution = uint32[32];
+	enum class EDumpUnit {Bytes, Number};
+
+	static void DumpDistribution(FOutputDevice& Out, const TCHAR* Title, const FLog2Distribution& Bins, uint32 Total, EDumpUnit Unit)
+	{
+		uint32 End = 32;
+		for (; End > 0 && Bins[End - 1] == 0; --End);
+
+		uint32 Begin = 0;
+		for (; Begin < End && Bins[Begin] == 0; ++Begin);
+
+		Out.Logf(TEXT(" --- %-24s ---"), Title);
+
+		for (uint32 Idx = Begin, Sum = 0; Idx < End; ++Idx)
+		{
+			uint32 Num = Bins[Idx];
+			Sum += Num;
+			uint32 Max = Idx == 0 ? 0 : 1u << (Idx - 1);
+			float Percentage = (100.f * Num) / Total;
+			float Cumulative = (100.f * Sum) / Total;
+			if (Unit == EDumpUnit::Number)
+			{
+				Out.Logf(TEXT(" %4d %5.1f%% %5.1f%%  %d"), Max, Percentage, Cumulative, Num);
+			}
+			else
+			{
+				bool bK = Idx > 10;
+				Out.Logf(TEXT(" %3d%c %5.1f%% %5.1f%%  %d"), Max >> bK*10, "BK"[bK], Percentage, Cumulative, Num);
+			}
+		}
+	}
+
+	static void DumpTokenStats(FOutputDevice& Out, const FTokenStreamStats& Stats)
+	{
+		static constexpr TCHAR FormatStr[] = TEXT(" %18s %5.1f%% %5.1f%%  %d");
+
+		uint32 Sum = 0;
+		for (uint32 Idx = 0; Idx < GCRT_Count; ++Idx)
+		{
+			if (uint32 Num = Stats.Counters[Idx])
+			{
+				Sum += Num;
+				float Percentage = (100.f * Num) / Stats.NumTokens;
+				float Cumulative = (100.f * Sum) / Stats.NumTokens;
+				Out.Logf(FormatStr, ToStatName((EGCReferenceType)Idx), Percentage, Cumulative, Num);
+			}
+		}
+
+		uint32 Num = Stats.NumPointers * 3; // EndOfPointer + 2 32-bit tokens for 64-bit pointer
+		Sum += Num;
+		float Percentage = (100.f * Num) / Stats.NumTokens;
+		float Cumulative = (100.f * Sum) / Stats.NumTokens;
+		Out.Logf(FormatStr, TEXT("Pointer tokens"), Percentage, Cumulative, Num);
+		Out.Logf(FormatStr, TEXT("Skip/Stride/Count"), 100 - Cumulative, 100.f, Stats.NumTokens - Sum);
+	}
+	
+	static void DumpSchemaStats(FOutputDevice& Out)
+	{
+		TMap<void*, FTokenStreamStats> AllStats;
+		uint32 NumClasses = 0;
+		uint32 NumTokens = 0;
+		for (TObjectIterator<UClass> It; It; ++It)
+		{
+			FTokenStreamOwner& Tokens = It->ReferenceTokens;
+			FTokenStreamStats& Stats = AllStats.FindOrAdd(Tokens.Strong.GetTokenData());
+			if (Stats.NumTokens == 0)
+			{
+				Stats = FTokenStreamStats(Tokens);
+				NumTokens += Stats.NumTokens;
+			}
+			++NumClasses;
+		}
+		
+		FTokenStreamStats Sum;
+		FLog2Distribution SizeBins = {};
+		FLog2Distribution MinOffsetBins = {};
+		FLog2Distribution OffsetRangeBins = {};
+		FLog2Distribution MaxJumpBins = {};
+		FLog2Distribution NumTokenBins = {};
+		FLog2Distribution TokenBins[GCRT_Count] = {};
+		uint32 NumObjects = 0;
+		for (TObjectIterator<UObject> It; It; ++It)
+		{
+			const uint32 Size = It->GetClass()->GetPropertiesSize();
+			const FTokenStreamStats& Stats = AllStats.FindChecked(It->GetClass()->ReferenceTokens.Strong.GetTokenData());
+			Sum += Stats;
+			++SizeBins[FMath::CeilLogTwo(Size)];
+			++MinOffsetBins[FMath::CeilLogTwo(Stats.MinOffset)];
+			++OffsetRangeBins[FMath::CeilLogTwo(Stats.MaxOffset - Stats.MinOffset)];
+			++MaxJumpBins[FMath::CeilLogTwo(Stats.MaxJump)];
+			++NumTokenBins[FMath::CeilLogTwo(Stats.NumTokens)];
+			++NumObjects;
+
+			for (uint32 Idx = 0; Idx < GCRT_Count; ++Idx)
+			{
+				 ++TokenBins[Idx][FMath::CeilLogTwo(Stats.Counters[Idx])];
+			}
+		}
+
+		Out.Logf(TEXT("-------------------- GC strong token statistics --------------------"));
+		Out.Logf(TEXT(" %d token streams shared by %d classes contain %d tokens"), AllStats.Num(), NumClasses, NumTokens);
+		Out.Logf(TEXT(" Sweeping %d objects will process %d tokens"), NumObjects, Sum.NumTokens);
+		Out.Logf(TEXT("--------------------------- Distributions --------------------------"));
+		DumpDistribution(Out, TEXT("Object size"), SizeBins, NumObjects, EDumpUnit::Bytes);
+		DumpDistribution(Out, TEXT("Min offset"),	MinOffsetBins, NumObjects, EDumpUnit::Bytes);
+		DumpDistribution(Out, TEXT("Offset range"),	OffsetRangeBins, NumObjects, EDumpUnit::Bytes);
+		DumpDistribution(Out, TEXT("Max jump"),	MaxJumpBins, NumObjects, EDumpUnit::Bytes);
+		DumpDistribution(Out, TEXT("Tokens"), NumTokenBins, NumObjects, EDumpUnit::Number);
+		for (EGCReferenceType TokenType : {	GCRT_Object, GCRT_ArrayObject, GCRT_ArrayStruct, GCRT_AddTMapReferencedObjects, GCRT_AddTSetReferencedObjects, GCRT_FixedArray })
+		{
+			DumpDistribution(Out, ToStatName(TokenType), TokenBins[TokenType], NumObjects, EDumpUnit::Number);
+		}
+		Out.Logf(TEXT("------------------------ Tokens to process -------------------------"));
+		DumpTokenStats(Out, Sum);
+		Out.Logf(TEXT("--------------------------------------------------------------------"));
+	}
+
+	static FAutoConsoleCommandWithOutputDevice GDumpSchemaStats(TEXT("gc.DumpSchemaStats"), TEXT("Dump GC token stream statistics"), FConsoleCommandWithOutputDeviceDelegate::CreateStatic(&DumpSchemaStats));
+
+#endif // !UE_BUILD_SHIPPING
+
 	struct FStats
 	{
 #if PERF_DETAILED_PER_CLASS_GC_STATS
@@ -5028,11 +5337,7 @@ static void UpdateStreamReference(FTokenStreamView& Out, FTokenStreamView New)
 	}
 }
 
-struct  
-{
-	FGCReferenceInfo Token{GCRT_EndOfStream, 0};
-	FName DebugName{"EndOfStreamToken"};
-} GEndOfStream;
+static FGCReferenceInfo GEndOfStreamToken{GCRT_EndOfStream, 0};
 
 uint32 UE::GC::FTokenStreamBuilder::MaxStackSize;
 
@@ -5223,21 +5528,20 @@ void FTokenStreamBuilder::EmitFinalTokens(AROFunc ARO)
 {
 	if (ARO)
 	{
-		static const FName TokenName("AROToken");
 		int32 SlowIndex = FindSlowImplementation(ARO);
 		if (SlowIndex == INDEX_NONE)
 		{
-			EmitReferenceInfo(FGCReferenceInfo(GCRT_AddReferencedObjects, 0), TokenName);
+			EmitReferenceInfo(FGCReferenceInfo(GCRT_AddReferencedObjects, 0), TokenDebugNames::ARO);
 			EmitPointer(reinterpret_cast<const void*>(ARO));
 		}
 		else
 		{
-			EmitReferenceInfo(FGCReferenceInfo(GCRT_SlowAddReferencedObjects, static_cast<uint32>(SlowIndex)), TokenName);
+			EmitReferenceInfo(FGCReferenceInfo(GCRT_SlowAddReferencedObjects, static_cast<uint32>(SlowIndex)), TokenDebugNames::ARO);
 		}
 	}
 
 	// Emit end of stream token.
-	EmitReferenceInfo(GEndOfStream.Token, GEndOfStream.DebugName);
+	EmitReferenceInfo(GEndOfStreamToken, TokenDebugNames::EndOfStream);
 }
 
  FTokenStreamView FTokenStreamBuilder::DropFinalTokens(FTokenStreamView Stream, AROFunc DropARO)
@@ -5280,7 +5584,7 @@ void FTokenStreamBuilder::Merge(FTokenStreamView& Out, const FTokenStreamBuilder
 	if (Num == 1)
 	{
 		// Return empty stream if there's only an EndOfStream terminator	
-		check(Super.IsEmpty() && Class.Tokens[0] == GEndOfStream.Token);
+		check(Super.IsEmpty() && Class.Tokens[0] == GEndOfStreamToken);
 		UpdateStreamReference(Out, FTokenStreamView());
 	}
 	else if (Num == Out.Num &&
@@ -5371,9 +5675,8 @@ uint32 FTokenStreamBuilder::EmitSkipIndexPlaceholder()
 {
 	uint32 TokenIndex = Tokens.Add(SkipIndexPlaceholderMagic);
 #if ENABLE_GC_TOKEN_DEBUG_INFO
-	static const FName TokenName("SkipIndexPlaceholder");
 	check(DebugNames.Num() == TokenIndex);
-	DebugNames.Add(TokenName);
+	DebugNames.Add(TokenDebugNames::SkipInfo);
 #endif
 	return TokenIndex;
 }
@@ -5397,9 +5700,8 @@ int32 FTokenStreamBuilder::EmitCount( uint32 Count )
 {
 	int32 TokenIndex = Tokens.Add( Count );
 #if ENABLE_GC_TOKEN_DEBUG_INFO
-	static const FName TokenName("CountToken");
 	check(DebugNames.Num() == TokenIndex);
-	DebugNames.Add(TokenName);
+	DebugNames.Add(TokenDebugNames::Count);
 #endif
 	return TokenIndex;
 }
@@ -5413,18 +5715,16 @@ int32 FTokenStreamBuilder::EmitPointer( void const* Ptr )
 	StorePointer(&Tokens[StoreIndex], Ptr);
 
 #if ENABLE_GC_TOKEN_DEBUG_INFO
-	static const FName TokenName("PointerToken");
 	check(DebugNames.Num() == StoreIndex);
 	for (int32 PointerTokenIndex = 0; PointerTokenIndex < NumTokensPerPointer; ++PointerTokenIndex)
 	{
-		DebugNames.Add(TokenName);
+		DebugNames.Add(TokenDebugNames::Pointer);
 	}
 #endif
 
 	// Now inser the end of pointer marker, this will mostly be used for storing ReturnCount value
 	// if the pointer was stored at the end of struct array stream.
-	static const FName EndOfPointerTokenName("EndOfPointerToken");
-	EmitReferenceInfo(FGCReferenceInfo(GCRT_EndOfPointer, 0), EndOfPointerTokenName);
+	EmitReferenceInfo(FGCReferenceInfo(GCRT_EndOfPointer, 0), TokenDebugNames::EndOfPointer);
 
 	return StoreIndex;
 }
@@ -5433,9 +5733,8 @@ int32 FTokenStreamBuilder::EmitStride( uint32 Stride )
 {
 	int32 TokenIndex = Tokens.Add( Stride );
 #if ENABLE_GC_TOKEN_DEBUG_INFO
-	static const FName TokenName("StrideToken");
 	check(DebugNames.Num() == TokenIndex);
-	DebugNames.Add(TokenName);
+	DebugNames.Add(TokenDebugNames::Stride);
 #endif
 	return TokenIndex;
 }
diff --git a/Engine/Source/Runtime/CoreUObject/Public/UObject/GarbageCollection.h b/Engine/Source/Runtime/CoreUObject/Public/UObject/GarbageCollection.h
--- a/Engine/Source/Runtime/CoreUObject/Public/UObject/GarbageCollection.h
+++ b/Engine/Source/Runtime/CoreUObject/Public/UObject/GarbageCollection.h
@@ -88,6 +88,7 @@ enum EGCReferenceType
 	GCRT_Optional,							// TOptional
 	GCRT_DynamicallyTypedValue,				// FDynamicallyTypedValue
 	GCRT_SlowAddReferencedObjects,			// AddReferencedObjects() function using UE::GC::RegisterSlowImplementation
+	GCRT_Count
 };
 
 enum class EGCTokenType : uint32
@@ -222,7 +223,7 @@ public:
 		return Tokens[CurrentIndex++];
 	}
 
-	FORCEINLINE void* ReadPointer(uint32& CurrentIndex)
+	FORCEINLINE void* ReadPointer(uint32& CurrentIndex) const
 	{
 		UPTRINT Result = UPTRINT(Tokens[CurrentIndex]) | (UPTRINT(Tokens[CurrentIndex+1]) << 32);
 		CurrentIndex += 2;
