Silence clang static analysis warning

#horde 269934
#rnx

diff --git a/Engine/Source/Runtime/CoreUObject/Private/UObject/GarbageCollection.cpp b/Engine/Source/Runtime/CoreUObject/Private/UObject/GarbageCollection.cpp
--- a/Engine/Source/Runtime/CoreUObject/Private/UObject/GarbageCollection.cpp
+++ b/Engine/Source/Runtime/CoreUObject/Private/UObject/GarbageCollection.cpp
@@ -467,7 +467,7 @@ namespace TokenDebugNames
 			}
 			else
 			{
-				bool bK = Idx > 10;
+				uint32 bK = Idx > 10;
 				Out.Logf(TEXT(" %3d%c %5.1f%% %5.1f%%  %d"), Max >> bK*10, "BK"[bK], Percentage, Cumulative, Num);
 			}
 		}
