Fixed cyclic header dependency. UnrealType.h was indirectly including itself via CoreNetTypes.h.

diff --git a/Engine/Source/Runtime/CoreUObject/Private/UObject/Property.cpp b/Engine/Source/Runtime/CoreUObject/Private/UObject/Property.cpp
--- a/Engine/Source/Runtime/CoreUObject/Private/UObject/Property.cpp
+++ b/Engine/Source/Runtime/CoreUObject/Private/UObject/Property.cpp
@@ -11,6 +11,7 @@
 #include "Misc/StringBuilder.h"
 #include "Math/RandomStream.h"
 #include "Logging/LogScopedCategoryAndVerbosityOverride.h"
+#include "UObject/CoreNetTypes.h"
 #include "UObject/ObjectMacros.h"
 #include "UObject/UObjectGlobals.h"
 #include "UObject/Class.h"
diff --git a/Engine/Source/Runtime/CoreUObject/Public/UObject/UnrealType.h b/Engine/Source/Runtime/CoreUObject/Public/UObject/UnrealType.h
--- a/Engine/Source/Runtime/CoreUObject/Public/UObject/UnrealType.h
+++ b/Engine/Source/Runtime/CoreUObject/Public/UObject/UnrealType.h
@@ -48,7 +48,6 @@
 #include "Templates/UnrealTypeTraits.h"
 #include "Trace/Detail/Channel.h"
 #include "UObject/Class.h"
-#include "UObject/CoreNetTypes.h"
 #include "UObject/Field.h"
 #include "UObject/LazyObjectPtr.h"
 #include "UObject/NameTypes.h"
@@ -71,6 +70,7 @@ class FBlake3;
 class FOutputDevice;
 class UPackageMap;
 class UPropertyWrapper;
+enum ELifetimeCondition : int;
 struct CGetTypeHashable;
 struct FUObjectSerializeContext;
 template <typename FuncType> class TFunctionRef;
diff --git a/Engine/Source/Runtime/Engine/Classes/Components/ActorComponent.h b/Engine/Source/Runtime/Engine/Classes/Components/ActorComponent.h
--- a/Engine/Source/Runtime/Engine/Classes/Components/ActorComponent.h
+++ b/Engine/Source/Runtime/Engine/Classes/Components/ActorComponent.h
@@ -8,6 +8,7 @@
 #include "Templates/SubclassOf.h"
 #include "Engine/EngineTypes.h"
 #include "Engine/EngineBaseTypes.h"
+#include "UObject/CoreNetTypes.h"
 #include "UObject/ScriptMacros.h"
 #if UE_ENABLE_INCLUDE_ORDER_DEPRECATED_IN_5_2
 #include "CoreMinimal.h"
diff --git a/Engine/Source/Runtime/Engine/Classes/Engine/Blueprint.h b/Engine/Source/Runtime/Engine/Classes/Engine/Blueprint.h
--- a/Engine/Source/Runtime/Engine/Classes/Engine/Blueprint.h
+++ b/Engine/Source/Runtime/Engine/Classes/Engine/Blueprint.h
@@ -292,11 +292,7 @@ struct FBPVariableDescription
 	UPROPERTY(EditAnywhere, Category=BPVariableDescription)
 	FString DefaultValue;
 
-	FBPVariableDescription()
-		: PropertyFlags(CPF_Edit)
-		, ReplicationCondition(ELifetimeCondition::COND_None)
-	{
-	}
+	ENGINE_API FBPVariableDescription();
 
 	/** Set a metadata value on the variable */
 	ENGINE_API void SetMetaData(FName Key, FString Value);
diff --git a/Engine/Source/Runtime/Engine/Private/Blueprint.cpp b/Engine/Source/Runtime/Engine/Private/Blueprint.cpp
--- a/Engine/Source/Runtime/Engine/Private/Blueprint.cpp
+++ b/Engine/Source/Runtime/Engine/Private/Blueprint.cpp
@@ -119,6 +119,12 @@ void UBlueprint::ConformNativeComponents()
 //////////////////////////////////////////////////////////////////////////
 // FBPVariableDescription
 
+FBPVariableDescription::FBPVariableDescription()
+	: PropertyFlags(CPF_Edit)
+	, ReplicationCondition(ELifetimeCondition::COND_None)
+{
+}
+
 int32 FBPVariableDescription::FindMetaDataEntryIndexForKey(const FName Key) const
 {
 	for(int32 i=0; i<MetaDataArray.Num(); i++)
diff --git a/Engine/Source/Runtime/Engine/Public/ReplayTypes.h b/Engine/Source/Runtime/Engine/Public/ReplayTypes.h
--- a/Engine/Source/Runtime/Engine/Public/ReplayTypes.h
+++ b/Engine/Source/Runtime/Engine/Public/ReplayTypes.h
@@ -19,6 +19,8 @@
 ENGINE_API DECLARE_LOG_CATEGORY_EXTERN(LogDemo, Log, All);
 
 class UNetConnection;
+enum class EChannelCloseReason : uint8;
+enum ELifetimeCondition : int;
 
 enum class EReplayHeaderFlags : uint32
 {
diff --git a/Engine/Source/Runtime/Experimental/Iris/Core/Private/Iris/ReplicationSystem/ReplicationOperations.cpp b/Engine/Source/Runtime/Experimental/Iris/Core/Private/Iris/ReplicationSystem/ReplicationOperations.cpp
--- a/Engine/Source/Runtime/Experimental/Iris/Core/Private/Iris/ReplicationSystem/ReplicationOperations.cpp
+++ b/Engine/Source/Runtime/Experimental/Iris/Core/Private/Iris/ReplicationSystem/ReplicationOperations.cpp
@@ -24,6 +24,7 @@
 #include "Net/Core/Trace/NetDebugName.h"
 #include "DequantizeAndApplyHelper.h"
 #include "HAL/IConsoleManager.h"
+#include "UObject/CoreNetTypes.h"
 
 namespace UE::Net::Private
 {
