ZenLoader in Editor: Respect the LoadObject flags LOAD_Quiet and LOAD_Warn flags to silence warnings about missing packages in this scenario.
Always silence warning about missing package for script packages, just like GetPackageLinker() simply returns nullptr for PKG_InMemoryOnly.
submitted on behalf of [REDACTED_HANDLE]

diff --git a/Engine/Source/Runtime/CoreUObject/Private/Serialization/AsyncLoading.cpp b/Engine/Source/Runtime/CoreUObject/Private/Serialization/AsyncLoading.cpp
--- a/Engine/Source/Runtime/CoreUObject/Private/Serialization/AsyncLoading.cpp
+++ b/Engine/Source/Runtime/CoreUObject/Private/Serialization/AsyncLoading.cpp
@@ -6567,7 +6567,7 @@ EAsyncPackageState::Type FAsyncPackage::LoadImports()
 						UE_LOG(LogStreaming, Verbose, TEXT("FAsyncPackage::LoadImports for %s: Loading %s (0x%llX) using IoStore package loader"), *Desc.PackagePath.GetDebugName(), *ImportPackageName, FPackageId::FromName(ImportPackageFName).ValueForDebugging());
 						check(IsInGameThread());
 						FAsyncLoadingThread::LeaveAsyncLoadingTick(AsyncLoadingThread.GetThreadIndex()); // Fix thread check that triggers if we recurse back into the uncooked package loader from here
-						int32 ImportRequestId = IoStorePackageLoader->LoadPackage(ImportedPackagePath, NAME_None, FLoadPackageAsyncDelegate(), PKG_None, INDEX_NONE, 0, nullptr);
+						int32 ImportRequestId = IoStorePackageLoader->LoadPackage(ImportedPackagePath, NAME_None, FLoadPackageAsyncDelegate(), PKG_None, INDEX_NONE, 0, nullptr, LOAD_None);
 						IoStorePackageLoader->FlushLoading(ImportRequestId);
 						FAsyncLoadingThread::EnterAsyncLoadingTick(AsyncLoadingThread.GetThreadIndex());
 						bLoadedFromIoStore = true;
@@ -7401,7 +7401,7 @@ bool FAsyncLoadingThread::ShouldAlwaysLoadPackageAsync(const FPackagePath& InPac
 	return FPlatformProperties::RequiresCookedData() && GEventDrivenLoaderEnabled && EVENT_DRIVEN_ASYNC_LOAD_ACTIVE_AT_RUNTIME;
 }
 
-int32 FAsyncLoadingThread::LoadPackage(const FPackagePath& InPackagePath, FName InCustomName, FLoadPackageAsyncDelegate InCompletionDelegate, EPackageFlags InPackageFlags, int32 InPIEInstanceID, int32 InPackagePriority, const FLinkerInstancingContext* InstancingContext)
+int32 FAsyncLoadingThread::LoadPackage(const FPackagePath& InPackagePath, FName InCustomName, FLoadPackageAsyncDelegate InCompletionDelegate, EPackageFlags InPackageFlags, int32 InPIEInstanceID, int32 InPackagePriority, const FLinkerInstancingContext* InInstancingContext, uint32 InLoadFlags)
 {
 	static bool bOnce = false;
 	if (!bOnce && GEventDrivenLoaderEnabled)
@@ -7440,9 +7440,9 @@ int32 FAsyncLoadingThread::LoadPackage(const FPackagePath& InPackagePath, FName
 
 	// Add new package request
 	FAsyncPackageDesc PackageDesc(RequestID, PackageName, InPackagePath, MoveTemp(CompletionDelegatePtr), InPackageFlags, InPIEInstanceID, InPackagePriority);
-	if (InstancingContext)
+	if (InInstancingContext)
 	{
-		PackageDesc.SetInstancingContext(*InstancingContext);
+		PackageDesc.SetInstancingContext(*InInstancingContext);
 	}
 	QueuePackage(PackageDesc);
 	return RequestID;
diff --git a/Engine/Source/Runtime/CoreUObject/Private/Serialization/AsyncLoading2.cpp b/Engine/Source/Runtime/CoreUObject/Private/Serialization/AsyncLoading2.cpp
--- a/Engine/Source/Runtime/CoreUObject/Private/Serialization/AsyncLoading2.cpp
+++ b/Engine/Source/Runtime/CoreUObject/Private/Serialization/AsyncLoading2.cpp
@@ -425,6 +425,7 @@ struct FPackageRequest
 	int32 Priority = -1;
 	EPackageFlags PackageFlags = PKG_None;
 #if WITH_EDITOR
+	uint32 LoadFlags = LOAD_None;
 	int32 PIEInstanceID = INDEX_NONE;
 #endif
 #if WITH_EDITORONLY_DATA
@@ -444,7 +445,7 @@ struct FPackageRequest
 #endif
 	}
 
-	static FPackageRequest Create(int32 RequestId, EPackageFlags PackageFlags, int32 PIEInstanceID, int32 Priority, const FLinkerInstancingContext* InstancingContext, const FPackagePath& PackagePath, FName CustomName, TUniquePtr<FLoadPackageAsyncDelegate> PackageLoadedDelegate)
+	static FPackageRequest Create(int32 RequestId, EPackageFlags PackageFlags, uint32 LoadFlags, int32 PIEInstanceID, int32 Priority, const FLinkerInstancingContext* InstancingContext, const FPackagePath& PackagePath, FName CustomName, TUniquePtr<FLoadPackageAsyncDelegate> PackageLoadedDelegate)
 	{
 		return FPackageRequest
 		{
@@ -452,6 +453,7 @@ struct FPackageRequest
 			Priority,
 			PackageFlags,
 #if WITH_EDITOR
+			LoadFlags,
 			PIEInstanceID,
 #endif
 #if WITH_EDITORONLY_DATA
@@ -474,6 +476,7 @@ struct FAsyncPackageDesc2
 	/** The flags that should be applied to the package */
 	EPackageFlags PackageFlags;
 #if WITH_EDITOR
+	uint32 LoadFlags;
 	/** PIE instance ID this package belongs to, INDEX_NONE otherwise */
 	int32 PIEInstanceID;
 #endif
@@ -502,6 +505,7 @@ struct FAsyncPackageDesc2
 			Request.Priority,
 			Request.PackageFlags,
 #if WITH_EDITOR
+			Request.LoadFlags,
 			Request.PIEInstanceID,
 #endif
 			FPackageId::FromName(UPackageName),
@@ -529,6 +533,7 @@ struct FAsyncPackageDesc2
 			ImportingPackageDesc.Priority,
 			PKG_None,
 #if WITH_EDITOR
+			LOAD_None,
 			INDEX_NONE,
 #endif
 			ImportedPackageId,
@@ -3101,7 +3106,8 @@ public:
 		EPackageFlags InPackageFlags,
 		int32 InPIEInstanceID,
 		int32 InPackagePriority,
-		const FLinkerInstancingContext* InstancingContext = nullptr) override;
+		const FLinkerInstancingContext* InstancingContext = nullptr,
+		uint32 InLoadFlags = LOAD_None) override;
 
 	EAsyncPackageState::Type ProcessLoadingFromGameThread(FAsyncLoadingThreadState2& ThreadState, bool bUseTimeLimit, bool bUseFullTimeLimit, double TimeLimit);
 
@@ -4542,7 +4548,7 @@ void FAsyncPackage2::ImportPackagesRecursiveInner(FAsyncLoadingThreadState2& Thr
 				IoBatch.Issue(); // The batch might already contain requests for packages being imported from the uncooked one we're going to load so make sure that those are started before blocking
 				ImportedPackagePath = FPackagePath::FromPackageNameUnchecked(ImportedPackageEntry.UncookedPackageName);
 				ImportedPackagePath.SetHeaderExtension(static_cast<EPackageExtension>(ImportedPackageEntry.UncookedPackageHeaderExtension));
-				int32 ImportRequestId = AsyncLoadingThread.UncookedPackageLoader->LoadPackage(ImportedPackagePath, NAME_None, FLoadPackageAsyncDelegate(), PKG_None, INDEX_NONE, 0, nullptr);
+				int32 ImportRequestId = AsyncLoadingThread.UncookedPackageLoader->LoadPackage(ImportedPackagePath, NAME_None, FLoadPackageAsyncDelegate(), PKG_None, INDEX_NONE, 0, nullptr, LOAD_None);
 				AsyncLoadingThread.UncookedPackageLoader->FlushLoading(ImportRequestId);
 				UncookedPackage = FindObjectFast<UPackage>(nullptr, ImportedPackagePath.GetPackageFName());
 				ImportedPackageRef.SetPackage(UncookedPackage);
@@ -8079,7 +8085,7 @@ void FAsyncPackage2::AddCompletionCallback(TUniquePtr<FLoadPackageAsyncDelegate>
 	CompletionCallbacks.Emplace(MoveTemp(Callback));
 }
 
-int32 FAsyncLoadingThread2::LoadPackage(const FPackagePath& InPackagePath, FName InCustomName, FLoadPackageAsyncDelegate InCompletionDelegate, EPackageFlags InPackageFlags, int32 InPIEInstanceID, int32 InPackagePriority, const FLinkerInstancingContext* InInstancingContext)
+int32 FAsyncLoadingThread2::LoadPackage(const FPackagePath& InPackagePath, FName InCustomName, FLoadPackageAsyncDelegate InCompletionDelegate, EPackageFlags InPackageFlags, int32 InPIEInstanceID, int32 InPackagePriority, const FLinkerInstancingContext* InInstancingContext, uint32 InLoadFlags)
 {
 	TRACE_CPUPROFILER_EVENT_SCOPE(LoadPackage);
 
@@ -8105,7 +8111,7 @@ int32 FAsyncLoadingThread2::LoadPackage(const FPackagePath& InPackagePath, FName
 		? MakeUnique<FLoadPackageAsyncDelegate>(MoveTemp(InCompletionDelegate))
 		: TUniquePtr<FLoadPackageAsyncDelegate>();
 
-	PackageRequestQueue.Enqueue(FPackageRequest::Create(RequestId, InPackageFlags, InPIEInstanceID, InPackagePriority, InInstancingContext, InPackagePath, InCustomName, MoveTemp(CompletionDelegate)));
+	PackageRequestQueue.Enqueue(FPackageRequest::Create(RequestId, InPackageFlags, InLoadFlags, InPIEInstanceID, InPackagePriority, InInstancingContext, InPackagePath, InCustomName, MoveTemp(CompletionDelegate)));
 	++QueuedPackagesCounter;
 	++PackagesWithRemainingWorkCounter;
 
@@ -8126,13 +8132,20 @@ void FAsyncLoadingThread2::QueueMissingPackage(FAsyncPackageDesc2& PackageDesc,
 
 	SkippedPackages.Add(FailedPackageName, &bIsAlreadySkipped);
 
-	if (!bIsAlreadySkipped)
+	bool bIssueWarning = !bIsAlreadySkipped;
+#if WITH_EDITOR
+	bIssueWarning &= ((PackageDesc.LoadFlags & (LOAD_NoWarn | LOAD_Quiet)) == 0);
+#endif
+	if (bIssueWarning)
 	{
-		UE_LOG(LogStreaming, Warning,
-			TEXT("LoadPackage: SkipPackage: %s (0x%llX) - The package to load does not exist on disk or in the loader"),
-			*FailedPackageName.ToString(), PackageDesc.PackageIdToLoad.ValueForDebugging());
+		bool bIsScriptPackage = FPackageName::IsScriptPackage(WriteToString<FName::StringBufferSize>(FailedPackageName));
+		bIssueWarning &= !bIsScriptPackage;
 	}
 
+	UE_CLOG(bIssueWarning, LogStreaming, Warning,
+		TEXT("LoadPackage: SkipPackage: %s (0x%llX) - The package to load does not exist on disk or in the loader"),
+		*FailedPackageName.ToString(), PackageDesc.PackageIdToLoad.ValueForDebugging());
+
 	if (PackageLoadedDelegate.IsValid())
 	{
 		FScopeLock LockMissingPackages(&FailedPackageRequestsCritical);
diff --git a/Engine/Source/Runtime/CoreUObject/Private/Serialization/AsyncLoadingThread.h b/Engine/Source/Runtime/CoreUObject/Private/Serialization/AsyncLoadingThread.h
--- a/Engine/Source/Runtime/CoreUObject/Private/Serialization/AsyncLoadingThread.h
+++ b/Engine/Source/Runtime/CoreUObject/Private/Serialization/AsyncLoadingThread.h
@@ -326,7 +326,8 @@ public:
 			EPackageFlags InPackageFlags,
 			int32 InPIEInstanceID,
 			int32 InPackagePriority,
-			const FLinkerInstancingContext* InstancingContext) override;
+			const FLinkerInstancingContext* InInstancingContext,
+			uint32 InLoadFlags) override;
 
 	EAsyncPackageState::Type ProcessLoading(bool bUseTimeLimit, bool bUseFullTimeLimit, double TimeLimit) override;
 
diff --git a/Engine/Source/Runtime/CoreUObject/Private/Serialization/AsyncPackageLoader.cpp b/Engine/Source/Runtime/CoreUObject/Private/Serialization/AsyncPackageLoader.cpp
--- a/Engine/Source/Runtime/CoreUObject/Private/Serialization/AsyncPackageLoader.cpp
+++ b/Engine/Source/Runtime/CoreUObject/Private/Serialization/AsyncPackageLoader.cpp
@@ -560,20 +560,6 @@ static FPackagePath GetLoadPackageAsyncPackagePath(FStringView InPackageNameOrFi
 	return PackagePath;
 }
 
-int32 IAsyncPackageLoader::LoadPackage(
-	const FString& InPackageName,
-	const FGuid* InGuid,
-	const TCHAR* InPackageToLoadFrom,
-	FLoadPackageAsyncDelegate InCompletionDelegate,
-	EPackageFlags InPackageFlags,
-	int32 InPIEInstanceID,
-	int32 InPackagePriority,
-	const FLinkerInstancingContext* InstancingContext)
-{
-	FPackagePath PackagePath = GetLoadPackageAsyncPackagePath(InPackageToLoadFrom ? FStringView(InPackageToLoadFrom) : FStringView(InPackageName));
-	return LoadPackage(PackagePath, FName(InPackageName), InCompletionDelegate, InPackageFlags, InPIEInstanceID, InPackagePriority, InstancingContext);
-}
-
 bool ShouldAlwaysLoadPackageAsync(const FPackagePath& InPackagePath)
 {
 	if (!GPackageLoader)
@@ -583,14 +569,21 @@ bool ShouldAlwaysLoadPackageAsync(const FPackagePath& InPackagePath)
 	return GPackageLoader->ShouldAlwaysLoadPackageAsync(InPackagePath);
 }
 
-int32 LoadPackageAsync(const FPackagePath& InPackagePath, FName InPackageNameToCreate /* = NAME_None*/, FLoadPackageAsyncDelegate InCompletionDelegate /*= FLoadPackageAsyncDelegate()*/, EPackageFlags InPackageFlags /*= PKG_None*/, int32 InPIEInstanceID /*= INDEX_NONE*/, int32 InPackagePriority /*= 0*/, const FLinkerInstancingContext* InstancingContext /*=nullptr*/)
+int32 LoadPackageAsync(const FPackagePath& InPackagePath,
+		FName InPackageNameToCreate /* = NAME_None*/,
+		FLoadPackageAsyncDelegate InCompletionDelegate /*= FLoadPackageAsyncDelegate()*/,
+		EPackageFlags InPackageFlags /*= PKG_None*/,
+		int32 InPIEInstanceID /*= INDEX_NONE*/,
+		int32 InPackagePriority /*= 0*/,
+		const FLinkerInstancingContext* InstancingContext /*=nullptr*/,
+		uint32 LoadFlags /*=LOAD_None*/)
 {
 	LLM_SCOPE(ELLMTag::AsyncLoading);
 	UE_CLOG(!GAsyncLoadingAllowed && !IsInAsyncLoadingThread(), LogStreaming, Fatal, TEXT("Requesting async load of \"%s\" when async loading is not allowed (after shutdown). Please fix higher level code."), *InPackagePath.GetDebugName());
 #if DO_TRACK_ASYNC_LOAD_REQUESTS
 	FTrackAsyncLoadRequests::Get().TrackRequest(InPackagePath.GetDebugName(), nullptr, InPackagePriority);
 #endif
-	return GetAsyncPackageLoader().LoadPackage(InPackagePath, InPackageNameToCreate, InCompletionDelegate, InPackageFlags, InPIEInstanceID, InPackagePriority, InstancingContext);
+	return GetAsyncPackageLoader().LoadPackage(InPackagePath, InPackageNameToCreate, InCompletionDelegate, InPackageFlags, InPIEInstanceID, InPackagePriority, InstancingContext, LoadFlags);
 }
 
 int32 LoadPackageAsync(const FString& InName, const FGuid* InGuid /* nullptr*/)
diff --git a/Engine/Source/Runtime/CoreUObject/Private/Serialization/EditorPackageLoader.cpp b/Engine/Source/Runtime/CoreUObject/Private/Serialization/EditorPackageLoader.cpp
--- a/Engine/Source/Runtime/CoreUObject/Private/Serialization/EditorPackageLoader.cpp
+++ b/Engine/Source/Runtime/CoreUObject/Private/Serialization/EditorPackageLoader.cpp
@@ -68,19 +68,20 @@ public:
 		EPackageFlags InPackageFlags,
 		int32 InPIEInstanceID,
 		int32 InPackagePriority,
-		const FLinkerInstancingContext* InstancingContext) override
+		const FLinkerInstancingContext* InInstancingContext,
+		uint32 InLoadFlags) override
 	{
 		// Use the old loader if an uncooked package exists on disk
 		if (!bHasInitializedCookedPackageLoader ||
 			FPackageName::DoesPackageExistEx(PackagePath, FPackageName::EPackageLocationFilter::FileSystem) != FPackageName::EPackageLocationFilter::None)
 		{
 			UE_LOG(LogEditorPackageLoader, Verbose, TEXT("Loading uncooked package '%s' from filesystem"), *PackagePath.GetDebugName());
-			return UncookedPackageLoader->LoadPackage(PackagePath, CustomPackageName, InCompletionDelegate, InPackageFlags, InPIEInstanceID, InPackagePriority, InstancingContext);
+			return UncookedPackageLoader->LoadPackage(PackagePath, CustomPackageName, InCompletionDelegate, InPackageFlags, InPIEInstanceID, InPackagePriority, InInstancingContext, InLoadFlags);
 		}
 		else
 		{
 			UE_LOG(LogEditorPackageLoader, Verbose, TEXT("Loading cooked package '%s' from I/O Store"), *PackagePath.GetDebugName());
-			return CookedPackageLoader->LoadPackage(PackagePath, CustomPackageName, InCompletionDelegate, InPackageFlags, InPIEInstanceID, InPackagePriority, InstancingContext);
+			return CookedPackageLoader->LoadPackage(PackagePath, CustomPackageName, InCompletionDelegate, InPackageFlags, InPIEInstanceID, InPackagePriority, InInstancingContext, InLoadFlags);
 		}
 	}
 
diff --git a/Engine/Source/Runtime/CoreUObject/Private/UObject/UObjectGlobals.cpp b/Engine/Source/Runtime/CoreUObject/Private/UObject/UObjectGlobals.cpp
--- a/Engine/Source/Runtime/CoreUObject/Private/UObject/UObjectGlobals.cpp
+++ b/Engine/Source/Runtime/CoreUObject/Private/UObject/UObjectGlobals.cpp
@@ -1575,7 +1575,7 @@ UPackage* LoadPackageInternal(UPackage* InOuter, const FPackagePath& PackagePath
 		}
 		constexpr int32 PIEInstanceID = INDEX_NONE;
 		constexpr int32 Priority = INT32_MAX;
-		int32 RequestID = LoadPackageAsync(PackagePath, PackageName, FLoadPackageAsyncDelegate(), PackageFlags, PIEInstanceID, Priority, InstancingContext);
+		int32 RequestID = LoadPackageAsync(PackagePath, PackageName, FLoadPackageAsyncDelegate(), PackageFlags, PIEInstanceID, Priority, InstancingContext, LoadFlags);
 
 		if (RequestID != INDEX_NONE)
 		{
diff --git a/Engine/Source/Runtime/CoreUObject/Public/Serialization/AsyncPackageLoader.h b/Engine/Source/Runtime/CoreUObject/Public/Serialization/AsyncPackageLoader.h
--- a/Engine/Source/Runtime/CoreUObject/Public/Serialization/AsyncPackageLoader.h
+++ b/Engine/Source/Runtime/CoreUObject/Public/Serialization/AsyncPackageLoader.h
@@ -100,26 +100,17 @@ public:
 
 	virtual bool ShouldAlwaysLoadPackageAsync(const FPackagePath& PackagePath) = 0;
 
-	UE_DEPRECATED(5.0, "Call the version that takes a FPackagePath without a FGuid instead")
-	COREUOBJECT_API int32 LoadPackage(
-		const FString& InPackageName,
-		const FGuid* InGuid,
-		const TCHAR* InPackageToLoadFrom,
-		FLoadPackageAsyncDelegate InCompletionDelegate,
-		EPackageFlags InPackageFlags,
-		int32 InPIEInstanceID,
-		int32 InPackagePriority,
-		const FLinkerInstancingContext* InstancingContext);
-
 	/**
 	 * Asynchronously load a package.
 	 *
-	 * @param	PackagePath		Name of package to load. The package is created if it does not already exist.
-	 * @param	CustomName				If not none, this is the name of the package to load into (and create if not yet existing). If none, the name is take from PackagePath.
+	 * @param	PackagePath				PackagePath to load. Must be a mounted path. The package is created if it does not already exist.
+	 * @param	CustomPackageName		If not none, this is the name of the package to load into (and create if not yet existing). If none, the name is take from PackagePath.
 	 * @param	InCompletionDelegate	Delegate to be invoked when the packages has finished streaming
 	 * @param	InPackageFlags			Package flags used to construct loaded package in memory
 	 * @param	InPIEInstanceID			Play in Editor instance ID
 	 * @param	InPackagePriority		Loading priority
+	 * @param   InstancingContext		Additional context to map object names to their instanced counterpart when loading an instanced package
+	 * @param	LoadFlags				Flags controlling loading behavior, from the ELoadFlags enum
 	 * @return Unique ID associated with this load request (the same package can be associated with multiple IDs).
 	 */
 	virtual int32 LoadPackage(
@@ -129,7 +120,8 @@ public:
 			EPackageFlags InPackageFlags,
 			int32 InPIEInstanceID,
 			int32 InPackagePriority,
-			const FLinkerInstancingContext* InstancingContext) = 0;
+			const FLinkerInstancingContext* InInstancingContext,
+			uint32 InLoadFlags) = 0;
 
 	/**
 	 * Process all currently loading package requests.
diff --git a/Engine/Source/Runtime/CoreUObject/Public/UObject/UObjectGlobals.h b/Engine/Source/Runtime/CoreUObject/Public/UObject/UObjectGlobals.h
--- a/Engine/Source/Runtime/CoreUObject/Public/UObject/UObjectGlobals.h
+++ b/Engine/Source/Runtime/CoreUObject/Public/UObject/UObjectGlobals.h
@@ -685,9 +685,10 @@ DECLARE_DELEGATE_ThreeParams(FLoadPackageAsyncDelegate, const FName& /*PackageNa
  * @param	InPIEInstanceID			Play in Editor instance ID
  * @param	InPackagePriority		Loading priority
  * @param   InstancingContext		Additional context to map object names to their instanced counterpart when loading an instanced package
+ * @param	LoadFlags				Flags controlling loading behavior, from the ELoadFlags enum
  * @return Unique ID associated with this load request (the same package can be associated with multiple IDs).
  */
-COREUOBJECT_API int32 LoadPackageAsync(const FPackagePath& InPackagePath, FName InPackageNameToCreate = NAME_None, FLoadPackageAsyncDelegate InCompletionDelegate = FLoadPackageAsyncDelegate(), EPackageFlags InPackageFlags = PKG_None, int32 InPIEInstanceID = INDEX_NONE, TAsyncLoadPriority InPackagePriority = 0, const FLinkerInstancingContext* InstancingContext = nullptr);
+COREUOBJECT_API int32 LoadPackageAsync(const FPackagePath& InPackagePath, FName InPackageNameToCreate = NAME_None, FLoadPackageAsyncDelegate InCompletionDelegate = FLoadPackageAsyncDelegate(), EPackageFlags InPackageFlags = PKG_None, int32 InPIEInstanceID = INDEX_NONE, TAsyncLoadPriority InPackagePriority = 0, const FLinkerInstancingContext* InstancingContext = nullptr, uint32 LoadFlags = LOAD_None);
 
 /**
  * Asynchronously load a package and all contained objects that match context flags. Non-blocking.
@@ -712,9 +713,6 @@ COREUOBJECT_API int32 LoadPackageAsync(const FString& InName, const FGuid* InGui
  */
 COREUOBJECT_API int32 LoadPackageAsync(const FString& InName, FLoadPackageAsyncDelegate InCompletionDelegate, TAsyncLoadPriority InPackagePriority = 0, EPackageFlags InPackageFlags = PKG_None, int32 InPIEInstanceID = INDEX_NONE);
 
-UE_DEPRECATED(5.0, "Use version that takes a FPackagePath without a FGuid")
-COREUOBJECT_API int32 LoadPackageAsync(const FString& InName, const FGuid* InGuid, const TCHAR* InPackageToLoadFrom, FLoadPackageAsyncDelegate InCompletionDelegate = FLoadPackageAsyncDelegate(), EPackageFlags InPackageFlags = PKG_None, int32 InPIEInstanceID = INDEX_NONE, TAsyncLoadPriority InPackagePriority = 0, const FLinkerInstancingContext* InstancingContext = nullptr);
-
 /**
 * Cancels all async package loading requests.
 */
