Allow the loaded objects to query the linker to know if there package was relocated.
Added a warning when loading a level with some external object or actors that was relocated because it is not properly supported and it might cause some data stumping when saving the external objects and actors.

#localization none

#tests Loaded some relocated level to validate the presence of the warning. Loaded some not relocated level to validate that the warning is not present. Run a preflight to validate that the change is not creating new cis issues.

diff --git a/Engine/Source/Runtime/CoreUObject/Private/UObject/LinkerLoad.cpp b/Engine/Source/Runtime/CoreUObject/Private/UObject/LinkerLoad.cpp
--- a/Engine/Source/Runtime/CoreUObject/Private/UObject/LinkerLoad.cpp
+++ b/Engine/Source/Runtime/CoreUObject/Private/UObject/LinkerLoad.cpp
@@ -985,11 +985,14 @@ FLinkerLoad::FLinkerLoad(UPackage* InParent, const FPackagePath& InPackagePath,
 , bUseTimeLimit(false)
 , bUseFullTimeLimit(false)
 , bLoaderNeedsEngineVersionChecks(true)
+#if WITH_EDITOR
+, bExportsDuplicatesFixed(false)
+, bIsPackageRelocated(false)
+#endif // WITH_EDITOR
 , IsTimeLimitExceededCallCount(0)
 , TimeLimit(0.0f)
 , TickStartTime(0.0)
 #if WITH_EDITOR
-, bExportsDuplicatesFixed(false)
 , LoadProgressScope( nullptr )
 #endif // WITH_EDITOR
 #if USE_CIRCULAR_DEPENDENCY_LOAD_DEFERRING
@@ -2117,8 +2120,11 @@ FLinkerLoad::ELinkerStatus FLinkerLoad::RelocateReferences()
 			UE_LOG(LogPackageRelocation, Verbose, TEXT("Loading relocated package (%s). The package was saved as (%s)."), *PackageNameToLoad, *Summary.PackageName);
 			UE::Package::Relocation::Private::ApplyRelocationToObjectImportMap(RelocationArgs, ImportMap);
 			UE::Package::Relocation::Private::ApplyRelocationToSoftObjectArray(RelocationArgs, SoftObjectPathList);
+
+			bIsPackageRelocated = true;
 		}
 #endif
+
 		// Avoid duplicate work in async case.
 		bHasRelocatedReferences = true;
 	}
diff --git a/Engine/Source/Runtime/CoreUObject/Public/UObject/LinkerLoad.h b/Engine/Source/Runtime/CoreUObject/Public/UObject/LinkerLoad.h
--- a/Engine/Source/Runtime/CoreUObject/Public/UObject/LinkerLoad.h
+++ b/Engine/Source/Runtime/CoreUObject/Public/UObject/LinkerLoad.h
@@ -200,6 +200,13 @@ public:
 		return PackageTrailer.Get();
 	}
 
+#if WITH_EDITOR
+	bool IsPackageRelocated() const
+	{
+		return bIsPackageRelocated;
+	}
+#endif
+
 private:
 
 	/** True if the linker is currently deleting loader */
@@ -450,17 +457,21 @@ private:
 	bool					bUseFullTimeLimit:1;
 	/** Whether the loader needs version and correctness checks (see OpenReadPackage)										*/
 	bool					bLoaderNeedsEngineVersionChecks : 1;
+
+#if WITH_EDITOR
+	/** Check to avoid multiple export duplicate fixups in case we don't save asset. */
+	bool bExportsDuplicatesFixed : 1;
+
+	/** Cache if the package is relocated or not */
+	bool bIsPackageRelocated : 1;
+#endif // WITH_EDITOR
+
 	/** Call count of IsTimeLimitExceeded.																					*/
 	int32					IsTimeLimitExceededCallCount;
 	/** Current time limit to use if bUseTimeLimit is true.																	*/
 	float					TimeLimit;
 	/** Time at begin of Tick function. Used for time limit determination.													*/
 	double					TickStartTime;
-
-#if WITH_EDITOR
-	/** Check to avoid multiple export duplicate fixups in case we don't save asset. */
-	bool bExportsDuplicatesFixed;
-#endif // WITH_EDITOR
 	/** Id of the thread that created this linker. This is to guard against using this linker on other threads than the one it was created on **/
 	int32					OwnerThread;
 
diff --git a/Engine/Source/Runtime/Engine/Private/Level.cpp b/Engine/Source/Runtime/Engine/Private/Level.cpp
--- a/Engine/Source/Runtime/Engine/Private/Level.cpp
+++ b/Engine/Source/Runtime/Engine/Private/Level.cpp
@@ -1086,6 +1086,28 @@ void ULevel::PostLoad()
 	}
 
 #if WITH_EDITOR
+	if ((IsUsingExternalActors() || IsUsingExternalObjects())
+		&& OwningWorld
+		&& OwningWorld->WorldType == EWorldType::Editor)
+	{
+		const FLinkerLoad* Linker = GetLinker();
+		if (Linker && Linker->IsPackageRelocated())
+		{
+			const FString LevelPackageName = GetPackage()->GetName();
+			UE_LOG(LogLevel, Error, TEXT("The level %s was moved on disk without the editor knowing it. This might cause some loading issues and some saving issue like an actor potentialy stomping the file of another actor."), *LevelPackageName);
+
+			if (GIsEditor)
+			{
+				FFormatNamedArguments Args;
+				Args.Add(TEXT("PackageName"), FText::FromString(LevelPackageName));
+				FNotificationInfo Info(FText::Format(LOCTEXT("UnsupportedRelocatedLevel", "The level ({PackageName}) was relocated.\nLevels that use one file per actor cannot be moved on disk without using the editor or the migration tool."), Args));
+				Info.ExpireDuration = 10.0f;
+
+				FSlateNotificationManager::Get().AddNotification(Info);
+			}
+		}
+	}
+
 	if (IsUsingActorFolders() && IsUsingExternalObjects() && IsActorFolderObjectsFeatureAvailable())
 	{
 		if (!bWasDuplicated)
