Do not broadcast OnAsyncLoadingFlush delegate when they are done in the context of a sync load when using the zen loader

diff --git a/Engine/Source/Runtime/CoreUObject/Private/Serialization/AsyncLoading2.cpp b/Engine/Source/Runtime/CoreUObject/Private/Serialization/AsyncLoading2.cpp
--- a/Engine/Source/Runtime/CoreUObject/Private/Serialization/AsyncLoading2.cpp
+++ b/Engine/Source/Runtime/CoreUObject/Private/Serialization/AsyncLoading2.cpp
@@ -8192,7 +8192,12 @@ void FAsyncLoadingThread2::FlushLoading(int32 RequestId)
 			return;
 		}
 
-		FCoreDelegates::OnAsyncLoadingFlush.Broadcast();
+		// if the sync count is 0, then this flush is not triggered from a sync load, broadcast the delegate in that case
+		FUObjectThreadContext& ThreadContext = FUObjectThreadContext::Get();
+		if (ThreadContext.SyncLoadUsingAsyncLoaderCount == 0)
+		{
+			FCoreDelegates::OnAsyncLoadingFlush.Broadcast();
+		}
 
 		double StartTime = FPlatformTime::Seconds();
 		double LogFlushTime = StartTime;
diff --git a/Engine/Source/Runtime/CoreUObject/Private/UObject/UObjectGlobals.cpp b/Engine/Source/Runtime/CoreUObject/Private/UObject/UObjectGlobals.cpp
--- a/Engine/Source/Runtime/CoreUObject/Private/UObject/UObjectGlobals.cpp
+++ b/Engine/Source/Runtime/CoreUObject/Private/UObject/UObjectGlobals.cpp
@@ -1330,13 +1330,10 @@ UObject* StaticLoadObjectInternal(UClass* ObjectClass, UObject* InOuter, const T
 	return Result;
 }
 
-// Track how many nested loads we're doing by triggering an async load and immediately flushing that request.
-static int32 GSyncLoadUsingAsyncLoaderCount = 0;
-
 UObject* StaticLoadObject(UClass* ObjectClass, UObject* InOuter, const TCHAR* InName, const TCHAR* Filename, uint32 LoadFlags, UPackageMap* Sandbox, bool bAllowObjectReconciliation, const FLinkerInstancingContext* InstancingContext)
 {
 	FUObjectThreadContext& ThreadContext = FUObjectThreadContext::Get();
-	if ((GSyncLoadUsingAsyncLoaderCount == 0) && ThreadContext.IsRoutingPostLoad && IsInAsyncLoadingThread())
+	if ((ThreadContext.SyncLoadUsingAsyncLoaderCount == 0) && ThreadContext.IsRoutingPostLoad && IsInAsyncLoadingThread())
 	{
 		UE_LOG(LogUObjectGlobals, Warning, TEXT("Calling StaticLoadObject(\"%s\", \"%s\", \"%s\") during PostLoad of %s is illegal and will crash in a cooked runtime"), 
 			*GetFullNameSafe(ObjectClass),
@@ -1550,6 +1547,7 @@ UPackage* LoadPackageInternal(UPackage* InOuter, const FPackagePath& PackagePath
 	}
 	checkf(IsInGameThread(), TEXT("Unable to load %s. Objects and Packages can only be loaded from the game thread."), *PackagePath.GetDebugName());
 
+	FUObjectThreadContext& ThreadContext = FUObjectThreadContext::Get();
 	if (ShouldAlwaysLoadPackageAsync(PackagePath))
 	{
 		checkf(!InOuter || !InOuter->GetOuter(), TEXT("Loading into subpackages is not implemented.")); // Subpackages are no longer supported in UE
@@ -1567,7 +1565,7 @@ UPackage* LoadPackageInternal(UPackage* InOuter, const FPackagePath& PackagePath
 			FCoreDelegates::OnSyncLoadPackage.Broadcast(PackageName.ToString());
 		}
 
-		GSyncLoadUsingAsyncLoaderCount++;
+		ThreadContext.SyncLoadUsingAsyncLoaderCount++;
 		EPackageFlags PackageFlags = PKG_None;
 		if ((!FApp::IsGame() || GIsEditor) && (LoadFlags & LOAD_PackageForPIE) != 0)
 		{
@@ -1582,7 +1580,7 @@ UPackage* LoadPackageInternal(UPackage* InOuter, const FPackagePath& PackagePath
 			UE_SCOPED_IO_ACTIVITY(*WriteToString<512>(TEXT("Sync "), PackagePath.GetDebugName()));
 			FlushAsyncLoading(RequestID);
 		}
-		GSyncLoadUsingAsyncLoaderCount--;
+		ThreadContext.SyncLoadUsingAsyncLoaderCount--;
 
 		return (InOuter ? InOuter : FindObjectFast<UPackage>(nullptr, PackageName));
 	}
@@ -1612,7 +1610,7 @@ UPackage* LoadPackageInternal(UPackage* InOuter, const FPackagePath& PackagePath
 	TRACE_LOADTIME_POSTLOAD_SCOPE;
 
 	// Set up a load context
-	TRefCountPtr<FUObjectSerializeContext> LoadContext = FUObjectThreadContext::Get().GetSerializeContext();
+	TRefCountPtr<FUObjectSerializeContext> LoadContext = ThreadContext.GetSerializeContext();
 
 	UE_SCOPED_IO_ACTIVITY(*WriteToString<512>(TEXT("Sync "), PackagePath.GetDebugName()));
 
diff --git a/Engine/Source/Runtime/CoreUObject/Private/UObject/UObjectThreadContext.cpp b/Engine/Source/Runtime/CoreUObject/Private/UObject/UObjectThreadContext.cpp
--- a/Engine/Source/Runtime/CoreUObject/Private/UObject/UObjectThreadContext.cpp
+++ b/Engine/Source/Runtime/CoreUObject/Private/UObject/UObjectThreadContext.cpp
@@ -12,10 +12,11 @@ DEFINE_LOG_CATEGORY(LogUObjectThreadContext);
 
 FUObjectThreadContext::FUObjectThreadContext()
 : IsRoutingPostLoad(false)
-, CurrentlyPostLoadedObjectByALT(nullptr)
 , IsDeletingLinkers(false)
+, SyncLoadUsingAsyncLoaderCount(0)
 , IsInConstructor(0)
 , ConstructedObject(nullptr)
+, CurrentlyPostLoadedObjectByALT(nullptr)
 , AsyncPackage(nullptr)
 , AsyncPackageLoader(nullptr)
 , SerializeContext(new FUObjectSerializeContext())
diff --git a/Engine/Source/Runtime/CoreUObject/Public/UObject/UObjectThreadContext.h b/Engine/Source/Runtime/CoreUObject/Public/UObject/UObjectThreadContext.h
--- a/Engine/Source/Runtime/CoreUObject/Public/UObject/UObjectThreadContext.h
+++ b/Engine/Source/Runtime/CoreUObject/Public/UObject/UObjectThreadContext.h
@@ -84,14 +84,16 @@ public:
 
 	/** true when we are routing ConditionalPostLoad/PostLoad to objects										*/
 	bool IsRoutingPostLoad;
-	/** The object we are routing PostLoad from the Async Loading code for */
-	UObject* CurrentlyPostLoadedObjectByALT;
 	/** true when FLinkerManager deletes linkers */
 	bool IsDeletingLinkers;
+	/* Global int to track how many nested loads we're doing by triggering an async load and immediately flushing that request. */
+	int32 SyncLoadUsingAsyncLoaderCount;
 	/* Global flag so that FObjectFinders know if they are called from inside the UObject constructors or not. */
 	int32 IsInConstructor;
 	/* Object that is currently being constructed with ObjectInitializer */
 	UObject* ConstructedObject;
+	/** The object we are routing PostLoad from the Async Loading code for */
+	UObject* CurrentlyPostLoadedObjectByALT;
 	/** Async Package currently processing objects */
 	void* AsyncPackage;
 	/** Async package loader currently processing objects */
