Additional fixes needed around deferring CDO exports on editor loads.

#rnx

diff --git a/Engine/Source/Runtime/CoreUObject/Private/Blueprint/BlueprintSupport.cpp b/Engine/Source/Runtime/CoreUObject/Private/Blueprint/BlueprintSupport.cpp
--- a/Engine/Source/Runtime/CoreUObject/Private/Blueprint/BlueprintSupport.cpp
+++ b/Engine/Source/Runtime/CoreUObject/Private/Blueprint/BlueprintSupport.cpp
@@ -1264,7 +1264,7 @@ bool FLinkerLoad::DeferExportCreation(const int32 Index, UObject* Outer)
 		// the actual CDO and execute its native ctor/initializer.
 		if (((LoadFlags & LOAD_DeferDependencyLoads) != 0) && bHasDeferredDependencies)
 		{
-			// This will cause IsBlueprintFinalizationPending() to fail below (which is what we want).
+			// This will cause IsBlueprintFinalizationPending() to return true (which is what we want).
 			// We'll then fall through and create a placeholder object for the CDO in order to defer its
 			// actual construction (along with serialization) until after we've resolved its dependencies.
 			DEFERRED_DEPENDENCY_CHECK(DeferredCDOIndex == INDEX_NONE);
@@ -1837,7 +1837,7 @@ void FLinkerLoad::FinalizeBlueprint(UClass* LoadClass)
 			// paramount that we force it through serialization (so we reset the 
 			// RF_NeedLoad guard, and leave it to ResolveDeferredExports, for it
 			// to re-run the serialization)
-			if ( (SuperCDO != nullptr) && !SuperCDO->HasAnyFlags(RF_NeedLoad|RF_LoadCompleted) )
+			if ((SuperCDO != nullptr) && !SuperCDO->HasAnyFlags(RF_NeedLoad|RF_LoadCompleted) && !FBlueprintSupport::IsDeferredDependencyPlaceholder(SuperCDO))
 			{
 				check(!GEventDrivenLoaderEnabled || !EVENT_DRIVEN_ASYNC_LOAD_ACTIVE_AT_RUNTIME);
 				SuperCDO->SetFlags(RF_NeedLoad);
diff --git a/Engine/Source/Runtime/CoreUObject/Private/UObject/LinkerLoad.cpp b/Engine/Source/Runtime/CoreUObject/Private/UObject/LinkerLoad.cpp
--- a/Engine/Source/Runtime/CoreUObject/Private/UObject/LinkerLoad.cpp
+++ b/Engine/Source/Runtime/CoreUObject/Private/UObject/LinkerLoad.cpp
@@ -5043,7 +5043,7 @@ UObject* FLinkerLoad::CreateExport( int32 Index )
 				// Don't do this for any packages that have previously fully loaded as they may have in memory changes
 				check(CurrentLoadContext);
 				CurrentLoadContext->AddLoadedObject(Export.Object);
-				if (!Export.Object->HasAnyFlags(RF_LoadCompleted) && !LinkerRoot->IsFullyLoaded())
+				if (!Export.Object->HasAnyFlags(RF_LoadCompleted) && (!LinkerRoot->IsFullyLoaded() || IsBlueprintFinalizationPending()))
 				{
 					check(!GEventDrivenLoaderEnabled || !EVENT_DRIVEN_ASYNC_LOAD_ACTIVE_AT_RUNTIME);
 
diff --git a/Engine/Source/Runtime/Engine/Private/BlueprintGeneratedClass.cpp b/Engine/Source/Runtime/Engine/Private/BlueprintGeneratedClass.cpp
--- a/Engine/Source/Runtime/Engine/Private/BlueprintGeneratedClass.cpp
+++ b/Engine/Source/Runtime/Engine/Private/BlueprintGeneratedClass.cpp
@@ -1550,6 +1550,9 @@ void UBlueprintGeneratedClass::GetPreloadDependencies(TArray<UObject*>& OutDeps)
 void UBlueprintGeneratedClass::GetDefaultObjectPreloadDependencies(TArray<UObject*>& OutDeps)
 {
 	Super::GetDefaultObjectPreloadDependencies(OutDeps);
+	
+	// Add the super classes CDO, as it needs to be in place to properly serialize.
+	OutDeps.Add(GetSuperClass()->GetDefaultObject());
 
 	// Ensure that BPGC-owned component templates (archetypes) are loaded prior to CDO serialization in order to support the following use cases:
 	//
@@ -1613,9 +1616,6 @@ void UBlueprintGeneratedClass::GetDefaultObjectPreloadDependencies(TArray<UObjec
 	// Add objects related to the sparse class data struct
 	if (SparseClassDataStruct)
 	{
-		// Add the super classes CDO, as it needs to be in place to properly serialize the sparse class data's archetype
-		OutDeps.Add(GetSuperClass()->GetDefaultObject());
-	
 		// Add the sparse class data struct, as it is serialized as part of the CDO
 		OutDeps.Add(SparseClassDataStruct);
 
