UObject data layout optimization

diff --git a/Engine/Source/Programs/Shared/EpicGames.UHT/Exporters/CodeGen/UhtHeaderCodeGeneratorCppFile.cs b/Engine/Source/Programs/Shared/EpicGames.UHT/Exporters/CodeGen/UhtHeaderCodeGeneratorCppFile.cs
--- a/Engine/Source/Programs/Shared/EpicGames.UHT/Exporters/CodeGen/UhtHeaderCodeGeneratorCppFile.cs
+++ b/Engine/Source/Programs/Shared/EpicGames.UHT/Exporters/CodeGen/UhtHeaderCodeGeneratorCppFile.cs
@@ -482,8 +482,8 @@ namespace EpicGames.UHT.Exporters.CodeGen
 				builder.Append("\t\t").AppendUTF8LiteralString(enumObj.SourceName).Append(",\r\n");
 				builder.Append("\t\t").AppendUTF8LiteralString(enumObj.CppType).Append(",\r\n");
 				builder.Append("\t\t").Append(staticsName).Append("::Enumerators,\r\n");
-				builder.Append("\t\tUE_ARRAY_COUNT(").Append(staticsName).Append("::Enumerators),\r\n");
 				builder.Append("\t\t").Append(ObjectFlags).Append(",\r\n");
+				builder.Append("\t\tUE_ARRAY_COUNT(").Append(staticsName).Append("::Enumerators),\r\n");
 				builder.Append("\t\t").Append(enumObj.EnumFlags.HasAnyFlags(EEnumFlags.Flags) ? "EEnumFlags::Flags" : "EEnumFlags::None").Append(",\r\n");
 				builder.Append("\t\t(uint8)UEnum::ECppForm::").Append(enumObj.CppForm.ToString()).Append(",\r\n");
 				builder.Append("\t\t").AppendMetaDataParams(enumObj, staticsName, MetaDataParamsName).Append("\r\n");
@@ -618,6 +618,8 @@ namespace EpicGames.UHT.Exporters.CodeGen
 				foreach (UhtScriptStruct noExportStruct in noExportStructs)
 				{
 					AppendMirrorsForNoexportStruct(builder, noExportStruct, 2);
+					builder.Append("\t\tstatic_assert(sizeof(").Append(noExportStruct.SourceName).Append(") < MAX_uint16);\r\n");
+					builder.Append("\t\tstatic_assert(alignof(").Append(noExportStruct.SourceName).Append(") < MAX_uint8);\r\n");
 				}
 
 				// Meta data
@@ -661,13 +663,15 @@ namespace EpicGames.UHT.Exporters.CodeGen
 					builder.Append("\t\tnullptr,\r\n");
 				}
 				builder.Append("\t\t").AppendUTF8LiteralString(scriptStruct.EngineName).Append(",\r\n");
+				builder.AppendPropertiesParams(scriptStruct, staticsName, 2, "\r\n");
 				builder.Append("\t\tsizeof(").Append(scriptStruct.SourceName).Append("),\r\n");
 				builder.Append("\t\talignof(").Append(scriptStruct.SourceName).Append("),\r\n");
-				builder.AppendPropertiesParams(scriptStruct, staticsName, 2, "\r\n");
 				builder.Append("\t\tRF_Public|RF_Transient|RF_MarkAsNative,\r\n");
 				builder.Append($"\t\tEStructFlags(0x{(uint)(scriptStruct.ScriptStructFlags & ~EStructFlags.ComputedFlags):X8}),\r\n");
 				builder.Append("\t\t").AppendMetaDataParams(scriptStruct, staticsName, MetaDataParamsName).Append("\r\n");
 				builder.Append("\t};\r\n");
+
+				builder.AppendPropertiesParamsCountStaticAssert(scriptStruct, staticsName);
 			}
 
 			// Generate the registration function
@@ -860,6 +864,9 @@ namespace EpicGames.UHT.Exporters.CodeGen
 					.AppendUTF8LiteralString(function.FunctionType == UhtFunctionType.SparseDelegate, function.SparseOwningClassName).Append(", ")
 					.AppendUTF8LiteralString(function.FunctionType == UhtFunctionType.SparseDelegate, function.SparseDelegateName).Append(", ");
 
+				builder.AppendPropertiesParams(function, staticsName, 0, " ");
+
+				string sizeOfStatic = "";
 				if (function.Children.Count > 0)
 				{
 					UhtFunction tempFunction = function;
@@ -870,12 +877,14 @@ namespace EpicGames.UHT.Exporters.CodeGen
 
 					if (paramsInStatic)
 					{
-						builder.Append("sizeof(").Append(staticsName).Append("::").Append(GetEventStructParametersName(tempFunction.Outer, tempFunction.StrippedFunctionName)).Append(')');
+						sizeOfStatic = "sizeof(" + staticsName + "::" + GetEventStructParametersName(tempFunction.Outer, tempFunction.StrippedFunctionName) + ")";
 					}
 					else
 					{
-						builder.Append("sizeof(").Append(GetEventStructParametersName(tempFunction.Outer, tempFunction.StrippedFunctionName)).Append(')');
+						sizeOfStatic = "sizeof(" + GetEventStructParametersName(tempFunction.Outer, tempFunction.StrippedFunctionName) + ")";
 					}
+
+					builder.Append(sizeOfStatic);
 				}
 				else
 				{
@@ -884,13 +893,19 @@ namespace EpicGames.UHT.Exporters.CodeGen
 				builder.Append(", ");
 
 				builder
-					.AppendPropertiesParams(function, staticsName, 0, " ")
 					.Append("RF_Public|RF_Transient|RF_MarkAsNative, ")
 					.Append($"(EFunctionFlags)0x{(uint)function.FunctionFlags:X8}, ")
 					.Append(isNet ? function.RPCId : 0).Append(", ")
 					.Append(isNet ? function.RPCResponseId : 0).Append(", ")
 					.AppendMetaDataParams(function, staticsName, MetaDataParamsName)
 					.Append(" };\r\n");
+
+				builder.AppendPropertiesParamsCountStaticAssert(function, staticsName);
+
+				if (!string.IsNullOrEmpty(sizeOfStatic))
+				{
+					builder.Append("\tstatic_assert(").Append(sizeOfStatic).Append(" < MAX_uint16);\r\n");
+				}
 			}
 
 			// Registration function
@@ -1500,6 +1515,8 @@ namespace EpicGames.UHT.Exporters.CodeGen
 				builder.Append("\t\t(UObject* (*)())").Append(GetSingletonName(Package, true)).Append(",\r\n");
 				builder.Append("\t};\r\n");
 
+				builder.Append("\tstatic_assert(UE_ARRAY_COUNT(").Append(staticsName).Append("::DependentSingletons) < 16);\r\n");
+
 				if (sortedFunctions.Count > 0)
 				{
 					builder.AppendBeginEditorOnlyGuard(allEditorOnlyFunctions);
@@ -1521,6 +1538,7 @@ namespace EpicGames.UHT.Exporters.CodeGen
 					}
 
 					builder.Append("\t};\r\n");
+					builder.Append("\tstatic_assert(UE_ARRAY_COUNT(").Append(staticsName).Append("::FuncInfo) < 2048);\r\n");
 					builder.AppendEndEditorOnlyGuard(allEditorOnlyFunctions);
 				}
 
@@ -1553,6 +1571,8 @@ namespace EpicGames.UHT.Exporters.CodeGen
 					}
 
 					builder.Append("\t\t};\r\n");
+
+					builder.Append("\tstatic_assert(UE_ARRAY_COUNT(").Append(staticsName).Append("::InterfaceParams) < 64);\r\n");
 				}
 
 				builder.Append("\tconst FCppClassTypeInfoStatic ").Append(staticsName).Append("::StaticCppClassTypeInfo = {\r\n");
@@ -1612,6 +1632,8 @@ namespace EpicGames.UHT.Exporters.CodeGen
 				builder.Append($"\t\t0x{(uint)classFlags:X8}u,\r\n");
 				builder.Append("\t\t").AppendMetaDataParams(classObj, staticsName, MetaDataParamsName).Append("\r\n");
 				builder.Append("\t};\r\n");
+
+				builder.AppendPropertiesParamsCountStaticAssert(classObj, staticsName);
 			}
 
 			// Class registration
@@ -2090,6 +2112,32 @@ namespace EpicGames.UHT.Exporters.CodeGen
 			return builder;
 		}
 
+		/// <summary>
+		/// Append static_assert to check that properties params count fit within struct
+		/// </summary>
+		/// <param name="builder">Destination builder</param>
+		/// <param name="structObj">Structure in question</param>
+		/// <param name="staticsName">Name of the statics section for this structure</param>
+		/// <returns>Destination builder</returns>
+		public static StringBuilder AppendPropertiesParamsCountStaticAssert(this StringBuilder builder, UhtStruct structObj, string staticsName)
+		{
+			if (!structObj.Children.Any(x => x is UhtProperty))
+			{
+				// do nothing here
+			}
+			else if (structObj.Properties.Any(x => !x.IsEditorOnlyProperty))
+			{
+				builder.Append("\tstatic_assert(UE_ARRAY_COUNT(").Append(staticsName).Append("::PropPointers) < 2048);\r\n");
+			}
+			else
+			{
+				builder.Append("#if WITH_EDITORONLY_DATA\r\n")
+					   .Append("\tstatic_assert(UE_ARRAY_COUNT(").Append(staticsName).Append("::PropPointers) < 2048);\r\n")
+					   .Append("#endif\r\n");
+			}
+			return builder;
+		}
+
 		/// <summary>
 		/// Append the given array list and count as arguments to a structure constructor
 		/// </summary>
diff --git a/Engine/Source/Programs/Shared/EpicGames.UHT/Types/UhtProperty.cs b/Engine/Source/Programs/Shared/EpicGames.UHT/Types/UhtProperty.cs
--- a/Engine/Source/Programs/Shared/EpicGames.UHT/Types/UhtProperty.cs
+++ b/Engine/Source/Programs/Shared/EpicGames.UHT/Types/UhtProperty.cs
@@ -1214,8 +1214,7 @@ namespace EpicGames.UHT.Types
 				.AppendNotifyFunc(this).Append(", ")
 				.AppendFlags(PropertyFlags).Append(", ")
 				.Append(paramsGenFlags).Append(", ")
-				.Append(ObjectFlags).Append(", ")
-				.AppendArrayDim(this, context).Append(", ");
+				.Append(ObjectFlags).Append(", ");
 
 			if (PropertyExportFlags.HasAnyFlags(UhtPropertyExportFlags.SetterFound))
 			{
@@ -1235,6 +1234,8 @@ namespace EpicGames.UHT.Types
 				builder.Append("nullptr, ");
 			}
 
+			builder.AppendArrayDim(this, context).Append(", ");
+
 			if (appendOffset)
 			{
 				if (!String.IsNullOrEmpty(offset))
diff --git a/Engine/Source/Programs/Shared/EpicGames.UHT/Utils/UhtStringBuilderExtensions.cs b/Engine/Source/Programs/Shared/EpicGames.UHT/Utils/UhtStringBuilderExtensions.cs
--- a/Engine/Source/Programs/Shared/EpicGames.UHT/Utils/UhtStringBuilderExtensions.cs
+++ b/Engine/Source/Programs/Shared/EpicGames.UHT/Utils/UhtStringBuilderExtensions.cs
@@ -140,14 +140,15 @@ namespace EpicGames.UHT.Utils
 			{
 				return builder
 					.Append("METADATA_PARAMS(")
+					.Append("UE_ARRAY_COUNT(")
 					.AppendNameDef(staticsName, namePrefix, name, nameSuffix).Append(metaNameSuffix)
-					.Append(", UE_ARRAY_COUNT(")
+					.Append("), ")
 					.AppendNameDef(staticsName, namePrefix, name, nameSuffix).Append(metaNameSuffix)
-					.Append("))");
+					.Append(")");
 			}
 			else
 			{
-				return builder.Append("METADATA_PARAMS(nullptr, 0)");
+				return builder.Append("METADATA_PARAMS(0, nullptr)");
 			}
 		}
 
diff --git a/Engine/Source/Runtime/CoreUObject/Public/UObject/ObjectMacros.h b/Engine/Source/Runtime/CoreUObject/Public/UObject/ObjectMacros.h
--- a/Engine/Source/Runtime/CoreUObject/Public/UObject/ObjectMacros.h
+++ b/Engine/Source/Runtime/CoreUObject/Public/UObject/ObjectMacros.h
@@ -480,7 +480,7 @@ ENUM_CLASS_FLAGS(EPropertyFlags)
 /**
  * Extra flags for array properties.
  */
-enum class EArrayPropertyFlags
+enum class EArrayPropertyFlags : uint8
 {
 	None,
 	UsesMemoryImageAllocator
@@ -491,7 +491,7 @@ ENUM_CLASS_FLAGS(EArrayPropertyFlags)
 /**
  * Extra flags for map properties.
  */
-enum class EMapPropertyFlags
+enum class EMapPropertyFlags : uint8
 {
 	None,
 	UsesMemoryImageAllocator
@@ -611,7 +611,7 @@ ENUM_CLASS_FLAGS(EInternalObjectFlags);
  * This MUST be kept in sync with EEnumFlags defined in
  * Engine\Source\Programs\Shared\EpicGames.Core\UnrealEngineTypes.cs
  */
-enum class EEnumFlags
+enum class EEnumFlags : uint8
 {
 	None,
 
diff --git a/Engine/Source/Runtime/CoreUObject/Public/UObject/UObjectGlobals.h b/Engine/Source/Runtime/CoreUObject/Public/UObject/UObjectGlobals.h
--- a/Engine/Source/Runtime/CoreUObject/Public/UObject/UObjectGlobals.h
+++ b/Engine/Source/Runtime/CoreUObject/Public/UObject/UObjectGlobals.h
@@ -3090,7 +3090,7 @@ typedef void (*GetterFuncPtr)(const void* InContainer, void* OutValue);
 /// @cond DOXYGEN_IGNORE
 namespace UECodeGen_Private
 {
-	enum class EPropertyGenFlags : uint32
+	enum class EPropertyGenFlags : uint8
 	{
 		None              = 0x00,
 
@@ -3163,9 +3163,9 @@ namespace UECodeGen_Private
 		EPropertyFlags    PropertyFlags;
 		EPropertyGenFlags Flags;
 		EObjectFlags   ObjectFlags;
-		int32          ArrayDim;
 		SetterFuncPtr  SetterFunc;
 		GetterFuncPtr  GetterFunc;
+		uint16         ArrayDim;
 	};
 
 	struct FPropertyParamsBaseWithoutOffset // : FPropertyParamsBase
@@ -3175,9 +3175,9 @@ namespace UECodeGen_Private
 		EPropertyFlags    PropertyFlags;
 		EPropertyGenFlags Flags;
 		EObjectFlags   ObjectFlags;
-		int32          ArrayDim;
 		SetterFuncPtr  SetterFunc;
 		GetterFuncPtr  GetterFunc;
+		uint16         ArrayDim;
 	};
 
 	struct FPropertyParamsBaseWithOffset // : FPropertyParamsBase
@@ -3187,10 +3187,10 @@ namespace UECodeGen_Private
 		EPropertyFlags    PropertyFlags;
 		EPropertyGenFlags Flags;
 		EObjectFlags   ObjectFlags;
-		int32          ArrayDim;
 		SetterFuncPtr  SetterFunc;
 		GetterFuncPtr  GetterFunc;
-		int32          Offset;
+		uint16         ArrayDim;
+		uint16         Offset;
 	};
 
 	struct FGenericPropertyParams // : FPropertyParamsBaseWithOffset
@@ -3200,13 +3200,13 @@ namespace UECodeGen_Private
 		EPropertyFlags    PropertyFlags;
 		EPropertyGenFlags Flags;
 		EObjectFlags     ObjectFlags;
-		int32            ArrayDim;
 		SetterFuncPtr  SetterFunc;
 		GetterFuncPtr  GetterFunc;
-		int32            Offset;
+		uint16           ArrayDim;
+		uint16           Offset;
 #if WITH_METADATA
+		uint16                              NumMetaData;
 		const FMetaDataPairParam*           MetaDataArray;
-		int32                               NumMetaData;
 #endif
 	};
 
@@ -3217,14 +3217,14 @@ namespace UECodeGen_Private
 		EPropertyFlags      PropertyFlags;
 		EPropertyGenFlags   Flags;
 		EObjectFlags     ObjectFlags;
-		int32            ArrayDim;
 		SetterFuncPtr  SetterFunc;
 		GetterFuncPtr  GetterFunc;
-		int32            Offset;
+		uint16           ArrayDim;
+		uint16           Offset;
 		UEnum*         (*EnumFunc)();
 #if WITH_METADATA
+		uint16                              NumMetaData;
 		const FMetaDataPairParam*           MetaDataArray;
-		int32                               NumMetaData;
 #endif
 	};
 
@@ -3235,15 +3235,15 @@ namespace UECodeGen_Private
 		EPropertyFlags      PropertyFlags;
 		EPropertyGenFlags   Flags;
 		EObjectFlags     ObjectFlags;
-		int32            ArrayDim;
 		SetterFuncPtr  SetterFunc;
 		GetterFuncPtr  GetterFunc;
-		uint32           ElementSize;
-		SIZE_T           SizeOfOuter;
+		uint16           ArrayDim;
+		uint16           ElementSize;
+		uint16           SizeOfOuter;
 		void           (*SetBitFunc)(void* Obj);
 #if WITH_METADATA
+		uint16                              NumMetaData;
 		const FMetaDataPairParam*           MetaDataArray;
-		int32                               NumMetaData;
 #endif
 	};
 
@@ -3254,13 +3254,13 @@ namespace UECodeGen_Private
 		EPropertyFlags      PropertyFlags;
 		EPropertyGenFlags   Flags;
 		EObjectFlags     ObjectFlags;
-		int32            ArrayDim;
 		SetterFuncPtr  SetterFunc;
 		GetterFuncPtr  GetterFunc;
-		int32            Offset;
+		uint16           ArrayDim;
+		uint16           Offset;
 #if WITH_METADATA
-		const FMetaDataPairParam* MetaDataArray;
-		int32                               NumMetaData;
+		uint16                              NumMetaData;
+		const FMetaDataPairParam*           MetaDataArray;
 #endif
 	};
 
@@ -3271,14 +3271,14 @@ namespace UECodeGen_Private
 		EPropertyFlags      PropertyFlags;
 		EPropertyGenFlags   Flags;
 		EObjectFlags     ObjectFlags;
-		int32            ArrayDim;
 		SetterFuncPtr  SetterFunc;
 		GetterFuncPtr  GetterFunc;
-		int32            Offset;
+		uint16           ArrayDim;
+		uint16           Offset;
 		UClass*        (*ClassFunc)();
 #if WITH_METADATA
-		const FMetaDataPairParam* MetaDataArray;
-		int32                               NumMetaData;
+		uint16                              NumMetaData;
+		const FMetaDataPairParam*           MetaDataArray;
 #endif
 	};
 
@@ -3289,14 +3289,14 @@ namespace UECodeGen_Private
 		EPropertyFlags      PropertyFlags;
 		EPropertyGenFlags   Flags;
 		EObjectFlags     ObjectFlags;
-		int32            ArrayDim;
 		SetterFuncPtr  SetterFunc;
 		GetterFuncPtr  GetterFunc;
-		int32            Offset;
+		uint16           ArrayDim;
+		uint16           Offset;
 		UClass*        (*ClassFunc)();
 #if WITH_METADATA
+		uint16                              NumMetaData;
 		const FMetaDataPairParam*           MetaDataArray;
-		int32                               NumMetaData;
 #endif
 	};
 
@@ -3307,15 +3307,15 @@ namespace UECodeGen_Private
 		EPropertyFlags      PropertyFlags;
 		EPropertyGenFlags   Flags;
 		EObjectFlags     ObjectFlags;
-		int32            ArrayDim;
 		SetterFuncPtr  SetterFunc;
 		GetterFuncPtr  GetterFunc;
-		int32            Offset;
+		uint16           ArrayDim;
+		uint16           Offset;
 		UClass*        (*ClassFunc)();
 		UClass*        (*MetaClassFunc)();		
 #if WITH_METADATA
+		uint16                              NumMetaData;
 		const FMetaDataPairParam*           MetaDataArray;
-		int32                               NumMetaData;
 #endif
 	};
 
@@ -3326,14 +3326,14 @@ namespace UECodeGen_Private
 		EPropertyFlags      PropertyFlags;
 		EPropertyGenFlags   Flags;
 		EObjectFlags     ObjectFlags;
-		int32            ArrayDim;
 		SetterFuncPtr  SetterFunc;
 		GetterFuncPtr  GetterFunc;
-		int32            Offset;
+		uint16           ArrayDim;
+		uint16           Offset;
 		UClass*        (*MetaClassFunc)();
 #if WITH_METADATA
+		uint16                              NumMetaData;
 		const FMetaDataPairParam*           MetaDataArray;
-		int32                               NumMetaData;
 #endif
 	};
 
@@ -3344,14 +3344,14 @@ namespace UECodeGen_Private
 		EPropertyFlags      PropertyFlags;
 		EPropertyGenFlags   Flags;
 		EObjectFlags     ObjectFlags;
-		int32            ArrayDim;
 		SetterFuncPtr  SetterFunc;
 		GetterFuncPtr  GetterFunc;
-		int32            Offset;
+		uint16           ArrayDim;
+		uint16           Offset;
 		UClass*        (*InterfaceClassFunc)();
 #if WITH_METADATA
+		uint16                              NumMetaData;
 		const FMetaDataPairParam*           MetaDataArray;
-		int32                               NumMetaData;
 #endif
 	};
 
@@ -3362,14 +3362,14 @@ namespace UECodeGen_Private
 		EPropertyFlags      PropertyFlags;
 		EPropertyGenFlags   Flags;
 		EObjectFlags     ObjectFlags;
-		int32            ArrayDim;
 		SetterFuncPtr  SetterFunc;
 		GetterFuncPtr  GetterFunc;
-		int32            Offset;
+		uint16           ArrayDim;
+		uint16           Offset;
 		UScriptStruct* (*ScriptStructFunc)();
 #if WITH_METADATA
+		uint16                              NumMetaData;
 		const FMetaDataPairParam*           MetaDataArray;
-		int32                               NumMetaData;
 #endif
 	};
 
@@ -3380,14 +3380,14 @@ namespace UECodeGen_Private
 		EPropertyFlags      PropertyFlags;
 		EPropertyGenFlags   Flags;
 		EObjectFlags     ObjectFlags;
-		int32            ArrayDim;
 		SetterFuncPtr  SetterFunc;
 		GetterFuncPtr  GetterFunc;
-		int32            Offset;
+		uint16           ArrayDim;
+		uint16           Offset;
 		UFunction*     (*SignatureFunctionFunc)();
 #if WITH_METADATA
+		uint16                              NumMetaData;
 		const FMetaDataPairParam*           MetaDataArray;
-		int32                               NumMetaData;
 #endif
 	};
 
@@ -3398,14 +3398,14 @@ namespace UECodeGen_Private
 		EPropertyFlags      PropertyFlags;
 		EPropertyGenFlags   Flags;
 		EObjectFlags     ObjectFlags;
-		int32            ArrayDim;
 		SetterFuncPtr  SetterFunc;
 		GetterFuncPtr  GetterFunc;
-		int32            Offset;
+		uint16           ArrayDim;
+		uint16           Offset;
 		UFunction*     (*SignatureFunctionFunc)();
 #if WITH_METADATA
+		uint16                              NumMetaData;
 		const FMetaDataPairParam*           MetaDataArray;
-		int32                               NumMetaData;
 #endif
 	};
 
@@ -3416,14 +3416,14 @@ namespace UECodeGen_Private
 		EPropertyFlags     PropertyFlags;
 		EPropertyGenFlags  Flags;
 		EObjectFlags     ObjectFlags;
-		int32            ArrayDim;
 		SetterFuncPtr  SetterFunc;
 		GetterFuncPtr  GetterFunc;
-		int32            Offset;
+		uint16           ArrayDim;
+		uint16           Offset;
 		UEnum*         (*EnumFunc)();
 #if WITH_METADATA
+		uint16                              NumMetaData;
 		const FMetaDataPairParam*           MetaDataArray;
-		int32                               NumMetaData;
 #endif
 	};
 
@@ -3434,14 +3434,14 @@ namespace UECodeGen_Private
 		EPropertyFlags     PropertyFlags;
 		EPropertyGenFlags  Flags;
 		EObjectFlags     ObjectFlags;
-		int32            ArrayDim;
 		SetterFuncPtr  SetterFunc;
 		GetterFuncPtr  GetterFunc;
-		int32            Offset;
+		uint16           ArrayDim;
+		uint16           Offset;
 		FFieldClass*     (*PropertyClassFunc)();
 #if WITH_METADATA
+		uint16                              NumMetaData;
 		const FMetaDataPairParam*           MetaDataArray;
-		int32                               NumMetaData;
 #endif
 	};
 
@@ -3452,14 +3452,14 @@ namespace UECodeGen_Private
 		EPropertyFlags      PropertyFlags;
 		EPropertyGenFlags   Flags;
 		EObjectFlags        ObjectFlags;
-		int32               ArrayDim;
 		SetterFuncPtr  SetterFunc;
 		GetterFuncPtr  GetterFunc;
-		int32               Offset;
+		uint16              ArrayDim;
+		uint16              Offset;
 		EArrayPropertyFlags ArrayFlags;
 #if WITH_METADATA
+		uint16                    NumMetaData;
 		const FMetaDataPairParam* MetaDataArray;
-		int32                     NumMetaData;
 #endif
 	};
 
@@ -3470,14 +3470,14 @@ namespace UECodeGen_Private
 		EPropertyFlags    PropertyFlags;
 		EPropertyGenFlags Flags;
 		EObjectFlags      ObjectFlags;
-		int32             ArrayDim;
 		SetterFuncPtr  SetterFunc;
 		GetterFuncPtr  GetterFunc;
-		int32             Offset;
+		uint16            ArrayDim;
+		uint16            Offset;
 		EMapPropertyFlags MapFlags;
 #if WITH_METADATA
+		uint16                    NumMetaData;
 		const FMetaDataPairParam* MetaDataArray;
-		int32                     NumMetaData;
 #endif
 	};
 
@@ -3511,16 +3511,16 @@ namespace UECodeGen_Private
 		const char*                         NameUTF8;
 		const char*                         OwningClassName;
 		const char*                         DelegateName;
-		SIZE_T                              StructureSize;
 		const FPropertyParamsBase* const*   PropertyArray;
-		int32                               NumProperties;
+		uint16                              NumProperties;
+		uint16                              StructureSize;
 		EObjectFlags                        ObjectFlags;
 		EFunctionFlags                      FunctionFlags;
 		uint16                              RPCId;
 		uint16                              RPCResponseId;
 #if WITH_METADATA
+		uint16                              NumMetaData;
 		const FMetaDataPairParam*           MetaDataArray;
-		int32                               NumMetaData;
 #endif
 	};
 
@@ -3531,13 +3531,13 @@ namespace UECodeGen_Private
 		const char*                 NameUTF8;
 		const char*                 CppTypeUTF8;
 		const FEnumeratorParam*     EnumeratorParams;
-		int32                       NumEnumerators;
 		EObjectFlags                ObjectFlags;
+		int16                       NumEnumerators;
 		EEnumFlags                  EnumFlags;
 		uint8                       CppForm; // this is of type UEnum::ECppForm
 #if WITH_METADATA
+		uint16                      NumMetaData;
 		const FMetaDataPairParam*   MetaDataArray;
-		int32                       NumMetaData;
 #endif
 	};
 
@@ -3547,15 +3547,15 @@ namespace UECodeGen_Private
 		UScriptStruct*                    (*SuperFunc)();
 		void*                             (*StructOpsFunc)(); // really returns UScriptStruct::ICppStructOps*
 		const char*                         NameUTF8;
-		SIZE_T                              SizeOf;
-		SIZE_T                              AlignOf;
 		const FPropertyParamsBase* const*   PropertyArray;
-		int32                               NumProperties;
+		uint16                              NumProperties;
+		uint16                              SizeOf;
+		uint8                               AlignOf;
 		EObjectFlags                        ObjectFlags;
 		uint32                              StructFlags; // EStructFlags
 #if WITH_METADATA
+		uint16                              NumMetaData;
 		const FMetaDataPairParam*           MetaDataArray;
-		int32                               NumMetaData;
 #endif
 	};
 
@@ -3568,8 +3568,8 @@ namespace UECodeGen_Private
 		uint32                             BodyCRC;
 		uint32                             DeclarationsCRC;
 #if WITH_METADATA
+		uint16                             NumMetaData;
 		const FMetaDataPairParam*          MetaDataArray;
-		int32                              NumMetaData;
 #endif
 	};
 
@@ -3589,14 +3589,14 @@ namespace UECodeGen_Private
 		const FClassFunctionLinkInfo*               FunctionLinkArray;
 		const FPropertyParamsBase* const*           PropertyArray;
 		const FImplementedInterfaceParams*          ImplementedInterfaceArray;
-		int32                                       NumDependencySingletons;
-		int32                                       NumFunctions;
-		int32                                       NumProperties;
-		int32                                       NumImplementedInterfaces;
+		uint32                                      NumDependencySingletons : 4;
+		uint32                                      NumFunctions : 11;
+		uint32                                      NumProperties : 11;
+		uint32                                      NumImplementedInterfaces : 6;
 		uint32                                      ClassFlags; // EClassFlags
 #if WITH_METADATA
+		uint16                                      NumMetaData;
 		const FMetaDataPairParam*                   MetaDataArray;
-		int32                                       NumMetaData;
 #endif
 	};
 
