Speculative fix for object being removed from object hash during gc when it was never inserted.

diff --git a/Engine/Source/Runtime/CoreUObject/Private/UObject/UObjectArray.cpp b/Engine/Source/Runtime/CoreUObject/Private/UObject/UObjectArray.cpp
--- a/Engine/Source/Runtime/CoreUObject/Private/UObject/UObjectArray.cpp
+++ b/Engine/Source/Runtime/CoreUObject/Private/UObject/UObjectArray.cpp
@@ -186,7 +186,7 @@ void FUObjectArray::DisableDisregardForGC()
 	}
 }
 
-void FUObjectArray::AllocateUObjectIndex(UObjectBase* Object, int32 AlreadyAllocatedIndex, int32 SerialNumber)
+void FUObjectArray::AllocateUObjectIndex(UObjectBase* Object, EInternalObjectFlags InitialFlags, int32 AlreadyAllocatedIndex, int32 SerialNumber)
 {
 	int32 Index = INDEX_NONE;
 	check(Object->InternalIndex == INDEX_NONE);
@@ -235,7 +235,7 @@ void FUObjectArray::AllocateUObjectIndex(UObjectBase* Object, int32 AlreadyAlloc
 	UE_CLOG(ObjectItem->Object != nullptr, LogUObjectArray, Fatal, TEXT("Attempting to add %s at index %d but another object (0x%016llx) exists at that index!"), *Object->GetFName().ToString(), Index, (int64)(PTRINT)ObjectItem->Object);
 	ObjectItem->Object = Object;
 	// At this point all not-compiled-in objects are not fully constructed yet and this is the earliest we can mark them as such
-	ObjectItem->Flags = (int32)EInternalObjectFlags::PendingConstruction;
+	ObjectItem->Flags = (int32)(EInternalObjectFlags::PendingConstruction | InitialFlags);
 	ObjectItem->ClusterRootIndex = 0;
 	ObjectItem->SerialNumber = SerialNumber;
 	Object->InternalIndex = Index;
diff --git a/Engine/Source/Runtime/CoreUObject/Private/UObject/UObjectBase.cpp b/Engine/Source/Runtime/CoreUObject/Private/UObject/UObjectBase.cpp
--- a/Engine/Source/Runtime/CoreUObject/Private/UObject/UObjectBase.cpp
+++ b/Engine/Source/Runtime/CoreUObject/Private/UObject/UObjectBase.cpp
@@ -218,13 +218,8 @@ void UObjectBase::AddObject(FName InName, EInternalObjectFlags InSetInternalFlag
 		InternalFlagsToSet |= EInternalObjectFlags::Native;
 		ObjectFlags &= ~RF_MarkAsNative;
 	}
-	GUObjectArray.AllocateUObjectIndex(this, InInternalIndex, InSerialNumber);
+	GUObjectArray.AllocateUObjectIndex(this, InternalFlagsToSet, InInternalIndex, InSerialNumber);
 	check(InName != NAME_None && InternalIndex >= 0);
-	if (InternalFlagsToSet != EInternalObjectFlags::None)
-	{
-		GUObjectArray.IndexToObject(InternalIndex)->SetFlags(InternalFlagsToSet);
-	
-	}	
 	HashObject(this);
 	check(IsValidLowLevel());
 }
diff --git a/Engine/Source/Runtime/CoreUObject/Public/UObject/UObjectArray.h b/Engine/Source/Runtime/CoreUObject/Public/UObject/UObjectArray.h
--- a/Engine/Source/Runtime/CoreUObject/Public/UObject/UObjectArray.h
+++ b/Engine/Source/Runtime/CoreUObject/Public/UObject/UObjectArray.h
@@ -693,10 +693,11 @@ public:
 	 * Adds a uobject to the global array which is used for uobject iteration
 	 *
 	 * @param	Object Object to allocate an index for
+	 * @param	InitialFlags Flags to set in the object array before the object pointer becomes visible to other threads. 
 	 * @param	AlreadyAllocatedIndex already allocated internal index to use, negative value means allocate a new index
 	 * @param	SerialNumber serial number to use
 	 */
-	void AllocateUObjectIndex(class UObjectBase* Object, int32 AlreadyAllocatedIndex = -1, int32 SerialNumber = 0);
+	void AllocateUObjectIndex(class UObjectBase* Object, EInternalObjectFlags InitialFlags, int32 AlreadyAllocatedIndex = -1, int32 SerialNumber = 0);
 
 	/**
 	 * Returns a UObject index top to the global uobject array
