changed MarkSearchableName to take a TObjectPtr to avoid resolving objects
https://p4-swarm.epicgames.net/reviews/24504202

diff --git a/Engine/Source/Runtime/Core/Public/Serialization/Archive.h b/Engine/Source/Runtime/Core/Public/Serialization/Archive.h
--- a/Engine/Source/Runtime/Core/Public/Serialization/Archive.h
+++ b/Engine/Source/Runtime/Core/Public/Serialization/Archive.h
@@ -37,6 +37,7 @@ class UObject;
 namespace UE::Serialization{ class FEditorBulkData; }
 struct FArchiveSerializedPropertyChain;
 template<class TEnum> class TEnumAsByte;
+template <typename T> struct TObjectPtr;
 
 typedef TFunction<bool (double RemainingTime)> FExternalReadCallback;
 class FField;
@@ -1810,7 +1811,7 @@ public:
 	/**
 	 * Called to register a reference to a specific name value, of type TypeObject (UEnum or UStruct normally). Const so it can be called from PostSerialize
 	 */
-	virtual void MarkSearchableName(const UObject* TypeObject, const FName& ValueName) const { }
+	virtual void MarkSearchableName(const TObjectPtr<const UObject>& TypeObject, const FName& ValueName) const { }
 
 	using FArchiveState::GetArchetypeFromLoader;
 
diff --git a/Engine/Source/Runtime/Core/Public/Serialization/ArchiveProxy.h b/Engine/Source/Runtime/Core/Public/Serialization/ArchiveProxy.h
--- a/Engine/Source/Runtime/Core/Public/Serialization/ArchiveProxy.h
+++ b/Engine/Source/Runtime/Core/Public/Serialization/ArchiveProxy.h
@@ -204,7 +204,7 @@ public:
 		InnerArchive.ResetCustomVersions();
 	}
 
-	virtual void MarkSearchableName(const UObject* TypeObject, const FName& ValueName) const override
+	virtual void MarkSearchableName(const TObjectPtr<const UObject>& TypeObject, const FName& ValueName) const override
 	{
 		InnerArchive.MarkSearchableName(TypeObject, ValueName);
 	}
diff --git a/Engine/Source/Runtime/CoreUObject/Private/UObject/SavePackage/PackageHarvester.cpp b/Engine/Source/Runtime/CoreUObject/Private/UObject/SavePackage/PackageHarvester.cpp
--- a/Engine/Source/Runtime/CoreUObject/Private/UObject/SavePackage/PackageHarvester.cpp
+++ b/Engine/Source/Runtime/CoreUObject/Private/UObject/SavePackage/PackageHarvester.cpp
@@ -446,7 +446,7 @@ FString FPackageHarvester::GetArchiveName() const
 	return FString::Printf(TEXT("PackageHarvester (%s)"), *SaveContext.GetPackage()->GetName());
 }
 
-void FPackageHarvester::MarkSearchableName(const UObject* TypeObject, const FName& ValueName) const
+void FPackageHarvester::MarkSearchableName(const TObjectPtr<const UObject>& TypeObject, const FName& ValueName) const
 {
 	if (TypeObject == nullptr)
 	{
@@ -455,7 +455,7 @@ void FPackageHarvester::MarkSearchableName(const UObject* TypeObject, const FNam
 
 	// Serialize object to make sure it ends up in import table
 	// This is doing a const cast to avoid backward compatibility issues
-	UObject* TempObject = const_cast<UObject*>(TypeObject);
+	UObject* TempObject = const_cast<UObject*>(TypeObject.Get());
 	FPackageHarvester* MutableArchive = const_cast<FPackageHarvester*>(this);
 	MutableArchive->HarvestSearchableName(TempObject, ValueName);
 }
diff --git a/Engine/Source/Runtime/CoreUObject/Private/UObject/SavePackage/PackageHarvester.h b/Engine/Source/Runtime/CoreUObject/Private/UObject/SavePackage/PackageHarvester.h
--- a/Engine/Source/Runtime/CoreUObject/Private/UObject/SavePackage/PackageHarvester.h
+++ b/Engine/Source/Runtime/CoreUObject/Private/UObject/SavePackage/PackageHarvester.h
@@ -85,7 +85,7 @@ public:
 
 	// FArchiveUObject implementation
 	virtual FString GetArchiveName() const override;
-	virtual void MarkSearchableName(const UObject* TypeObject, const FName& ValueName) const override;
+	virtual void MarkSearchableName(const TObjectPtr<const UObject>& TypeObject, const FName& ValueName) const override;
 	virtual FArchive& operator<<(UObject*& Obj) override;
 	virtual FArchive& operator<<(struct FWeakObjectPtr& Value) override;
 	virtual FArchive& operator<<(FLazyObjectPtr& LazyObjectPtr) override;
