Added ElementType to TSubclassOf

diff --git a/Engine/Source/Runtime/CoreUObject/Public/Templates/SubclassOf.h b/Engine/Source/Runtime/CoreUObject/Public/Templates/SubclassOf.h
--- a/Engine/Source/Runtime/CoreUObject/Public/Templates/SubclassOf.h
+++ b/Engine/Source/Runtime/CoreUObject/Public/Templates/SubclassOf.h
@@ -31,6 +31,8 @@ private:
 	friend class TSubclassOf;
 
 public:
+	using ElementType = T;
+
 	TSubclassOf() = default;
 	TSubclassOf(TSubclassOf&&) = default;
 	TSubclassOf(const TSubclassOf&) = default;
