PR #10110: Skip unnecessary iteration in GetMaxEnumValue (Contributed by [REDACTED_PERSON])

diff --git a/Engine/Source/Runtime/CoreUObject/Private/UObject/Enum.cpp b/Engine/Source/Runtime/CoreUObject/Private/UObject/Enum.cpp
--- a/Engine/Source/Runtime/CoreUObject/Private/UObject/Enum.cpp
+++ b/Engine/Source/Runtime/CoreUObject/Private/UObject/Enum.cpp
@@ -216,7 +216,7 @@ int64 UEnum::GetMaxEnumValue() const
 	}
 
 	int64 MaxValue = Names[0].Value;
-	for (int32 i = 0; i < NamesNum; ++i)
+	for (int32 i = 1; i < NamesNum; ++i)
 	{
 		int64 CurrentValue = Names[i].Value;
 		if (CurrentValue > MaxValue)
