ZenLoader in editor: Fixed issue where fully loaded packages could sometimes be reloaded after a GC
#rnx

diff --git a/Engine/Source/Runtime/CoreUObject/Private/Serialization/AsyncLoading2.cpp b/Engine/Source/Runtime/CoreUObject/Private/Serialization/AsyncLoading2.cpp
--- a/Engine/Source/Runtime/CoreUObject/Private/Serialization/AsyncLoading2.cpp
+++ b/Engine/Source/Runtime/CoreUObject/Private/Serialization/AsyncLoading2.cpp
@@ -988,10 +988,13 @@ public:
 					PackageRef.SetPackage(FoundPackage);
 					FoundPackage->SetCanBeImportedFlag(true);
 					FoundPackage->SetPackageId(PackageId);
-					if (FoundPackage->bHasBeenFullyLoaded)
-					{
-						PackageRef.SetAllPublicExportsLoaded();
-					}
+				}
+			}
+			if (PackageRef.HasPackage())
+			{
+				if (PackageRef.GetPackage()->bHasBeenFullyLoaded)
+				{
+					PackageRef.SetAllPublicExportsLoaded();
 				}
 			}
 #endif
