ElementProp=ElementProp -> this lambda capture

diff --git a/Engine/Source/Runtime/CoreUObject/Private/UObject/PropertySet.cpp b/Engine/Source/Runtime/CoreUObject/Private/UObject/PropertySet.cpp
--- a/Engine/Source/Runtime/CoreUObject/Private/UObject/PropertySet.cpp
+++ b/Engine/Source/Runtime/CoreUObject/Private/UObject/PropertySet.cpp
@@ -1065,7 +1065,7 @@ void FSetProperty::AppendSchemaHash(FBlake3& Builder, bool bSkipEditorOnly) cons
 void FScriptSetHelper::Rehash()
 {
 	// Moved out-of-line to maybe fix a weird link error
-	Set->Rehash(SetLayout, [ElementProp=ElementProp](const void* Src) {
+	Set->Rehash(SetLayout, [this](const void* Src) {
 		return ElementProp->GetValueTypeHash(Src);
 	});
 }
