Saving packages with the package trailer stucture is no longer optional.

#rnx

- Removed the config file option [Core.System]UsePackageTrailer so we no longer read it, the system is now on permantly
- FPackageTrailer::IsEnabled was the only place reading that config file option. To avoid compiler issues it will now always return true but is marked as deprecated so if any 3rd party has been using this (and I cannot think why they would) they will be alerted to the change.

diff --git a/Engine/Config/BaseEngine.ini b/Engine/Config/BaseEngine.ini
--- a/Engine/Config/BaseEngine.ini
+++ b/Engine/Config/BaseEngine.ini
@@ -1309,7 +1309,6 @@ GameThreadHeartBeatStartSuspended=False
 GameThreadHeartBeatHitchDuration=-1.0
 GameThreadHeartBeatStackWalk=false
 IgnoreInlineBulkDataReloadEnsures=False
-UsePackageTrailer=True
 HangDuration=0.0
 AssetLogShowsDiskPath=True
 AssetLogShowsAbsolutePath=True
diff --git a/Engine/Source/Runtime/CoreUObject/Private/UObject/PackageTrailer.cpp b/Engine/Source/Runtime/CoreUObject/Private/UObject/PackageTrailer.cpp
--- a/Engine/Source/Runtime/CoreUObject/Private/UObject/PackageTrailer.cpp
+++ b/Engine/Source/Runtime/CoreUObject/Private/UObject/PackageTrailer.cpp
@@ -503,22 +503,6 @@ void FPackageTrailerBuilder::RemoveDuplicateEntries()
 	}
 }
 
-bool FPackageTrailer::IsEnabled()
-{
-	static struct FUsePackageTrailer
-	{
-		bool bEnabled = true;
-
-		FUsePackageTrailer()
-		{
-			GConfig->GetBool(TEXT("Core.System"), TEXT("UsePackageTrailer"), bEnabled, GEngineIni);
-			UE_LOG(LogSerialization, Log, TEXT("UsePackageTrailer: '%s'"), bEnabled ? TEXT("true") : TEXT("false"));
-		}
-	} UsePackageTrailer;
-
-	return UsePackageTrailer.bEnabled;
-}
-
 bool FPackageTrailer::TryLoadFromPackage(const FPackagePath& PackagePath, FPackageTrailer& OutTrailer)
 {
 	//TODO: Need to do this in a way that supports text based assets
diff --git a/Engine/Source/Runtime/CoreUObject/Private/UObject/SavePackage2.cpp b/Engine/Source/Runtime/CoreUObject/Private/UObject/SavePackage2.cpp
--- a/Engine/Source/Runtime/CoreUObject/Private/UObject/SavePackage2.cpp
+++ b/Engine/Source/Runtime/CoreUObject/Private/UObject/SavePackage2.cpp
@@ -961,22 +961,19 @@ ESavePackageResult CreateLinker(FSaveContext& SaveContext)
 
 		SaveContext.GetLinker()->bRehydratePayloads = SaveContext.ShouldRehydratePayloads();
 		
-		if (UE::FPackageTrailer::IsEnabled())
+		// The package trailer is not supported for text based assets yet
+		if (!SaveContext.IsTextFormat() && !SaveContext.IsProceduralSave())
 		{
-			// The package trailer is not supported for text based assets yet
-			if (!SaveContext.IsTextFormat() && !SaveContext.IsProceduralSave())
-			{		
-				SaveContext.GetLinker()->PackageTrailerBuilder = MakeUnique<UE::FPackageTrailerBuilder>(SaveContext.GetPackage()->GetName());
-			}
-			else if ((SaveContext.GetSaveArgs().SaveFlags & SAVE_BulkDataByReference) != 0)
+			SaveContext.GetLinker()->PackageTrailerBuilder = MakeUnique<UE::FPackageTrailerBuilder>(SaveContext.GetPackage()->GetName());
+		}
+		else if ((SaveContext.GetSaveArgs().SaveFlags & SAVE_BulkDataByReference) != 0)
+		{
+			if (const FLinkerLoad* LinkerLoad = FLinkerLoad::FindExistingLinkerForPackage(SaveContext.GetPackage()))
 			{
-				if (const FLinkerLoad* LinkerLoad = FLinkerLoad::FindExistingLinkerForPackage(SaveContext.GetPackage()))
+				const UE::FPackageTrailer* Trailer = LinkerLoad->GetPackageTrailer();
+				if (Trailer && Trailer->GetNumPayloads(UE::EPayloadStorageType::Any) > 0)
 				{
-					const UE::FPackageTrailer* Trailer = LinkerLoad->GetPackageTrailer();
-					if (Trailer && Trailer->GetNumPayloads(UE::EPayloadStorageType::Any) > 0)
-					{
-						SaveContext.GetLinker()->PackageTrailerBuilder = UE::FPackageTrailerBuilder::CreateReferenceToTrailer(*Trailer, SaveContext.GetPackage()->GetName());
-					}
+					SaveContext.GetLinker()->PackageTrailerBuilder = UE::FPackageTrailerBuilder::CreateReferenceToTrailer(*Trailer, SaveContext.GetPackage()->GetName());
 				}
 			}
 		}
diff --git a/Engine/Source/Runtime/CoreUObject/Public/UObject/PackageTrailer.h b/Engine/Source/Runtime/CoreUObject/Public/UObject/PackageTrailer.h
--- a/Engine/Source/Runtime/CoreUObject/Public/UObject/PackageTrailer.h
+++ b/Engine/Source/Runtime/CoreUObject/Public/UObject/PackageTrailer.h
@@ -411,7 +411,11 @@ public:
 	 * 
 	 * Note that this is for development purposes only and should ship as always enabled!
 	 */
-	[[nodiscard]] static bool IsEnabled();
+	UE_DEPRECATED(5.3, "FPackageTrailer::IsEnabled will always return true as the system is no longer optional")
+	[[nodiscard]] static bool IsEnabled()
+	{
+		return true;
+	}
 
 	/** Try to load a trailer from a given package path. Note that it will always try to load the trailer from the workspace domain */
 	[[nodiscard]] static bool TryLoadFromPackage(const FPackagePath& PackagePath, FPackageTrailer& OutTrailer);
