Remove temporary path created whenever PackageName testcase runs.
#rnx

diff --git a/Engine/Source/Runtime/CoreUObject/Private/Misc/PackageName.cpp b/Engine/Source/Runtime/CoreUObject/Private/Misc/PackageName.cpp
--- a/Engine/Source/Runtime/CoreUObject/Private/Misc/PackageName.cpp
+++ b/Engine/Source/Runtime/CoreUObject/Private/Misc/PackageName.cpp
@@ -3023,14 +3023,52 @@ bool FPackageNameTests::RunTest(const FString& Parameters)
 		}
 	}
 
+	// DEPRECATED Mar 2023
+#if WITH_EDITOR
+	{
+		// This block was created on Mar 2023 to clean up previous test results that were not cleaned up.
+		const TCHAR* OldTestMountLocalPathRoot = TEXT("../../../PackageNameTestDNE");
+		if (IFileManager::Get().DirectoryExists(OldTestMountLocalPathRoot))
+		{
+			const TCHAR* OldTestMountLocalPathContent = TEXT("../../../PackageNameTestDNE/Content");
+			bool bHasContent = IFileManager::Get().DirectoryExists(OldTestMountLocalPathContent);
+			TArray<FString> Filenames;
+			IFileManager::Get().FindFiles(Filenames, *WriteToString<64>(OldTestMountLocalPathRoot, TEXT("/*")), true /* Files */, true /* Directories */);
+			if (bHasContent)
+			{
+				IFileManager::Get().FindFiles(Filenames, *WriteToString<64>(OldTestMountLocalPathContent, TEXT("/*")), true /* Files */, true /* Directories */);
+			}
+			int32 NumExpectedFiles = bHasContent ? 1 : 0;
+			if (Filenames.Num() == NumExpectedFiles)
+			{
+				IFileManager::Get().DeleteDirectory(OldTestMountLocalPathRoot, false /* bRequireExists */, true /* Delete Tree */);
+			}
+		};
+	}
+#endif
+
+
+	return true;
+}
+
+// Tests that are too expensive to run as a SmokeFilter
+IMPLEMENT_SIMPLE_AUTOMATION_TEST(FPackageNameTestsExtended, "System.Core.Misc.PackageNamesExtended", EAutomationTestFlags::ApplicationContextMask | EAutomationTestFlags::EngineFilter)
+
+bool FPackageNameTestsExtended::RunTest(const FString& Parameters)
+{
 	// Mounting and Unmounting Paths
 	{
 		FString TestMountLongPackageName(TEXT("/PackageNameTestDNE"));
-		FString TestMountLocalPath(TEXT("../../../PackageNameTestDNE/Content"));
+		FString TestMountLocalPathRoot(TEXT("../../../Engine/Intermediate/PackageNameTestDNE"));
+		FString TestMountLocalPath(TEXT("../../../Engine/Intermediate/PackageNameTestDNE/Content"));
 		FString TestLongPackageName(TEXT("/PackageNameTestDNE/Package1"));
-		FString TestLocalPath(TEXT("../../../PackageNameTestDNE/Content/Package1.uasset"));
+		FString TestLocalPath(TEXT("../../../Engine/Intermediate/PackageNameTestDNE/Content/Package1.uasset"));
 		FString TestExtension(TEXT(".uasset"));
 		FString ActualLocalPath;
+		ON_SCOPE_EXIT
+		{
+			IFileManager::Get().DeleteDirectory(*TestMountLocalPathRoot, false /* bRequireExists */, true /* Delete Tree */);
+		};
 
 		if (FPackageName::TryConvertLongPackageNameToFilename(TestLongPackageName, ActualLocalPath, TestExtension))
 		{
