Remove unnecessary access to linker annotation in UObject::AddReferencedObjects.
Lifetime of Linker should be set by Package's BeginDestroy.
Speeds up reference traversal in editor significantly. (Approx 25% decrease in bp reinstancing time in a large test case.)

diff --git a/Engine/Source/Runtime/CoreUObject/Private/UObject/GarbageCollection.cpp b/Engine/Source/Runtime/CoreUObject/Private/UObject/GarbageCollection.cpp
--- a/Engine/Source/Runtime/CoreUObject/Private/UObject/GarbageCollection.cpp
+++ b/Engine/Source/Runtime/CoreUObject/Private/UObject/GarbageCollection.cpp
@@ -4901,12 +4901,6 @@ void UObject::CallAddReferencedObjects(FReferenceCollector& Collector)
 void UObject::AddReferencedObjects(UObject* This, FReferenceCollector& Collector)
 {
 #if WITH_EDITOR
-	//@todo UE - This seems to be required and it should not be. Seems to be related to the texture streamer.
-	FLinkerLoad* LinkerLoad = This->GetLinker();
-	if (LinkerLoad)
-	{
-		LinkerLoad->AddReferencedObjects(Collector);
-	}
 	// Required by the unified GC when running in the editor
 	if (GIsEditor)
 	{
diff --git a/Engine/Source/Runtime/CoreUObject/Private/UObject/Linker.cpp b/Engine/Source/Runtime/CoreUObject/Private/UObject/Linker.cpp
--- a/Engine/Source/Runtime/CoreUObject/Private/UObject/Linker.cpp
+++ b/Engine/Source/Runtime/CoreUObject/Private/UObject/Linker.cpp
@@ -216,16 +216,6 @@ void FLinker::Serialize( FArchive& Ar )
 	}
 }
 
-void FLinker::AddReferencedObjects(FReferenceCollector& Collector)
-{
-#if WITH_EDITOR
-	if (GIsEditor)
-	{
-		Collector.AddReferencedObject(*(UObject**)&LinkerRoot);
-	}
-#endif
-}
-
 /**
  * Return the path name of the UObject represented by the specified import. 
  * (can be used with StaticFindObject)
diff --git a/Engine/Source/Runtime/CoreUObject/Public/UObject/Linker.h b/Engine/Source/Runtime/CoreUObject/Public/UObject/Linker.h
--- a/Engine/Source/Runtime/CoreUObject/Public/UObject/Linker.h
+++ b/Engine/Source/Runtime/CoreUObject/Public/UObject/Linker.h
@@ -483,7 +483,6 @@ public:
 	 * @param	Ar	the archive to read/write into
 	 */
 	void Serialize(FArchive& Ar);
-	void AddReferencedObjects(FReferenceCollector& Collector);
 	
 	/**
 	 * Return the path name of the UObject represented by the specified export.
