change to FObjectPropertyBase ParseObjectPropertyValue and FindImportedObject to take in a `TObjectPtr<UObject>` instead of UObject
added `SetObjectPtrPropertyValue(void* PropertyValueAddress, TObjectPtr<UObject> Value)`

this allows objects to be imported from text without resolving the object

diff --git a/Engine/Source/Editor/PropertyEditor/Private/PropertyHandleImpl.cpp b/Engine/Source/Editor/PropertyEditor/Private/PropertyHandleImpl.cpp
--- a/Engine/Source/Editor/PropertyEditor/Private/PropertyHandleImpl.cpp
+++ b/Engine/Source/Editor/PropertyEditor/Private/PropertyHandleImpl.cpp
@@ -3944,7 +3944,7 @@ FPropertyAccess::Result FPropertyHandleObject::SetValueFromFormattedString(const
 	TSharedPtr<FPropertyNode> PropertyNodePin = Implementation->GetPropertyNode();
 	FProperty* NodeProperty = PropertyNodePin.IsValid() ? PropertyNodePin->GetProperty() : nullptr;
 	const TCHAR* ObjectBuffer = *InValue;
-	UObject* QualifiedObject = nullptr;
+	TObjectPtr<UObject> QualifiedObject = nullptr;
 
 	// Only allow finding any object with the same name if package path not provided
 	const bool bAllowAnyPackage = !InValue.Contains(TEXT("/"));
diff --git a/Engine/Source/Runtime/CoreUObject/Private/UObject/PropertyBaseObject.cpp b/Engine/Source/Runtime/CoreUObject/Private/UObject/PropertyBaseObject.cpp
--- a/Engine/Source/Runtime/CoreUObject/Private/UObject/PropertyBaseObject.cpp
+++ b/Engine/Source/Runtime/CoreUObject/Private/UObject/PropertyBaseObject.cpp
@@ -12,6 +12,7 @@
 #include "UObject/PropertyHelper.h"
 #include "UObject/LinkerPlaceholderClass.h"
 #include "UObject/LinkerPlaceholderExportObject.h"
+#include "UObject/LinkerLoadImportBehavior.h"
 #include "Misc/StringBuilder.h"
 
 /*-----------------------------------------------------------------------------
@@ -296,7 +297,7 @@ void FObjectPropertyBase::ExportText_Internal( FString& ValueStr, const void* Pr
  *
  * @return	true if the text is successfully resolved into a valid object reference of the correct type, false otherwise.
  */
-bool FObjectPropertyBase::ParseObjectPropertyValue(const FProperty* Property, UObject* OwnerObject, UClass* RequiredMetaClass, uint32 PortFlags, const TCHAR*& Buffer, UObject*& out_ResolvedValue, FUObjectSerializeContext* InSerializeContext /*= nullptr*/, bool bAllowAnyPackage /*= true*/)
+bool FObjectPropertyBase::ParseObjectPropertyValue(const FProperty* Property, UObject* OwnerObject, UClass* RequiredMetaClass, uint32 PortFlags, const TCHAR*& Buffer, TObjectPtr<UObject>& out_ResolvedValue, FUObjectSerializeContext* InSerializeContext /*= nullptr*/, bool bAllowAnyPackage /*= true*/)
 {
 	check(Property);
 	if (!RequiredMetaClass)
@@ -351,7 +352,7 @@ bool FObjectPropertyBase::ParseObjectPropertyValue(const FProperty* Property, UO
 			out_ResolvedValue = FObjectPropertyBase::FindImportedObject(Property, OwnerObject, ObjectClass, RequiredMetaClass, Temp.ToString(), PortFlags, InSerializeContext, bAllowAnyPackage);
 		}
 
-		if ( out_ResolvedValue != nullptr && !out_ResolvedValue->GetClass()->IsChildOf(RequiredMetaClass) )
+		if ( out_ResolvedValue != nullptr && !out_ResolvedValue.GetClass()->IsChildOf(RequiredMetaClass) )
 		{
 			if (bWarnOnnullptr )
 			{
@@ -379,7 +380,7 @@ bool FObjectPropertyBase::ParseObjectPropertyValue(const FProperty* Property, UO
 const TCHAR* FObjectPropertyBase::ImportText_Internal( const TCHAR* InBuffer, void* ContainerOrPropertyPtr, EPropertyPointerType PropertyPointerType, UObject* Parent, int32 PortFlags, FOutputDevice* ErrorText ) const
 {
 	const TCHAR* Buffer = InBuffer;
-	UObject* Result = nullptr;
+	TObjectPtr<UObject> Result = nullptr;
 	FLinkerLoad* Linker = GetLinker();
 
 	bool bOk = ParseObjectPropertyValue(this, Parent, PropertyClass, PortFlags, Buffer, Result, Linker ? Linker->GetSerializeContext() : nullptr);
@@ -410,18 +411,18 @@ const TCHAR* FObjectPropertyBase::ImportText_Internal( const TCHAR* InBuffer, vo
 
 	if (PropertyPointerType == EPropertyPointerType::Container && HasSetter())
 	{
-		SetObjectPropertyValue_InContainer(ContainerOrPropertyPtr, Result);
+		SetObjectPropertyValue_InContainer(ContainerOrPropertyPtr, Result); //TODO change this to not resolve TObjectPtr's
 	}
 	else
 	{
-		SetObjectPropertyValue(PointerToValuePtr(ContainerOrPropertyPtr, PropertyPointerType), Result);
+		SetObjectPtrPropertyValue(PointerToValuePtr(ContainerOrPropertyPtr, PropertyPointerType), Result);
 	}
 	return Buffer;
 }
 
-UObject* FObjectPropertyBase::FindImportedObject( const FProperty* Property, UObject* OwnerObject, UClass* ObjectClass, UClass* RequiredMetaClass, const TCHAR* Text, uint32 PortFlags/*=0*/, FUObjectSerializeContext* InSerializeContext /*= nullptr*/, bool bAllowAnyPackage /*= true*/)
+TObjectPtr<UObject> FObjectPropertyBase::FindImportedObject( const FProperty* Property, UObject* OwnerObject, UClass* ObjectClass, UClass* RequiredMetaClass, const TCHAR* Text, uint32 PortFlags/*=0*/, FUObjectSerializeContext* InSerializeContext /*= nullptr*/, bool bAllowAnyPackage /*= true*/)
 {
-	UObject*	Result = nullptr;
+	TObjectPtr<UObject>	Result = nullptr;
 	check( ObjectClass->IsChildOf(RequiredMetaClass) );
 
 	bool AttemptNonQualifiedSearch = (PortFlags & PPF_AttemptNonQualifiedSearch) != 0; 
@@ -522,6 +523,14 @@ UObject* FObjectPropertyBase::FindImportedObject( const FProperty* Property, UOb
 		// If we still can't find it, try to load it. (Only try to load fully qualified names)
 		if(!Result && Dot && !GIsSavingPackage)
 		{
+#if UE_WITH_OBJECT_HANDLE_LATE_RESOLVE 
+			FSoftObjectPath Path(Text);
+			if (UE::LinkerLoad::TryLazyLoad(*RequiredMetaClass, Path, Result))
+			{
+				return Result;
+			}
+#endif
+
 #if USE_CIRCULAR_DEPENDENCY_LOAD_DEFERRING
 			FLinkerLoad* Linker = (OwnerObject != nullptr) ? OwnerObject->GetClass()->GetLinker() : nullptr;
 			if (Linker == nullptr)
@@ -579,6 +588,11 @@ TObjectPtr<UObject> FObjectPropertyBase::GetObjectPtrPropertyValue(const void* P
 	return TObjectPtr<UObject>();
 }
 
+void FObjectPropertyBase::SetObjectPtrPropertyValue(void* PropertyValueAddress, TObjectPtr<UObject> Ptr) const
+{
+	SetObjectPropertyValue(PropertyValueAddress, Ptr.Get());
+}
+
 UObject* FObjectPropertyBase::GetObjectPropertyValue(const void* PropertyValueAddress) const
 {
 	checkf(false, TEXT("%s is missing implementation of GetObjectPropertyValue"), *GetFullName());
@@ -611,7 +625,7 @@ bool FObjectPropertyBase::AllowObjectTypeReinterpretationTo(const FObjectPropert
 	return false;
 }
 
-void FObjectPropertyBase::CheckValidObject(void* ValueAddress, UObject* OldValue) const
+void FObjectPropertyBase::CheckValidObject(void* ValueAddress, TObjectPtr<UObject> OldValue) const
 {
 	const TObjectPtr<UObject> Object = GetObjectPtrPropertyValue(ValueAddress);
 	if (!Object)
@@ -672,11 +686,11 @@ void FObjectPropertyBase::CheckValidObject(void* ValueAddress, UObject* OldValue
 					TEXT("Serialized %s for a property of %s. Reference will be reverted back to %s.\n    Property = %s\n    Item = %s"),
 					*ObjectClass->GetFullName(),
 					*PropertyClass->GetFullName(),
-					OldValue ? *OldValue->GetFullName() : TEXT("None"),
+					OldValue ? *OldValue.GetFullName() : TEXT("None"),
 					*GetFullName(),
 					*Object.GetFullName()
 				);
-				SetObjectPropertyValue(ValueAddress, OldValue);
+				SetObjectPtrPropertyValue(ValueAddress, OldValue);
 			}
 		}
 	}
diff --git a/Engine/Source/Runtime/CoreUObject/Private/UObject/PropertyInterface.cpp b/Engine/Source/Runtime/CoreUObject/Private/UObject/PropertyInterface.cpp
--- a/Engine/Source/Runtime/CoreUObject/Private/UObject/PropertyInterface.cpp
+++ b/Engine/Source/Runtime/CoreUObject/Private/UObject/PropertyInterface.cpp
@@ -230,7 +230,7 @@ void FInterfaceProperty::ExportText_Internal( FString& ValueStr, const void* Pro
 const TCHAR* FInterfaceProperty::ImportText_Internal( const TCHAR* InBuffer, void* ContainerOrPropertyPtr, EPropertyPointerType PropertyPointerType, UObject* Parent, int32 PortFlags, FOutputDevice* ErrorText) const
 {
 	FScriptInterface* InterfaceValue = (FScriptInterface*)PointerToValuePtr(ContainerOrPropertyPtr, PropertyPointerType);
-	UObject* ResolvedObject = InterfaceValue->GetObject();
+	TObjectPtr<UObject> ResolvedObject = InterfaceValue->GetObject();
 
 	auto SetInterfaceValue = [InterfaceValue, ContainerOrPropertyPtr, PropertyPointerType, this](UObject* NewObject, void* NewInterfaceAddress)
 	{
diff --git a/Engine/Source/Runtime/CoreUObject/Private/UObject/PropertyObject.cpp b/Engine/Source/Runtime/CoreUObject/Private/UObject/PropertyObject.cpp
--- a/Engine/Source/Runtime/CoreUObject/Private/UObject/PropertyObject.cpp
+++ b/Engine/Source/Runtime/CoreUObject/Private/UObject/PropertyObject.cpp
@@ -157,7 +157,7 @@ const TCHAR* FObjectProperty::ImportText_Internal(const TCHAR* Buffer, void* Con
 	if (Result)
 	{
 		void* Data = PointerToValuePtr(ContainerOrPropertyPtr, PropertyPointerType);
-		UObject* ObjectValue = GetObjectPropertyValue(Data);
+		TObjectPtr<UObject> ObjectValue = GetObjectPtrPropertyValue(Data);
 		CheckValidObject(Data, ObjectValue);
 
 #if USE_CIRCULAR_DEPENDENCY_LOAD_DEFERRING		
@@ -203,6 +203,18 @@ UObject* FObjectProperty::GetObjectPropertyValue_InContainer(const void* Contain
 	return Result;
 }
 
+void FObjectProperty::SetObjectPtrPropertyValue(void* PropertyValueAddress, TObjectPtr<UObject> Value) const
+{
+	if (Value || !HasAnyPropertyFlags(CPF_NonNullable))
+	{
+		SetPropertyValue(PropertyValueAddress, Value);
+	}
+	else
+	{
+		UE_LOG(LogProperty, Verbose /*Warning*/, TEXT("Trying to assign null object value to non-nullable \"%s\""), *GetFullName());
+	}
+}
+
 void FObjectProperty::SetObjectPropertyValue(void* PropertyValueAddress, UObject* Value) const
 {
 	if (Value || !HasAnyPropertyFlags(CPF_NonNullable))
diff --git a/Engine/Source/Runtime/CoreUObject/Public/UObject/UnrealType.h b/Engine/Source/Runtime/CoreUObject/Public/UObject/UnrealType.h
--- a/Engine/Source/Runtime/CoreUObject/Public/UObject/UnrealType.h
+++ b/Engine/Source/Runtime/CoreUObject/Public/UObject/UnrealType.h
@@ -2650,8 +2650,8 @@ public:
 	 * @param bAllowAnyPackage		allows ignoring package name to find any object that happens to be loaded with the same name
 	 * @return	true if the text is successfully resolved into a valid object reference of the correct type, false otherwise.
 	 */
-	static bool ParseObjectPropertyValue( const FProperty* Property, UObject* OwnerObject, UClass* RequiredMetaClass, uint32 PortFlags, const TCHAR*& Buffer, UObject*& out_ResolvedValue, FUObjectSerializeContext* InSerializeContext = nullptr, bool bAllowAnyPackage = true );
-	static UObject* FindImportedObject( const FProperty* Property, UObject* OwnerObject, UClass* ObjectClass, UClass* RequiredMetaClass, const TCHAR* Text, uint32 PortFlags = 0, FUObjectSerializeContext* InSerializeContext = nullptr, bool bAllowAnyPackage = true );
+	static bool ParseObjectPropertyValue( const FProperty* Property, UObject* OwnerObject, UClass* RequiredMetaClass, uint32 PortFlags, const TCHAR*& Buffer, TObjectPtr<UObject>& out_ResolvedValue, FUObjectSerializeContext* InSerializeContext = nullptr, bool bAllowAnyPackage = true );
+	static TObjectPtr<UObject> FindImportedObject( const FProperty* Property, UObject* OwnerObject, UClass* ObjectClass, UClass* RequiredMetaClass, const TCHAR* Text, uint32 PortFlags = 0, FUObjectSerializeContext* InSerializeContext = nullptr, bool bAllowAnyPackage = true );
 	
 	/**
 	 * Returns the qualified export path for a given object, parent, and export root scope
@@ -2687,6 +2687,7 @@ public:
 	virtual UObject* GetObjectPropertyValue(const void* PropertyValueAddress) const;
 	virtual UObject* GetObjectPropertyValue_InContainer(const void* ContainerAddress, int32 ArrayIndex = 0) const;
 	virtual void SetObjectPropertyValue(void* PropertyValueAddress, UObject* Value) const;
+	virtual void SetObjectPtrPropertyValue(void* PropertyValueAddress, TObjectPtr<UObject> Value) const;
 	virtual void SetObjectPropertyValue_InContainer(void* ContainerAddress, UObject* Value, int32 ArrayIndex = 0) const;
 
 	/**
@@ -2704,7 +2705,7 @@ public:
 	FORCEINLINE void SetPropertyClass(UClass* NewPropertyClass) { PropertyClass = NewPropertyClass; }
 #endif // USE_CIRCULAR_DEPENDENCY_LOAD_DEFERRING
 
-	virtual void CheckValidObject(void* ValueAddress, UObject* OldValue) const;
+	virtual void CheckValidObject(void* ValueAddress, TObjectPtr<UObject> OldValue) const;
 
 protected:
 	virtual bool AllowCrossLevel() const;
@@ -2957,6 +2958,7 @@ public:
 	virtual UObject* GetObjectPropertyValue(const void* PropertyValueAddress) const override;
 	virtual UObject* GetObjectPropertyValue_InContainer(const void* ContainerAddress, int32 ArrayIndex = 0) const override;
 	virtual void SetObjectPropertyValue(void* PropertyValueAddress, UObject* Value) const override;
+	virtual void SetObjectPtrPropertyValue(void* PropertyValueAddress, TObjectPtr<UObject> Value) const override;
 	virtual void SetObjectPropertyValue_InContainer(void* ContainerAddress, UObject* Value, int32 ArrayIndex = 0) const override;
 	virtual FString GetCPPTypeCustom(FString* ExtendedTypeText, uint32 CPPExportFlags, const FString& InnerNativeTypeName)  const override;
 	// End of FObjectPropertyBase interface
diff --git a/Engine/Source/Runtime/CoreUObject/Tests/ObjectPropertyTest.cpp b/Engine/Source/Runtime/CoreUObject/Tests/ObjectPropertyTest.cpp
--- a/Engine/Source/Runtime/CoreUObject/Tests/ObjectPropertyTest.cpp
+++ b/Engine/Source/Runtime/CoreUObject/Tests/ObjectPropertyTest.cpp
@@ -12,6 +12,8 @@
 #include "UObject/LinkerPlaceholderClass.h"
 #include "UObject/ObjectHandleTracking.h"
 #include "LowLevelTestsRunner/WarnFilterScope.h"
+#include "Misc/AssetRegistryInterface.h"
+#include "AssetRegistry/AssetData.h"
 
 TEST_CASE("UE::CoreUObject::FObjectProperty::CheckValidAddress")
 {
@@ -216,4 +218,132 @@ TEST_CASE("UE::CoreUObject::FObjectProperty::StaticSerializeItem")
 	TestSerializeItem<UObjectWithRawProperty>(TEXT("Object2"));
 }
 
+
+class MockAssetRegistryInterface : public IAssetRegistryInterface
+{
+public:
+	MockAssetRegistryInterface()
+		: Old(IAssetRegistryInterface::Default)
+	{
+		IAssetRegistryInterface::Default = this;
+	}
+
+	virtual ~MockAssetRegistryInterface()
+	{
+		IAssetRegistryInterface::Default = Old;
+	}
+
+	virtual void GetDependencies(FName InPackageName, TArray<FName>& OutDependencies, UE::AssetRegistry::EDependencyCategory Category = UE::AssetRegistry::EDependencyCategory::Package, const UE::AssetRegistry::FDependencyQuery& Flags = UE::AssetRegistry::FDependencyQuery()) override
+	{
+
+	}
+
+	virtual UE::AssetRegistry::EExists TryGetAssetByObjectPath(const FSoftObjectPath& ObjectPath, FAssetData& OutAssetData) const override
+	{
+
+		const FAssetData* Found = AssetData.Find(ObjectPath);
+		if (Found)
+		{
+			OutAssetData = *Found;
+			return UE::AssetRegistry::EExists::Exists;
+		}
+		return UE::AssetRegistry::EExists::DoesNotExist;
+	}
+
+	virtual UE::AssetRegistry::EExists TryGetAssetPackageData(FName PackageName, FAssetPackageData& OutPackageData) const override
+	{
+		return UE::AssetRegistry::EExists::Exists;
+	}
+
+	IAssetRegistryInterface* Old;
+	TMap<FSoftObjectPath, FAssetData> AssetData;
+};
+
+TEST_CASE("UE::CoreUObject::FObjectProperty::ParseObjectPropertyValue")
+{
+	UClass* Class = UObjectPtrTestClassWithRef::StaticClass();
+	FObjectProperty* Property = CastField<FObjectProperty>(Class->FindPropertyByName(TEXT("ObjectPtr")));
+
+	//make a fake entry for the asset registry
+	//this allows lazy load to return an unresolved handle
+	MockAssetRegistryInterface MockAssetRegistry;
+	const TCHAR* Text = TEXT("/TestPackageName.Other");
+	FSoftObjectPath Path(Text);
+	FAssetData AssetData;
+	AssetData.PackageName = "/TestPackageName";
+	AssetData.PackagePath = "/";
+	AssetData.AssetName = "Other";
+	AssetData.AssetClassPath.TrySetPath(UObjectPtrTestClass::StaticClass());
+
+	MockAssetRegistry.AssetData.Add(Path, AssetData);
+	UPackage* TestPackage = NewObject<UPackage>(nullptr, TEXT("TestPackageName"), RF_Transient);
+	TestPackage->AddToRoot();
+	ON_SCOPE_EXIT
+	{
+		TestPackage->RemoveFromRoot();
+	};
+	UObjectPtrTestClassWithRef* Obj = NewObject<UObjectPtrTestClassWithRef>(TestPackage, "ObjectName");
+	TObjectPtr<UObject> Result;
+	FObjectPropertyBase::ParseObjectPropertyValue(Property, Obj, UObjectPtrTestClass::StaticClass(), 0, Text, Result);
+#if UE_WITH_OBJECT_HANDLE_LATE_RESOLVE
+	CHECK(Result != nullptr);
+	CHECK(!Result.IsResolved());
+#else
+	CHECK(Result == nullptr);
+#endif
+
+	{
+		const TCHAR * Buffer = TEXT("ObjectPtr=/TestPackageName.Other");
+		TArray<FDefinedProperty> DefinedProperties;
+		Buffer = FProperty::ImportSingleProperty(Buffer, Obj, Class, Obj, PPF_Delimited, nullptr, DefinedProperties);
+
+#if UE_WITH_OBJECT_HANDLE_LATE_RESOLVE
+		CHECK(Obj->ObjectPtr != nullptr);
+		CHECK(!Obj->ObjectPtr.IsResolved());
+#else
+		CHECK(Obj->ObjectPtr == nullptr);
+#endif
+	}
+
+	{
+		const TCHAR* Buffer = TEXT("ObjectPtr=/TestPackageName.Other");
+
+		FProperty* ObjProperty = Class->FindPropertyByName(TEXT("ObjectPtr"));
+		ObjProperty->ImportText_InContainer(Buffer, Obj, Obj, 0);
+		CHECK(Obj->ObjectPtr == nullptr); //TODO this should not resolve and not be null
+		CHECK(Obj->ObjectPtr.IsResolved());
+
+	}
+
+	{
+		const TCHAR * Buffer = TEXT("ArrayObjPtr=(/TestPackageName.Other)");
+		TArray<FDefinedProperty> DefinedProperties;
+		Buffer = FProperty::ImportSingleProperty(Buffer, Obj, Class, Obj, PPF_Delimited, nullptr, DefinedProperties);
+
+#if UE_WITH_OBJECT_HANDLE_LATE_RESOLVE
+		CHECK(Obj->ArrayObjPtr.Num() == 1);
+		CHECK(!Obj->ArrayObjPtr[0].IsResolved());
+#else
+		CHECK(Obj->ArrayObjPtr.Num() == 0);
+#endif
+	}
+
+
+
+	{
+		const TCHAR* Buffer = TEXT("ArrayObjPtr=(/TestPackageName.Other)");
+		
+		FProperty* ArrayProperty = Class->FindPropertyByName(TEXT("ArrayObjPtr"));
+
+		ArrayProperty->ImportText_InContainer(Buffer, Obj, Obj, 0);
+
+#if UE_WITH_OBJECT_HANDLE_LATE_RESOLVE
+		CHECK(Obj->ArrayObjPtr.Num() == 1);
+		CHECK(!Obj->ArrayObjPtr[0].IsResolved());
+#else
+		CHECK(Obj->ArrayObjPtr.Num() == 0);
+#endif
+	}
+}
+
 #endif
diff --git a/Engine/Source/Runtime/CoreUObject/Tests/ObjectPtrTestClass.cpp b/Engine/Source/Runtime/CoreUObject/Tests/ObjectPtrTestClass.cpp
--- a/Engine/Source/Runtime/CoreUObject/Tests/ObjectPtrTestClass.cpp
+++ b/Engine/Source/Runtime/CoreUObject/Tests/ObjectPtrTestClass.cpp
@@ -7,6 +7,7 @@
 #include "UObject/Package.h"
 #include "UObject/MetaData.h"
 #include "UObject/UnrealType.h"
+#include "UObject/UObjectGlobals.h"
 
 //can not put the #if inside as the expansion of IMPLEMENT_CORE_INTRINSIC_CLASS fails
 #if WITH_EDITORONLY_DATA
@@ -30,14 +31,41 @@ IMPLEMENT_CORE_INTRINSIC_CLASS(UObjectPtrTestClass, UObject,
 IMPLEMENT_CORE_INTRINSIC_CLASS(UObjectPtrTestClassWithRef, UObject,
 	{
 		//add reflection info for the two properties
-		auto Property1 = new FObjectPtrProperty(Class, TEXT("ObjectPtr"), EObjectFlags::RF_NoFlags);
-		Property1->PropertyClass = UObjectPtrTestClass::StaticClass();
-		Class->AddCppProperty(Property1);
+		{
+			UECodeGen_Private::FObjectPropertyParams Params = { };
+			Params.NameUTF8 = "ObjectPtr";
+			Params.Offset = STRUCT_OFFSET(UObjectPtrTestClassWithRef, ObjectPtr);
+			Params.PropertyFlags = CPF_None;
+			Params.ObjectFlags = RF_Public | RF_Transient | RF_MarkAsNative;
+			Params.ClassFunc = nullptr;
+			auto Property = new FObjectPtrProperty(Class, Params );
+			Property->PropertyClass = UObjectPtrTestClass::StaticClass();
+		}
 
-		auto Property2 = new FObjectPtrProperty(Class, TEXT("ObjectPtrNonNullable"), EObjectFlags::RF_NoFlags);
-		Property2->PropertyClass = UObjectPtrTestClass::StaticClass();
-		Property2->SetPropertyFlags(EPropertyFlags::CPF_NonNullable); //make non nullable
-		Class->AddCppProperty(Property2);
+		{
+
+			UECodeGen_Private::FObjectPropertyParams Params = { };
+			Params.NameUTF8 = "ObjectPtrNonNullable";
+			Params.Offset = STRUCT_OFFSET(UObjectPtrTestClassWithRef, ObjectPtrNonNullable);
+			Params.PropertyFlags = EPropertyFlags::CPF_NonNullable;
+			Params.ObjectFlags = RF_Public | RF_Transient | RF_MarkAsNative;
+			Params.ClassFunc = nullptr;
+			auto Property = new FObjectPtrProperty(Class, Params);
+			Property->PropertyClass = UObjectPtrTestClass::StaticClass();
+		}
+
+
+		{
+			UECodeGen_Private::FArrayPropertyParams Params = { };
+			Params.NameUTF8 = "ArrayObjPtr";
+			Params.Offset = STRUCT_OFFSET(UObjectPtrTestClassWithRef, ArrayObjPtr);
+			Params.PropertyFlags = CPF_None;
+			Params.ObjectFlags = RF_Public | RF_Transient | RF_MarkAsNative;
+			auto Property = new FArrayProperty(Class, Params);
+			auto InnerProperty = new FObjectPtrProperty(Property, TEXT("Inner"), EObjectFlags::RF_NoFlags);
+			InnerProperty->PropertyClass = UObjectPtrTestClass::StaticClass();
+			Property->AddCppProperty(InnerProperty);
+		}
 
 		auto MetaData = Class->GetOutermost()->GetMetaData();
 		if (MetaData)
diff --git a/Engine/Source/Runtime/CoreUObject/Tests/ObjectPtrTestClass.h b/Engine/Source/Runtime/CoreUObject/Tests/ObjectPtrTestClass.h
--- a/Engine/Source/Runtime/CoreUObject/Tests/ObjectPtrTestClass.h
+++ b/Engine/Source/Runtime/CoreUObject/Tests/ObjectPtrTestClass.h
@@ -23,6 +23,7 @@ class UObjectPtrTestClassWithRef : public UObject
 public:
 	TObjectPtr<UObjectPtrTestClass> ObjectPtr;
 	TObjectPtr<UObjectPtrTestClass> ObjectPtrNonNullable;
+	TArray<TObjectPtr<UObjectPtrTestClass>> ArrayObjPtr;
 };
 
 
