lazyload support for zenloader
https://p4-swarm.epicgames.net/reviews/24507657

diff --git a/Engine/Source/Runtime/CoreUObject/Internal/UObject/LinkerLoadImportBehavior.h b/Engine/Source/Runtime/CoreUObject/Internal/UObject/LinkerLoadImportBehavior.h
--- a/Engine/Source/Runtime/CoreUObject/Internal/UObject/LinkerLoadImportBehavior.h
+++ b/Engine/Source/Runtime/CoreUObject/Internal/UObject/LinkerLoadImportBehavior.h
@@ -31,6 +31,14 @@ EImportBehavior GetPropertyImportLoadBehavior(const FObjectImport& Import, const
 /// @return true if import was lazy loaded otherwise false
 bool TryLazyImport(const IAssetRegistryInterface& AssetRegistry, const FObjectImport& Import, const FLinkerLoad& LinkerLoad, FObjectPtr& ObjectPtr);
 
+/// @brief determines if an import can be lazy loaded
+/// 
+/// @param AssetRegistry to see if import exists
+/// @param Import to try to lazy load
+/// @param LinkerLoad for the Import
+/// @return true if import can be lazy loaded otherwise false
+bool CanLazyImport(const IAssetRegistryInterface& AssetRegistry, const FObjectImport& Import, const FLinkerLoad& LinkerLoad);
+
 /// @brief tries to resolve an object path to a lazy loaded TObjectPtr
 ///
 /// @param Class type of the object that will be resolved
@@ -39,7 +47,6 @@ bool TryLazyImport(const IAssetRegistryInterface& AssetRegistry, const FObjectIm
 /// @return true if was lazy loaded otherwise false
 bool TryLazyLoad(const UClass& Class, const FSoftObjectPath& ObjectPath, TObjectPtr<UObject>& OutObjectPtr);
 
-
 /// @brief Finds LoadBehavior meta data recursively
 /// @return Eager by default in not found
 EImportBehavior FindLoadBehavior(const UClass& Class);
diff --git a/Engine/Source/Runtime/CoreUObject/Private/Serialization/AsyncLoading2.cpp b/Engine/Source/Runtime/CoreUObject/Private/Serialization/AsyncLoading2.cpp
--- a/Engine/Source/Runtime/CoreUObject/Private/Serialization/AsyncLoading2.cpp
+++ b/Engine/Source/Runtime/CoreUObject/Private/Serialization/AsyncLoading2.cpp
@@ -57,6 +57,7 @@
 #include "Serialization/MemoryReader.h"
 #include "UObject/UObjectClusters.h"
 #include "UObject/LinkerInstancingContext.h"
+#include "UObject/LinkerLoadImportBehavior.h"
 #include "UObject/ObjectSerializeAccessScope.h"
 #include "ProfilingDebugging/CountersTrace.h"
 #include "ProfilingDebugging/AssetMetadataTrace.h"
@@ -67,6 +68,7 @@
 #include "Modules/ModuleManager.h"
 #include "Containers/MpscQueue.h"
 #include "Containers/SpscQueue.h"
+#include "Misc/AssetRegistryInterface.h"
 #include "Misc/PathViews.h"
 #include "UObject/LinkerLoad.h"
 #include "Containers/SpscQueue.h"
@@ -5038,6 +5040,13 @@ EEventLoadNodeExecutionResult FAsyncPackage2::ProcessLinkerLoadPackageSummary(FA
 		}
 		else
 		{
+#if UE_WITH_OBJECT_HANDLE_LATE_RESOLVE
+			auto AssetRegistry = IAssetRegistryInterface::GetPtr();
+			if (UE::LinkerLoad::CanLazyImport(*AssetRegistry, LinkerImport, *LinkerLoadState->Linker))
+			{
+				continue;
+			}
+#endif
 			FPackageId ImportedPackageId = FPackageId::FromName(ImportPackageName);
 			int32 ImportedPackageIndex = ImportedPackageIds.AddUnique(ImportedPackageId);
 			if (ImportedPackageIndex == ImportedPackageNames.Num())
diff --git a/Engine/Source/Runtime/CoreUObject/Private/UObject/LinkerLoadImportBehavior.cpp b/Engine/Source/Runtime/CoreUObject/Private/UObject/LinkerLoadImportBehavior.cpp
--- a/Engine/Source/Runtime/CoreUObject/Private/UObject/LinkerLoadImportBehavior.cpp
+++ b/Engine/Source/Runtime/CoreUObject/Private/UObject/LinkerLoadImportBehavior.cpp
@@ -49,7 +49,7 @@ EImportBehavior FindLoadBehavior(const UClass& Class)
 
 EImportBehavior GetPropertyImportLoadBehavior(const FObjectImport& Import, const FLinkerLoad& LinkerLoad)
 {
-	if (Import.bImportSearchedFor)
+	if (Import.bImportSearchedFor || Import.XObject)
 	{
 		// If it was something that's been searched for, we've already attempted a resolve, might as well use it
 		return EImportBehavior::Eager;
@@ -79,7 +79,7 @@ EImportBehavior GetPropertyImportLoadBehavior(const FObjectImport& Import, const
 }
 
 // recursively handles FAssetData redirectors
-static bool HandleRedirector(const IAssetRegistryInterface& AssetRegistry, const FAssetData& InAssetData, TSet<FSoftObjectPath>& SeenPaths, FObjectPtr& OutObjectPtr)
+static bool HandleRedirector(const IAssetRegistryInterface& AssetRegistry, const FAssetData& InAssetData, TSet<FSoftObjectPath>& SeenPaths, FAssetData& AssetData)
 {
 	using namespace UE::CoreUObject::Private;
 
@@ -100,7 +100,6 @@ static bool HandleRedirector(const IAssetRegistryInterface& AssetRegistry, const
 		return false;
 	}
 
-	FAssetData AssetData;
 	UE::AssetRegistry::EExists Exists = AssetRegistry.TryGetAssetByObjectPath(RedirectedPathName, AssetData);
 	if (Exists == UE::AssetRegistry::EExists::Unknown)
 	{
@@ -114,24 +113,14 @@ static bool HandleRedirector(const IAssetRegistryInterface& AssetRegistry, const
 	}
 	if (AssetData.IsRedirector())
 	{
-		return HandleRedirector(AssetRegistry, AssetData, SeenPaths, OutObjectPtr);
+		return HandleRedirector(AssetRegistry, AssetData, SeenPaths, AssetData);
 	}
-
-	FNameBuilder NameBuilder;
-	AssetData.AssetName.AppendString(NameBuilder);
-	FObjectRef ImportRef( AssetData.PackageName, AssetData.AssetClassPath.GetPackageName(), AssetData.AssetClassPath.GetAssetName(), FObjectPathId(NameBuilder) );
-	FPackedObjectRef PackedObjectRef = MakePackedObjectRef(ImportRef);
-	FObjectPtr Ptr({ PackedObjectRef.EncodedRef });
-	OutObjectPtr = Ptr;
 	return true;
 }
 
-static bool TryLazyLoad(const IAssetRegistryInterface& AssetRegistry, const FSoftObjectPath& ObjectPath, FObjectPtr& OutObjectPtr)
+static bool CanLazyImport(const IAssetRegistryInterface& AssetRegistry, const FSoftObjectPath& ObjectPath, FAssetData& AssetData)
 {
 	using namespace UE::CoreUObject::Private;
-	OutObjectPtr = nullptr;
-	
-	FAssetData AssetData;
 	UE::AssetRegistry::EExists Exists = AssetRegistry.TryGetAssetByObjectPath(ObjectPath, AssetData);
 	if (Exists == UE::AssetRegistry::EExists::Unknown)
 	{
@@ -145,23 +134,58 @@ static bool TryLazyLoad(const IAssetRegistryInterface& AssetRegistry, const FSof
 	if (AssetData.IsRedirector())
 	{
 		TSet<FSoftObjectPath> SeenPaths;
-		return HandleRedirector(AssetRegistry, AssetData, SeenPaths, OutObjectPtr);
+		return HandleRedirector(AssetRegistry, AssetData, SeenPaths, AssetData);
 	}
-	FNameBuilder NameBuilder;
-	AssetData.AssetName.AppendString(NameBuilder);
-
-	FObjectRef ImportRef = { AssetData.PackageName, AssetData.AssetClassPath.GetPackageName(), AssetData.AssetClassPath.GetAssetName(), FObjectPathId(NameBuilder) };
-	FPackedObjectRef PackedObjectRef = MakePackedObjectRef(ImportRef);
-	FObjectPtr Ptr({ PackedObjectRef.EncodedRef });
-	OutObjectPtr = Ptr;
 	return true;
 }
 
-bool TryLazyImport(const IAssetRegistryInterface& AssetRegistry, const FObjectImport& Import, const FLinkerLoad& LinkerLoad, FObjectPtr& ObjectPtr)
+
+bool CanLazyImport(const IAssetRegistryInterface& AssetRegistry, const FObjectImport& Import, const FLinkerLoad& LinkerLoad)
 {
+	//most of this is duplicated from TryLazyImport with the out parameter but it avoids a bunch of string allocations
 	using namespace UE::CoreUObject::Private;
 	EImportBehavior Behavior = GetPropertyImportLoadBehavior(Import, LinkerLoad);
-	if (Behavior != EImportBehavior::LazyOnDemand)
+	if (GIsPlayInEditorWorld || Behavior != EImportBehavior::LazyOnDemand)
+	{
+		return false;
+	}
+
+	static const FName Name_CoreUObjectPackage("/Script/CoreUObject");
+
+	//packages need to handles differently since the can't be found by object path in the AssetRegistry
+	if (Import.ClassPackage == Name_CoreUObjectPackage && Import.ClassName == NAME_Package)
+	{
+		FAssetPackageData PackageData;
+		UE::AssetRegistry::EExists Exists = AssetRegistry.TryGetAssetPackageData(Import.ObjectName, PackageData);
+		if (Exists == UE::AssetRegistry::EExists::Unknown)
+		{
+			return false;
+		}
+		if (Exists == UE::AssetRegistry::EExists::DoesNotExist)
+		{
+			//package doesn't exist. resolve to nullptr
+			return true;
+		}
+		return true;
+	}
+
+	//build the a full objectpath for the AssetRegistry
+	FObjectPathId ObjectPath;
+	FName PackageName = FObjectPathId::MakeImportPathIdAndPackageName(Import, LinkerLoad, ObjectPath);
+	FObjectRef ImportRef(PackageName, Import.ClassPackage, Import.ClassName, ObjectPath);
+
+	TStringBuilder<FName::StringBufferSize> PathName;
+	ImportRef.AppendPathName(PathName);
+
+	FAssetData AssetData;
+	return CanLazyImport(AssetRegistry, FSoftObjectPath(PathName), AssetData);
+}
+
+
+bool TryLazyImport(const IAssetRegistryInterface& AssetRegistry, const FObjectImport& Import, const FLinkerLoad& LinkerLoad, FObjectPtr& ObjectPtr)
+{
+	using namespace UE::CoreUObject::Private;
+	if (GIsPlayInEditorWorld || GetPropertyImportLoadBehavior(Import, LinkerLoad) != EImportBehavior::LazyOnDemand )
 	{
 		return false;
 	}
@@ -185,6 +209,12 @@ bool TryLazyImport(const IAssetRegistryInterface& AssetRegistry, const FObjectIm
 			return true;
 		}
 
+		UPackage* Package = FindObjectFast<UPackage>(nullptr, Import.ObjectName);
+		if (Package)
+		{
+			//already loaded don't bother with setting up lazy load
+			return false;
+		}
 		FObjectPathId ObjectPath;
 		FObjectRef ImportRef(Import.ObjectName, Import.ClassPackage, Import.ClassName, ObjectPath);
 		FPackedObjectRef PackedObjectRef = MakePackedObjectRef(ImportRef);
@@ -200,13 +230,26 @@ bool TryLazyImport(const IAssetRegistryInterface& AssetRegistry, const FObjectIm
 	TStringBuilder<FName::StringBufferSize> PathName;
 	ImportRef.AppendPathName(PathName);
 
-	return TryLazyLoad(AssetRegistry, FSoftObjectPath(PathName), ObjectPtr);
+	FAssetData AssetData;
+	if (!CanLazyImport(AssetRegistry, FSoftObjectPath(PathName), AssetData))
+	{
+		return false;
+	}
+	
+	FNameBuilder NameBuilder;
+	AssetData.AssetName.AppendString(NameBuilder);
+	FObjectRef Ref = { AssetData.PackageName, AssetData.AssetClassPath.GetPackageName(), AssetData.AssetClassPath.GetAssetName(), FObjectPathId(NameBuilder) };
+	FPackedObjectRef PackedObjectRef = MakePackedObjectRef(Ref);
+	FObjectPtr Ptr({ PackedObjectRef.EncodedRef });
+	ObjectPtr = Ptr;
+	return true;
 }
 
 
 bool TryLazyLoad(const UClass& Class, const FSoftObjectPath& ObjectPath, TObjectPtr<UObject>& OutObjectPtr)
 {
-	if (!FLinkerLoad::IsImportLazyLoadEnabled())
+	using namespace UE::CoreUObject::Private;
+	if (GIsPlayInEditorWorld || !FLinkerLoad::IsImportLazyLoadEnabled())
 	{
 		return false;
 	}
@@ -222,7 +265,19 @@ bool TryLazyLoad(const UClass& Class, const FSoftObjectPath& ObjectPath, TObject
 		return false;
 	}
 	FObjectPtr& ObjectPtr = reinterpret_cast<FObjectPtr&>(OutObjectPtr);
-	return TryLazyLoad(*AssetRegistry, ObjectPath, ObjectPtr);
+
+	FAssetData AssetData;
+	if (!CanLazyImport(*AssetRegistry, ObjectPath, AssetData))
+	{
+		return false;
+	}
+	FNameBuilder NameBuilder;
+	AssetData.AssetName.AppendString(NameBuilder);
+	FObjectRef ImportRef = { AssetData.PackageName, AssetData.AssetClassPath.GetPackageName(), AssetData.AssetClassPath.GetAssetName(), FObjectPathId(NameBuilder) };
+	FPackedObjectRef PackedObjectRef = MakePackedObjectRef(ImportRef);
+	FObjectPtr Ptr({ PackedObjectRef.EncodedRef });
+	ObjectPtr = Ptr;
+	return true;
 }
 
 bool IsImportLazyLoadEnabled()
