Fix TSAN warning by making bHasFoundRequestedPackage a barrier for reading RequestedPackage

#rmx

diff --git a/Engine/Source/Runtime/CoreUObject/Private/Serialization/AsyncLoading2.cpp b/Engine/Source/Runtime/CoreUObject/Private/Serialization/AsyncLoading2.cpp
--- a/Engine/Source/Runtime/CoreUObject/Private/Serialization/AsyncLoading2.cpp
+++ b/Engine/Source/Runtime/CoreUObject/Private/Serialization/AsyncLoading2.cpp
@@ -2109,7 +2109,7 @@ public:
 	int32 RequestId;
 	FAsyncPackage2* RequestedPackage = nullptr;
 	FAsyncPackage2* RequestingPackageDebug = nullptr;
-	bool bHasFoundRequestedPackage = false;
+	std::atomic<bool> bHasFoundRequestedPackage { false };
 
 private:
 	std::atomic<int32> RefCount = 1;
@@ -6629,18 +6629,20 @@ void FAsyncLoadingThread2::UpdateSyncLoadContext(FAsyncLoadingThreadState2& Thre
 	{
 		return;
 	}
-	if (ThreadState.bCanAccessAsyncLoadingThreadData && !SyncLoadContext->bHasFoundRequestedPackage)
+	if (ThreadState.bCanAccessAsyncLoadingThreadData && !SyncLoadContext->bHasFoundRequestedPackage.load(std::memory_order_relaxed))
 	{
 		// Ensure that we've created the package we're waiting for
 		CreateAsyncPackagesFromQueue(ThreadState);
 		if (FAsyncPackage2* RequestedPackage = RequestIdToPackageMap.FindRef(SyncLoadContext->RequestId))
 		{
-			SyncLoadContext->bHasFoundRequestedPackage = true;
+			// Set RequestedPackage before setting bHasFoundRequestedPackage so that another thread looking at RequestedPackage
+			// after validating that bHasFoundRequestedPackage is true would see the proper value.
 			SyncLoadContext->RequestedPackage = RequestedPackage;
+			SyncLoadContext->bHasFoundRequestedPackage.store(true, std::memory_order_release);
 			IncludePackageInSyncLoadContextRecursive(ThreadState, SyncLoadContext->ContextId, RequestedPackage);
 		}
 	}
-	if (SyncLoadContext->bHasFoundRequestedPackage && ThreadState.PackagesOnStack.Contains(SyncLoadContext->RequestedPackage))
+	if (SyncLoadContext->bHasFoundRequestedPackage.load(std::memory_order_acquire) && ThreadState.PackagesOnStack.Contains(SyncLoadContext->RequestedPackage))
 	{
 		// Flushing a package while it's already being processed on the stack, if we're done preloading we let it pass and remove the request id
 		bool bPreloadIsDone = SyncLoadContext->RequestedPackage->AsyncPackageLoadingState >= EAsyncPackageLoadingState2::DeferredPostLoad;
