Deprecating UE_GC_TRACK_OBJ_AVAILABLE as it's not required. Simplified available UObject count tracking.

+ Removed GetObjectArrayNumAvailable which was a temp workaround

diff --git a/Engine/Source/Editor/UnrealEd/Private/CookOnTheFlyServer.cpp b/Engine/Source/Editor/UnrealEd/Private/CookOnTheFlyServer.cpp
--- a/Engine/Source/Editor/UnrealEd/Private/CookOnTheFlyServer.cpp
+++ b/Engine/Source/Editor/UnrealEd/Private/CookOnTheFlyServer.cpp
@@ -4510,7 +4510,6 @@ void UCookOnTheFlyServer::TickPrecacheObjectsForPlatforms(const float TimeSlice,
 
 bool UCookOnTheFlyServer::PumpHasExceededMaxMemory(uint32& OutResultFlags)
 {
-#if UE_GC_TRACK_OBJ_AVAILABLE
 	if (GUObjectArray.GetObjectArrayEstimatedAvailable() < MinFreeUObjectIndicesBeforeGC)
 	{
 		UE_LOG(LogCook, Display, TEXT("Running out of available UObject indices (%d remaining)"), GUObjectArray.GetObjectArrayEstimatedAvailable());
@@ -4524,8 +4523,6 @@ bool UCookOnTheFlyServer::PumpHasExceededMaxMemory(uint32& OutResultFlags)
 		OutResultFlags |= COSR_RequiresGC | COSR_RequiresGC_OOM | COSR_YieldTick;
 		return true;
 	}
-#endif // UE_GC_TRACK_OBJ_AVAILABLE
-
 
 	TStringBuilder<256> TriggerMessages;
 	const FPlatformMemoryStats MemStats = FPlatformMemory::GetStats();
diff --git a/Engine/Source/Runtime/CoreUObject/Private/UObject/UObjectArray.cpp b/Engine/Source/Runtime/CoreUObject/Private/UObject/UObjectArray.cpp
--- a/Engine/Source/Runtime/CoreUObject/Private/UObject/UObjectArray.cpp
+++ b/Engine/Source/Runtime/CoreUObject/Private/UObject/UObjectArray.cpp
@@ -217,10 +217,8 @@ void FUObjectArray::AllocateUObjectIndex(UObjectBase* Object, EInternalObjectFla
 		if (ObjAvailableList.Num() > 0)
 		{
 			Index = ObjAvailableList.Pop();
-#if UE_GC_TRACK_OBJ_AVAILABLE
-			const int32 AvailableCount = ObjAvailableCount.Decrement();
+			const int32 AvailableCount = ObjAvailableList.Num();
 			checkSlow(AvailableCount >= 0);
-#endif
 		}
 		else
 		{
@@ -306,9 +304,6 @@ void FUObjectArray::FreeUObjectIndex(UObjectBase* Object)
 	if (Index > ObjLastNonGCIndex && !GExitPurge && bShouldRecycleObjectIndices)
 	{
 		ObjAvailableList.Add(Index);
-#if UE_GC_TRACK_OBJ_AVAILABLE
-		ObjAvailableCount.Increment();
-#endif
 	}
 }
 
diff --git a/Engine/Source/Runtime/CoreUObject/Private/UObject/UObjectHash.cpp b/Engine/Source/Runtime/CoreUObject/Private/UObject/UObjectHash.cpp
--- a/Engine/Source/Runtime/CoreUObject/Private/UObject/UObjectHash.cpp
+++ b/Engine/Source/Runtime/CoreUObject/Private/UObject/UObjectHash.cpp
@@ -21,11 +21,8 @@ DECLARE_CYCLE_STAT( TEXT( "GetObjectsOfClass" ), STAT_Hash_GetObjectsOfClass, ST
 #if !UE_BUILD_TEST && !UE_BUILD_SHIPPING
 DECLARE_CYCLE_STAT( TEXT( "HashObject" ), STAT_Hash_HashObject, STATGROUP_UObjectHash );
 DECLARE_CYCLE_STAT( TEXT( "UnhashObject" ), STAT_Hash_UnhashObject, STATGROUP_UObjectHash );
-#endif
-
-#if UE_GC_TRACK_OBJ_AVAILABLE
 DEFINE_STAT( STAT_Hash_NumObjects );
-#endif
+#endif // !UE_BUILD_TEST && !UE_BUILD_SHIPPING
 
 LLM_DEFINE_TAG(UObjectHash, TEXT("UObject hashtables"));
 
diff --git a/Engine/Source/Runtime/CoreUObject/Public/UObject/UObjectArray.h b/Engine/Source/Runtime/CoreUObject/Public/UObject/UObjectArray.h
--- a/Engine/Source/Runtime/CoreUObject/Public/UObject/UObjectArray.h
+++ b/Engine/Source/Runtime/CoreUObject/Public/UObject/UObjectArray.h
@@ -16,7 +16,7 @@
 * By default it is only tracked in WITH_EDITOR builds as it adds a small amount of tracking overhead
 */
 #if !defined(UE_GC_TRACK_OBJ_AVAILABLE)
-#define UE_GC_TRACK_OBJ_AVAILABLE (WITH_EDITOR)
+#define UE_GC_TRACK_OBJ_AVAILABLE UE_DEPRECATED_MACRO(5.2, "The UE_GC_TRACK_OBJ_AVAILABLE macro has been deprecated because it is no longer necessary.") 1
 #endif
 
 /**
@@ -897,26 +897,24 @@ public:
 		return ObjLastNonGCIndex + 1;
 	}
 
-#if UE_GC_TRACK_OBJ_AVAILABLE
 	/**
 	 * Returns the number of actual object indices that are claimed (the total size of the global object array minus
 	 * the number of available object array elements
 	 *
 	 * @return	The number of objects claimed
 	 */
-	int32 GetObjectArrayNumMinusAvailable()
+	int32 GetObjectArrayNumMinusAvailable() const
 	{
-		return ObjObjects.Num() - ObjAvailableCount.GetValue();
+		return ObjObjects.Num() - ObjAvailableList.Num();
 	}
 
 	/**
 	* Returns the estimated number of object indices available for allocation
 	*/
-	int32 GetObjectArrayEstimatedAvailable()
+	int32 GetObjectArrayEstimatedAvailable() const
 	{
 		return ObjObjects.Capacity() - GetObjectArrayNumMinusAvailable();
 	}
-#endif
 
 	/**
 	* Returns the estimated number of object indices available for allocation
@@ -926,14 +924,6 @@ public:
 		return ObjObjects.Capacity();
 	}
 
-	/**
-	* Returns the estimated number of object indices available for allocation
-	*/
-	int32 GetObjectArrayNumAvailable() const
-	{
-		return ObjObjects.Capacity() - ObjObjects.Num() + ObjAvailableList.Num();
-	}
-
 	/**
 	 * Clears some internal arrays to get rid of false memory leaks
 	 */
@@ -1113,10 +1103,7 @@ private:
 	mutable FCriticalSection ObjObjectsCritical;
 	/** Available object indices.											*/
 	TArray<int32> ObjAvailableList;
-#if UE_GC_TRACK_OBJ_AVAILABLE
-	/** Available object index count.										*/
-	FThreadSafeCounter ObjAvailableCount;
-#endif
+
 	/**
 	 * Array of things to notify when a UObjectBase is created
 	 */
diff --git a/Engine/Source/Runtime/CoreUObject/Public/UObject/UObjectHash.h b/Engine/Source/Runtime/CoreUObject/Public/UObject/UObjectHash.h
--- a/Engine/Source/Runtime/CoreUObject/Public/UObject/UObjectHash.h
+++ b/Engine/Source/Runtime/CoreUObject/Public/UObject/UObjectHash.h
@@ -27,9 +27,9 @@ class UPackage;
 
 DECLARE_STATS_GROUP_VERBOSE(TEXT("UObject Hash"), STATGROUP_UObjectHash, STATCAT_Advanced);
 
-#if UE_GC_TRACK_OBJ_AVAILABLE
+#if !UE_BUILD_TEST && !UE_BUILD_SHIPPING
 DECLARE_DWORD_COUNTER_STAT_EXTERN(TEXT("NumObjects"), STAT_Hash_NumObjects, STATGROUP_UObjectHash, COREUOBJECT_API);
-#endif
+#endif // !UE_BUILD_TEST && !UE_BUILD_SHIPPING
 
 /**
  * Private internal version of StaticFindObjectFast that allows using 0 exclusion flags.
diff --git a/Engine/Source/Runtime/Launch/Private/LaunchEngineLoop.cpp b/Engine/Source/Runtime/Launch/Private/LaunchEngineLoop.cpp
--- a/Engine/Source/Runtime/Launch/Private/LaunchEngineLoop.cpp
+++ b/Engine/Source/Runtime/Launch/Private/LaunchEngineLoop.cpp
@@ -6098,9 +6098,9 @@ void FEngineLoop::Tick()
 		SET_FLOAT_STAT( STAT_CPUTimePctRelative, CPUTime.CPUTimePctRelative );
 
 		// Set the UObject count stat
-#if UE_GC_TRACK_OBJ_AVAILABLE
+#if !UE_BUILD_TEST && !UE_BUILD_SHIPPING
 		SET_DWORD_STAT(STAT_Hash_NumObjects, GUObjectArray.GetObjectArrayNumMinusAvailable());
-#endif
+#endif // !UE_BUILD_TEST && !UE_BUILD_SHIPPING
 	}
 
 	TRACE_END_FRAME(TraceFrameType_Game);
