fix for failing tests on console

diff --git a/Engine/Source/Runtime/CoreUObject/Tests/ObjectPropertyTest.cpp b/Engine/Source/Runtime/CoreUObject/Tests/ObjectPropertyTest.cpp
--- a/Engine/Source/Runtime/CoreUObject/Tests/ObjectPropertyTest.cpp
+++ b/Engine/Source/Runtime/CoreUObject/Tests/ObjectPropertyTest.cpp
@@ -259,6 +259,7 @@ public:
 	TMap<FSoftObjectPath, FAssetData> AssetData;
 };
 
+#if UE_WITH_OBJECT_HANDLE_LATE_RESOLVE
 TEST_CASE("UE::CoreUObject::FObjectProperty::ParseObjectPropertyValue")
 {
 	UClass* Class = UObjectPtrTestClassWithRef::StaticClass();
@@ -285,24 +286,15 @@ TEST_CASE("UE::CoreUObject::FObjectProperty::ParseObjectPropertyValue")
 	UObjectPtrTestClassWithRef* Obj = NewObject<UObjectPtrTestClassWithRef>(TestPackage, "ObjectName");
 	TObjectPtr<UObject> Result;
 	FObjectPropertyBase::ParseObjectPropertyValue(Property, Obj, UObjectPtrTestClass::StaticClass(), 0, Text, Result);
-#if UE_WITH_OBJECT_HANDLE_LATE_RESOLVE
 	CHECK(Result != nullptr);
 	CHECK(!Result.IsResolved());
-#else
-	CHECK(Result == nullptr);
-#endif
-
 	{
 		const TCHAR * Buffer = TEXT("ObjectPtr=/TestPackageName.Other");
 		TArray<FDefinedProperty> DefinedProperties;
 		Buffer = FProperty::ImportSingleProperty(Buffer, Obj, Class, Obj, PPF_Delimited, nullptr, DefinedProperties);
 
-#if UE_WITH_OBJECT_HANDLE_LATE_RESOLVE
 		CHECK(Obj->ObjectPtr != nullptr);
 		CHECK(!Obj->ObjectPtr.IsResolved());
-#else
-		CHECK(Obj->ObjectPtr == nullptr);
-#endif
 	}
 
 	{
@@ -320,15 +312,48 @@ TEST_CASE("UE::CoreUObject::FObjectProperty::ParseObjectPropertyValue")
 		TArray<FDefinedProperty> DefinedProperties;
 		Buffer = FProperty::ImportSingleProperty(Buffer, Obj, Class, Obj, PPF_Delimited, nullptr, DefinedProperties);
 
-#if UE_WITH_OBJECT_HANDLE_LATE_RESOLVE
 		CHECK(Obj->ArrayObjPtr.Num() == 1);
 		CHECK(!Obj->ArrayObjPtr[0].IsResolved());
-#else
-		CHECK(Obj->ArrayObjPtr.Num() == 0);
-#endif
 	}
 
+	{
+		const TCHAR* Buffer = TEXT("ArrayObjPtr=(/TestPackageName.Other)");
+		
+		FProperty* ArrayProperty = Class->FindPropertyByName(TEXT("ArrayObjPtr"));
+
+		ArrayProperty->ImportText_InContainer(Buffer, Obj, Obj, 0);
+
+		CHECK(Obj->ArrayObjPtr.Num() == 1);
+		CHECK(!Obj->ArrayObjPtr[0].IsResolved());
+	}
+
+	{
+		const TCHAR * Buffer = TEXT("ObjectPtr=/TestPackageName.Other");
+		TArray<FDefinedProperty> DefinedProperties;
+		Buffer = FProperty::ImportSingleProperty(Buffer, Obj, Class, Obj, PPF_Delimited, nullptr, DefinedProperties);
+
+		CHECK(Obj->ObjectPtr != nullptr);
+		CHECK(!Obj->ObjectPtr.IsResolved());
+	}
+
+	{
+		const TCHAR* Buffer = TEXT("ObjectPtr=/TestPackageName.Other");
+
+		FProperty* ObjProperty = Class->FindPropertyByName(TEXT("ObjectPtr"));
+		ObjProperty->ImportText_InContainer(Buffer, Obj, Obj, 0);
+		CHECK(Obj->ObjectPtr == nullptr); //TODO this should not resolve and not be null
+		CHECK(Obj->ObjectPtr.IsResolved());
+
+	}
+
+	{
+		const TCHAR * Buffer = TEXT("ArrayObjPtr=(/TestPackageName.Other)");
+		TArray<FDefinedProperty> DefinedProperties;
+		Buffer = FProperty::ImportSingleProperty(Buffer, Obj, Class, Obj, PPF_Delimited, nullptr, DefinedProperties);
 
+		CHECK(Obj->ArrayObjPtr.Num() == 1);
+		CHECK(!Obj->ArrayObjPtr[0].IsResolved());
+	}
 
 	{
 		const TCHAR* Buffer = TEXT("ArrayObjPtr=(/TestPackageName.Other)");
@@ -337,13 +362,10 @@ TEST_CASE("UE::CoreUObject::FObjectProperty::ParseObjectPropertyValue")
 
 		ArrayProperty->ImportText_InContainer(Buffer, Obj, Obj, 0);
 
-#if UE_WITH_OBJECT_HANDLE_LATE_RESOLVE
 		CHECK(Obj->ArrayObjPtr.Num() == 1);
 		CHECK(!Obj->ArrayObjPtr[0].IsResolved());
-#else
-		CHECK(Obj->ArrayObjPtr.Num() == 0);
-#endif
 	}
 }
+#endif
 
 #endif
