Fix core redirect package substring being case sensitive. (regression from CL# 19199624)

Since package names are stored in case insensitive FNames, the redirectors should be case insensitive as well.

diff --git a/Engine/Source/Runtime/CoreUObject/Private/UObject/CoreRedirects.cpp b/Engine/Source/Runtime/CoreUObject/Private/UObject/CoreRedirects.cpp
--- a/Engine/Source/Runtime/CoreUObject/Private/UObject/CoreRedirects.cpp
+++ b/Engine/Source/Runtime/CoreUObject/Private/UObject/CoreRedirects.cpp
@@ -105,8 +105,8 @@ bool FCoreRedirectObjectName::Matches(const FCoreRedirectObjectName& Other, EMat
 		}
 		else if (bSubString)
 		{
-			// Much slower
-			return WriteToString<FName::StringBufferSize>(RHS).ToView().Contains(WriteToString<FName::StringBufferSize>(LHS));
+			// Much slower, see if RHS contains the substring LHS.
+			return INDEX_NONE != UE::String::FindFirst(WriteToString<FName::StringBufferSize>(RHS), WriteToString<FName::StringBufferSize>(LHS), ESearchCase::IgnoreCase);
 		}
 		else
 		{
