Created UEnum::GetValueOrBitfieldAsString to allow UENUM(Flags) types to be converted to a human readable string of flags (FlagA | FlagB).

diff --git a/Engine/Source/Runtime/CoreUObject/Private/UObject/Enum.cpp b/Engine/Source/Runtime/CoreUObject/Private/UObject/Enum.cpp
--- a/Engine/Source/Runtime/CoreUObject/Private/UObject/Enum.cpp
+++ b/Engine/Source/Runtime/CoreUObject/Private/UObject/Enum.cpp
@@ -499,6 +499,36 @@ FString UEnum::GetNameStringByValue(int64 Value) const
 	return GetNameStringByIndex(Index);
 }
 
+FString UEnum::GetValueOrBitfieldAsString(int64 InValue) const
+{
+	if (!HasAnyEnumFlags(EEnumFlags::Flags) || InValue == 0)
+	{
+		return GetNameStringByValue(InValue);
+	}
+	else
+	{
+		FString BitfieldString;
+		bool WroteFirstFlag = false;
+		while (InValue != 0)
+		{
+			int64 NextValue = 1ll << FMath::CountTrailingZeros64(InValue);
+			InValue = InValue & ~NextValue;
+			if (WroteFirstFlag)
+			{
+				// We don't just want to use the NameValuePair.Key because we want to strip enum class prefixes
+				BitfieldString.Appendf(TEXT(" | %s"), *GetNameStringByValue(NextValue));
+			}
+			else
+			{
+				// We don't just want to use the NameValuePair.Key because we want to strip enum class prefixes
+				BitfieldString.Appendf(TEXT("%s"), *GetNameStringByValue(NextValue));
+				WroteFirstFlag = true;
+			}
+		}
+		return BitfieldString;
+	}
+}
+
 bool UEnum::FindNameStringByValue(FString& Out, int64 InValue) const
 {
 	int32 Index = GetIndexByValue(InValue);
@@ -927,3 +957,77 @@ IMPLEMENT_CORE_INTRINSIC_CLASS(UEnum, UField,
 	{
 	}
 );
+
+#if WITH_DEV_AUTOMATION_TESTS 
+
+#include "Misc/AutomationTest.h"
+IMPLEMENT_SIMPLE_AUTOMATION_TEST(FEnumBitfieldTest, "System.CoreUObject.EnumBitfields", EAutomationTestFlags::ApplicationContextMask | EAutomationTestFlags::EngineFilter);
+bool FEnumBitfieldTest::RunTest(const FString& Parameters)
+{
+	UPackage* NativePackage = CreatePackage(TEXT("/Script/TestEnumBitfieldsPackage"));
+	NativePackage->SetPackageFlags(PKG_CompiledIn);
+
+	{
+		const FName BitfieldTestEnumName(TEXT("EBitfieldTestEnum"));
+		TArray<TPair<FName, int64>> BitfieldTestEnumValueNames =
+		{
+		 { TEXT("EBitfieldTestEnum::TestValue0"), 0 },
+		 { TEXT("EBitfieldTestEnum::TestValue1"), 1 << 0 },
+		 { TEXT("EBitfieldTestEnum::TestValue2"), 1 << 1 },
+		 { TEXT("EBitfieldTestEnum::TestValue4"), 1 << 2 },
+		 { TEXT("EBitfieldTestEnum::TestValue2to63"), 1ll << 63 },
+		};
+
+		UEnum* NativeEnum = NewObject<UEnum>(NativePackage, BitfieldTestEnumName);
+		NativeEnum->SetEnums(BitfieldTestEnumValueNames, UEnum::ECppForm::EnumClass);
+		NativeEnum->SetEnumFlags(EEnumFlags::Flags);
+
+		FString BitfieldTestValueOf3 = NativeEnum->GetValueOrBitfieldAsString(3);
+		TestEqual(TEXT("Test bitfield with value 3"), BitfieldTestValueOf3, FString(TEXT("TestValue1 | TestValue2")));
+
+		FString BitfieldTestValueOf1 = NativeEnum->GetValueOrBitfieldAsString(1);
+		TestEqual(TEXT("Test bitfield with value 1"), BitfieldTestValueOf1, FString(TEXT("TestValue1")));
+
+		FString BitfieldTestValueOf0 = NativeEnum->GetValueOrBitfieldAsString(0);
+		TestEqual(TEXT("Test bitfield with value 0"), BitfieldTestValueOf0, FString(TEXT("TestValue0")));
+
+		FString BitfieldTestValueOf2to63 = NativeEnum->GetValueOrBitfieldAsString(1ll << 63);
+		TestEqual(TEXT("Test bitfield with value 2^63"), BitfieldTestValueOf2to63, FString(TEXT("TestValue2to63")));
+
+		FString BitfieldTestValueOf8 = NativeEnum->GetValueOrBitfieldAsString(1ll << 3);
+		TestEqual(TEXT("Test bitfield with invalid value 8"), BitfieldTestValueOf8, FString(TEXT("")));
+	}
+
+	{
+		const FName BitfieldTestStandardEnumName(TEXT("EBitfieldTestStandardEnum"));
+		TArray<TPair<FName, int64>> BitfieldTestStandardEnumValueNames =
+		{
+		 { TEXT("EBitfieldTestStandardEnum::TestValue0"), 0 },
+		 { TEXT("EBitfieldTestStandardEnum::TestValue1"), 1 },
+		 { TEXT("EBitfieldTestStandardEnum::TestValue2"), 2 },
+		 { TEXT("EBitfieldTestStandardEnum::TestValue3"), 3 },
+		 { TEXT("EBitfieldTestStandardEnum::TestValue2to63"), 1ll << 63 },
+		};
+
+		UEnum* NativeEnum = NewObject<UEnum>(NativePackage, BitfieldTestStandardEnumName);
+		NativeEnum->SetEnums(BitfieldTestStandardEnumValueNames, UEnum::ECppForm::EnumClass);
+
+		FString BitfieldTestValueOf3 = NativeEnum->GetValueOrBitfieldAsString(3);
+		TestEqual(TEXT("Test non bitfield with value 3"), BitfieldTestValueOf3, FString(TEXT("TestValue3")));
+
+		FString BitfieldTestValueOf1 = NativeEnum->GetValueOrBitfieldAsString(1);
+		TestEqual(TEXT("Test non bitfield with value 1"), BitfieldTestValueOf1, FString(TEXT("TestValue1")));
+
+		FString BitfieldTestValueOf0 = NativeEnum->GetValueOrBitfieldAsString(0);
+		TestEqual(TEXT("Test non bitfield with value 0"), BitfieldTestValueOf0, FString(TEXT("TestValue0")));
+
+		FString BitfieldTestValueOf2to63 = NativeEnum->GetValueOrBitfieldAsString(1ll << 63);
+		TestEqual(TEXT("Test non bitfield with value 2^63"), BitfieldTestValueOf2to63, FString(TEXT("TestValue2to63")));
+
+		FString BitfieldTestValueOf8 = NativeEnum->GetValueOrBitfieldAsString(8);
+		TestEqual(TEXT("Test non bitfield with bad value"), BitfieldTestValueOf8, FString(TEXT("")));
+	}
+	return true;
+}
+
+#endif // WITH_DEV_AUTOMATION_TESTS
diff --git a/Engine/Source/Runtime/CoreUObject/Public/UObject/Class.h b/Engine/Source/Runtime/CoreUObject/Public/UObject/Class.h
--- a/Engine/Source/Runtime/CoreUObject/Public/UObject/Class.h
+++ b/Engine/Source/Runtime/CoreUObject/Public/UObject/Class.h
@@ -1999,6 +1999,9 @@ public:
 	/** Returns the short name matching the enum Value, returns empty string if invalid */
 	COREUOBJECT_API FString GetNameStringByValue(int64 InValue) const;
 
+	/** If the enumeration is declared as UENUM(Flags), returns a string of the form A | B | C representing set bits A, B, and C. If it is not a bitfield, the result is the same as calling GetNameStringByValue*/
+	FString GetValueOrBitfieldAsString(int64 InValue) const;
+
 	/** Looks for a name with a given value and returns true and writes the name to Out if one was found */
 	COREUOBJECT_API bool FindNameStringByValue(FString& Out, int64 InValue) const;
 
@@ -2261,6 +2264,25 @@ public:
 		out_StringValue = GetValueAsString( EnumPath, EnumeratorValue );
 	}
 
+	template <typename T>
+	FORCEINLINE static FString GetValueOrBitfieldAsString(const TCHAR* EnumPath, const T EnumeratorValue)
+	{
+		// For the C++ enum.
+		static_assert(TIsEnum<T>::Value, "Should only call this with enum types");
+		return GetValueOrBitfieldAsString_Internal(EnumPath, (int64)EnumeratorValue);
+	}
+
+	template <typename T>
+	FORCEINLINE static FString GetValueOrBitfieldAsString(const TCHAR* EnumPath, const TEnumAsByte<T> EnumeratorValue)
+	{
+		return GetValueOrBitfieldAsString_Internal(EnumPath, (int64)EnumeratorValue.GetValue());
+	}
+
+	template< class T >
+	FORCEINLINE static void GetValueOrBitfieldAsString(const TCHAR* EnumPath, const T EnumeratorValue, FString& out_StringValue)
+	{
+		out_StringValue = GetValueOrBitfieldAsString( EnumPath, EnumeratorValue );
+	}
 	/**
 	 * @param EnumPath         Full enum path.
 	 * @param EnumeratorValue  Enumerator Value.
@@ -2422,6 +2444,13 @@ private:
 		return EnumClass->GetNameStringByValue(EnumeratorValue);
 	}
 
+	FORCEINLINE static FString GetValueOrBitfieldAsString_Internal(const TCHAR* EnumPath, const int64 EnumeratorValue)
+	{
+		UEnum* EnumClass = FindObject<UEnum>(nullptr, EnumPath);
+		UE_CLOG(!EnumClass, LogClass, Fatal, TEXT("Couldn't find enum '%s'"), EnumPath);
+		return EnumClass->GetValueOrBitfieldAsString(EnumeratorValue);
+	}
+
 	FORCEINLINE static FText GetDisplayValueAsText_Internal( const TCHAR* EnumPath, const int64 EnumeratorValue )
 	{
 		UEnum* EnumClass = FindObject<UEnum>(nullptr, EnumPath);
