ZenLoader in editor: Defer editor specific callbacks to fix incompatiblity with editordomain
#rnx

diff --git a/Engine/Source/Runtime/CoreUObject/Private/Serialization/AsyncLoading2.cpp b/Engine/Source/Runtime/CoreUObject/Private/Serialization/AsyncLoading2.cpp
--- a/Engine/Source/Runtime/CoreUObject/Private/Serialization/AsyncLoading2.cpp
+++ b/Engine/Source/Runtime/CoreUObject/Private/Serialization/AsyncLoading2.cpp
@@ -2573,10 +2573,6 @@ struct FAsyncPackage2
 	void DetachLinker();
 #endif
 
-#if WITH_EDITOR
-	void GetLoadedAssetsAndPackages(TSet<FWeakObjectPtr>& AssetList, TSet<UPackage*>& PackageList);
-#endif
-
 private:
 	void ImportPackagesRecursiveInner(FAsyncLoadingThreadState2& ThreadState, FIoBatch& IoBatch, FPackageStore& PackageStore, const TArrayView<FPackageId>& ImportedPackageIds, const TArrayView<FName>& ImportedPackageNames, int32& ImportedPackageIndex);
 
@@ -2890,8 +2886,9 @@ private:
 	/** [GAME THREAD] Game thread CompletedPackages list */
 	TArray<FAsyncPackage2*> CompletedPackages;
 #if WITH_EDITOR
-	/** [GAME THREAD] Game thread LoadedAssets list */
-	TSet<FWeakObjectPtr> LoadedAssets;
+	/** [GAME THREAD] */
+	TArray<UObject*> EditorLoadedAssets;
+	TArray<UPackage*> EditorCompletedUPackages;
 #endif
 	/** [ASYNC/GAME THREAD] Packages to be deleted from async thread */
 	TMpscQueue<FAsyncPackage2*> DeferredDeletePackages;
@@ -3214,6 +3211,9 @@ private:
 	}
 #endif
 
+#if WITH_EDITOR
+	void ConditionalProcessEditorCallbacks();
+#endif
 	void ConditionalBeginPostLoad(FAsyncLoadingThreadState2& ThreadState, FAsyncLoadingPostLoadGroup* PostLoadGroup);
 	void MergePostLoadGroups(FAsyncLoadingThreadState2& ThreadState, FAsyncLoadingPostLoadGroup* Target, FAsyncLoadingPostLoadGroup* Source);
 
@@ -6783,9 +6783,6 @@ EAsyncPackageState::Type FAsyncLoadingThread2::ProcessLoadedPackagesFromGameThre
 
 		bLocalDidSomething |= LoadedPackagesToProcess.Num() > 0;
 		TArray<FAsyncPackage2*, TInlineAllocator<4>> PackagesReadyForCallback;
-#if WITH_EDITOR
-		TSet<UPackage*> CompletedUPackages;
-#endif
 		for (int32 PackageIndex = 0; PackageIndex < LoadedPackagesToProcess.Num(); ++PackageIndex)
 		{
 			SCOPED_LOADTIMER(ProcessLoadedPackagesTime);
@@ -6887,8 +6884,18 @@ EAsyncPackageState::Type FAsyncLoadingThread2::ProcessLoadedPackagesFromGameThre
 				if (!Package->bLoadHasFailed)
 				{
 #if WITH_EDITOR
-					// In the editor we need to find any assets and packages and add them to list for later callback
-					Package->GetLoadedAssetsAndPackages(LoadedAssets, CompletedUPackages);
+					if (GIsEditor)
+					{
+						// In the editor we need to find any assets and packages and add them to list for later callback
+						EditorCompletedUPackages.Add(Package->LinkerRoot);
+						for (UObject* Object : Package->ConstructedObjects)
+						{
+							if (Object->IsAsset())
+							{
+								EditorLoadedAssets.Add(Object);
+							}
+						}
+					}
 #endif
 					Package->ClearConstructedObjects();
 				}
@@ -6928,16 +6935,6 @@ EAsyncPackageState::Type FAsyncLoadingThread2::ProcessLoadedPackagesFromGameThre
 			UE_ASYNC_PACKAGE_LOG(Verbose, Package->Desc, TEXT("GameThread: LoadCompleted"),
 				TEXT("All loading of package is done, and the async package and load request will be deleted."));
 		}
-#if WITH_EDITOR
-		// Call the global delegate for package endloads and set the bHasBeenLoaded flag that is used to
-		// check which packages have reached this state
-		for (UPackage* CompletedUPackage : CompletedUPackages)
-		{
-			CompletedUPackage->SetHasBeenEndLoaded(true);
-		}
-		FCoreUObjectDelegates::OnEndLoadPackage.Broadcast(
-			FEndLoadPackageContext{ CompletedUPackages.Array(), 0, false /* bSynchronous */ });
-#endif
 		
 		{
 			TArray<FFailedPackageRequest> LocalFailedPackageRequests;
@@ -6986,21 +6983,6 @@ EAsyncPackageState::Type FAsyncLoadingThread2::ProcessLoadedPackagesFromGameThre
 
 	if (Result == EAsyncPackageState::Complete)
 	{
-#if WITH_EDITOR
-		// In editor builds, call the asset load callback. This happens in both editor and standalone to match EndLoad
-		TSet<FWeakObjectPtr> TempLoadedAssets = MoveTemp(LoadedAssets);
-
-		// Make a copy because LoadedAssets could be modified by one of the OnAssetLoaded callbacks
-		for (const FWeakObjectPtr& WeakAsset : TempLoadedAssets)
-		{
-			// It may have been unloaded/marked pending kill since being added, ignore those cases
-			if (UObject* LoadedAsset = WeakAsset.Get())
-			{
-				FCoreUObjectDelegates::OnAssetLoaded.Broadcast(LoadedAsset);
-			}
-		}
-#endif
-
 		// We're not done until all packages have been deleted
 		Result = CompletedPackages.Num() ? EAsyncPackageState::PendingImports  : EAsyncPackageState::Complete;
 	}
@@ -7071,6 +7053,10 @@ EAsyncPackageState::Type FAsyncLoadingThread2::TickAsyncLoadingFromGameThread(FA
 		FCoreDelegates::OnAsyncLoadingFlushUpdate.Broadcast();
 	}
 
+#if WITH_EDITOR
+	ConditionalProcessEditorCallbacks();
+#endif
+
 	return Result;
 }
 
@@ -7861,25 +7847,6 @@ double FAsyncPackage2::GetLoadStartTime() const
 	return LoadStartTime;
 }
 
-#if WITH_EDITOR
-void FAsyncPackage2::GetLoadedAssetsAndPackages(TSet<FWeakObjectPtr>& AssetList, TSet<UPackage*>& PackageList)
-{
-	for (UObject* Object : ConstructedObjects)
-	{
-		if (IsValid(Object) && Object->IsAsset())
-		{
-			AssetList.Add(Object);
-		}
-	}
-
-	// All ConstructedObjects belong to this package, so we only have to consider the single package in this->LinkerRoot
-	if (LinkerRoot && !LinkerRoot->HasAnyFlags(RF_Transient) && !LinkerRoot->HasAnyPackageFlags(PKG_InMemoryOnly))
-	{
-		PackageList.Add(LinkerRoot);
-	}
-}
-#endif
-
 /**
  * Begin async loading process. Simulates parts of BeginLoad.
  *
@@ -8068,11 +8035,38 @@ void FAsyncPackage2::CallCompletionCallbacks(EAsyncLoadingResult::Type LoadingRe
 	CompletionCallbacks.Empty();
 }
 
-UPackage* FAsyncPackage2::GetLoadedPackage()
+#if WITH_EDITOR
+void FAsyncLoadingThread2::ConditionalProcessEditorCallbacks()
 {
-	UPackage* LoadedPackage = (!bLoadHasFailed) ? LinkerRoot : nullptr;
-	return LoadedPackage;
+	check(IsInGameThread());
+	if (!GameThreadState->SyncLoadContextStack.IsEmpty())
+	{
+		return;
+	}
+	while (!EditorCompletedUPackages.IsEmpty() || !EditorLoadedAssets.IsEmpty())
+	{
+		TArray<UObject*> LocalEditorLoadedAssets;
+		Swap(LocalEditorLoadedAssets, EditorLoadedAssets);
+		TArray<UPackage*> LocalEditorCompletedUPackages;
+		Swap(LocalEditorCompletedUPackages, EditorCompletedUPackages);
+
+		// Call the global delegate for package endloads and set the bHasBeenLoaded flag that is used to
+		// check which packages have reached this state
+		for (UPackage* CompletedUPackage : LocalEditorCompletedUPackages)
+		{
+			CompletedUPackage->SetHasBeenEndLoaded(true);
+		}
+		FCoreUObjectDelegates::OnEndLoadPackage.Broadcast(
+			FEndLoadPackageContext{ LocalEditorCompletedUPackages, 0, false /* bSynchronous */ });
+
+		// In editor builds, call the asset load callback. This happens in both editor and standalone to match EndLoad
+		for (UObject* LoadedObject : LocalEditorLoadedAssets)
+		{
+			FCoreUObjectDelegates::OnAssetLoaded.Broadcast(LoadedObject);
+		}
+	}
 }
+#endif
 
 void FAsyncPackage2::Cancel()
 {
@@ -8325,6 +8319,10 @@ void FAsyncLoadingThread2::FlushLoading(int32 RequestId)
 			GameThreadState->PackagesOnStack.Pop();
 		}
 
+#if WITH_EDITOR
+		ConditionalProcessEditorCallbacks();
+#endif
+
 		check(RequestId != INDEX_NONE || !IsAsyncLoadingPackages());
 	}
 }
