LinkerLoad: Fixed unsafe cast in GetExportLoadClass

diff --git a/Engine/Source/Runtime/CoreUObject/Private/UObject/LinkerLoad.cpp b/Engine/Source/Runtime/CoreUObject/Private/UObject/LinkerLoad.cpp
--- a/Engine/Source/Runtime/CoreUObject/Private/UObject/LinkerLoad.cpp
+++ b/Engine/Source/Runtime/CoreUObject/Private/UObject/LinkerLoad.cpp
@@ -3863,7 +3863,7 @@ UClass* FLinkerLoad::GetExportLoadClass(int32 Index)
 		VerifyImport(Export.ClassIndex.ToImport());
 	}
 
-	return (UClass*)IndexToObject(Export.ClassIndex);
+	return dynamic_cast<UClass*>(IndexToObject(Export.ClassIndex));
 }
 
 #if WITH_EDITORONLY_DATA
@@ -4757,21 +4757,7 @@ UObject* FLinkerLoad::CreateExport( int32 Index )
 			LoadClass = UClass::StaticClass();
 		}
 
-		UObjectRedirector* LoadClassRedirector = dynamic_cast<UObjectRedirector*>(LoadClass);
-		if( LoadClassRedirector)
-		{
-			// mark this export as unloadable (so that other exports that
-			// reference this one won't continue to execute the above logic), then return NULL
-			Export.bExportLoadFailed = true;
-
-			// otherwise, return NULL and let the calling code determine what to do
-			FString OuterName = Export.OuterIndex.IsNull() ? LinkerRoot->GetFullName() : GetFullImpExpName(Export.OuterIndex);
-			UE_ASSET_LOG(LogLinker, Warning, PackagePath, TEXT("CreateExport: Failed to load Outer for resource because its class is a redirector '%s': %s"), *Export.ObjectName.ToString(), *OuterName);
-			return nullptr;
-		}
-
 		check(LoadClass);
-		check(dynamic_cast<UClass*>(LoadClass) != NULL);
 
 		// Check for a valid superstruct while there is still time to safely bail, if this export has one
 		if( !Export.SuperIndex.IsNull() )
