Make ObjectFlags operations thread-safe by default

#rnx

diff --git a/Engine/Source/Runtime/CoreUObject/Private/UObject/UObjectBase.cpp b/Engine/Source/Runtime/CoreUObject/Private/UObject/UObjectBase.cpp
--- a/Engine/Source/Runtime/CoreUObject/Private/UObject/UObjectBase.cpp
+++ b/Engine/Source/Runtime/CoreUObject/Private/UObject/UObjectBase.cpp
@@ -284,11 +284,11 @@ void UObjectBase::SetExternalPackage(UPackage* InPackage)
 	HashObjectExternalPackage(this, InPackage);
 	if (InPackage)
 	{
-		SetFlagsTo(GetFlags() | RF_HasExternalPackage);
+		AtomicallySetFlags(RF_HasExternalPackage);
 	}
 	else
 	{
-		SetFlagsTo(GetFlags() & ~RF_HasExternalPackage);
+		AtomicallyClearFlags(RF_HasExternalPackage);
 	}
 }
 
@@ -354,7 +354,7 @@ bool UObjectBase::IsValidLowLevelFast(bool bRecursive /*= true*/) const
 	}
 
 	// These should all be 0.
-	const UPTRINT CheckZero = (ObjectFlags & ~RF_AllFlags) | ((UPTRINT)ClassPrivate & AlignmentCheck) | ((UPTRINT)OuterPrivate & AlignmentCheck);
+	const UPTRINT CheckZero = (GetFlagsInternal() & ~RF_AllFlags) | ((UPTRINT)ClassPrivate & AlignmentCheck) | ((UPTRINT)OuterPrivate & AlignmentCheck);
 	if (!!CheckZero)
 	{
 		UE_LOG(LogUObjectBase, Error, TEXT("Object flags are invalid or either Class or Outer is misaligned"));
diff --git a/Engine/Source/Runtime/CoreUObject/Public/UObject/UObjectBase.h b/Engine/Source/Runtime/CoreUObject/Public/UObject/UObjectBase.h
--- a/Engine/Source/Runtime/CoreUObject/Public/UObject/UObjectBase.h
+++ b/Engine/Source/Runtime/CoreUObject/Public/UObject/UObjectBase.h
@@ -184,6 +184,7 @@ protected:
 	 * Set the object flags directly
 	 *
 	 **/
+	UE_DEPRECATED(5.3, "This function is not thread-safe. Use AtomicallySetFlags or AtomicallyClearFlags instead.")
 	FORCEINLINE void SetFlagsTo( EObjectFlags NewFlags )
 	{
 		checkfSlow((NewFlags & ~RF_AllFlags) == 0, TEXT("%s flagged as 0x%x but is trying to set flags to RF_AllFlags"), *GetFName().ToString(), (int)ObjectFlags);
@@ -198,48 +199,54 @@ public:
 	 **/
 	FORCEINLINE EObjectFlags GetFlags() const
 	{
-		checkfSlow((ObjectFlags & ~RF_AllFlags) == 0, TEXT("%s flagged as RF_AllFlags"), *GetFName().ToString());
-		return ObjectFlags;
+		EObjectFlags Flags = (EObjectFlags)GetFlagsInternal();
+		checkfSlow((Flags & ~RF_AllFlags) == 0, TEXT("%s flagged as RF_AllFlags"), *GetFName().ToString());
+		return Flags;
 	}
 
 	/**
 	 *	Atomically adds the specified flags.
-	 *	Do not use unless you know what you are doing.
-	 *	Designed to be used only by parallel GC and UObject loading thread.
 	 */
 	FORCENOINLINE void AtomicallySetFlags( EObjectFlags FlagsToAdd )
 	{
-		int32 OldFlags = 0;
-		int32 NewFlags = 0;
-		do 
+		int32 OldFlags = GetFlagsInternal();
+		int32 NewFlags = OldFlags | FlagsToAdd;
+
+		// Fast path without atomics if already set
+		if (NewFlags == OldFlags)
 		{
-			OldFlags = ObjectFlags;
-			NewFlags = OldFlags | FlagsToAdd;
+			return;
 		}
-		while( FPlatformAtomics::InterlockedCompareExchange( (int32*)&ObjectFlags, NewFlags, OldFlags) != OldFlags );
+
+		FPlatformAtomics::InterlockedOr((int32*)&ObjectFlags, FlagsToAdd);
 	}
 
 	/**
 	 *	Atomically clears the specified flags.
-	 *	Do not use unless you know what you are doing.
-	 *	Designed to be used only by parallel GC and UObject loading thread.
 	 */
 	FORCENOINLINE void AtomicallyClearFlags( EObjectFlags FlagsToClear )
 	{
-		int32 OldFlags = 0;
-		int32 NewFlags = 0;
-		do 
+		int32 OldFlags = GetFlagsInternal();
+		int32 NewFlags = OldFlags & ~FlagsToClear;
+
+		// Fast path without atomics if already cleared
+		if (NewFlags == OldFlags)
 		{
-			OldFlags = ObjectFlags;
-			NewFlags = OldFlags & ~FlagsToClear;
+			return;
 		}
-		while( FPlatformAtomics::InterlockedCompareExchange( (int32*)&ObjectFlags, NewFlags, OldFlags) != OldFlags );
+
+		FPlatformAtomics::InterlockedAnd((int32*)&ObjectFlags, ~FlagsToClear);
 	}
 
 	static void PrefetchClass(UObject* Object) { FPlatformMisc::Prefetch(Object, offsetof(UObjectBase, ClassPrivate)); }
 	static void PrefetchOuter(UObject* Object) { FPlatformMisc::Prefetch(Object, offsetof(UObjectBase, OuterPrivate)); }
 
 private:
+	FORCEINLINE int32 GetFlagsInternal() const
+	{
+		static_assert(sizeof(int32) == sizeof(ObjectFlags), "Flags must be 32-bit for atomics.");
+		return FPlatformAtomics::AtomicRead_Relaxed((int32*)&ObjectFlags);
+	}
 
 	/** Flags used to track and report various object states. This needs to be 8 byte aligned on 32-bit
 	    platforms to reduce memory waste */
diff --git a/Engine/Source/Runtime/CoreUObject/Public/UObject/UObjectBaseUtility.h b/Engine/Source/Runtime/CoreUObject/Public/UObject/UObjectBaseUtility.h
--- a/Engine/Source/Runtime/CoreUObject/Public/UObject/UObjectBaseUtility.h
+++ b/Engine/Source/Runtime/CoreUObject/Public/UObject/UObjectBaseUtility.h
@@ -58,22 +58,22 @@ class COREUOBJECT_API UObjectBaseUtility : public UObjectBase
 {
 	FORCEINLINE void MarkPendingKillOnlyInternal()
 	{
-		SetFlagsTo(GetFlags() | RF_InternalPendingKill);
+		AtomicallySetFlags(RF_InternalPendingKill);
 		GUObjectArray.IndexToObject(InternalIndex)->SetPendingKill();
 	}
 	FORCEINLINE void ClearPendingKillOnlyInternal()
 	{
-		SetFlagsTo(GetFlags() & ~RF_InternalPendingKill);
+		AtomicallyClearFlags(RF_InternalPendingKill);
 		GUObjectArray.IndexToObject(InternalIndex)->ClearPendingKill();
 	}
 	FORCEINLINE void MarkAsGarbageOnlyInternal()
 	{
-		SetFlagsTo(GetFlags() | RF_InternalGarbage);
+		AtomicallySetFlags(RF_InternalGarbage);
 		GUObjectArray.IndexToObject(InternalIndex)->ThisThreadAtomicallySetFlag(EInternalObjectFlags::Garbage);
 	}
 	FORCEINLINE void ClearGarbageOnlyInternal()
 	{
-		SetFlagsTo(GetFlags() & ~RF_InternalGarbage);
+		AtomicallyClearFlags(RF_InternalGarbage);
 		GUObjectArray.IndexToObject(InternalIndex)->ThisThreadAtomicallyClearedFlag(EInternalObjectFlags::Garbage);
 	}
 
@@ -101,15 +101,15 @@ public:
 	{
 		checkSlow(!(NewFlags & (RF_MarkAsNative | RF_MarkAsRootSet | RF_InternalPendingKill | RF_InternalGarbage))); // These flags can't be used outside of constructors / internal code
 		checkf(!(NewFlags & RF_InternalMirroredFlags) || (GetFlags() & (NewFlags & RF_InternalMirroredFlags)) == (NewFlags & RF_InternalMirroredFlags), TEXT("RF_PendingKill and RF_garbage can not be set through SetFlags function. Use MarkAsGarbage() instead"));
-		SetFlagsTo(GetFlags() | NewFlags);
+		AtomicallySetFlags(NewFlags);
 	}
 
 	/** Clears subset of flags for a specific object */
-	FORCEINLINE void ClearFlags( EObjectFlags NewFlags )
+	FORCEINLINE void ClearFlags( EObjectFlags FlagsToClear )
 	{
-		checkSlow(!(NewFlags & (RF_MarkAsNative | RF_MarkAsRootSet | RF_InternalPendingKill | RF_InternalGarbage)) || NewFlags == RF_AllFlags); // These flags can't be used outside of constructors / internal code
-		checkf(!(NewFlags & RF_InternalMirroredFlags) || (GetFlags() & (NewFlags & RF_InternalMirroredFlags)) == RF_NoFlags, TEXT("RF_PendingKill and RF_garbage can not be cleared through ClearFlags function. Use ClearGarbage() instead"));
-		SetFlagsTo(GetFlags() & ~NewFlags);
+		checkSlow(!(FlagsToClear & (RF_MarkAsNative | RF_MarkAsRootSet | RF_InternalPendingKill | RF_InternalGarbage)) || FlagsToClear == RF_AllFlags); // These flags can't be used outside of constructors / internal code
+		checkf(!(FlagsToClear & RF_InternalMirroredFlags) || (GetFlags() & (FlagsToClear & RF_InternalMirroredFlags)) == RF_NoFlags, TEXT("RF_PendingKill and RF_garbage can not be cleared through ClearFlags function. Use ClearGarbage() instead"));
+		AtomicallyClearFlags(FlagsToClear);
 	}
 
 	/**
