[Engine] Fix for ImportText crashing when ErrorText is null.

   * Add if check to detect if ErrorText is null and setting it to a NullOutputDevice.

diff --git a/Engine/Source/Runtime/CoreUObject/Private/UObject/Class.cpp b/Engine/Source/Runtime/CoreUObject/Private/UObject/Class.cpp
--- a/Engine/Source/Runtime/CoreUObject/Private/UObject/Class.cpp
+++ b/Engine/Source/Runtime/CoreUObject/Private/UObject/Class.cpp
@@ -14,6 +14,7 @@
 #include "Misc/AutomationTest.h"
 #include "Misc/EnumClassFlags.h"
 #include "Misc/StringBuilder.h"
+#include "Misc/OutputDeviceNull.h"
 #include "UObject/CoreNet.h"
 #include "Modules/ModuleManager.h"
 #include "UObject/UObjectAllocator.h"
@@ -2828,6 +2829,12 @@ const TCHAR* UScriptStruct::ImportText(const TCHAR* InBuffer, void* Value, UObje
 
 const TCHAR* UScriptStruct::ImportText(const TCHAR* InBuffer, void* Value, UObject* OwnerObject, int32 PortFlags, FOutputDevice* ErrorText, const TFunctionRef<FString()>& StructNameGetter, bool bAllowNativeOverride)
 {
+	FOutputDeviceNull NullErrorText;
+	if (!ErrorText)
+	{
+		ErrorText = &NullErrorText;
+	}
+
 	if (bAllowNativeOverride && StructFlags & STRUCT_ImportTextItemNative)
 	{
 		UScriptStruct::ICppStructOps* TheCppStructOps = GetCppStructOps();
