Addition flavors of UEnum::GetValueOrBitfieldAsString for ease of use.

diff --git a/Engine/Source/Runtime/CoreUObject/Public/UObject/Class.h b/Engine/Source/Runtime/CoreUObject/Public/UObject/Class.h
--- a/Engine/Source/Runtime/CoreUObject/Public/UObject/Class.h
+++ b/Engine/Source/Runtime/CoreUObject/Public/UObject/Class.h
@@ -2361,6 +2361,52 @@ public:
 		out_StringValue = GetValueAsString(EnumeratorValue );
 	}
 
+	/**
+	 * @param EnumeratorValue  Enumerator Value.
+	 *
+	 * @return a string representing the combination of set bits in the bitfield or, if not a bitfield/flags enumeration, the string associated with the enumerator for the specified enum value for the enum specified by the template type.
+	 */
+	template<typename EnumType UE_REQUIRES(TIsEnum<EnumType>::Value)>
+	FORCEINLINE static FString GetValueOrBitfieldAsString(const EnumType EnumeratorValue)
+	{
+		// For the C++ enum.
+		static_assert(TIsEnum<EnumType>::Value, "Should only call this with enum types");
+		UEnum* EnumClass = StaticEnum<EnumType>();
+		check(EnumClass != nullptr);
+		return EnumClass->GetValueOrBitfieldAsString((int64)EnumeratorValue);
+	}
+
+	// TEnumAsByte produces a warning if you use it with EnumClass, so this UE_REQUIRES keeps this overload
+	// from being matched in that case
+	template<typename EnumType UE_REQUIRES(!TIsEnumClass<EnumType>::Value)>
+	FORCEINLINE static FString GetValueOrBitfieldAsString(const TEnumAsByte<EnumType> EnumeratorValue)
+	{
+		return GetValueOrBitfieldAsString(EnumeratorValue.GetValue());
+	}
+
+	template<typename EnumType UE_REQUIRES(TIsEnum<EnumType>::Value),
+		typename IntegralType UE_REQUIRES(std::is_integral_v<IntegralType>)>
+	FORCEINLINE static FString GetValueOrBitfieldAsString(const IntegralType EnumeratorValue)
+	{
+		// For the C++ enum.
+		static_assert(TIsEnum<EnumType>::Value, "Should only call this with enum types");
+		UEnum* EnumClass = StaticEnum<EnumType>();
+		check(EnumClass != nullptr);
+		return EnumClass->GetValueOrBitfieldAsString((int64)EnumeratorValue);
+	}
+
+	template<typename EnumType>
+	FORCEINLINE static void GetValueOrBitfieldAsString(const EnumType EnumeratorValue, FString& out_StringValue)
+	{
+		out_StringValue = GetValueOrBitfieldAsString(EnumeratorValue);
+	}
+
+	template<typename EnumType UE_REQUIRES(TIsEnum<EnumType>::Value), 
+		typename IntegralType UE_REQUIRES(std::is_integral_v<IntegralType>)>
+	FORCEINLINE static void GetValueOrBitfieldAsString(const IntegralType EnumeratorValue, FString& out_StringValue)
+	{
+		out_StringValue = GetValueOrBitfieldAsString<EnumType>(EnumeratorValue);
+	}
 
 	/**
 	 * @param EnumeratorValue  Enumerator Value.
