AsyncLoading2: Fix broken support for FlushAsyncLoading(-1) from package load completion callbacks.

#rnx

diff --git a/Engine/Source/Runtime/CoreUObject/Private/Serialization/AsyncLoading2.cpp b/Engine/Source/Runtime/CoreUObject/Private/Serialization/AsyncLoading2.cpp
--- a/Engine/Source/Runtime/CoreUObject/Private/Serialization/AsyncLoading2.cpp
+++ b/Engine/Source/Runtime/CoreUObject/Private/Serialization/AsyncLoading2.cpp
@@ -2536,11 +2536,6 @@ struct FAsyncPackage2
 		return SyncLoadContextId;
 	}
 
-	/**
-	* Cancel loading this package.
-	*/
-	void Cancel();
-
 	void AddConstructedObject(UObject* Object, bool bSubObjectThatAlreadyExists)
 	{
 		if (bSubObjectThatAlreadyExists)
@@ -2733,9 +2728,8 @@ private:
 	/** List of ConstructedObjects = Exports + UPackage + ObjectsCreatedFromExports */
 	TArray<UObject*> ConstructedObjects;
 	TArray<FExternalReadCallback> ExternalReadDependencies;
-	/** Call backs called when we finished loading this package											*/
-	using FCompletionCallback = TUniquePtr<FLoadPackageAsyncDelegate>;
-	TArray<FCompletionCallback, TInlineAllocator<2>> CompletionCallbacks;
+	/** Callbacks called when we finished loading this package */
+	TArray<TUniquePtr<FLoadPackageAsyncDelegate>, TInlineAllocator<2>> CompletionCallbacks;
 
 	/** Set when the package is being loaded as an instance; null otherwise. */
 	TUniquePtr<FLinkerInstancingContext> InstanceContext;
@@ -2883,8 +2877,6 @@ private:
 	/** [ASYNC/GAME THREAD] Event used to signal that the async loading thread has resumed */
 	FEvent* ThreadResumedEvent;
 	TArray<FAsyncPackage2*> LoadedPackagesToProcess;
-	/** [GAME THREAD] Game thread CompletedPackages list */
-	TArray<FAsyncPackage2*> CompletedPackages;
 #if WITH_EDITOR
 	/** [GAME THREAD] */
 	TArray<UObject*> EditorLoadedAssets;
@@ -2893,13 +2885,55 @@ private:
 	/** [ASYNC/GAME THREAD] Packages to be deleted from async thread */
 	TMpscQueue<FAsyncPackage2*> DeferredDeletePackages;
 	
-	struct FFailedPackageRequest
+	struct FCompletedPackageRequest
 	{
 		int32 RequestID = INDEX_NONE;
 		FName PackageName;
-		TUniquePtr<FLoadPackageAsyncDelegate> Callback;
+		EAsyncLoadingResult::Type Result = EAsyncLoadingResult::Succeeded;
+		UPackage* Package = nullptr;
+		TArray<TUniquePtr<FLoadPackageAsyncDelegate>, TInlineAllocator<2>> CompletionCallbacks;
+
+		static FCompletedPackageRequest FromMissingPackage(
+			const FAsyncPackageDesc2& Desc,
+			TUniquePtr<FLoadPackageAsyncDelegate>&& Callback)
+		{
+			FCompletedPackageRequest Res = {Desc.RequestID, Desc.UPackageName, EAsyncLoadingResult::Failed};
+			Res.CompletionCallbacks.Add(MoveTemp(Callback));
+			return Res;
+		}
+
+		static FCompletedPackageRequest FromLoadedPackage(
+			const FAsyncPackage2& Package,
+			TArray<TUniquePtr<FLoadPackageAsyncDelegate>, TInlineAllocator<2>>&& Callbacks)
+		{
+			return FCompletedPackageRequest
+			{
+				INDEX_NONE,
+				Package.Desc.UPackageName,
+				Package.HasLoadFailed() ? EAsyncLoadingResult::Failed : EAsyncLoadingResult::Succeeded,
+				Package.LinkerRoot,
+				MoveTemp(Callbacks)
+			};
+		}
+
+		void CallCompletionCallbacks()
+		{
+			checkSlow(IsInGameThread());
+			
+			if (CompletionCallbacks.Num() == 0)
+			{
+				return;
+			}
+
+			TRACE_CPUPROFILER_EVENT_SCOPE(PackageCompletionCallbacks);
+			for (TUniquePtr<FLoadPackageAsyncDelegate>& CompletionCallback : CompletionCallbacks)
+			{
+				CompletionCallback->ExecuteIfBound(PackageName, Package, Result);
+			}
+			CompletionCallbacks.Empty();
+		}
 	};
-	TArray<FFailedPackageRequest> FailedPackageRequests;
+	TArray<FCompletedPackageRequest> FailedPackageRequests;
 	FCriticalSection FailedPackageRequestsCritical;
 
 	FCriticalSection AsyncPackagesCritical;
@@ -6740,9 +6774,6 @@ EAsyncPackageState::Type FAsyncLoadingThread2::ProcessLoadedPackagesFromGameThre
 {
 	EAsyncPackageState::Type Result = EAsyncPackageState::Complete;
 
-	// This is for debugging purposes only. @todo remove
-	volatile int32 CurrentAsyncLoadingCounter = AsyncLoadingTickCounter;
-
 	if (IsMultithreaded() &&
 		ENamedThreads::GetRenderThread() == ENamedThreads::GameThread &&
 		!FTaskGraphInterface::Get().IsThreadProcessingTasks(ENamedThreads::GameThread)) // render thread tasks are actually being sent to the game thread.
@@ -6755,6 +6786,8 @@ EAsyncPackageState::Type FAsyncLoadingThread2::ProcessLoadedPackagesFromGameThre
 			return EAsyncPackageState::TimeOut;
 		}
 	}
+
+	TArray<FAsyncPackage2*, TInlineAllocator<4>> LocalCompletedAsyncPackages;
 	for (;;)
 	{
 		FPlatformMisc::PumpEssentialAppMessages();
@@ -6782,7 +6815,6 @@ EAsyncPackageState::Type FAsyncLoadingThread2::ProcessLoadedPackagesFromGameThre
 		}
 
 		bLocalDidSomething |= LoadedPackagesToProcess.Num() > 0;
-		TArray<FAsyncPackage2*, TInlineAllocator<4>> PackagesReadyForCallback;
 		for (int32 PackageIndex = 0; PackageIndex < LoadedPackagesToProcess.Num(); ++PackageIndex)
 		{
 			SCOPED_LOADTIMER(ProcessLoadedPackagesTime);
@@ -6913,60 +6945,45 @@ EAsyncPackageState::Type FAsyncLoadingThread2::ProcessLoadedPackagesFromGameThre
 
 			TRACE_COUNTER_SET(AsyncLoadingLoadingPackages, LoadingPackagesCounter);
 
-			PackagesReadyForCallback.Add(Package);
+			LocalCompletedAsyncPackages.Add(Package);
 		}
 
-		// Call callbacks in a batch in a stack-local array. This is to ensure that callbacks that trigger
-		// on each package load and call FlushAsyncLoading do not stack overflow by adding one FlushAsyncLoading
-		// call per LoadedPackageToProcess onto the stack
-		for (FAsyncPackage2* Package : PackagesReadyForCallback)
+		TArray<FCompletedPackageRequest> LocalCompletedPackageRequests;
 		{
-			// Call external callbacks
-			const EAsyncLoadingResult::Type LoadingResult = Package->HasLoadFailed() ? EAsyncLoadingResult::Failed : EAsyncLoadingResult::Succeeded;
-			{
-				TRACE_CPUPROFILER_EVENT_SCOPE(PackageCompletionCallbacks);
-				Package->CallCompletionCallbacks(LoadingResult);
-			}
-
-			check(!CompletedPackages.Contains(Package));
-			CompletedPackages.Add(Package);
-			RemovePendingRequests(Package->RequestIDs);
-
-			UE_ASYNC_PACKAGE_LOG(Verbose, Package->Desc, TEXT("GameThread: LoadCompleted"),
-				TEXT("All loading of package is done, and the async package and load request will be deleted."));
+			FScopeLock _(&FailedPackageRequestsCritical);
+			Swap(LocalCompletedPackageRequests, FailedPackageRequests);
 		}
-		
+		for (FCompletedPackageRequest& FailedPackageRequest : LocalCompletedPackageRequests)
 		{
-			TArray<FFailedPackageRequest> LocalFailedPackageRequests;
-			{
-				FScopeLock _(&FailedPackageRequestsCritical);
-				Swap(LocalFailedPackageRequests, FailedPackageRequests);
-			}
-
-			bLocalDidSomething |= LocalFailedPackageRequests.Num() > 0;
-			for (FFailedPackageRequest& FailedPackageRequest : LocalFailedPackageRequests)
-			{
-				FailedPackageRequest.Callback->ExecuteIfBound(FailedPackageRequest.PackageName, nullptr, EAsyncLoadingResult::Failed);
-				RemovePendingRequests(TArrayView<int32>(&FailedPackageRequest.RequestID, 1));
-				--PackagesWithRemainingWorkCounter;
-				TRACE_COUNTER_SET(AsyncLoadingPackagesWithRemainingWork, PackagesWithRemainingWorkCounter);
-			}
+			RemovePendingRequests(TArrayView<int32>(&FailedPackageRequest.RequestID, 1));
+			--PackagesWithRemainingWorkCounter;
+			TRACE_COUNTER_SET(AsyncLoadingPackagesWithRemainingWork, PackagesWithRemainingWorkCounter);
 		}
 
-		bLocalDidSomething |= CompletedPackages.Num() > 0;
-		for (int32 PackageIndex = 0; PackageIndex < CompletedPackages.Num(); ++PackageIndex)
+		LocalCompletedPackageRequests.Reserve(LocalCompletedPackageRequests.Num() + LocalCompletedAsyncPackages.Num());
+		for (FAsyncPackage2* Package : LocalCompletedAsyncPackages)
 		{
-			FAsyncPackage2* Package = CompletedPackages[PackageIndex];
 			UE_ASYNC_PACKAGE_DEBUG(Package->Desc);
+			UE_ASYNC_PACKAGE_LOG(Verbose, Package->Desc, TEXT("GameThread: LoadCompleted"),
+				TEXT("All loading of package is done, and the async package and load request will be deleted."));
 
+			LocalCompletedPackageRequests.Add(FCompletedPackageRequest::FromLoadedPackage(*Package, MoveTemp(Package->CompletionCallbacks)));
+			RemovePendingRequests(Package->RequestIDs);
 			check(Package->AsyncPackageLoadingState == EAsyncPackageLoadingState2::Complete);
 			Package->AsyncPackageLoadingState = EAsyncPackageLoadingState2::DeferredDelete;
-
-			CompletedPackages.RemoveAtSwap(PackageIndex--);
 			Package->ClearImportedPackages();
 			Package->ReleaseRef();
 		}
+		LocalCompletedAsyncPackages.Reset();
 		
+		// Call callbacks in a batch in a stack-local array after all other work has been done to handle
+		// callbacks that may call FlushAsyncLoading
+		for (FCompletedPackageRequest& CompletedPackageRequest : LocalCompletedPackageRequests)
+		{
+			CompletedPackageRequest.CallCompletionCallbacks();
+		}
+
+		bLocalDidSomething |= LocalCompletedPackageRequests.Num() > 0;
 		if (!bLocalDidSomething)
 		{
 			break;
@@ -6983,8 +7000,6 @@ EAsyncPackageState::Type FAsyncLoadingThread2::ProcessLoadedPackagesFromGameThre
 
 	if (Result == EAsyncPackageState::Complete)
 	{
-		// We're not done until all packages have been deleted
-		Result = CompletedPackages.Num() ? EAsyncPackageState::PendingImports  : EAsyncPackageState::Complete;
 	}
 
 	return Result;
@@ -8023,18 +8038,6 @@ void FAsyncPackage2::FinishUPackage()
 	}
 }
 
-void FAsyncPackage2::CallCompletionCallbacks(EAsyncLoadingResult::Type LoadingResult)
-{
-	checkSlow(IsInGameThread());
-
-	UPackage* LoadedPackage = (!bLoadHasFailed) ? LinkerRoot : nullptr;
-	for (FCompletionCallback& CompletionCallback : CompletionCallbacks)
-	{
-		CompletionCallback->ExecuteIfBound(Desc.UPackageName, LoadedPackage, LoadingResult);
-	}
-	CompletionCallbacks.Empty();
-}
-
 #if WITH_EDITOR
 void FAsyncLoadingThread2::ConditionalProcessEditorCallbacks()
 {
@@ -8068,24 +8071,6 @@ void FAsyncLoadingThread2::ConditionalProcessEditorCallbacks()
 }
 #endif
 
-void FAsyncPackage2::Cancel()
-{
-	// Call any completion callbacks specified.
-	bLoadHasFailed = true;
-	const EAsyncLoadingResult::Type Result = EAsyncLoadingResult::Canceled;
-	CallCompletionCallbacks(Result);
-
-	if (LinkerRoot)
-	{
-		if (bCreatedLinkerRoot)
-		{
-			LinkerRoot->ClearFlags(RF_WasLoaded);
-			LinkerRoot->bHasBeenFullyLoaded = false;
-			LinkerRoot->Rename(*MakeUniqueObjectName(GetTransientPackage(), UPackage::StaticClass()).ToString(), nullptr, REN_DontCreateRedirectors | REN_DoNotDirty | REN_ForceNoResetLoaders | REN_NonTransactional);
-		}
-	}
-}
-
 void FAsyncPackage2::AddCompletionCallback(TUniquePtr<FLoadPackageAsyncDelegate>&& Callback)
 {
 	// This is to ensure that there is no one trying to subscribe to a already loaded package
@@ -8156,13 +8141,9 @@ void FAsyncLoadingThread2::QueueMissingPackage(FAsyncPackageDesc2& PackageDesc,
 
 	if (PackageLoadedDelegate.IsValid())
 	{
-		FScopeLock LockMissingPackages(&FailedPackageRequestsCritical);
-		FailedPackageRequests.Add(FFailedPackageRequest
-		{
-			PackageDesc.RequestID,
-			FailedPackageName,
-			MoveTemp(PackageLoadedDelegate)
-		});
+		FScopeLock _(&FailedPackageRequestsCritical);
+		FailedPackageRequests.Add(
+			FCompletedPackageRequest::FromMissingPackage(PackageDesc, MoveTemp(PackageLoadedDelegate)));
 	}
 	else
 	{
