Fix TSAN warning by adding missing TSAN_SAFE to ExpandChunksToIndex.

#rnx

diff --git a/Engine/Source/Runtime/CoreUObject/Public/UObject/UObjectArray.h b/Engine/Source/Runtime/CoreUObject/Public/UObject/UObjectArray.h
--- a/Engine/Source/Runtime/CoreUObject/Public/UObject/UObjectArray.h
+++ b/Engine/Source/Runtime/CoreUObject/Public/UObject/UObjectArray.h
@@ -383,7 +383,7 @@ class FChunkedFixedUObjectArray
 	/**
 	* Allocates new chunk for the array
 	**/
-	void ExpandChunksToIndex(int32 Index)
+	void ExpandChunksToIndex(int32 Index) TSAN_SAFE
 	{
 		check(Index >= 0 && Index < MaxElements);
 		int32 ChunkIndex = Index / NumElementsPerChunk;
@@ -395,7 +395,7 @@ class FChunkedFixedUObjectArray
 			if (FPlatformAtomics::InterlockedCompareExchangePointer((void**)Chunk, NewChunk, nullptr))
 			{
 				// someone else beat us to the add, we don't support multiple concurrent adds
-				check(0)
+				check(0);
 			}
 			else
 			{
