Fix new function not being exported
#rnx

diff --git a/Engine/Source/Runtime/CoreUObject/Public/UObject/Class.h b/Engine/Source/Runtime/CoreUObject/Public/UObject/Class.h
--- a/Engine/Source/Runtime/CoreUObject/Public/UObject/Class.h
+++ b/Engine/Source/Runtime/CoreUObject/Public/UObject/Class.h
@@ -2000,7 +2000,7 @@ public:
 	COREUOBJECT_API FString GetNameStringByValue(int64 InValue) const;
 
 	/** If the enumeration is declared as UENUM(Flags), returns a string of the form A | B | C representing set bits A, B, and C. If it is not a bitfield, the result is the same as calling GetNameStringByValue*/
-	FString GetValueOrBitfieldAsString(int64 InValue) const;
+	COREUOBJECT_API FString GetValueOrBitfieldAsString(int64 InValue) const;
 
 	/** Looks for a name with a given value and returns true and writes the name to Out if one was found */
 	COREUOBJECT_API bool FindNameStringByValue(FString& Out, int64 InValue) const;
