PackageName: Add console commands for ConvertFilenameToLongPackageName and ConvertLongPackageNameToFilename.
#rnx

diff --git a/Engine/Source/Runtime/CoreUObject/Private/Misc/PackageName.cpp b/Engine/Source/Runtime/CoreUObject/Private/Misc/PackageName.cpp
--- a/Engine/Source/Runtime/CoreUObject/Private/Misc/PackageName.cpp
+++ b/Engine/Source/Runtime/CoreUObject/Private/Misc/PackageName.cpp
@@ -2708,6 +2708,56 @@ FText FPackageName::FormatErrorAsText(FStringView InPath, EErrorCode ErrorCode)
 	}
 }
 
+#if !UE_BUILD_SHIPPING
+void ConsoleCommandImplConvertFilenameToLongPackageName(const TArray<FString>& Args)
+{
+	if (Args.IsEmpty())
+	{
+		UE_LOG(LogPackageName, Display, TEXT("Invalid arguments. Usage: PackageName.ConvertFilenameToLongPackageName <Filename>."));
+		return;
+	}
+
+	FString LongPackageName;
+	if (!FPackageName::TryConvertFilenameToLongPackageName(Args[0], LongPackageName))
+	{
+		UE_LOG(LogPackageName, Display, TEXT(""));
+		return;
+	}
+
+	UE_LOG(LogPackageName, Display, TEXT("%s"), *LongPackageName);
+}
+
+static FAutoConsoleCommand ConsoleCommandConvertFilenameToLongPackageName(
+	TEXT("PackageName.ConvertFilenameToLongPackageName"),
+	TEXT("Prints the corresponding packagename for a filename at a given localpath, according to the current registered mount points. Prints empty string if not mounted."),
+	FConsoleCommandWithArgsDelegate::CreateStatic(ConsoleCommandImplConvertFilenameToLongPackageName)
+);
+
+void ConsoleCommandImplConvertLongPackageNameToFilename(const TArray<FString>& Args)
+{
+	if (Args.IsEmpty())
+	{
+		UE_LOG(LogPackageName, Display, TEXT("Invalid arguments. Usage: PackageName.ConvertFilenameToLongPackageName <Filename>."));
+		return;
+	}
+
+	FString Filename;
+	if (!FPackageName::TryConvertLongPackageNameToFilename(Args[0], Filename))
+	{
+		UE_LOG(LogPackageName, Display, TEXT(""));
+		return;
+	}
+
+	UE_LOG(LogPackageName, Display, TEXT("%s"), *Filename);
+}
+
+static FAutoConsoleCommand ConsoleCommandConvertLongPackageNameToFilename(
+	TEXT("PackageName.ConvertLongPackageNameToFilename"),
+	TEXT("Prints the corresponding local filename for a given packagename, according to the current registered mount points. Prints empty string if not mounted."),
+	FConsoleCommandWithArgsDelegate::CreateStatic(ConsoleCommandImplConvertLongPackageNameToFilename)
+);
+#endif
+
 #if WITH_DEV_AUTOMATION_TESTS
 
 IMPLEMENT_SIMPLE_AUTOMATION_TEST(FPackageNameTests, "System.Core.Misc.PackageNames", EAutomationTestFlags::ApplicationContextMask | EAutomationTestFlags::SmokeFilter)
