Fix TSAN warning from ALT by turning SetDirtyFlag as a no-op if called with the same value as before

#rnx

diff --git a/Engine/Source/Runtime/CoreUObject/Private/UObject/Package.cpp b/Engine/Source/Runtime/CoreUObject/Private/UObject/Package.cpp
--- a/Engine/Source/Runtime/CoreUObject/Private/UObject/Package.cpp
+++ b/Engine/Source/Runtime/CoreUObject/Private/UObject/Package.cpp
@@ -93,33 +93,35 @@ void UPackage::PostInitProperties()
 /**
  * Marks/Unmarks the package's bDirty flag
  */
-void UPackage::SetDirtyFlag( bool bIsDirty )
+void UPackage::SetDirtyFlag( bool bInIsDirty )
 {
-	if ( GetOutermost() != GetTransientPackage() )
+	// Early out if there is no change to the flag
+	if (bDirty != bInIsDirty)
 	{
-		if ( GUndo != nullptr
-		// PIE and script/class packages should never end up in the transaction buffer as we cannot undo during gameplay.
-		&& !GetOutermost()->HasAnyPackageFlags(PKG_PlayInEditor|PKG_ContainsScript|PKG_CompiledIn) )
+		if ( GetOutermost() != GetTransientPackage() )
 		{
-			// make sure we're marked as transactional
-			SetFlags(RF_Transactional);
+			if ( GUndo != nullptr
+				// PIE and script/class packages should never end up in the transaction buffer as we cannot undo during gameplay.
+				&& !GetOutermost()->HasAnyPackageFlags(PKG_PlayInEditor|PKG_ContainsScript|PKG_CompiledIn) )
+			{
+				// make sure we're marked as transactional
+				SetFlags(RF_Transactional);
 
-			// don't call Modify() since it calls SetDirtyFlag()
-			GUndo->SaveObject( this );
-		}
+				// don't call Modify() since it calls SetDirtyFlag()
+				GUndo->SaveObject( this );
+			}
 
-		// Update dirty bit
-		const bool bWasDirty = bDirty;
-		bDirty = bIsDirty;
+			// Update dirty bit after we saved the object in the transaction buffer.
+			bDirty = bInIsDirty;
 
-		if( bWasDirty != bIsDirty						// Only fire the callback if the dirty state actually changes
-			&& GIsEditor								// Only fire the callback in editor mode
-			&& !HasAnyPackageFlags(PKG_ContainsScript)	// Skip script packages
-			&& !HasAnyPackageFlags(PKG_PlayInEditor)	// Skip packages for PIE
-			&& GetTransientPackage() != this )			// Skip the transient package
-		{
-			// Package is changing dirty state, let the editor know so we may prompt for source control checkout
-			PackageDirtyStateChangedEvent.Broadcast(this);
+			if (GIsEditor									// Only fire the callback in editor mode
+				&& !HasAnyPackageFlags(PKG_ContainsScript)	// Skip script packages
+				&& !HasAnyPackageFlags(PKG_PlayInEditor)	// Skip packages for PIE
+				&& GetTransientPackage() != this )			// Skip the transient package
+			{
+				// Package is changing dirty state, let the editor know so we may prompt for source control checkout
+				PackageDirtyStateChangedEvent.Broadcast(this);
+			}
 		}
 	}
 }
