Fix TSAN warning in AActor::Modify when called from ALT

#rnx

diff --git a/Engine/Source/Runtime/CoreUObject/Private/UObject/Obj.cpp b/Engine/Source/Runtime/CoreUObject/Private/UObject/Obj.cpp
--- a/Engine/Source/Runtime/CoreUObject/Private/UObject/Obj.cpp
+++ b/Engine/Source/Runtime/CoreUObject/Private/UObject/Obj.cpp
@@ -1261,7 +1261,12 @@ void UObject::PreSave(FObjectPreSaveContext SaveContext)
 #if WITH_EDITOR
 bool UObject::CanModify() const
 {
-	return (!HasAnyFlags(RF_NeedInitialization) && !IsGarbageCollecting() && !GExitPurge && !IsUnreachable());
+	return
+		!HasAnyFlags(RF_NeedInitialization) && !IsGarbageCollecting() && !GExitPurge && !IsUnreachable() &&
+		// Prevent modification while loading
+		!HasAnyInternalFlags(EInternalObjectFlags::Async | EInternalObjectFlags::AsyncLoading) &&
+		// Only the game-thread should be allowed to touch the transaction buffer at all
+		IsInGameThread();
 }
 
 bool UObject::Modify( bool bAlwaysMarkDirty/*=true*/ )
@@ -1270,29 +1275,25 @@ bool UObject::Modify( bool bAlwaysMarkDirty/*=true*/ )
 
 	if (CanModify())
 	{
-		// Only the game-thread should be allowed to touch the transaction buffer at all
-		if (!HasAnyInternalFlags(EInternalObjectFlags::Async | EInternalObjectFlags::AsyncLoading) && IsInGameThread())
+		// Do not consider script packages, as they should never end up in the
+		// transaction buffer and we don't want to mark them dirty here either.
+		// We do want to consider PIE objects however
+		if (GetOutermost()->HasAnyPackageFlags(PKG_ContainsScript | PKG_CompiledIn) == false || GetClass()->HasAnyClassFlags(CLASS_DefaultConfig | CLASS_Config))
 		{
-			// Do not consider script packages, as they should never end up in the
-			// transaction buffer and we don't want to mark them dirty here either.
-			// We do want to consider PIE objects however
-			if (GetOutermost()->HasAnyPackageFlags(PKG_ContainsScript | PKG_CompiledIn) == false || GetClass()->HasAnyClassFlags(CLASS_DefaultConfig | CLASS_Config))
+			// Attempt to mark the package dirty and save a copy of the object to the transaction
+			// buffer. The save will fail if there isn't a valid transactor, the object isn't
+			// transactional, etc.
+			bSavedToTransactionBuffer = SaveToTransactionBuffer(this, bAlwaysMarkDirty);
+
+			// If we failed to save to the transaction buffer, but the user requested the package
+			// marked dirty anyway, do so
+			if (!bSavedToTransactionBuffer && bAlwaysMarkDirty)
 			{
-				// Attempt to mark the package dirty and save a copy of the object to the transaction
-				// buffer. The save will fail if there isn't a valid transactor, the object isn't
-				// transactional, etc.
-				bSavedToTransactionBuffer = SaveToTransactionBuffer(this, bAlwaysMarkDirty);
-
-				// If we failed to save to the transaction buffer, but the user requested the package
-				// marked dirty anyway, do so
-				if (!bSavedToTransactionBuffer && bAlwaysMarkDirty)
-				{
-					MarkPackageDirty();
-				}
+				MarkPackageDirty();
 			}
-
-			FCoreUObjectDelegates::BroadcastOnObjectModified(this);
 		}
+
+		FCoreUObjectDelegates::BroadcastOnObjectModified(this);
 	}
 
 	return bSavedToTransactionBuffer;
