Fix typo and deprecate GetMaxAnnottations

#rnx

diff --git a/Engine/Source/Runtime/CoreUObject/Public/UObject/UObjectAnnotation.h b/Engine/Source/Runtime/CoreUObject/Public/UObject/UObjectAnnotation.h
--- a/Engine/Source/Runtime/CoreUObject/Public/UObject/UObjectAnnotation.h
+++ b/Engine/Source/Runtime/CoreUObject/Public/UObject/UObjectAnnotation.h
@@ -835,6 +835,17 @@ public:
 	* Thread safe, but you know, someone might have added more elements before this even returns
 	* @return	the maximum number of elements in the array
 	**/
+	FORCEINLINE int32 GetMaxAnnotations() const TSAN_SAFE
+	{
+		return MaxAnnotations;
+	}
+
+	/**
+	* Return the number max capacity of the array
+	* Thread safe, but you know, someone might have added more elements before this even returns
+	* @return	the maximum number of elements in the array
+	**/
+	UE_DEPRECATED(5.3, "Use GetMaxAnnotations instead")
 	FORCEINLINE int32 GetMaxAnnottations() const TSAN_SAFE
 	{
 		return MaxAnnotations;
