Add a -nozenloader option to disable the loader, takes priority over other method of enabling zenloader if present

diff --git a/Engine/Source/Runtime/CoreUObject/Private/Serialization/AsyncPackageLoader.cpp b/Engine/Source/Runtime/CoreUObject/Private/Serialization/AsyncPackageLoader.cpp
--- a/Engine/Source/Runtime/CoreUObject/Private/Serialization/AsyncPackageLoader.cpp
+++ b/Engine/Source/Runtime/CoreUObject/Private/Serialization/AsyncPackageLoader.cpp
@@ -189,16 +189,18 @@ void InitAsyncThread()
 	{
 		bool bSettingsEnabled = false;
 		bool bCommandLineEnabled = false;
+		bool bCommandLineDisabled = false;
 		bool bHasUseIoStoreParamInEditor = false;
 #if WITH_EDITOR
 		bCommandLineEnabled = FParse::Param(FCommandLine::Get(), TEXT("ZenLoader"));
+		bCommandLineDisabled = FParse::Param(FCommandLine::Get(), TEXT("NoZenLoader"));
 		check(GConfig);
 		GConfig->GetBool(TEXT("/Script/Engine.EditorStreamingSettings"), TEXT("s.ZenLoaderEnabled"), bSettingsEnabled, GEngineIni);
 		bHasUseIoStoreParamInEditor = UE_FORCE_USE_IOSTORE || FParse::Param(FCommandLine::Get(), TEXT("UseIoStore"));
 #endif
 		FIoDispatcher& IoDispatcher = FIoDispatcher::Get();
 		bool bHasScriptObjectsChunk = IoDispatcher.DoesChunkExist(CreateIoChunkId(0, 0, EIoChunkType::ScriptObjects));
-		if (bSettingsEnabled || bCommandLineEnabled)
+		if (!bCommandLineDisabled && (bSettingsEnabled || bCommandLineEnabled))
 		{
 			GPackageLoader.Reset(MakeAsyncPackageLoader2(IoDispatcher));
 		}
