removing usage of GIsPlayInEditorWorld as it did not work as expected
https://jira.it.epicgames.com/browse/UE-180246
https://p4-swarm.epicgames.net/reviews/24658760

diff --git a/Engine/Source/Runtime/CoreUObject/Private/UObject/LinkerLoadImportBehavior.cpp b/Engine/Source/Runtime/CoreUObject/Private/UObject/LinkerLoadImportBehavior.cpp
--- a/Engine/Source/Runtime/CoreUObject/Private/UObject/LinkerLoadImportBehavior.cpp
+++ b/Engine/Source/Runtime/CoreUObject/Private/UObject/LinkerLoadImportBehavior.cpp
@@ -145,7 +145,7 @@ bool CanLazyImport(const IAssetRegistryInterface& AssetRegistry, const FObjectIm
 	//most of this is duplicated from TryLazyImport with the out parameter but it avoids a bunch of string allocations
 	using namespace UE::CoreUObject::Private;
 	EImportBehavior Behavior = GetPropertyImportLoadBehavior(Import, LinkerLoad);
-	if (GIsPlayInEditorWorld || Behavior != EImportBehavior::LazyOnDemand)
+	if (Behavior != EImportBehavior::LazyOnDemand)
 	{
 		return false;
 	}
@@ -185,7 +185,7 @@ bool CanLazyImport(const IAssetRegistryInterface& AssetRegistry, const FObjectIm
 bool TryLazyImport(const IAssetRegistryInterface& AssetRegistry, const FObjectImport& Import, const FLinkerLoad& LinkerLoad, FObjectPtr& ObjectPtr)
 {
 	using namespace UE::CoreUObject::Private;
-	if (GIsPlayInEditorWorld || GetPropertyImportLoadBehavior(Import, LinkerLoad) != EImportBehavior::LazyOnDemand )
+	if (GetPropertyImportLoadBehavior(Import, LinkerLoad) != EImportBehavior::LazyOnDemand )
 	{
 		return false;
 	}
@@ -249,7 +249,7 @@ bool TryLazyImport(const IAssetRegistryInterface& AssetRegistry, const FObjectIm
 bool TryLazyLoad(const UClass& Class, const FSoftObjectPath& ObjectPath, TObjectPtr<UObject>& OutObjectPtr)
 {
 	using namespace UE::CoreUObject::Private;
-	if (GIsPlayInEditorWorld || !FLinkerLoad::IsImportLazyLoadEnabled())
+	if (!FLinkerLoad::IsImportLazyLoadEnabled())
 	{
 		return false;
 	}
