ZenLoader in editor: Fixed issue where exports from cooked packages could sometimes be requested from uncooked packages before being found in the global import store
#rnx

diff --git a/Engine/Source/Runtime/CoreUObject/Private/UObject/LinkerLoad.cpp b/Engine/Source/Runtime/CoreUObject/Private/UObject/LinkerLoad.cpp
--- a/Engine/Source/Runtime/CoreUObject/Private/UObject/LinkerLoad.cpp
+++ b/Engine/Source/Runtime/CoreUObject/Private/UObject/LinkerLoad.cpp
@@ -3193,7 +3193,6 @@ FLinkerLoad::EVerifyResult FLinkerLoad::VerifyImport(int32 ImportIndex)
 // Internal Load package call so that we can pass the linker that requested this package as an import dependency
 UPackage* LoadPackageInternal(UPackage* InOuter, const FPackagePath& PackagePath, uint32 LoadFlags, FLinkerLoad* ImportLinker, FArchive* InReaderOverride, const FLinkerInstancingContext* InstancingContext, const FPackagePath* DiffPackagePath);
 
-#if WITH_IOSTORE_IN_EDITOR
 /**
  * Finds and populates the import table for the specified package import.
  *
@@ -3266,7 +3265,6 @@ void StaticFindAllImportObjects(TArray<FObjectImport>& InOutImportMap, FPackageI
 		}
 	}
 }
-#endif // WITH_IOSTORE_IN_EDITOR
 
 /**
  * Safely verify that an import in the ImportMap points to a good object. This decides whether or not
@@ -3314,6 +3312,14 @@ bool FLinkerLoad::VerifyImportInner(const int32 ImportIndex, FString& WarningSuf
 				return nullptr;
 			}
 			Import.SourceLinker = FindExistingLinkerForPackage(Package);
+			if (!Import.SourceLinker && Package->HasAnyPackageFlags(PKG_Cooked))
+			{
+				// Special case where we're verifying an import from a cooked package before we've performed the global import store lookup for this package in AsyncLoading2.cpp
+				// Find the imports by name instead
+				// Note: The cooked package might not be marked as fully loaded at this stage, but we will have created and serialized all its exports
+				Import.XObject = Package;
+				StaticFindAllImportObjects(ImportMap, FPackageIndex::FromImport(ImportIndex));
+			}
 			return Package;
 		}
 		if (Package == nullptr || !Package->IsFullyLoaded())
