ZenLoader in editor: Fixed imports from redirected outers sometimes not being resolved correctly
#rnx

diff --git a/Engine/Source/Runtime/CoreUObject/Private/Serialization/AsyncLoading2.cpp b/Engine/Source/Runtime/CoreUObject/Private/Serialization/AsyncLoading2.cpp
--- a/Engine/Source/Runtime/CoreUObject/Private/Serialization/AsyncLoading2.cpp
+++ b/Engine/Source/Runtime/CoreUObject/Private/Serialization/AsyncLoading2.cpp
@@ -5301,9 +5301,11 @@ bool FAsyncPackage2::ProcessLinkerLoadPackageImports(FAsyncLoadingThreadState2&
 			}
 #endif
 			FObjectImport& LinkerImport = LinkerLoadState->Linker->ImportMap[ImportIndex];
-			if (!LinkerImport.XObject)
+			if (!LinkerImport.XObject && FromImportStore)
 			{
 				LinkerImport.XObject = FromImportStore;
+				LinkerImport.SourceIndex = FromImportStore->GetLinkerIndex();
+				LinkerImport.SourceLinker = FromImportStore->GetLinker();
 			}
 		}
 	}
