ZenLoader in editor: Add support for package core redirects
#rnx

diff --git a/Engine/Source/Runtime/CoreUObject/Private/Serialization/AsyncLoading2.cpp b/Engine/Source/Runtime/CoreUObject/Private/Serialization/AsyncLoading2.cpp
--- a/Engine/Source/Runtime/CoreUObject/Private/Serialization/AsyncLoading2.cpp
+++ b/Engine/Source/Runtime/CoreUObject/Private/Serialization/AsyncLoading2.cpp
@@ -73,6 +73,7 @@
 #include "UObject/LinkerLoad.h"
 #include "Containers/SpscQueue.h"
 #include "IO/IoPriorityQueue.h"
+#include "UObject/CoreRedirects.h"
 
 #include <atomic>
 
@@ -3859,6 +3860,14 @@ bool FAsyncLoadingThread2::CreateAsyncPackagesFromQueue(FAsyncLoadingThreadState
 				}
 			}
 
+#if WITH_EDITOR
+			const FCoreRedirectObjectName RedirectedPackageName = FCoreRedirects::GetRedirectedName(ECoreRedirectFlags::Type_Package, FCoreRedirectObjectName(NAME_None, NAME_None, PackageNameToLoad));
+			if (RedirectedPackageName.PackageName != PackageNameToLoad)
+			{
+				PackageNameToLoad = RedirectedPackageName.PackageName;
+			}
+#endif
+
 			FPackageId PackageIdToLoad = FPackageId::FromName(PackageNameToLoad);
 			FName UPackageName = PackageNameToLoad;
 			{
@@ -4515,6 +4524,24 @@ void FAsyncPackage2::ImportPackagesRecursiveInner(FAsyncLoadingThreadState2& Thr
 		FName ImportedPackageUPackageName = bHasImportedPackageNames ? ImportedPackageNames[LocalImportedPackageIndex] : NAME_None;
 		FName ImportedPackageNameToLoad = ImportedPackageUPackageName;
 		FPackageId ImportedPackageIdToLoad = ImportedPackageId;
+
+#if WITH_EDITOR
+		if (!ImportedPackageNameToLoad.IsNone())
+		{
+			const FCoreRedirectObjectName RedirectedPackageName = FCoreRedirects::GetRedirectedName(ECoreRedirectFlags::Type_Package, FCoreRedirectObjectName(NAME_None, NAME_None, ImportedPackageNameToLoad));
+			if (RedirectedPackageName.PackageName != ImportedPackageNameToLoad)
+			{
+				ImportedPackageNameToLoad = RedirectedPackageName.PackageName;
+				ImportedPackageIdToLoad = FPackageId::FromName(ImportedPackageNameToLoad);
+				ImportedPackageUPackageName = ImportedPackageNameToLoad;
+				ImportedPackageId = ImportedPackageIdToLoad;
+				// Rewrite the import table
+				ImportedPackageIds[LocalImportedPackageIndex] = ImportedPackageId;
+				ImportedPackageNames[LocalImportedPackageIndex] = ImportedPackageUPackageName;
+			}
+		}
+#endif
+
 #if WITH_EDITORONLY_DATA && ALT2_ENABLE_LINKERLOAD_SUPPORT
 		if (bHasImportedPackageNames && LinkerLoadState.IsSet())
 		{
diff --git a/Engine/Source/Runtime/CoreUObject/Private/UObject/UObjectGlobals.cpp b/Engine/Source/Runtime/CoreUObject/Private/UObject/UObjectGlobals.cpp
--- a/Engine/Source/Runtime/CoreUObject/Private/UObject/UObjectGlobals.cpp
+++ b/Engine/Source/Runtime/CoreUObject/Private/UObject/UObjectGlobals.cpp
@@ -1582,7 +1582,24 @@ UPackage* LoadPackageInternal(UPackage* InOuter, const FPackagePath& PackagePath
 		}
 		ThreadContext.SyncLoadUsingAsyncLoaderCount--;
 
-		return (InOuter ? InOuter : FindObjectFast<UPackage>(nullptr, PackageName));
+		if (InOuter)
+		{
+			return InOuter;
+		}
+		else
+		{
+			UPackage* Result = FindObjectFast<UPackage>(nullptr, PackageName);
+			if (!Result)
+			{
+				// Might have been redirected
+				const FCoreRedirectObjectName NewPackageName = FCoreRedirects::GetRedirectedName(ECoreRedirectFlags::Type_Package, FCoreRedirectObjectName(NAME_None, NAME_None, PackageName));
+				if (NewPackageName.PackageName != PackageName)
+				{
+					Result = FindObjectFast<UPackage>(nullptr, NewPackageName.PackageName);
+				}
+			}
+			return Result;
+		}
 	}
 
 	UPackage* Result = nullptr;
