Fix race/crash on static inside CheckMissingSection

#rnx

diff --git a/Engine/Source/Runtime/CoreUObject/Private/UObject/Obj.cpp b/Engine/Source/Runtime/CoreUObject/Private/UObject/Obj.cpp
--- a/Engine/Source/Runtime/CoreUObject/Private/UObject/Obj.cpp
+++ b/Engine/Source/Runtime/CoreUObject/Private/UObject/Obj.cpp
@@ -2292,20 +2292,46 @@ void UObject::ReloadConfig( UClass* ConfigClass/*=NULL*/, const TCHAR* InFilenam
 #if !UE_BUILD_SHIPPING
 void CheckMissingSection(const FString& SectionName, const FString& IniFilename)
 {
-	static TSet<FString> MissingSections;
+	// Apply lock striping to reduce contention.
+	constexpr int32 MISSINGSECTIONS_BUCKETS = 31; /* prime number for best distribution using modulo */
+
+	struct FMissingSections
+	{
+		FRWLock Lock;
+		TSet<FString> Sections;
+	};
+	static FMissingSections MissingSections[MISSINGSECTIONS_BUCKETS];
+
 	FConfigSection* Sec = GConfig->GetSectionPrivate(*SectionName, false, true, *IniFilename);
-	if (!Sec && MissingSections.Contains(SectionName) == false)
+
+	if (Sec == nullptr)
 	{
-		FString ShortSectionName = FPackageName::GetShortName(SectionName);
-		if (ShortSectionName != SectionName)
+		const uint32 SectionNameHash = GetTypeHash(SectionName);
+		FMissingSections& Bucket = MissingSections[SectionNameHash % MISSINGSECTIONS_BUCKETS];
+
 		{
-			Sec = GConfig->GetSectionPrivate(*ShortSectionName, false, true, *IniFilename);
-			if (Sec != NULL)
+			FReadScopeLock ScopeLock(Bucket.Lock);
+			if (Bucket.Sections.ContainsByHash(SectionNameHash, SectionName))
 			{
-				UE_LOG(LogObj, Fatal, TEXT("Short class section names (%s) are not supported, please use long name: %s"), *ShortSectionName, *SectionName);
+				return;
+			}
+		}
+
+		FWriteScopeLock ScopeLock(Bucket.Lock);
+
+		if (Bucket.Sections.ContainsByHash(SectionNameHash, SectionName) == false)
+		{
+			FString ShortSectionName = FPackageName::GetShortName(SectionName);
+			if (ShortSectionName != SectionName)
+			{
+				Sec = GConfig->GetSectionPrivate(*ShortSectionName, false, true, *IniFilename);
+				if (Sec != nullptr)
+				{
+					UE_LOG(LogObj, Fatal, TEXT("Short class section names (%s) are not supported, please use long name: %s"), *ShortSectionName, *SectionName);
+				}
 			}
+			Bucket.Sections.AddByHash(SectionNameHash, SectionName);
 		}
-		MissingSections.Add(SectionName);		
 	}
 }
 #endif
