Fixed typo.

diff --git a/Engine/Source/Runtime/CoreUObject/Private/Blueprint/BlueprintSupport.cpp b/Engine/Source/Runtime/CoreUObject/Private/Blueprint/BlueprintSupport.cpp
--- a/Engine/Source/Runtime/CoreUObject/Private/Blueprint/BlueprintSupport.cpp
+++ b/Engine/Source/Runtime/CoreUObject/Private/Blueprint/BlueprintSupport.cpp
@@ -2535,15 +2535,15 @@ bool FObjectInitializer::InitNonNativeProperty(FProperty* Property, UObject* Dat
  * FDeferredInitializationTrackerBase
  ******************************************************************************/
 
-FObjectInitializer* FDeferredInitializationTrackerBase::Add(const UObject* InitDependecy, const FObjectInitializer& DeferringInitializer)
+FObjectInitializer* FDeferredInitializationTrackerBase::Add(const UObject* InitDependency, const FObjectInitializer& DeferringInitializer)
 {
 	FObjectInitializer* DeferredInitializerCopy = nullptr;
 
-	DEFERRED_DEPENDENCY_CHECK(InitDependecy);
-	if (InitDependecy)
+	DEFERRED_DEPENDENCY_CHECK(InitDependency);
+	if (InitDependency)
 	{
 		UObject* InstanceObj = DeferringInitializer.GetObj();
-		ArchetypeInstanceMap.AddUnique(InitDependecy, InstanceObj);
+		ArchetypeInstanceMap.AddUnique(InitDependency, InstanceObj);
 
 		DEFERRED_DEPENDENCY_CHECK(DeferredInitializers.Find(InstanceObj) == nullptr); // did we try to init the object twice?
 
@@ -2553,17 +2553,17 @@ FObjectInitializer* FDeferredInitializationTrackerBase::Add(const UObject* InitD
 	return DeferredInitializerCopy;
 }
 
-void FDeferredInitializationTrackerBase::ResolveArchetypeInstances(UObject* InitDependecy)
+void FDeferredInitializationTrackerBase::ResolveArchetypeInstances(UObject* InitDependency)
 {
 	TArray<UObject*> ArchetypeInstances;
-	ArchetypeInstanceMap.MultiFind(InitDependecy, ArchetypeInstances);
+	ArchetypeInstanceMap.MultiFind(InitDependency, ArchetypeInstances);
 
 	for (UObject* Instance : ArchetypeInstances)
 	{
 		DEFERRED_DEPENDENCY_CHECK(ResolvingObjects.Contains(Instance) == false);
 		ResolvingObjects.Push(Instance);
 
-		if (ResolveDeferredInitialization(InitDependecy, Instance))
+		if (ResolveDeferredInitialization(InitDependency, Instance))
 		{
 			// For sub-objects, this has to come after ResolveDeferredInitialization(), since InitSubObjectProperties() is 
 			// invoked there (which is where we fill this sub-object with values from the super)
@@ -2574,7 +2574,7 @@ void FDeferredInitializationTrackerBase::ResolveArchetypeInstances(UObject* Init
 		ResolvingObjects.Pop();
 	}
 
-	ArchetypeInstanceMap.Remove(InitDependecy);
+	ArchetypeInstanceMap.Remove(InitDependency);
 }
 
 bool FDeferredInitializationTrackerBase::IsInitializationDeferred(const UObject* Object) const
diff --git a/Engine/Source/Runtime/CoreUObject/Public/Blueprint/BlueprintSupport.h b/Engine/Source/Runtime/CoreUObject/Public/Blueprint/BlueprintSupport.h
--- a/Engine/Source/Runtime/CoreUObject/Public/Blueprint/BlueprintSupport.h
+++ b/Engine/Source/Runtime/CoreUObject/Public/Blueprint/BlueprintSupport.h
@@ -291,13 +291,13 @@ public:
 	 * Makes a copy of the specified initializer and stores it (mapped under its 
 	 * dependency), so that it can instead be executed later via ResolveArchetypeInstances().
 	 * 
-	 * @param  InitDependecy			The object (usually the initializer's archetype) that this initializer is dependent on. The key 
+	 * @param  InitDependency			The object (usually the initializer's archetype) that this initializer is dependent on. The key 
 	 *									you'll pass to ResolveArchetypeInstances() later to run this initializer.
 	 * @param  DeferringInitializer		The initializer you want to defer.
 	 *
-	 * @raturn A copy of the specified initializer. This should always succeed (only returns null if InitDependecy is null).
+	 * @raturn A copy of the specified initializer. This should always succeed (only returns null if InitDependency is null).
 	 */
-	FObjectInitializer* Add(const UObject* InitDependecy, const FObjectInitializer& DeferringInitializer);
+	FObjectInitializer* Add(const UObject* InitDependency, const FObjectInitializer& DeferringInitializer);
 
 	/** 
 	 * Runs all deferred initializers that were dependent on the specified archetype 
