[Backout] - CL24690253 - CIS UE Error
Original CL Desc
-----------------------------------------------------------------
GC duration API + debug facilities

Added console commands to show/hide a live GC debug graph + generate or unlink a dense UObject tree as a stress test for GC marking.

diff --git a/Engine/Source/Runtime/CoreUObject/Private/UObject/GarbageCollection.cpp b/Engine/Source/Runtime/CoreUObject/Private/UObject/GarbageCollection.cpp
--- a/Engine/Source/Runtime/CoreUObject/Private/UObject/GarbageCollection.cpp
+++ b/Engine/Source/Runtime/CoreUObject/Private/UObject/GarbageCollection.cpp
@@ -37,7 +37,6 @@
 #include "HAL/RunnableThread.h"
 #include "UObject/FieldPathProperty.h"
 #include "UObject/GarbageCollectionHistory.h"
-#include "UObject/GarbageCollectionTesting.h"
 
 #include <atomic>
 
@@ -85,13 +84,7 @@ static bool GObjCurrentPurgeObjectIndexNeedsReset = true;
 static TArray<FUObjectItem*> GUnreachableObjects;
 static FCriticalSection GUnreachableObjectsCritical;
 static int32 GUnrechableObjectIndex = 0;
-
-struct FGCTimingInfo
-{
-	double LastGCTime{};
-	double LastGCDuration = -1;
-};
-static FGCTimingInfo GTimingInfo;
+static double GLastGCTime = 0;
 
 bool GIsGarbageCollecting = false;
 
@@ -249,36 +242,6 @@ static FAutoConsoleVariableRef CMultithreadedDestructionEnabled(
 	ECVF_Default
 );
 
-static UObjectReachabilityStressData* GReachabilityStressData;
-static void AllocateReachabilityStressData(FOutputDevice&)
-{
-	if (GReachabilityStressData)
-	{
-		return;
-	}
-	GReachabilityStressData = GenerateReachabilityStressData();
-}
-
-static void UnlinkReachabilityStressData(FOutputDevice&)
-{
-	if (!GReachabilityStressData)
-	{
-		return;
-	}
-	UnlinkReachabilityStressData(GReachabilityStressData);
-	GReachabilityStressData = nullptr;
-}
-
-static FAutoConsoleCommandWithOutputDevice GGenerateReachabilityStressDataCmd(TEXT("gc.GenerateReachabilityStressData"),
-		        						      TEXT("Allocate deeply-nested UObject tree to "
-						        			   "stress test reachability analysis."),
-FConsoleCommandWithOutputDeviceDelegate::CreateStatic(&AllocateReachabilityStressData));
-
-static FAutoConsoleCommandWithOutputDevice GUnlinkReachabilityStressDataCmd(TEXT("gc.UnlinkReachabilityStressData"),
-                                 					    TEXT("Unlink previously-generated reachability analysis "
-										 "stress test data for collection in the next cycle."),
-FConsoleCommandWithOutputDeviceDelegate::CreateStatic(&UnlinkReachabilityStressData));
-
 #if UE_BUILD_SHIPPING
 static constexpr int32 GGarbageReferenceTrackingEnabled = 0;
 #else
@@ -4537,8 +4500,6 @@ FORCENOINLINE static void CollectGarbageFull(EObjectFlags KeepFlags)
 
 FORCEINLINE void CollectGarbageInternal(EObjectFlags KeepFlags, bool bPerformFullPurge)
 {
-	const double StartTime = FPlatformTime::Seconds();
-
 	if (bPerformFullPurge)
 	{
 		CollectGarbageFull(KeepFlags);
@@ -4547,8 +4508,6 @@ FORCEINLINE void CollectGarbageInternal(EObjectFlags KeepFlags, bool bPerformFul
 	{
 		CollectGarbageIncremental(KeepFlags);
 	}
-
-	GTimingInfo.LastGCDuration = FPlatformTime::Seconds() - StartTime;
 }
 
 /** 
@@ -4763,7 +4722,7 @@ void CollectGarbageImpl(EObjectFlags KeepFlags)
 		FCoreUObjectDelegates::GetPostGarbageCollect().Broadcast();
 	}
 
-	GTimingInfo.LastGCTime = FPlatformTime::Seconds();
+	GLastGCTime = FPlatformTime::Seconds();
 	STAT_ADD_CUSTOMMESSAGE_NAME( STAT_NamedMarker, TEXT( "GarbageCollection - End" ) );
 }
 #endif // !UE_WITH_GC
@@ -4789,12 +4748,7 @@ FString FGarbageReferenceInfo::GetReferencingObjectInfo() const
 
 double GetLastGCTime()
 {
-	return GTimingInfo.LastGCTime;
-}
-
-double GetLastGCDuration()
-{
-	return GTimingInfo.LastGCDuration;
+	return GLastGCTime;
 }
 
 bool IsIncrementalUnhashPending()
diff --git a/Engine/Source/Runtime/CoreUObject/Private/UObject/GarbageCollectionTesting.cpp b/Engine/Source/Runtime/CoreUObject/Private/UObject/GarbageCollectionTesting.cpp
deleted file mode 100644
diff --git a/Engine/Source/Runtime/CoreUObject/Private/UObject/GarbageCollectionTesting.h b/Engine/Source/Runtime/CoreUObject/Private/UObject/GarbageCollectionTesting.h
deleted file mode 100644
diff --git a/Engine/Source/Runtime/CoreUObject/Public/UObject/GarbageCollection.h b/Engine/Source/Runtime/CoreUObject/Public/UObject/GarbageCollection.h
--- a/Engine/Source/Runtime/CoreUObject/Public/UObject/GarbageCollection.h
+++ b/Engine/Source/Runtime/CoreUObject/Public/UObject/GarbageCollection.h
@@ -518,13 +518,6 @@ extern COREUOBJECT_API bool GIsGarbageCollecting;
  */
 COREUOBJECT_API double GetLastGCTime();
 
-/**
- * Gets the duration of the last GC run.
- *
- * @return	Returns the last GC duration, -1 if GC has never run.
- */
-COREUOBJECT_API double GetLastGCDuration();
-
 /**
 * Whether we are inside garbage collection
 */
diff --git a/Engine/Source/Runtime/Engine/Private/DebugGarbageCollectionGraph.cpp b/Engine/Source/Runtime/Engine/Private/DebugGarbageCollectionGraph.cpp
deleted file mode 100644
diff --git a/Engine/Source/Runtime/Engine/Private/DebugGarbageCollectionGraph.h b/Engine/Source/Runtime/Engine/Private/DebugGarbageCollectionGraph.h
deleted file mode 100644
