RedirectCollector: For commandlets using the function ResolveAllSoftObjectPaths to load all softobjectpaths, downgrade the warning about missing objects from Warning to Display. It can occur for SoftObjectPaths that were added to config for some plugins that are not loaded when the commandlet is run.
#rnx

diff --git a/Engine/Source/Runtime/CoreUObject/Private/Misc/RedirectCollector.cpp b/Engine/Source/Runtime/CoreUObject/Private/Misc/RedirectCollector.cpp
--- a/Engine/Source/Runtime/CoreUObject/Private/Misc/RedirectCollector.cpp
+++ b/Engine/Source/Runtime/CoreUObject/Private/Misc/RedirectCollector.cpp
@@ -123,7 +123,7 @@ void FRedirectCollector::ResolveAllSoftObjectPaths(FName FilterPackage)
 			else
 			{
 				const FString Referencer = SoftObjectPathProperty.GetPropertyName().ToString().Len() ? SoftObjectPathProperty.GetPropertyName().ToString() : TEXT("Unknown");
-				UE_LOG(LogRedirectors, Warning, TEXT("Soft Object Path '%s' was not found when resolving paths! (Referencer '%s:%s')"), *ToLoad, *ReferencerPackageName.ToString(), *Referencer);
+				UE_LOG(LogRedirectors, Display, TEXT("Soft Object Path '%s' was not found when resolving paths! (Referencer '%s:%s')"), *ToLoad, *ReferencerPackageName.ToString(), *Referencer);
 			}
 		}
 	};
