Improve asset metadata tracking by accurately using object names

diff --git a/Engine/Source/Runtime/CoreUObject/Private/Serialization/AsyncLoading.cpp b/Engine/Source/Runtime/CoreUObject/Private/Serialization/AsyncLoading.cpp
--- a/Engine/Source/Runtime/CoreUObject/Private/Serialization/AsyncLoading.cpp
+++ b/Engine/Source/Runtime/CoreUObject/Private/Serialization/AsyncLoading.cpp
@@ -3158,10 +3158,11 @@ void FAsyncPackage::EventDrivenCreateExport(int32 LocalExportIndex)
 
 	TRACE_LOADTIME_CREATE_EXPORT_SCOPE(this, &Export.Object);
 
+	const UClass* Class = CastEventDrivenIndexToObject<UClass>(Export.ClassIndex, false);
 	LLM_SCOPE(ELLMTag::AsyncLoading);
 	LLM_SCOPED_TAG_WITH_OBJECT_IN_SET(GetLinkerRoot(), ELLMTagSet::Assets);
-	LLM_SCOPED_TAG_WITH_OBJECT_IN_SET(CastEventDrivenIndexToObject<UClass>(Export.ClassIndex, false), ELLMTagSet::AssetClasses);
-   	UE_TRACE_METADATA_SCOPE_ASSET(GetLinkerRoot(), CastEventDrivenIndexToObject<UClass>(Export.ClassIndex, false));
+	LLM_SCOPED_TAG_WITH_OBJECT_IN_SET(Class, ELLMTagSet::AssetClasses);
+   	UE_TRACE_METADATA_SCOPE_ASSET_FNAME(Export.ObjectName, Class->GetFName(), GetLinkerRoot()->GetFName());
 
 	// Check whether we already loaded the object and if not whether the context flags allow loading it.
 	//check(!Export.Object || Export.Object->HasAnyFlags(RF_ClassDefaultObject)); // we should not have this yet, unless it is a CDO
@@ -3445,14 +3446,14 @@ void FAsyncPackage::EventDrivenSerializeExport(int32 LocalExportIndex)
 
 	FObjectExport& Export = Linker->ExportMap[LocalExportIndex];
 
+	const UClass* Class = CastEventDrivenIndexToObject<UClass>(Export.ClassIndex, false);
 	LLM_SCOPED_TAG_WITH_OBJECT_IN_SET(GetLinkerRoot(), ELLMTagSet::Assets);
-	LLM_SCOPED_TAG_WITH_OBJECT_IN_SET(CastEventDrivenIndexToObject<UClass>(Export.ClassIndex, false), ELLMTagSet::AssetClasses);
-  	UE_TRACE_METADATA_SCOPE_ASSET(GetLinkerRoot(), CastEventDrivenIndexToObject<UClass>(Export.ClassIndex, false));
+	LLM_SCOPED_TAG_WITH_OBJECT_IN_SET(Class, ELLMTagSet::AssetClasses);
 
 	UObject* Object = Export.Object;
 	if (Object && Object->HasAnyFlags(RF_NeedLoad))
 	{
-
+  		UE_TRACE_METADATA_SCOPE_ASSET(Object, Class);
 		Linker->GetAsyncLoader()->LogItem(TEXT("EventDrivenSerializeExport"), Export.SerialOffset, Export.SerialSize);
 
 		LastTypeOfWorkPerformed = TEXT("EventDrivenSerializeExport");
diff --git a/Engine/Source/Runtime/CoreUObject/Private/Serialization/AsyncLoading2.cpp b/Engine/Source/Runtime/CoreUObject/Private/Serialization/AsyncLoading2.cpp
--- a/Engine/Source/Runtime/CoreUObject/Private/Serialization/AsyncLoading2.cpp
+++ b/Engine/Source/Runtime/CoreUObject/Private/Serialization/AsyncLoading2.cpp
@@ -5914,7 +5914,7 @@ void FAsyncPackage2::EventDrivenCreateExport(const FAsyncPackageHeaderData& Head
 
 	LLM_SCOPED_TAG_WITH_OBJECT_IN_SET(GetLinkerRoot(), ELLMTagSet::Assets);
 	LLM_SCOPED_TAG_WITH_OBJECT_IN_SET(LoadClass, ELLMTagSet::AssetClasses);
-    UE_TRACE_METADATA_SCOPE_ASSET(GetLinkerRoot(), LoadClass);
+    UE_TRACE_METADATA_SCOPE_ASSET_FNAME(ObjectName, LoadClass->GetFName(), GetLinkerRoot()->GetFName());
 
 	// Try to find existing object first as we cannot in-place replace objects, could have been created by other export in this package
 	{
diff --git a/Engine/Source/Runtime/CoreUObject/Private/UObject/Obj.cpp b/Engine/Source/Runtime/CoreUObject/Private/UObject/Obj.cpp
--- a/Engine/Source/Runtime/CoreUObject/Private/UObject/Obj.cpp
+++ b/Engine/Source/Runtime/CoreUObject/Private/UObject/Obj.cpp
@@ -1116,7 +1116,7 @@ void UObject::ConditionalPostLoad()
 				UE_SCOPED_COOK_STAT(Package->GetFName(), EPackageEventStatType::LoadPackage);
 				LLM_SCOPED_TAG_WITH_OBJECT_IN_SET(Package, ELLMTagSet::Assets);
 				LLM_SCOPED_TAG_WITH_OBJECT_IN_SET(GetClass(), ELLMTagSet::AssetClasses);
-				UE_TRACE_METADATA_SCOPE_ASSET(Package, GetClass());
+				UE_TRACE_METADATA_SCOPE_ASSET(this, GetClass());
 				TRACE_LOADTIME_POSTLOAD_OBJECT_SCOPE(this);
 				
 				PostLoad();
