Fix race on UPackage bitfields when using ALT

#rnx

diff --git a/Engine/Source/Runtime/CoreUObject/Private/UObject/Package.cpp b/Engine/Source/Runtime/CoreUObject/Private/UObject/Package.cpp
--- a/Engine/Source/Runtime/CoreUObject/Private/UObject/Package.cpp
+++ b/Engine/Source/Runtime/CoreUObject/Private/UObject/Package.cpp
@@ -395,8 +395,7 @@ void FixupPackageEditorOnlyFlag(FName PackageThatGotEditorOnlyFlagCleared, bool
 
 void UPackage::SetLoadedByEditorPropertiesOnly(bool bIsEditorOnly, bool bRecursive /*= false*/)
 {
-	const bool bWasEditorOnly = bLoadedByEditorPropertiesOnly;
-	bLoadedByEditorPropertiesOnly = bIsEditorOnly;
+	const bool bWasEditorOnly = bLoadedByEditorPropertiesOnly.exchange(bIsEditorOnly);
 	if (bWasEditorOnly && !bIsEditorOnly)
 	{
 		FixupPackageEditorOnlyFlag(GetFName(), bRecursive);
diff --git a/Engine/Source/Runtime/CoreUObject/Public/UObject/Package.h b/Engine/Source/Runtime/CoreUObject/Public/UObject/Package.h
--- a/Engine/Source/Runtime/CoreUObject/Public/UObject/Package.h
+++ b/Engine/Source/Runtime/CoreUObject/Public/UObject/Package.h
@@ -191,9 +191,6 @@ private:
 	uint8	bDirty:1;
 
 #if WITH_EDITORONLY_DATA
-	/** True if this package is only referenced by editor-only properties */
-	uint8 bLoadedByEditorPropertiesOnly:1;
-
 	/** True if this package is a dynamic PIE package with external objects still loading */
 	uint8 bIsDynamicPIEPackagePending:1;
 #endif
@@ -208,29 +205,36 @@ public:
 	/** Whether this package has been fully loaded (aka had all it's exports created) at some point.															*/
 	mutable uint8 bHasBeenFullyLoaded:1;
 
+	/**
+	 * Whether this package can be imported, i.e. its package name is a package that exists on disk.
+	 * Note: This includes all normal packages where the Name matches the FileName
+	 * and localized packages shadowing an existing source package,
+	 * but excludes level streaming packages with /Temp/ names.
+	 */
+	uint8 bCanBeImported:1;
+
 #if WITH_EDITORONLY_DATA
-	/** Set to true after serialization and postload and before returning from LoadPackage or calling load completion delegate, for loaded packages.			*/
 private:
-	uint8 bHasBeenEndLoaded : 1;
+	/** True if this package is only referenced by editor-only properties 
+	*   Note: This flag is manipulated on different threads, do not mix with other bitfields above as they might get corrupted.
+	*/
+	std::atomic<bool> bLoadedByEditorPropertiesOnly{ false };
+
+	/** Set to true after serialization and postload and before returning from LoadPackage or calling load completion delegate, for loaded packages.			
+	*   Note: This flag is manipulated on different threads, do not mix with other bitfields above as they might get corrupted.
+	*/
+	std::atomic<bool> bHasBeenEndLoaded{ false };
 public:
 	bool GetHasBeenEndLoaded() const
 	{
-		return bHasBeenEndLoaded != 0;
+		return bHasBeenEndLoaded.load(std::memory_order_acquire) != 0;
 	}
 	void SetHasBeenEndLoaded(bool bValue)
 	{
-		bHasBeenEndLoaded = bValue ? 1 : 0;
+		bHasBeenEndLoaded.store(bValue, std::memory_order_release);
 	}
 #endif
 
-	/**
-	 * Whether this package can be imported, i.e. its package name is a package that exists on disk.
-	 * Note: This includes all normal packages where the Name matches the FileName
-	 * and localized packages shadowing an existing source package,
-	 * but excludes level streaming packages with /Temp/ names.
-	 */
-	uint8 bCanBeImported:1;
-
 private:
 	// @note this should probably be entirely deprecated and removed but certain stat dump function are still using it, just compile it out in shipping for now
 #if !UE_BUILD_SHIPPING
