Fix race and typo in EditorBulkData trace counters

#rnx

diff --git a/Engine/Source/Runtime/CoreUObject/Private/Serialization/EditorBulkData.cpp b/Engine/Source/Runtime/CoreUObject/Private/Serialization/EditorBulkData.cpp
--- a/Engine/Source/Runtime/CoreUObject/Private/Serialization/EditorBulkData.cpp
+++ b/Engine/Source/Runtime/CoreUObject/Private/Serialization/EditorBulkData.cpp
@@ -41,8 +41,8 @@
 	#define UE_CORRUPTED_DATA_SEVERITY Error
 #endif // VBD_CORRUPTED_PAYLOAD_IS_FATAL
 
-TRACE_DECLARE_INT_COUNTER(EditorBulkData_PayloadDataLoaded, TEXT("EditorBulkData/PayloadDataLoaded"));
-TRACE_DECLARE_INT_COUNTER(EditorBulkData_PayloadDataPulled, TEXT("EditorBulkData/PayloadDatPulled"));
+TRACE_DECLARE_ATOMIC_INT_COUNTER(EditorBulkData_PayloadDataLoaded, TEXT("EditorBulkData/PayloadDataLoaded"));
+TRACE_DECLARE_ATOMIC_INT_COUNTER(EditorBulkData_PayloadDataPulled, TEXT("EditorBulkData/PayloadDataPulled"));
 
 namespace UE::Serialization
 {
