PackageName testing: Remove deprecation-handling code that cleaned-up temporary paths from the previous version of the test.
#rnx

diff --git a/Engine/Source/Runtime/CoreUObject/Private/Misc/PackageName.cpp b/Engine/Source/Runtime/CoreUObject/Private/Misc/PackageName.cpp
--- a/Engine/Source/Runtime/CoreUObject/Private/Misc/PackageName.cpp
+++ b/Engine/Source/Runtime/CoreUObject/Private/Misc/PackageName.cpp
@@ -3073,31 +3073,6 @@ bool FPackageNameTests::RunTest(const FString& Parameters)
 		}
 	}
 
-	// DEPRECATED Mar 2023
-#if WITH_EDITOR
-	{
-		// This block was created on Mar 2023 to clean up previous test results that were not cleaned up.
-		const TCHAR* OldTestMountLocalPathRoot = TEXT("../../../PackageNameTestDNE");
-		if (IFileManager::Get().DirectoryExists(OldTestMountLocalPathRoot))
-		{
-			const TCHAR* OldTestMountLocalPathContent = TEXT("../../../PackageNameTestDNE/Content");
-			bool bHasContent = IFileManager::Get().DirectoryExists(OldTestMountLocalPathContent);
-			TArray<FString> Filenames;
-			IFileManager::Get().FindFiles(Filenames, *WriteToString<64>(OldTestMountLocalPathRoot, TEXT("/*")), true /* Files */, true /* Directories */);
-			if (bHasContent)
-			{
-				IFileManager::Get().FindFiles(Filenames, *WriteToString<64>(OldTestMountLocalPathContent, TEXT("/*")), true /* Files */, true /* Directories */);
-			}
-			int32 NumExpectedFiles = bHasContent ? 1 : 0;
-			if (Filenames.Num() == NumExpectedFiles)
-			{
-				IFileManager::Get().DeleteDirectory(OldTestMountLocalPathRoot, false /* bRequireExists */, true /* Delete Tree */);
-			}
-		};
-	}
-#endif
-
-
 	return true;
 }
 
