Refactored DuplicateObject out of UObjectGlobals.h and replaced its call to GetTransientPackage with GetTransientPackageForRename to better support Within classes.

diff --git a/Engine/Source/Runtime/CoreUObject/Private/UObject/UObjectGlobals.cpp b/Engine/Source/Runtime/CoreUObject/Private/UObject/UObjectGlobals.cpp
--- a/Engine/Source/Runtime/CoreUObject/Private/UObject/UObjectGlobals.cpp
+++ b/Engine/Source/Runtime/CoreUObject/Private/UObject/UObjectGlobals.cpp
@@ -4280,6 +4280,19 @@ void CheckIsClassChildOf_Internal(const UClass* Parent, const UClass* Child)
 }
 #endif
 
+UObject* DuplicateObject_Internal(UClass* Class, const UObject* SourceObject, UObject* Outer, FName Name)
+{
+	if (SourceObject != nullptr)
+	{
+		if (Outer == nullptr || Outer == INVALID_OBJECT)
+		{
+			Outer = (UObject*)GetTransientOuterForRename(Class);
+		}
+		return StaticDuplicateObject(SourceObject,Outer,Name);
+	}
+	return nullptr;
+}
+
 FStaticConstructObjectParameters::FStaticConstructObjectParameters(const UClass* InClass)
 	: Class(InClass)
 	, Outer((UObject*)GetTransientPackage())
diff --git a/Engine/Source/Runtime/CoreUObject/Public/UObject/UObjectGlobals.h b/Engine/Source/Runtime/CoreUObject/Public/UObject/UObjectGlobals.h
--- a/Engine/Source/Runtime/CoreUObject/Public/UObject/UObjectGlobals.h
+++ b/Engine/Source/Runtime/CoreUObject/Public/UObject/UObjectGlobals.h
@@ -1668,6 +1668,18 @@ FUNCTION_NON_NULL_RETURN_END
 	return static_cast<T*>(StaticConstructObject_Internal(Params));
 }
 
+/**
+ * Convenience function for duplicating an object
+ *
+ * @param Class the class of the object being copied
+ * @param SourceObject the object being copied
+ * @param Outer the outer to use for the object
+ * @param Name the optional name of the object
+ *
+ * @return the copied object or null if it failed for some reason
+ */
+COREUOBJECT_API UObject* DuplicateObject_Internal(UClass* Class, const UObject* SourceObject, UObject* Outer, const FName NAME_None);
+
 /**
  * Convenience template for duplicating an object
  *
@@ -1680,28 +1692,13 @@ FUNCTION_NON_NULL_RETURN_END
 template< class T >
 T* DuplicateObject(T const* SourceObject,UObject* Outer, const FName Name = NAME_None)
 {
-	if (SourceObject != nullptr)
-	{
-		if (Outer == nullptr || Outer == INVALID_OBJECT)
-		{
-			Outer = (UObject*)GetTransientPackage();
-		}
-		return (T*)StaticDuplicateObject(SourceObject,Outer,Name);
-	}
-	return nullptr;
+	return static_cast<T*>(DuplicateObject_Internal(T::StaticClass(), SourceObject, Outer, Name));
 }
+
 template <typename T>
 T* DuplicateObject(const TObjectPtr<T>& SourceObject,UObject* Outer, const FName Name = NAME_None)
 {
-	if (SourceObject != nullptr)
-	{
-		if (Outer == nullptr || Outer == INVALID_OBJECT)
-		{
-			Outer = (UObject*)GetTransientPackage();
-		}
-		return (T*)StaticDuplicateObject(SourceObject,Outer,Name);
-	}
-	return nullptr;
+	return static_cast<T*>(DuplicateObject_Internal(T::StaticClass(), SourceObject, Outer, Name));
 }
 
 /**
