Replaced uses of FLazyEvent with FManualResetEvent

#rnx

diff --git a/Engine/Source/Editor/UnrealEd/Private/Subsystems/EditorAssetSubsystem.cpp b/Engine/Source/Editor/UnrealEd/Private/Subsystems/EditorAssetSubsystem.cpp
--- a/Engine/Source/Editor/UnrealEd/Private/Subsystems/EditorAssetSubsystem.cpp
+++ b/Engine/Source/Editor/UnrealEd/Private/Subsystems/EditorAssetSubsystem.cpp
@@ -8,7 +8,6 @@
 #include "Editor.h"
 #include "Editor/Transactor.h"
 #include "EditorScriptingHelpers.h"
-#include "Experimental/Async/LazyEvent.h"
 #include "FileHelpers.h"
 #include "HAL/FileManager.h"
 #include "Misc/PackageName.h"
diff --git a/Engine/Source/Runtime/CoreUObject/Private/Serialization/BulkDataStreaming.cpp b/Engine/Source/Runtime/CoreUObject/Private/Serialization/BulkDataStreaming.cpp
--- a/Engine/Source/Runtime/CoreUObject/Private/Serialization/BulkDataStreaming.cpp
+++ b/Engine/Source/Runtime/CoreUObject/Private/Serialization/BulkDataStreaming.cpp
@@ -1,8 +1,8 @@
 // Copyright Epic Games, Inc. All Rights Reserved.
 #include "Serialization/BulkData.h"
+#include "Async/ManualResetEvent.h"
 #include "Async/MappedFileHandle.h"
 #include "Containers/ChunkedArray.h"
-#include "Experimental/Async/LazyEvent.h"
 #include "HAL/CriticalSection.h"
 #include "IO/IoDispatcher.h"
 #include "Misc/ConfigCacheIni.h"
@@ -132,15 +132,14 @@ protected:
 
 private:
 
-	UE::FLazyEvent		DoneEvent;
-	FIoRequest			Request;
-	int64				SizeResult;
-	std::atomic<uint32>	Status;
+	UE::FManualResetEvent	DoneEvent;
+	FIoRequest				Request;
+	int64					SizeResult;
+	std::atomic<uint32>		Status;
 };
 
 FChunkRequest::FChunkRequest(FIoBuffer&& InBuffer)
 	: Buffer(MoveTemp(InBuffer))
-	, DoneEvent(EEventMode::ManualReset)
 	, SizeResult(-1)
 	, Status{uint32(EChunkRequestStatus::None)}
 {
@@ -173,7 +172,7 @@ void FChunkRequest::Issue(FIoChunkId ChunkId, FIoReadOptions Options, int32 Prio
 		HandleChunkResult(MoveTemp(Result));
 		Status.store(uint32(ReadyOrCanceled | EChunkRequestStatus::CallbackTriggered), std::memory_order_release); 
 
-		DoneEvent.Trigger();
+		DoneEvent.Notify();
 	});
 
 	IoBatch.Issue();
@@ -183,8 +182,7 @@ bool FChunkRequest::WaitForChunkRequest(float TimeLimitSeconds)
 {
 	checkf(GetStatus() != EChunkRequestStatus::None, TEXT("The request must be issued before waiting for completion"));
 
-	const uint32 TimeLimitMilliseconds = TimeLimitSeconds <= 0.0f ? MAX_uint32 : (uint32)(TimeLimitSeconds * 1000.0f);
-	return DoneEvent.Wait(TimeLimitMilliseconds);
+	return DoneEvent.WaitFor(TimeLimitSeconds <= 0.0f ? FMonotonicTimeSpan::Infinity() : FMonotonicTimeSpan::FromSeconds(TimeLimitSeconds));
 }
 
 void FChunkRequest::CancelChunkRequest()
@@ -791,7 +789,7 @@ public:
 	virtual bool Wait(uint32 Milliseconds) override final
 	{
 		check(GetStatus() != FBulkDataRequest::EStatus::None);
-		return DoneEvent.Wait(Milliseconds);
+		return DoneEvent.WaitFor(UE::FMonotonicTimeSpan::FromMilliseconds(Milliseconds));
 	}
 
 	void Read(
@@ -856,15 +854,15 @@ private:
 		}
 
 		SetStatus(CompletionStatus);
-		DoneEvent.Trigger();
+		DoneEvent.Notify();
 	}
 
 	static constexpr int32 TargetBytesPerChunk = sizeof(FChunkBatchReadRequest) * 8;
 	using FRequests = TChunkedArray<FChunkBatchReadRequest, TargetBytesPerChunk>;
 
-	FIoBatch		IoBatch;
-	FRequests		Requests; //TOOD: Optimize if there's only a single read request 
-	UE::FLazyEvent	DoneEvent{EEventMode::ManualReset};
+	FIoBatch				IoBatch;
+	FRequests				Requests; //TOOD: Optimize if there's only a single read request 
+	UE::FManualResetEvent	DoneEvent;
 	FBulkDataRequest::FCompletionCallback CompletionCallback;
 };
 
diff --git a/Engine/Source/Runtime/Engine/Private/Streaming/Texture2DStreamIn_DerivedData.cpp b/Engine/Source/Runtime/Engine/Private/Streaming/Texture2DStreamIn_DerivedData.cpp
--- a/Engine/Source/Runtime/Engine/Private/Streaming/Texture2DStreamIn_DerivedData.cpp
+++ b/Engine/Source/Runtime/Engine/Private/Streaming/Texture2DStreamIn_DerivedData.cpp
@@ -15,7 +15,6 @@ namespace UE
 
 FTexture2DStreamIn_DerivedData::FTexture2DStreamIn_DerivedData(UTexture2D* InTexture, bool bInHighPriority)
 	: FTexture2DStreamIn(InTexture)
-	, ResponseComplete(EEventMode::AutoReset)
 	, bHighPriority(bInHighPriority)
 {
 }
@@ -78,7 +77,7 @@ void FTexture2DStreamIn_DerivedData::DoBeginIoRequests(const FContext& Context)
 			}
 		}
 
-		ResponseComplete.Trigger();
+		ResponseComplete.Notify();
 
 		// Tick to schedule the task or cancellation callback. Use TT_None to avoid a deadlock.
 		Tick(FTexture2DUpdate::TT_None);
diff --git a/Engine/Source/Runtime/Engine/Private/Streaming/Texture2DStreamIn_DerivedData.h b/Engine/Source/Runtime/Engine/Private/Streaming/Texture2DStreamIn_DerivedData.h
--- a/Engine/Source/Runtime/Engine/Private/Streaming/Texture2DStreamIn_DerivedData.h
+++ b/Engine/Source/Runtime/Engine/Private/Streaming/Texture2DStreamIn_DerivedData.h
@@ -2,7 +2,7 @@
 
 #pragma once
 
-#include "Experimental/Async/LazyEvent.h"
+#include "Async/ManualResetEvent.h"
 #include "Serialization/DerivedData.h"
 #include "Texture2DStreamIn.h"
 
@@ -31,7 +31,7 @@ protected:
 
 private:
 	FDerivedDataIoResponse Response;
-	FLazyEvent ResponseComplete;
+	FManualResetEvent ResponseComplete;
 	bool bHighPriority;
 };
 
