Reverted the conversion of FLazyEvent to FManualResetEvent while investigating an issue

#rnx

diff --git a/Engine/Source/Developer/DerivedDataCache/Private/DerivedDataBuildJobContext.cpp b/Engine/Source/Developer/DerivedDataCache/Private/DerivedDataBuildJobContext.cpp
--- a/Engine/Source/Developer/DerivedDataCache/Private/DerivedDataBuildJobContext.cpp
+++ b/Engine/Source/Developer/DerivedDataCache/Private/DerivedDataBuildJobContext.cpp
@@ -140,7 +140,7 @@ void FBuildJobContext::BeginBuild(IRequestOwner& InOwner, TUniqueFunction<void (
 void FBuildJobContext::EndBuild()
 {
 	OnEndBuild();
-	BuildCompleteEvent.Notify();
+	BuildCompleteEvent.Trigger();
 	Owner = nullptr;
 }
 
diff --git a/Engine/Source/Developer/DerivedDataCache/Private/DerivedDataBuildJobContext.h b/Engine/Source/Developer/DerivedDataCache/Private/DerivedDataBuildJobContext.h
--- a/Engine/Source/Developer/DerivedDataCache/Private/DerivedDataBuildJobContext.h
+++ b/Engine/Source/Developer/DerivedDataCache/Private/DerivedDataBuildJobContext.h
@@ -2,7 +2,6 @@
 
 #pragma once
 
-#include "Async/ManualResetEvent.h"
 #include "Compression/CompressedBuffer.h"
 #include "Containers/Map.h"
 #include "Containers/StringView.h"
@@ -10,6 +9,7 @@
 #include "DerivedDataCacheKey.h"
 #include "DerivedDataRequest.h"
 #include "DerivedDataSharedString.h"
+#include "Experimental/Async/LazyEvent.h"
 #include "Memory/MemoryFwd.h"
 #include "Serialization/CompactBinary.h"
 #include "Templates/Function.h"
@@ -88,12 +88,12 @@ private:
 	FBuildOutputBuilder& OutputBuilder;
 	TMap<FUtf8SharedString, FCbObject> Constants;
 	TMap<FUtf8SharedString, FCompressedBuffer> Inputs;
+	UE::FLazyEvent BuildCompleteEvent{EEventMode::ManualReset};
 	TUniqueFunction<void ()> OnEndBuild;
 	IRequestOwner* Owner{};
 	uint64 RequiredMemory{};
 	ECachePolicy CachePolicyMask;
 	EBuildPolicy BuildPolicyMask;
-	FManualResetEvent BuildCompleteEvent;
 	bool bIsAsyncBuild{false};
 	bool bIsAsyncBuildComplete{false};
 	bool bDeterministicOutputCheck{true};
diff --git a/Engine/Source/Developer/DerivedDataCache/Private/DerivedDataThreadPoolTask.cpp b/Engine/Source/Developer/DerivedDataCache/Private/DerivedDataThreadPoolTask.cpp
--- a/Engine/Source/Developer/DerivedDataCache/Private/DerivedDataThreadPoolTask.cpp
+++ b/Engine/Source/Developer/DerivedDataCache/Private/DerivedDataThreadPoolTask.cpp
@@ -3,10 +3,10 @@
 #include "DerivedDataThreadPoolTask.h"
 
 #include "Async/InheritedContext.h"
-#include "Async/ManualResetEvent.h"
 #include "DerivedDataRequest.h"
 #include "DerivedDataRequestOwner.h"
 #include "DerivedDataRequestTypes.h"
+#include "Experimental/Async/LazyEvent.h"
 #include "Misc/IQueuedWork.h"
 #include "Misc/QueuedThreadPool.h"
 #include "Stats/Stats.h"
@@ -56,7 +56,7 @@ private:
 	IRequestOwner& Owner;
 	FQueuedThreadPool& ThreadPool;
 	TUniqueFunction<void ()> TaskBody;
-	FManualResetEvent Event;
+	FLazyEvent Event{EEventMode::ManualReset};
 	std::atomic<bool> bClaimed = false;
 };
 
@@ -91,7 +91,7 @@ TRefCountPtr<IRequest> FThreadPoolTaskRequest::TryEnd()
 		FInheritedContextScope InheritedContextScope = RestoreInheritedContext();
 		FScopeCycleCounter Scope(GetStatId(), /*bAlways*/ true);
 		TaskBody();
-		Event.Notify();
+		Event.Trigger();
 	});
 }
 
diff --git a/Engine/Source/Developer/DerivedDataCache/Private/FileSystemCacheStore.cpp b/Engine/Source/Developer/DerivedDataCache/Private/FileSystemCacheStore.cpp
--- a/Engine/Source/Developer/DerivedDataCache/Private/FileSystemCacheStore.cpp
+++ b/Engine/Source/Developer/DerivedDataCache/Private/FileSystemCacheStore.cpp
@@ -7,7 +7,6 @@
 #include "Algo/StableSort.h"
 #include "Algo/Transform.h"
 #include "Async/Async.h"
-#include "Async/ManualResetEvent.h"
 #include "Containers/StaticBitArray.h"
 #include "DerivedDataBackendInterface.h"
 #include "DerivedDataCacheInterface.h"
@@ -18,6 +17,7 @@
 #include "DerivedDataChunk.h"
 #include "DerivedDataRequestOwner.h"
 #include "DerivedDataValue.h"
+#include "Experimental/Async/LazyEvent.h"
 #include "Features/IModularFeatures.h"
 #include "HAL/Event.h"
 #include "HAL/FileManager.h"
@@ -265,7 +265,7 @@ private:
 
 	IFileManager& FileManager = IFileManager::Get();
 
-	mutable FManualResetEvent IdleEvent;
+	mutable FLazyEvent IdleEvent;
 	FEventRef WaitEvent;
 	FThread Thread;
 
@@ -315,6 +315,7 @@ FFileSystemCacheStoreMaintainer::FFileSystemCacheStoreMaintainer(
 	: Params(InParams)
 	, CachePath(InCachePath)
 	, bExitAfterScan(Params.ScanFrequency.GetTotalDays() > MaxScanFrequencyDays)
+	, IdleEvent(EEventMode::ManualReset)
 	, WaitEvent(EEventMode::AutoReset)
 	, Thread(
 		TEXT("FileSystemCacheStoreMaintainer"),
@@ -349,7 +350,7 @@ void FFileSystemCacheStoreMaintainer::Tick()
 		Loop();
 	}
 	bIdle = true;
-	IdleEvent.Notify();
+	IdleEvent.Trigger();
 }
 
 void FFileSystemCacheStoreMaintainer::Loop()
@@ -369,7 +370,7 @@ void FFileSystemCacheStoreMaintainer::Loop()
 		bIdle = false;
 		Scan();
 		bIdle = true;
-		IdleEvent.Notify();
+		IdleEvent.Trigger();
 		bIgnoreScanRate = false;
 		const FDateTime ScanEnd = FDateTime::Now();
 
@@ -394,7 +395,7 @@ void FFileSystemCacheStoreMaintainer::Loop()
 	}
 
 	bIdle = true;
-	IdleEvent.Notify();
+	IdleEvent.Trigger();
 }
 
 void FFileSystemCacheStoreMaintainer::Scan()
diff --git a/Engine/Source/Developer/DerivedDataCache/Private/Http/CurlHttpClient.cpp b/Engine/Source/Developer/DerivedDataCache/Private/Http/CurlHttpClient.cpp
--- a/Engine/Source/Developer/DerivedDataCache/Private/Http/CurlHttpClient.cpp
+++ b/Engine/Source/Developer/DerivedDataCache/Private/Http/CurlHttpClient.cpp
@@ -3,10 +3,10 @@
 #include "HttpClient.h"
 
 #include "Async/InheritedContext.h"
-#include "Async/ManualResetEvent.h"
 #include "Containers/ConsumeAllMpmcQueue.h"
 #include "Containers/LockFreeList.h"
 #include "Containers/StringView.h"
+#include "HAL/Event.h"
 #include "HAL/PlatformProcess.h"
 #include "HAL/Thread.h"
 #include "Memory/CompositeBuffer.h"
@@ -294,7 +294,7 @@ public:
 	TRefCountPtr<IHttpResponseMonitor> GetMonitor() final { return this; }
 
 	void Cancel() final;
-	void Wait() const final { CompleteEvent.Wait(); }
+	void Wait() const final { CompleteEvent->Wait(); }
 	bool Poll() const final { return EnumHasAnyFlags(State.load(), ECurlHttpResponseState::Complete); }
 	bool IsCanceled() const { return EnumHasAnyFlags(State.load(std::memory_order_relaxed), ECurlHttpResponseState::Canceled); }
 
@@ -340,7 +340,7 @@ private:
 	FAnsiStringView Error;
 	ANSICHAR ErrorBuffer[CURL_ERROR_SIZE]{};
 	FHttpResponseStats Stats;
-	mutable FManualResetEvent CompleteEvent;
+	FEventRef CompleteEvent{EEventMode::ManualReset};
 
 	enum class EHttpReceiverFunction : uint8
 	{
@@ -1138,7 +1138,7 @@ void FCurlHttpResponse::SetComplete(CURLcode Code)
 	}
 
 	CurlHttp::Private::AtomicFetchOr(State, ECurlHttpResponseState::Complete);
-	CompleteEvent.Notify();
+	CompleteEvent->Trigger();
 
 	if (UE_LOG_ACTIVE(LogHttp, Verbose))
 	{
@@ -1184,7 +1184,7 @@ void FCurlHttpResponse::Cancel()
 	// Waiting within a receiver function would prevent the response from completing.
 	if (LocalReceiverFunction == EHttpReceiverFunction::None)
 	{
-		CompleteEvent.Wait();
+		CompleteEvent->Wait();
 	}
 }
 
diff --git a/Engine/Source/Developer/DerivedDataCache/Private/Http/HttpClient.cpp b/Engine/Source/Developer/DerivedDataCache/Private/Http/HttpClient.cpp
--- a/Engine/Source/Developer/DerivedDataCache/Private/Http/HttpClient.cpp
+++ b/Engine/Source/Developer/DerivedDataCache/Private/Http/HttpClient.cpp
@@ -2,8 +2,8 @@
 
 #include "HttpClient.h"
 
-#include "Async/ManualResetEvent.h"
 #include "Containers/LockFreeList.h"
+#include "HAL/Event.h"
 #include "Memory/MemoryView.h"
 #include "Misc/AsciiSet.h"
 #include "ProfilingDebugging/CpuProfilerTrace.h"
@@ -271,7 +271,7 @@ struct Http::Private::FHttpRequestQueueData
 	struct FWaiter
 	{
 		THttpUniquePtr<IHttpRequest> Request;
-		FManualResetEvent Event;
+		FEventRef Event{EEventMode::ManualReset};
 	};
 
 	FHttpRequestQueueData(IHttpConnectionPool& ConnectionPool, const FHttpClientParams& ClientParams)
@@ -290,7 +290,7 @@ struct Http::Private::FHttpRequestQueueData
 					if (FWaiter* Waiter = Waiters.Pop())
 					{
 						Waiter->Request = MoveTemp(Request);
-						Waiter->Event.Notify();
+						Waiter->Event->Trigger();
 					}
 				}
 			}
@@ -310,7 +310,7 @@ struct Http::Private::FHttpRequestQueueData
 			if (FWaiter* Waiter = Waiters.Pop())
 			{
 				Waiter->Request = MoveTemp(Request);
-				Waiter->Event.Notify();
+				Waiter->Event->Trigger();
 			}
 			else
 			{
@@ -326,16 +326,16 @@ struct Http::Private::FHttpRequestQueueData
 			if (FWaiter* Waiter = Waiters.Pop())
 			{
 				Waiter->Request = MoveTemp(Request);
-				Waiter->Event.Notify();
+				Waiter->Event->Trigger();
 			}
-			if (LocalWaiter.Event.WaitFor(FMonotonicTimeSpan::Zero()))
+			if (LocalWaiter.Event->Wait(0))
 			{
 				return MoveTemp(LocalWaiter.Request);
 			}
 		}
 
 		TRACE_CPUPROFILER_EVENT_SCOPE(HttpRequestQueue_Wait);
-		LocalWaiter.Event.Wait();
+		LocalWaiter.Event->Wait();
 		return MoveTemp(LocalWaiter.Request);
 	}
 
diff --git a/Engine/Source/Runtime/CoreUObject/Private/Serialization/BulkDataStreaming.cpp b/Engine/Source/Runtime/CoreUObject/Private/Serialization/BulkDataStreaming.cpp
--- a/Engine/Source/Runtime/CoreUObject/Private/Serialization/BulkDataStreaming.cpp
+++ b/Engine/Source/Runtime/CoreUObject/Private/Serialization/BulkDataStreaming.cpp
@@ -1,8 +1,8 @@
 // Copyright Epic Games, Inc. All Rights Reserved.
 #include "Serialization/BulkData.h"
-#include "Async/ManualResetEvent.h"
 #include "Async/MappedFileHandle.h"
 #include "Containers/ChunkedArray.h"
+#include "Experimental/Async/LazyEvent.h"
 #include "HAL/CriticalSection.h"
 #include "IO/IoDispatcher.h"
 #include "Misc/ConfigCacheIni.h"
@@ -132,14 +132,15 @@ protected:
 
 private:
 
-	UE::FManualResetEvent	DoneEvent;
-	FIoRequest				Request;
-	int64					SizeResult;
-	std::atomic<uint32>		Status;
+	UE::FLazyEvent		DoneEvent;
+	FIoRequest			Request;
+	int64				SizeResult;
+	std::atomic<uint32>	Status;
 };
 
 FChunkRequest::FChunkRequest(FIoBuffer&& InBuffer)
 	: Buffer(MoveTemp(InBuffer))
+	, DoneEvent(EEventMode::ManualReset)
 	, SizeResult(-1)
 	, Status{uint32(EChunkRequestStatus::None)}
 {
@@ -172,7 +173,7 @@ void FChunkRequest::Issue(FIoChunkId ChunkId, FIoReadOptions Options, int32 Prio
 		HandleChunkResult(MoveTemp(Result));
 		Status.store(uint32(ReadyOrCanceled | EChunkRequestStatus::CallbackTriggered), std::memory_order_release); 
 
-		DoneEvent.Notify();
+		DoneEvent.Trigger();
 	});
 
 	IoBatch.Issue();
@@ -182,7 +183,8 @@ bool FChunkRequest::WaitForChunkRequest(float TimeLimitSeconds)
 {
 	checkf(GetStatus() != EChunkRequestStatus::None, TEXT("The request must be issued before waiting for completion"));
 
-	return DoneEvent.WaitFor(TimeLimitSeconds <= 0.0f ? FMonotonicTimeSpan::Infinity() : FMonotonicTimeSpan::FromSeconds(TimeLimitSeconds));
+	const uint32 TimeLimitMilliseconds = TimeLimitSeconds <= 0.0f ? MAX_uint32 : (uint32)(TimeLimitSeconds * 1000.0f);
+	return DoneEvent.Wait(TimeLimitMilliseconds);
 }
 
 void FChunkRequest::CancelChunkRequest()
@@ -789,7 +791,7 @@ public:
 	virtual bool Wait(uint32 Milliseconds) override final
 	{
 		check(GetStatus() != FBulkDataRequest::EStatus::None);
-		return DoneEvent.WaitFor(UE::FMonotonicTimeSpan::FromMilliseconds(Milliseconds));
+		return DoneEvent.Wait(Milliseconds);
 	}
 
 	void Read(
@@ -854,15 +856,15 @@ private:
 		}
 
 		SetStatus(CompletionStatus);
-		DoneEvent.Notify();
+		DoneEvent.Trigger();
 	}
 
 	static constexpr int32 TargetBytesPerChunk = sizeof(FChunkBatchReadRequest) * 8;
 	using FRequests = TChunkedArray<FChunkBatchReadRequest, TargetBytesPerChunk>;
 
-	FIoBatch				IoBatch;
-	FRequests				Requests; //TOOD: Optimize if there's only a single read request 
-	UE::FManualResetEvent	DoneEvent;
+	FIoBatch		IoBatch;
+	FRequests		Requests; //TOOD: Optimize if there's only a single read request 
+	UE::FLazyEvent	DoneEvent{EEventMode::ManualReset};
 	FBulkDataRequest::FCompletionCallback CompletionCallback;
 };
 
diff --git a/Engine/Source/Runtime/CoreUObject/Private/Serialization/DerivedData.cpp b/Engine/Source/Runtime/CoreUObject/Private/Serialization/DerivedData.cpp
--- a/Engine/Source/Runtime/CoreUObject/Private/Serialization/DerivedData.cpp
+++ b/Engine/Source/Runtime/CoreUObject/Private/Serialization/DerivedData.cpp
@@ -29,7 +29,7 @@
 ///////////////////////////////////////////////////////////////////////////////////////////////////////////////////////
 // TEMPORARY SERIALIZATION AS BULK DATA
 #if WITH_EDITORONLY_DATA
-#include "Async/ManualResetEvent.h"
+#include "Experimental/Async/LazyEvent.h"
 #include "Serialization/BulkData.h"
 #endif // WITH_EDITORONLY_DATA
 // TEMPORARY SERIALIZATION AS BULK DATA
@@ -336,11 +336,11 @@ void FDerivedData::Serialize(FArchive& Ar, UObject* Owner)
 
 		// Blocking read to populate the bulk data to be serialized later by the linker.
 		{
-			FManualResetEvent BatchComplete;
+			FLazyEvent BatchComplete(EEventMode::ManualReset);
 			FDerivedDataIoBatch Batch;
 			FDerivedDataIoRequest Request = Batch.Read(*this);
 			FDerivedDataIoResponse Response;
-			Batch.Dispatch(Response, [&BatchComplete] { BatchComplete.Notify(); });
+			Batch.Dispatch(Response, [&BatchComplete] { BatchComplete.Trigger(); });
 			BatchComplete.Wait();
 
 			const FSharedBuffer Data = Response.GetData(Request);
diff --git a/Engine/Source/Runtime/Engine/Private/Streaming/Texture2DStreamIn_DerivedData.cpp b/Engine/Source/Runtime/Engine/Private/Streaming/Texture2DStreamIn_DerivedData.cpp
--- a/Engine/Source/Runtime/Engine/Private/Streaming/Texture2DStreamIn_DerivedData.cpp
+++ b/Engine/Source/Runtime/Engine/Private/Streaming/Texture2DStreamIn_DerivedData.cpp
@@ -15,6 +15,7 @@ namespace UE
 
 FTexture2DStreamIn_DerivedData::FTexture2DStreamIn_DerivedData(UTexture2D* InTexture, bool bInHighPriority)
 	: FTexture2DStreamIn(InTexture)
+	, ResponseComplete(EEventMode::AutoReset)
 	, bHighPriority(bInHighPriority)
 {
 }
@@ -77,7 +78,7 @@ void FTexture2DStreamIn_DerivedData::DoBeginIoRequests(const FContext& Context)
 			}
 		}
 
-		ResponseComplete.Notify();
+		ResponseComplete.Trigger();
 
 		// Tick to schedule the task or cancellation callback. Use TT_None to avoid a deadlock.
 		Tick(FTexture2DUpdate::TT_None);
diff --git a/Engine/Source/Runtime/Engine/Private/Streaming/Texture2DStreamIn_DerivedData.h b/Engine/Source/Runtime/Engine/Private/Streaming/Texture2DStreamIn_DerivedData.h
--- a/Engine/Source/Runtime/Engine/Private/Streaming/Texture2DStreamIn_DerivedData.h
+++ b/Engine/Source/Runtime/Engine/Private/Streaming/Texture2DStreamIn_DerivedData.h
@@ -2,7 +2,7 @@
 
 #pragma once
 
-#include "Async/ManualResetEvent.h"
+#include "Experimental/Async/LazyEvent.h"
 #include "Serialization/DerivedData.h"
 #include "Texture2DStreamIn.h"
 
@@ -31,7 +31,7 @@ protected:
 
 private:
 	FDerivedDataIoResponse Response;
-	FManualResetEvent ResponseComplete;
+	FLazyEvent ResponseComplete;
 	bool bHighPriority;
 };
 
