Refactor and optimize reference chain search
Use int32-based graph format to allow use of shared algorithms.
Delay callstack/property name gathering until relevant paths are discovered.
Parallelize initial reference graph building.
Cache callstack addresses during printing.

diff --git a/Engine/Source/Runtime/CoreUObject/Private/UObject/Obj.cpp b/Engine/Source/Runtime/CoreUObject/Private/UObject/Obj.cpp
--- a/Engine/Source/Runtime/CoreUObject/Private/UObject/Obj.cpp
+++ b/Engine/Source/Runtime/CoreUObject/Private/UObject/Obj.cpp
@@ -4355,6 +4355,10 @@ bool StaticExec( UWorld* InWorld, const TCHAR* Cmd, FOutputDevice& Ar )
 					{
 						SearchModeFlags |= EReferenceChainSearchMode::FullChain;
 					}
+					else if (FCString::Stricmp(*Tok, TEXT("minimal")) == 0)
+					{
+						SearchModeFlags |= EReferenceChainSearchMode::Minimal;
+					}
 #if ENABLE_GC_HISTORY
 					else if (FParse::Value(Str, TEXT("history="), HistoryLevel))
 					{
diff --git a/Engine/Source/Runtime/CoreUObject/Private/UObject/ReferenceChainSearch.cpp b/Engine/Source/Runtime/CoreUObject/Private/UObject/ReferenceChainSearch.cpp
--- a/Engine/Source/Runtime/CoreUObject/Private/UObject/ReferenceChainSearch.cpp
+++ b/Engine/Source/Runtime/CoreUObject/Private/UObject/ReferenceChainSearch.cpp
@@ -1,328 +1,1200 @@
 // Copyright Epic Games, Inc. All Rights Reserved.
 
 #include "UObject/ReferenceChainSearch.h"
+
+#include "Algo/AnyOf.h"
+#include "Algo/RemoveIf.h"
+#include "Algo/Reverse.h"
+#include "Experimental/Graph/GraphConvert.h"
 #include "HAL/PlatformStackWalk.h"
-#include "UObject/UObjectIterator.h"
-#include "UObject/UnrealType.h"
+#include "HAL/ThreadHeartBeat.h"
 #include "UObject/FastReferenceCollector.h"
 #include "UObject/GCObject.h"
-#include "HAL/ThreadHeartBeat.h"
+#include "UObject/UObjectIterator.h"
+#include "UObject/UnrealType.h"
 
 DEFINE_LOG_CATEGORY_STATIC(LogReferenceChain, Log, All);
 
-static bool GFGCObjectPropertyNames = true;
-static FAutoConsoleVariableRef CVarFGCObjectPropertyNames(
-	TEXT("gc.ReferenceChainSearch.FGCObjectPropertyNames"),
-	GFGCObjectPropertyNames,
-	TEXT("Whether or not to extract property names from FGCObject implementors during reference chain search - can be very expensive."),
-	ECVF_Default
-);
-
-static bool GAddReferencedObjectsStackTrace = true;
-static FAutoConsoleVariableRef CVarAddReferencedObjectsStackTrace(
-	TEXT("gc.ReferenceChainSearch.AddReferencedObjectsStackTrace"),
-	GAddReferencedObjectsStackTrace,
-	TEXT("Whether to gather a stack trace for calls to AddReferencedObjects."),
-	ECVF_Default
-);
-
-// Returns true if the object can't be collected by GC
-static FORCEINLINE bool IsNonGCObject(FGCObjectInfo* Object, EReferenceChainSearchMode SearchMode)
+namespace UE::ReferenceChainSearch
 {
-	return Object->HasAnyInternalFlags(EInternalObjectFlags::GarbageCollectionKeepFlags | EInternalObjectFlags::RootSet) ||
-		(GARBAGE_COLLECTION_KEEPFLAGS != RF_NoFlags && Object->HasAnyFlags(GARBAGE_COLLECTION_KEEPFLAGS) && !(SearchMode & EReferenceChainSearchMode::FullChain));
-}
+	using UE::Graph::FGraph;
+	using UE::Graph::FVertex;
 
-FReferenceChainSearch::FGraphNode* FReferenceChainSearch::FindOrAddNode(FGCObjectInfo* InObjectInfo)
-{
-	FGraphNode* ObjectNode = nullptr;
-	FGraphNode** ExistingObjectNode = AllNodes.Find(InObjectInfo);
-	if (ExistingObjectNode)
+	// This policy causes us to allocate a dense graph in all cases.
+	// If memory usage is a concern for minimal/direct search cases, or when excluding DisregardForGC objects, we
+	// could build a sparse list of relevant objects and map them to vertices that way rather than using GUObjectArray.
+	struct FPolicyUObjectHeap
 	{
-		ObjectNode = *ExistingObjectNode;
-		check(ObjectNode->ObjectInfo == InObjectInfo);
-	}
-	else
+		using ObjectType = UObject*; // FGCObjectInfo
+		using ConstObjectType = const UObject*;
+
+		TMap<const UObject*, FGCObjectInfo*>& ObjectToInfoMap;
+		TMap<FGCObjectInfo*, FReferenceChainSearch::FGraphNode*>& AllNodes;
+
+		FPolicyUObjectHeap(TMap<const UObject*, FGCObjectInfo*>& InObjectToInfoMap,
+			TMap<FGCObjectInfo*, FReferenceChainSearch::FGraphNode*>& InAllNodes)
+			: ObjectToInfoMap(InObjectToInfoMap)
+			, AllNodes(InAllNodes)
+		{
+		}
+
+		FVertex GetGCObjectReferencerVertex() const
+		{
+			return FGCObject::GGCObjectReferencer ? ObjectToVertex(FGCObject::GGCObjectReferencer) : UE::Graph::InvalidVertex;
+		}
+		int32 GetNumVertices() const
+		{
+			return GUObjectArray.GetObjectArrayNum();
+		}
+		FVertex ObjectToVertex(ConstObjectType Object) const
+		{
+			return GUObjectArray.ObjectToIndex(Object);
+		}
+		ObjectType VertexToObject(FVertex Vertex) const
+		{
+			return reinterpret_cast<UObject*>(GUObjectArray.IndexToObject(Vertex)->Object);
+		}
+
+		bool IsIn(ConstObjectType Outer, ConstObjectType PotentialInner) const
+		{
+			return PotentialInner->IsIn(Outer);
+		}
+		bool IsIn(FVertex Outer, FVertex PotentialInner) const
+		{
+			return VertexToObject(PotentialInner)->IsIn(VertexToObject(Outer));
+		}
+
+		bool IsRoot(FVertex Vertex, EReferenceChainSearchMode SearchMode) const
+		{
+			const UObject* Object = VertexToObject(Vertex);
+			return Object->HasAnyInternalFlags(EInternalObjectFlags::GarbageCollectionKeepFlags | EInternalObjectFlags::RootSet)
+				|| (GARBAGE_COLLECTION_KEEPFLAGS != RF_NoFlags && Object->HasAnyFlags(GARBAGE_COLLECTION_KEEPFLAGS)
+					&& !(SearchMode & EReferenceChainSearchMode::FullChain));
+		}
+
+		bool IsGarbage(FVertex Vertex) const
+		{
+			return GUObjectArray.IndexToObject(Vertex)->HasAnyFlags(EInternalObjectFlags::Garbage);
+		}
+
+		FGCObjectInfo* FindOrCreateObjectInfoAndGraphNode(FVertex InVertex)
+		{
+			ObjectType Object = VertexToObject(InVertex);
+			FGCObjectInfo* Info = ObjectToInfoMap.FindRef(Object);
+			if (!Info)
+			{
+				Info = FGCObjectInfo::FindOrAddInfoHelper(Object, ObjectToInfoMap);
+			}
+			if (!AllNodes.Contains(Info))
+			{
+				FReferenceChainSearch::FGraphNode* Node = new FReferenceChainSearch::FGraphNode;
+				Node->ObjectInfo = Info;
+				AllNodes.Add(Info, Node);
+			}
+			return Info;
+		}
+
+		FReferenceChainSearch::FGraphNode* VertexToGraphNode(FVertex InVertex) const
+		{
+			ObjectType Object = VertexToObject(InVertex);
+			FGCObjectInfo* ObjectInfo = ObjectToInfoMap.FindChecked(Object);
+			return AllNodes.FindChecked(ObjectInfo);
+		}
+
+		/**
+		 * @param InOutReferenceInfo a map whose keys are vertices we want to find references FROM and whose valuesa re maps whose keys we want to
+		 * find references TO.
+		 */
+		void GatherReferenceInfo(TMap<FVertex, TMap<FVertex, FReferenceChainSearch::FObjectReferenceInfo>>& InOutReferenceInfo);
+	};
+
+	struct FPolicyGCHistory
 	{
-		ObjectNode = new FGraphNode();
-		ObjectNode->ObjectInfo = InObjectInfo;
-		AllNodes.Add(InObjectInfo, ObjectNode);
-	}
-	return ObjectNode;
-}
+		using ObjectType = FGCObjectInfo*;
 
-FReferenceChainSearch::FGraphNode* FReferenceChainSearch::FindOrAddNode(const UObject* InObjectToFindNodeFor)
-{
-	return FindOrAddNode(FGCObjectInfo::FindOrAddInfoHelper(InObjectToFindNodeFor, ObjectToInfoMap));
-}
+		const FGCSnapshot& Snapshot;
+		TMap<FGCObjectInfo*, FReferenceChainSearch::FGraphNode*>& AllNodes;
 
-void FReferenceChainSearch::LinkNodes(FGCObjectInfo* From, FGCObjectInfo* To, EReferenceType ReferenceType, FName PropertyName, TConstArrayView<uint64> StackFrames)
-{
-	FGraphNode* ToNode = FindOrAddNode(To);
-	FGraphNode* FromNode = FindOrAddNode(From);
-	FromNode->ReferencedObjects.Emplace(FNodeReferenceInfo(ToNode, ReferenceType, PropertyName, StackFrames));
-	ToNode->ReferencedByObjects.Add(FromNode);
-}
+		TMap<FGCObjectInfo*, FVertex> ObjectInfoToVertex;
+		int32 NumVertexIndices = 0;
 
-void FReferenceChainSearch::LinkNodes(UObject* From, UObject* To, EReferenceType ReferenceType, FName PropertyName, TConstArrayView<uint64> StackFrames)
-{
-	LinkNodes(FGCObjectInfo::FindOrAddInfoHelper(From, ObjectToInfoMap), FGCObjectInfo::FindOrAddInfoHelper(To, ObjectToInfoMap), ReferenceType, PropertyName, StackFrames);
-}
+		FPolicyGCHistory(const FGCSnapshot& InSnapshot, TMap<FGCObjectInfo*, FReferenceChainSearch::FGraphNode*>& InAllNodes)
+			: Snapshot(InSnapshot)
+			, AllNodes(InAllNodes)
+		{
+			ObjectInfoToVertex.Reserve(Snapshot.ObjectToInfoMap.GetMaxIndex() + 1);
+			for (auto It = Snapshot.ObjectToInfoMap.CreateConstIterator(); It; ++It)
+			{
+				int32 Index = It.GetId().AsInteger();
+				ObjectInfoToVertex.Add(It.Value(), Index);
+				NumVertexIndices = FMath::Max(Index + 1, NumVertexIndices);
+			}
+		}
 
-int32 FReferenceChainSearch::BuildReferenceChainsRecursive(FGraphNode* TargetNode, TArray<FReferenceChain*>& ProducedChains, int32 ChainDepth, const int32 VisitCounter, EReferenceChainSearchMode SearchMode, FGraphNode* GCObjReferencerNode)
-{
-	int32 ProducedChainsCount = 0;
+		FVertex GetGCObjectReferencerVertex() const
+		{
+			if (FGCObject::GGCObjectReferencer)
+			{
+				if (FGCObjectInfo* Info = Snapshot.ObjectToInfoMap.FindRef(FGCObject::GGCObjectReferencer))
+				{
+					return ObjectToVertex(Info);
+				}
+			}
+			return UE::Graph::InvalidVertex;
+		}
+
+		int32 GetNumVertices() const
+		{
+			return Snapshot.ObjectToInfoMap.GetMaxIndex() + 1;
+		}
+
+		FVertex ObjectToVertex(ObjectType Object) const
+		{
+			return ObjectInfoToVertex.FindChecked(Object);
+		}
 
-	// Always create a chain for the FGCObject referencer node even if we've found another path to it, because it handles many types of references
-	if (TargetNode == GCObjReferencerNode)
+		ObjectType VertexToObject(FVertex Vertex) const
+		{
+			return Snapshot.ObjectToInfoMap.Get(FSetElementId::FromInteger(Vertex)).Value;
+		}
+
+		bool IsIn(FGCObjectInfo* Outer, FGCObjectInfo* PotentialInner) const
+		{
+			return PotentialInner->IsIn(Outer);
+		}
+		bool IsIn(FVertex Outer, FVertex PotentialInner) const
+		{
+			return VertexToObject(PotentialInner)->IsIn(VertexToObject(Outer));
+		}
+
+		bool IsRoot(FVertex Vertex, EReferenceChainSearchMode SearchMode) const
+		{
+			ObjectType Object = VertexToObject(Vertex);
+			return Object->HasAnyInternalFlags(EInternalObjectFlags::GarbageCollectionKeepFlags | EInternalObjectFlags::RootSet)
+				|| (GARBAGE_COLLECTION_KEEPFLAGS != RF_NoFlags && Object->HasAnyFlags(GARBAGE_COLLECTION_KEEPFLAGS)
+					&& !(SearchMode & EReferenceChainSearchMode::FullChain));
+		}
+
+		bool IsGarbage(FVertex Vertex) const
+		{
+			return VertexToObject(Vertex)->IsGarbage();
+		}
+
+		FGCObjectInfo* FindOrCreateObjectInfoAndGraphNode(FVertex InVertex)
+		{
+			FGCObjectInfo* Info = VertexToObject(InVertex);
+			if (!AllNodes.Contains(Info))
+			{
+				FReferenceChainSearch::FGraphNode* Node = new FReferenceChainSearch::FGraphNode;
+				Node->ObjectInfo = Info;
+				AllNodes.Add(Info, Node);
+			}
+			return Info;
+		}
+
+		FReferenceChainSearch::FGraphNode* VertexToGraphNode(FVertex InVertex) const
+		{
+			ObjectType Object = VertexToObject(InVertex);
+			return AllNodes.FindChecked(Object);
+		}
+
+		/**
+		 * @param InOutReferenceInfo a map whose keys are vertices we want to find references FROM and whose valuesa re maps whose keys we want to
+		 * find references TO.
+		 */
+		void GatherReferenceInfo(TMap<FVertex, TMap<FVertex, FReferenceChainSearch::FObjectReferenceInfo>>& InOutReferenceInfo);
+	};
+
+	template<typename Derived>
+	struct TReferenceSearchBase
 	{
-		// This is a root so we can construct a chain from this node up to the target node
-		FReferenceChain* Chain = new FReferenceChain(ChainDepth);
-		Chain->InsertNode(TargetNode);
-		ProducedChains.Add(Chain);
-		ProducedChainsCount = 1;
-		return ProducedChainsCount;
-	}
+		// We don't actually implement reference processing on the processor type, we just store a thread index and delegate to the derived
+		// class
+		struct FProcessor
+		{
+			static constexpr EGCOptions Options = EGCOptions::None;
+			int32 ThreadIndex;
+			Derived* This;
+
+			void BeginTimingObject(UObject* CurrentObject) { }
+			void UpdateDetailedStats(UObject* CurrentObject) { }
+			void LogDetailedStatsSummary() { }
+			void HandleTokenStreamObjectReference(UE::GC::FWorkerContext& Context,
+				UObject* ReferencingObject,
+				UObject*& Object,
+				UE::GC::FTokenId TokenIndex,
+				EGCTokenType TokenType,
+				bool bAllowReferenceElimination)
+			{
+				if (Object)
+				{
+					ReferencingObject = Context.GetReferencingObject();
+					if (!ReferencingObject)
+					{
+						ReferencingObject = FGCObject::GGCObjectReferencer;
+					}
+					check(ReferencingObject);
+
+					if (Object != ReferencingObject)
+					{
+						This->HandleObjectReference(ThreadIndex, Object, ReferencingObject, TokenIndex, TokenType);
+					}
+				}
+			}
+		};
+
+		struct FCollector : public FReferenceCollector
+		{
+		protected:
+			FProcessor& Processor;
+			UE::GC::FWorkerContext& Context;
+
+		public:
+			FCollector(FProcessor& InProcessor, FGCArrayStruct& InContext)
+				: Processor(InProcessor)
+				, Context(InContext) { }
+
+			virtual void HandleObjectReference(UObject*& Object, const UObject* ReferencingObject, const FProperty* ReferencingProperty) override
+			{
+				Processor.HandleTokenStreamObjectReference(Context,
+					const_cast<UObject*>(ReferencingObject),
+					Object,
+					UE::GC::ETokenlessId::Collector,
+					EGCTokenType::Native,
+					false);
+			}
+			virtual void HandleObjectReferences(
+				UObject** InObjects, const int32 ObjectNum, const UObject* ReferencingObject, const FProperty* InReferencingProperty) override
+			{
+				for (int32 ObjectIndex = 0; ObjectIndex < ObjectNum; ++ObjectIndex)
+				{
+					UObject*& Object = InObjects[ObjectIndex];
+					Processor.HandleTokenStreamObjectReference(Context,
+						const_cast<UObject*>(ReferencingObject),
+						Object,
+						UE::GC::ETokenlessId::Collector,
+						EGCTokenType::Native,
+						false);
+				}
+			}
+
+			virtual bool NeedsPropertyReferencer() const
+			{
+				return false;
+			}
+			virtual bool IsIgnoringArchetypeRef() const override
+			{
+				return false;
+			}
+			virtual bool IsIgnoringTransient() const override
+			{
+				return false;
+			}
+			virtual bool MarkWeakObjectReferenceForClearing(UObject** WeakReference)
+			{
+				// To avoid false positives we need to implement this method just like GC does
+				// as these references will be treated as weak and should not be reported
+				return true;
+			}
+		};
+
+		struct FThreadData
+		{
+			UObject* PreviousReferencingObject = nullptr;
+			TSet<FVertex> CurrentEdgeList;
+
+			TArray64<FVertex> GraphBuffer;
+			TArray<TArrayView<FVertex>> EdgeLists;
+			FVertex StartVertex;
+
+			SIZE_T GetAllocatedSize() const
+			{
+				return CurrentEdgeList.GetAllocatedSize() + GraphBuffer.GetAllocatedSize() + EdgeLists.GetAllocatedSize();
+			}
+
+			void FlushEdgeList(const FPolicyUObjectHeap& Policy)
+			{
+				if (PreviousReferencingObject && CurrentEdgeList.Num() > 0)
+				{
+					// Flush current edge list to graph buffer
+					FVertex SourceVertex = Policy.ObjectToVertex(PreviousReferencingObject);
+					TArray<FVertex> EdgeList = CurrentEdgeList.Array();
+					int32 BufferStartIndex = GraphBuffer.Num();
+					// During construction store edge lists relative to null to be fixed up later
+					FVertex* StartAddress = nullptr;
+					GraphBuffer.Append(EdgeList);
+					EdgeLists[SourceVertex - StartVertex] = TArrayView<FVertex>(StartAddress + BufferStartIndex, EdgeList.Num());
+					CurrentEdgeList.Reset();
+				}
+			}
+		};
+		TArray<FThreadData> AllThreadData;
+
+		SIZE_T GetAllocatedSize() const
+		{
+			SIZE_T Size = 0;
+			for (const FThreadData& Thread : AllThreadData)
+			{
+				Size += Thread.GetAllocatedSize();
+			}
+			return Size + AllThreadData.GetAllocatedSize();
+		}
+
+		void SetNumThreads(int32 InNumThreads, int32 NumberOfObjectsPerThread, int32 GlobalStartIndex, int32 MaxNumberOfObjects)
+		{
+			AllThreadData.Reset();
+			AllThreadData.AddDefaulted(InNumThreads);
+			for (int32 i = 0; i < InNumThreads; ++i)
+			{
+				FThreadData& Thread = AllThreadData[i];
+				Thread.StartVertex = GlobalStartIndex + i * NumberOfObjectsPerThread;
+				const int32 NumObjects =
+					(i < (InNumThreads - 1)) ? NumberOfObjectsPerThread : (MaxNumberOfObjects - (InNumThreads - 1) * NumberOfObjectsPerThread);
+				Thread.EdgeLists.Reserve(NumObjects);
+				Thread.EdgeLists.SetNumZeroed(NumObjects);
+			}
+		}
+
+		void CollectAllReferences(bool bGCOnly)
+		{
+			constexpr bool bParallel = Derived::bParallel;
+			Derived* This = static_cast<Derived*>(this);
+			const int32 GlobalStartObjectIndex = bGCOnly ? GUObjectArray.GetFirstGCIndex() : 0;
+			const int32 MaxNumberOfObjects = GUObjectArray.GetObjectArrayNum() - GlobalStartObjectIndex;
+			if constexpr (bParallel)
+			{
+				const int32 NumThreads = GetNumCollectReferenceWorkers();
+				const int32 NumberOfObjectsPerThread = (MaxNumberOfObjects / NumThreads) + 1;
+
+				SetNumThreads(NumThreads, NumberOfObjectsPerThread, GlobalStartObjectIndex, MaxNumberOfObjects);
+
+				ParallelFor(NumThreads,
+					[this, This, GlobalStartObjectIndex, MaxNumberOfObjects, NumThreads, NumberOfObjectsPerThread](int32 ThreadIndex) {
+						FProcessor Processor{ ThreadIndex, This };
+						TSet<UObject*> ThreadResult;
+						TArray<UObject*> ObjectsToSerialize;
+						ObjectsToSerialize.Reserve(NumberOfObjectsPerThread);
+
+						const int32 FirstObjectIndex = GlobalStartObjectIndex + ThreadIndex * NumberOfObjectsPerThread;
+						const int32 NumObjects = (ThreadIndex < (NumThreads - 1))
+												   ? NumberOfObjectsPerThread
+												   : (MaxNumberOfObjects - (NumThreads - 1) * NumberOfObjectsPerThread);
+
+						// First cache all potential referencers
+						for (int32 ObjectIndex = 0; ObjectIndex < NumObjects && (FirstObjectIndex + ObjectIndex) < MaxNumberOfObjects; ++ObjectIndex)
+						{
+							FUObjectItem& ObjectItem = GUObjectArray.GetObjectItemArrayUnsafe()[FirstObjectIndex + ObjectIndex];
+							UObject* Object = static_cast<UObject*>(ObjectItem.Object);
+							if (Object && !ObjectItem.IsUnreachable() && !This->ShouldSkipReferencer(ThreadIndex, Object))
+							{
+								ObjectsToSerialize.Add(Object);
+							}
+						}
+
+						UE::GC::FWorkerContext Context;
+						Context.SetInitialObjectsUnpadded(ObjectsToSerialize);
+						CollectReferences<UE::GC::TDefaultCollector<FProcessor>>(Processor, Context);
+					});
+				return;
+			}
+
+			SetNumThreads(1, MaxNumberOfObjects, GlobalStartObjectIndex, MaxNumberOfObjects);
+			FProcessor Processor{ 0, This };
+			TArray<UObject*> ObjectsToProcess;
+			for (FRawObjectIterator It; It; ++It)
+			{
+				FUObjectItem* ObjItem = *It;
+				UObject* Object = static_cast<UObject*>(ObjItem->Object);
+
+				// We can't ask the iterator for only GC objects because that would skip the GC Object referencer
+				if (bGCOnly && GUObjectArray.IsDisregardForGC(Object) && (Object != FGCObject::GGCObjectReferencer))
+				{
+					continue;
+				}
+
+				if (This->ShouldSkipReferencer(0, Object))
+				{
+					continue;
+				}
+
+				// Find direct references
+				UE::GC::FWorkerContext Context;
+				ObjectsToProcess = { Object };
+				Context.SetInitialObjectsUnpadded(ObjectsToProcess);
+				CollectReferences<UE::GC::TDefaultCollector<FProcessor>>(Processor, Context);
+			}
+		}
+
+		void CollectReferencesFromObject(UObject* FromObject)
+		{
+			Derived* This = static_cast<Derived*>(this);
+			FProcessor Processor{ 0, This };
+			// Find direct references
+			UE::GC::FWorkerContext Context;
+			TArray<UObject*> ObjectsToProcess;
+			ObjectsToProcess = { FromObject };
+			Context.SetInitialObjectsUnpadded(ObjectsToProcess);
+			CollectReferences<UE::GC::TDefaultCollector<FProcessor>>(Processor, Context);
+		}
+	};
+
+	struct FDirectReferenceSearch : public TReferenceSearchBase<FDirectReferenceSearch>
+	{
+		using Super = TReferenceSearchBase<FDirectReferenceSearch>;
+		static constexpr bool bParallel = true;
+		FPolicyUObjectHeap& Policy;
+
+		FDirectReferenceSearch(FPolicyUObjectHeap& InPolicy)
+			: Policy(InPolicy) { }
+
+		bool ShouldSkipReferencer(int32 ThreadIndex, UObject* Object)
+		{
+			return false;
+		}
+
+		void HandleObjectReference(
+			int32 ThreadIndex, UObject* Object, UObject* ReferencingObject, UE::GC::FTokenId TokenIndex, EGCTokenType TokenType)
+		{
+			FThreadData& ThreadData = AllThreadData[ThreadIndex];
+			if (ThreadData.PreviousReferencingObject != ReferencingObject)
+			{
+				ThreadData.FlushEdgeList(Policy);
+			}
+
+			ThreadData.PreviousReferencingObject = ReferencingObject;
+			ThreadData.CurrentEdgeList.Add(Policy.ObjectToVertex(Object));
+		}
+
+		// Merge all edge lists into a single edge list, leaving the actual data in individual buffers
+		void MergeGraph(TArray<TConstArrayView<FVertex>>& OutMergedEdgeLists)
+		{
+			FVertex* OldBase = nullptr;
+			for (FThreadData& Thread : AllThreadData)
+			{
+				Thread.FlushEdgeList(Policy);
+				FVertex* EdgeListBase = Thread.GraphBuffer.GetData();
+				for (int32 i = 0; i < Thread.EdgeLists.Num(); ++i)
+				{
+					TConstArrayView<FVertex> TempEdgeList = Thread.EdgeLists[i];
+					FVertex SourceVertex = Thread.StartVertex + i;
+					TConstArrayView<FVertex> EdgeList{ TempEdgeList.GetData() - OldBase + EdgeListBase, TempEdgeList.Num() };
+					OutMergedEdgeLists[SourceVertex] = EdgeList;
+				}
+			}
+		}
+	};
 
-	if (TargetNode->Visited != VisitCounter)
+	struct FMinimalReferenceSearch : public TReferenceSearchBase<FMinimalReferenceSearch>
 	{
-		TargetNode->Visited = VisitCounter;
+		using Super = TReferenceSearchBase<FMinimalReferenceSearch>;
+		static constexpr bool bParallel = true;
+		FPolicyUObjectHeap& Policy;
+		TSet<const UObject*> TargetObjects;
+		TSet<UObject*> FoundTargets;
+
+		FMinimalReferenceSearch(FPolicyUObjectHeap& InPolicy)
+			: Policy(InPolicy) { }
+
+		SIZE_T GetAllocatedSize() const
+		{
+			return TargetObjects.GetAllocatedSize() + FoundTargets.GetAllocatedSize() + Super::GetAllocatedSize();
+		}
+
+		bool ShouldSkipReferencer(int32 ThreadIndex, UObject* Object)
+		{
+			// Skip objects which are in any of what we're looking for so we only report reference external to that entire group
+			for (UObject* OuterObject = Object; OuterObject; OuterObject = OuterObject->GetOuter())
+			{
+				if (TargetObjects.Contains(OuterObject) && Object->GetOuter() != nullptr)
+				{
+					FThreadData& ThreadData = AllThreadData[ThreadIndex];
+					// Add an edge between this object and its outer so the outer object may be reachable via this object
+					ThreadData.FlushEdgeList(Policy);
+					ThreadData.PreviousReferencingObject = Object;
+					ThreadData.CurrentEdgeList.Add(Policy.ObjectToVertex(Object->GetOuter()));
+					return true;
+				}
+			}
+			return false;
+		}
+
+		void HandleObjectReference(
+			int32 ThreadIndex,
+			UObject* Object,
+			UObject* ReferencingObject,
+			UE::GC::FTokenId TokenIndex,
+			EGCTokenType TokenType)
+		{
+			FThreadData& ThreadData = AllThreadData[ThreadIndex];
+			if (!TargetObjects.Contains(Object))
+			{
+				UObject* OuterObject = Object->GetOuter();
+				for (; OuterObject; OuterObject = OuterObject->GetOuter())
+				{
+					if (TargetObjects.Contains(OuterObject))
+					{
+						break;
+					}
+				}
 
-		// Stop at root objects
-		if (!IsNonGCObject(TargetNode->ObjectInfo, SearchMode))
-		{			
-			for (FGraphNode* ReferencedByNode : TargetNode->ReferencedByObjects)
+				// Only track references to target objects or their inners
+				if (!OuterObject)
+				{
+					return;
+				}
+			}
+
+			// Add an edge from ReferencingObject to Object
+			if (ThreadData.PreviousReferencingObject != ReferencingObject)
+			{
+				ThreadData.FlushEdgeList(Policy);
+			}
+
+			ThreadData.PreviousReferencingObject = ReferencingObject;
+			ThreadData.CurrentEdgeList.Add(Policy.ObjectToVertex(Object));
+			if (TargetObjects.Contains(Object))
+			{
+				FoundTargets.Add(Object);
+			}
+		}
+
+		// Merge all edge lists into a single edge list, leaving the actual data in individual buffers
+		// For each edge list, if it contains any direct references to the targets, remove any other references
+		void MergeGraph(TArray<TConstArrayView<FVertex>>& OutMergedEdgeLists)
+		{
+			for (FThreadData& Thread : AllThreadData)
 			{
-				// For each of the referencers of this node, duplicate the current chain and continue processing
-				if (ReferencedByNode->Visited != VisitCounter || ReferencedByNode == GCObjReferencerNode)
+				Thread.FlushEdgeList(Policy);
+			}
+
+			TSet<FVertex> TargetVertices;
+			for (const UObject* TargetObject : TargetObjects)
+			{
+				TargetVertices.Add(Policy.ObjectToVertex(TargetObject));
+			}
+
+			FVertex* OldBase = nullptr;
+			for (FThreadData& Thread : AllThreadData)
+			{
+				FVertex* EdgeListBase = Thread.GraphBuffer.GetData();
+				for (int32 i = 0; i < Thread.EdgeLists.Num(); ++i)
 				{
-					int32 OldChainsCount = ProducedChains.Num();
-					int32 NewChainsCount = BuildReferenceChainsRecursive(ReferencedByNode, ProducedChains, ChainDepth + 1, VisitCounter, SearchMode, GCObjReferencerNode);
-					// Insert the current node to all chains produced recursively
-					for (int32 NewChainIndex = OldChainsCount; NewChainIndex < (NewChainsCount + OldChainsCount); ++NewChainIndex)
+					TArrayView<FVertex> EdgeList{ Thread.EdgeLists[i].GetData() - OldBase + EdgeListBase, Thread.EdgeLists[i].Num() };
+					int32 NumDirectRefs = Algo::RemoveIf(EdgeList, [&](FVertex Vertex) { return !TargetVertices.Contains(Vertex); });
+					if (NumDirectRefs != 0)
 					{
-						FReferenceChain* Chain = ProducedChains[NewChainIndex];
-						Chain->InsertNode(TargetNode);
+						EdgeList = EdgeList.Slice(0, NumDirectRefs);
 					}
-					ProducedChainsCount += NewChainsCount;
+
+					FVertex SourceVertex = Thread.StartVertex + i;
+					OutMergedEdgeLists[SourceVertex] = EdgeList;
 				}
-			}			
+			}
+		}
+	};
+
+	/**
+	 * Create an initial graph from all UObject references on the heap for later analysis.
+	 * Do not record any additional info such as property names/callstacks in this pass.
+	 *
+	 */
+	void PerformInitialGatherFromLiveUObjectHeap(FPolicyUObjectHeap Policy,
+		TConstArrayView<const UObject*> ObjectsToFindReferencesTo,
+		EReferenceChainSearchMode Mode,
+		UE::Graph::FGraph& OutGraph)
+	{
+		TRACE_CPUPROFILER_EVENT_SCOPE(UE::ReferenceChainSearch::PerformInitialGatherFromLiveUObjectHeap);
+		// Minimal mode tries to minimize the amount of heap memory used by searching for a mixture of direct references and references via
+		// outer chains
+		const bool bMinimal = !!(Mode & EReferenceChainSearchMode::Minimal);
+		const bool bGCOnly = !!(Mode & EReferenceChainSearchMode::GCOnly);
+		const bool bDirectOnly = !!(Mode & EReferenceChainSearchMode::Direct);
+
+		if (bMinimal)
+		{
+			FMinimalReferenceSearch Search(Policy);
+			for (const UObject* Object : ObjectsToFindReferencesTo)
+			{
+				Search.TargetObjects.Add(Object);
+			}
+
+			Search.CollectAllReferences(bGCOnly);
+			TArray<TConstArrayView<FVertex>> MergedEdgeLists;
+			MergedEdgeLists.Reserve(Policy.GetNumVertices());
+			MergedEdgeLists.SetNumUninitialized(Policy.GetNumVertices());
+			Search.MergeGraph(MergedEdgeLists);
+
+			OutGraph = UE::Graph::ConstructTransposeGraph(MergedEdgeLists);
+			UE_LOG(LogReferenceChain, Display, TEXT("InitialGather memory usage: %.2f"), (OutGraph.GetAllocatedSize() + Search.GetAllocatedSize()) / 1024.0 / 1024.0);
 		}
 		else
 		{
-			// This is a root so we can construct a chain from this node up to the target node
-			FReferenceChain* Chain = new FReferenceChain(ChainDepth);
-			Chain->InsertNode(TargetNode);
-			ProducedChains.Add(Chain);
-			ProducedChainsCount = 1;
+			FDirectReferenceSearch Search(Policy);
+			Search.CollectAllReferences(bGCOnly);
+			TArray<TConstArrayView<FVertex>> MergedEdgeLists;
+			MergedEdgeLists.Reserve(Policy.GetNumVertices());
+			MergedEdgeLists.SetNumUninitialized(Policy.GetNumVertices());
+			Search.MergeGraph(MergedEdgeLists);
+
+			OutGraph = UE::Graph::ConstructTransposeGraph(MergedEdgeLists);
+			UE_LOG(LogReferenceChain, Display, TEXT("InitialGather memory usage: %.2f"), (OutGraph.GetAllocatedSize() + Search.GetAllocatedSize()) / 1024.0 / 1024.0);
 		}
 	}
 
-	return ProducedChainsCount;
-}
+	void PerformInitialGatherFromGCHistory(const FPolicyGCHistory& Policy, FGraph& OutGraph)
+	{
+		TRACE_CPUPROFILER_EVENT_SCOPE(UE::ReferenceChainSearch::PerformInitialGatherFromGCHistory);
+		int64 TotalEdges = 0;
+		for (const TPair<FGCObjectInfo*, TArray<FGCDirectReferenceInfo>*>& Pair : Policy.Snapshot.DirectReferences)
+		{
+			FVertex FromVertex = Policy.ObjectToVertex(Pair.Key);
+			TotalEdges += Pair.Value->Num();
+		}
 
-void FReferenceChainSearch::RemoveChainsWithDuplicatedRoots(TArray<FReferenceChain*>& AllChains, FGraphNode* GCObjReferencerNode)
-{
-	// This is going to be rather slow but it depends on the number of chains which shouldn't be too bad (usually)
-	for (int32 FirstChainIndex = 0; FirstChainIndex < AllChains.Num(); ++FirstChainIndex)
+		FGraph TempGraph;
+		TempGraph.Buffer.Reserve(TotalEdges);
+		TempGraph.EdgeLists.Reserve(Policy.GetNumVertices());
+		TempGraph.EdgeLists.SetNumZeroed(Policy.GetNumVertices());
+
+		for (const TPair<FGCObjectInfo*, TArray<FGCDirectReferenceInfo>*>& Pair : Policy.Snapshot.DirectReferences)
+		{
+			int64 StartOffset = TempGraph.Buffer.Num();
+			FVertex FromVertex = Policy.ObjectToVertex(Pair.Key);
+			for (FGCDirectReferenceInfo& ReferenceInfo : *Pair.Value)
+			{
+				FVertex ToVertex = Policy.ObjectToVertex(ReferenceInfo.ReferencedObjectInfo);
+				check(TempGraph.EdgeLists.IsValidIndex(ToVertex));
+				TempGraph.Buffer.Add(ToVertex);
+			}
+			TempGraph.EdgeLists[FromVertex] =
+				TConstArrayView<int32>(TempGraph.Buffer.GetData() + StartOffset, static_cast<int32>(TempGraph.Buffer.Num() - StartOffset));
+		}
+
+		OutGraph = UE::Graph::ConstructTransposeGraph(TempGraph.EdgeLists);
+	}
+
+	// Temporary storage for a reference chain before gathering information
+	struct FGraphPath
+	{
+		FVertex Target;
+		// Reversed after construction.
+		// 	During construction:
+		// 		First element is root
+		// 	After construction:
+		//		First element is target, last element is root
+		TArray<FVertex> Vertices;
+	};
+
+	template<typename PolicyType>
+	int32 BuildGraphPathRecursive(PolicyType& Policy,
+		const UE::Graph::FGraph& Graph,
+		TMap<FVertex, int32>& VisitCounts,
+		FVertex ThisVertex,
+		TArray<FGraphPath>& ProducedPaths,
+		int32 ChainDepth, // For reserving space in new chains so we don't reallocate as we return from recursion
+		const int32 VisitCounter,
+		EReferenceChainSearchMode Mode,
+		FVertex GCObjReferencerVertex)
 	{
-		const FGraphNode* RootNode = AllChains[FirstChainIndex]->GetRootNode(GCObjReferencerNode);
-		for (int32 SecondChainIndex = AllChains.Num() - 1; SecondChainIndex > FirstChainIndex; --SecondChainIndex)
+		int32 ProducedPathsCount = 0; // ?
+
+		// Always create a chain for the FGCObject referencer node even if we've found another path to it, because it handles many types of
+		// references
+		if (ThisVertex == GCObjReferencerVertex)
+		{
+			FGraphPath& Path = ProducedPaths.AddDefaulted_GetRef();
+			Path.Vertices.Reserve(ChainDepth);
+			Path.Vertices.Add(ThisVertex);
+			ProducedPathsCount = 1;
+			return ProducedPathsCount;
+		}
+
+		if (int32& ThisVertexVisitCount = VisitCounts.FindOrAdd(ThisVertex); ThisVertexVisitCount != VisitCounter)
 		{
-			if (AllChains[SecondChainIndex]->GetRootNode(GCObjReferencerNode) == RootNode)
+			ThisVertexVisitCount = VisitCounter;
+
+			const bool bIsRoot = Policy.IsRoot(ThisVertex, Mode);
+			if (bIsRoot)
+			{
+				FGraphPath& Path = ProducedPaths.AddDefaulted_GetRef();
+				Path.Vertices.Reserve(ChainDepth);
+				Path.Vertices.Add(ThisVertex);
+				ProducedPathsCount = 1;
+			}
+			else
 			{
-				delete AllChains[SecondChainIndex];
-				AllChains.RemoveAt(SecondChainIndex);
+				for (FVertex ReferencedByVertex : Graph.EdgeLists[ThisVertex])
+				{
+					// For each of the referencers of this node, duplicate the current chain and continue processing
+					if (ReferencedByVertex == GCObjReferencerVertex || VisitCounts.FindOrAdd(ReferencedByVertex) != VisitCounter)
+					{
+						int32 OldCount = ProducedPaths.Num();
+						int32 NewCount = BuildGraphPathRecursive(Policy,
+							Graph,
+							VisitCounts,
+							ReferencedByVertex,
+							ProducedPaths,
+							ChainDepth + 1,
+							VisitCounter,
+							Mode,
+							GCObjReferencerVertex);
+						for (int32 NewIndex = OldCount; NewIndex < (NewCount + OldCount); ++NewIndex)
+						{
+							FGraphPath& Path = ProducedPaths[NewIndex];
+							Path.Vertices.Add(ThisVertex);
+						}
+						ProducedPathsCount += NewCount;
+					}
+				}
 			}
 		}
+
+		return ProducedPathsCount;
 	}
-}
 
-void FReferenceChainSearch::RemoveDuplicateGarbageChains(TArray<FReferenceChain*>& AllChains, FGraphNode* GCObjReferencerNode)
-{
-	typedef TPair<FGraphNode*, FGraphNode*> FGarbageChainKey; 
+	template<typename PolicyType>
+	void RemoveDuplicateGarbageChains(PolicyType& Policy, TArray<FGraphPath>& InOutGraphPaths, FVertex GCObjectReferencerVertex)
+	{
+		typedef TPair<FVertex, FVertex> FGarbageChainKey;
+		TMap<FGarbageChainKey, FGraphPath*> GarbageChains;
+		TArray<FGraphPath*> NoGarbageChains;
+
+		for (FGraphPath& Path : InOutGraphPaths)
+		{
+			const int32 GarbageIndex = Path.Vertices.FindLastByPredicate([&Policy](FVertex Vertex) { return Policy.IsGarbage(Vertex); });
+			if (GarbageIndex == INDEX_NONE)
+			{
+				NoGarbageChains.Add(&Path);
+				continue;
+			}
 
-	TMap<FGarbageChainKey, FReferenceChain*> GarbageChains;
-	TArray<FReferenceChain*> NoGarbageChains;
+			FVertex Garbage = Path.Vertices[GarbageIndex];
+			FVertex RootReferencer = Path.Vertices.Last();
+			if (RootReferencer == GCObjectReferencerVertex && Path.Vertices.Num() > 2)
+			{
+				RootReferencer = Path.Vertices[Path.Vertices.Num() - 2];
+			}
 
-	for (int32 ChainIndex = 0; ChainIndex < AllChains.Num(); ++ChainIndex)
+			FGarbageChainKey Key(Garbage, RootReferencer);
+			if (FGraphPath* Existing = GarbageChains.FindRef(Key))
+			{
+				int32 ExistingLen = Existing->Vertices.Num() - Existing->Vertices.Find(Garbage);
+				int32 NewLen = Path.Vertices.Num() - Path.Vertices.Find(Garbage);
+				if (NewLen < ExistingLen)
+				{
+					GarbageChains.FindOrAdd(Key, &Path);
+				}
+			}
+			else
+			{
+				GarbageChains.Add(Key, &Path);
+			}
+		}
+
+		// Copy relevant paths out of InOutGraphPaths before ovewriting its contents.
+		TArray<FGraphPath> OutPaths;
+		OutPaths.Reserve(GarbageChains.Num() + NoGarbageChains.Num());
+		for (const TPair<FGarbageChainKey, FGraphPath*>& Pair : GarbageChains)
+		{
+			OutPaths.Emplace(MoveTemp(*Pair.Value));
+		}
+		for (FGraphPath* Path : NoGarbageChains)
+		{
+			OutPaths.Emplace(MoveTemp(*Path));
+		}
+		InOutGraphPaths = MoveTemp(OutPaths);
+	}
+
+	void RemoveChainsWithDuplicatedRoots(TArray<FGraphPath>& InOutPaths, FVertex GCObjectReferencerVertex)
 	{
-		FReferenceChain* Chain = AllChains[ChainIndex];
-		const int32 GarbageIndex = Chain->Nodes.FindLastByPredicate([](FGraphNode* Node){ return Node->ObjectInfo->IsGarbage(); });
+		const auto Predicate = [GCObjectReferencerVertex](const FVertex& V) { return V != GCObjectReferencerVertex; };
+		// This is going to be rather slow but it depends on the number of chains which shouldn't be too bad (usually)
+		for (int32 FirstChainIndex = 0; FirstChainIndex < InOutPaths.Num(); ++FirstChainIndex)
+		{
+			const FGraphPath& FirstPath = InOutPaths[FirstChainIndex];
+			const FVertex RootNode = FirstPath.Vertices[FirstPath.Vertices.FindLastByPredicate(Predicate)];
+			for (int32 SecondChainIndex = InOutPaths.Num() - 1; SecondChainIndex > FirstChainIndex; --SecondChainIndex)
+			{
+				const FGraphPath& SecondPath = InOutPaths[SecondChainIndex];
+				if (SecondPath.Vertices[SecondPath.Vertices.FindLastByPredicate(Predicate)] == RootNode)
+				{
+					InOutPaths.RemoveAt(SecondChainIndex);
+				}
+			}
+		}
+	}
 
-		if (GarbageIndex == INDEX_NONE)
+	void RemoveDuplicatedChains(TArray<FGraphPath>& InOutPaths, FVertex GCObjectReferencerVertex)
+	{
+		// We consider chains identical if the direct referencer of the target node and the root node are identical
+		typedef TTuple<FVertex, FVertex, FVertex> FChainDedupKey;
+		TMap<FChainDedupKey, FGraphPath*> UniqueChains;
+
+		for (int32 ChainIndex = 0; ChainIndex < InOutPaths.Num(); ++ChainIndex)
 		{
-			NoGarbageChains.Add(Chain);
-			continue;
+			FGraphPath& Path = InOutPaths[ChainIndex];
+			FVertex DirectReferencer = Path.Vertices[1];
+
+			// If the chain referencers is the GUObjectReferencer then go one down to find a more unique object to use for deduplication
+			FVertex RootReferencer = Path.Vertices.Last();
+			if (RootReferencer == GCObjectReferencerVertex && Path.Vertices.Num() > 2)
+			{
+				RootReferencer = Path.Vertices[Path.Vertices.Num() - 2];
+			}
+
+			FChainDedupKey ChainRootAndReferencer(Path.Target, DirectReferencer, RootReferencer);
+			if (FGraphPath* ExistingChain = UniqueChains.FindRef(ChainRootAndReferencer))
+			{
+				if (ExistingChain->Vertices.Num() > Path.Vertices.Num())
+				{
+					UniqueChains[ChainRootAndReferencer] = &Path;
+				}
+			}
+			else
+			{
+				UniqueChains.Add(ChainRootAndReferencer, &Path);
+			}
 		}
 
-		FGraphNode* Garbage = Chain->Nodes[GarbageIndex];
-		FGraphNode* ChainReferencer = Chain->Nodes.Last();
-		if (ChainReferencer == GCObjReferencerNode && Chain->Nodes.Num() > 2)
+		TArray<FGraphPath> OutPaths;
+		for (const TPair<FChainDedupKey, FGraphPath*>& Pair : UniqueChains)
 		{
-			ChainReferencer = Chain->Nodes[Chain->Nodes.Num() - 2];
+			OutPaths.Emplace(MoveTemp(*Pair.Value));
 		}
+		InOutPaths = MoveTemp(OutPaths);
+	}
 
-		FGarbageChainKey Key(Garbage, ChainReferencer);
-		if (FReferenceChain* Existing = GarbageChains.FindRef(Key))
+	/**
+	 * Search a transpose object graph for paths from the target objects to the roots.
+	 *
+	 * @param Graph reverse object reference graph i.e. edges are from objects to the objects that reference them.
+	 */
+	template<typename PolicyType>
+	void PerformSearch(PolicyType& Policy,
+		TConstArrayView<typename PolicyType::ObjectType> ObjectsToFindReferencesTo,
+		const UE::Graph::FGraph& Graph,
+		EReferenceChainSearchMode Mode,
+		TArray<FGraphPath>& OutGraphPaths)
+	{
+		TRACE_CPUPROFILER_EVENT_SCOPE(UE::ReferenceChainSearch::PerformSearch);
+		const bool bDirectOnly = !!(Mode & EReferenceChainSearchMode::Direct);
+		const bool bExternalOnly = !!(Mode & EReferenceChainSearchMode::ExternalOnly);
+		const bool bLongest = !!(Mode & EReferenceChainSearchMode::Longest);
+		const bool bShortest = !!(Mode & EReferenceChainSearchMode::Shortest);
+		const bool bShortestToGarbage = !!(Mode & EReferenceChainSearchMode::ShortestToGarbage);
+
+		using ObjectType = typename PolicyType::ObjectType;
+		if (bDirectOnly)
 		{
-			int32 ExistingLen = Existing->Nodes.Num() - Existing->Nodes.Find(Garbage);
-			int32 NewLen = Chain->Nodes.Num() - Chain->Nodes.Find(Garbage);
-			if (NewLen < ExistingLen)
+			for (ObjectType TargetObject : ObjectsToFindReferencesTo)
 			{
-				GarbageChains.FindOrAdd(Key, Chain);
+				FVertex TargetVertex = Policy.ObjectToVertex(TargetObject);
+				TConstArrayView<FVertex> InEdges = Graph.EdgeLists[TargetVertex];
+				for (FVertex ReferencerVertex : InEdges)
+				{
+					ObjectType Referencer = Policy.VertexToObject(ReferencerVertex);
+					if (!bExternalOnly || !Policy.IsIn(TargetObject, Referencer))
+					{
+						FGraphPath Path;
+						Path.Target = TargetVertex;
+						Path.Vertices.Add(TargetVertex);
+						Path.Vertices.Add(ReferencerVertex);
+						OutGraphPaths.Emplace(MoveTemp(Path));
+					}
+				}
 			}
 		}
 		else
 		{
-			GarbageChains.Add(Key, Chain);
+			int32 VisitCounter = 0;
+			// Sparse mapping of visit count for relevant nodes
+			TMap<FVertex, int32> VisitCounts;
+			TArray<FGraphPath> TempChains;
+			FVertex GCObjReferencerVertex = Policy.GetGCObjectReferencerVertex();
+			for (ObjectType TargetObject : ObjectsToFindReferencesTo)
+			{
+				const FVertex TargetVertex = Policy.ObjectToVertex(TargetObject);
+				TArray<FGraphPath> ThisTargetChains;
+				for (FVertex ReferencedByVertex : Graph.EdgeLists[TargetVertex])
+				{
+					VisitCounts.FindOrAdd(TargetVertex) = ++VisitCounter;
+					TempChains.Reset();
+					const int32 MinChainDepth = 2;
+					BuildGraphPathRecursive(Policy,
+						Graph,
+						VisitCounts,
+						ReferencedByVertex,
+						TempChains,
+						MinChainDepth,
+						VisitCounter,
+						Mode,
+						GCObjReferencerVertex);
+					for (FGraphPath& Chain : TempChains)
+					{
+						Chain.Target = TargetVertex;
+						Chain.Vertices.Add(TargetVertex);
+						Algo::Reverse(Chain.Vertices);
+					}
+
+					if (bExternalOnly)
+					{
+						auto IsExternal = [](const FGraphPath& Chain) { return false; };
+						TempChains.RemoveAllSwap([&IsExternal](const FGraphPath& Chain) { return !IsExternal(Chain); });
+					}
+
+					ThisTargetChains.Append(TempChains);
+				}
+
+				if (bLongest)
+				{
+					Algo::SortBy(
+						ThisTargetChains,
+						[](const FGraphPath& Chain) { return Chain.Vertices.Num(); },
+						TGreater<int32>{});
+				}
+				else
+				{
+					Algo::SortBy(
+						ThisTargetChains,
+						[](const FGraphPath& Chain) { return Chain.Vertices.Num(); },
+						TLess<int32>{});
+				}
+
+				if (bShortestToGarbage)
+				{
+					RemoveDuplicateGarbageChains(Policy, ThisTargetChains, GCObjReferencerVertex);
+				}
+				else if (bShortest || bLongest)
+				{
+					RemoveChainsWithDuplicatedRoots(ThisTargetChains, GCObjReferencerVertex);
+				}
+
+				OutGraphPaths.Append(MoveTemp(ThisTargetChains));
+			}
+
+			if (!bLongest && !bShortest && !bShortestToGarbage)
+			{
+				RemoveDuplicatedChains(OutGraphPaths, GCObjReferencerVertex);
+			}
 		}
-	}
-	AllChains.Reset();
-	GarbageChains.GenerateValueArray(AllChains);
-	AllChains.Append(NoGarbageChains);
-}
 
-typedef TTuple<FReferenceChainSearch::FGraphNode*, FReferenceChainSearch::FGraphNode*, FReferenceChainSearch::FGraphNode*> FChainDedupKey;
+		// TODO: Improvements to relevance of emitted paths using edge-list based algorithms
+		//  Condense graph and search for paths through there, then pick arbitrary paths through cyclical components?
+		//  Should we remove vertices that can't participate in any relevant chains to simplify condensation?
+	}
 
-void FReferenceChainSearch::RemoveDuplicatedChains(TArray<FReferenceChain*>& AllChains, FGraphNode* GCObjectReferencerNode)
-{
-	// We consider chains identical if the direct referencer of the target node and the root node are identical
-	TMap<FChainDedupKey, FReferenceChain*> UniqueChains;
-	
-	for (int32 ChainIndex = 0; ChainIndex < AllChains.Num(); ++ChainIndex)
+	struct FReferenceInfoSearch : public TReferenceSearchBase<FReferenceInfoSearch>
 	{
-		FReferenceChain* Chain = AllChains[ChainIndex];
-		FGraphNode* ChainRoot = Chain->Nodes[1];
+		FPolicyUObjectHeap& Policy;
+		const TMap<const UObject*, FGCObjectInfo*>& ObjectToInfoMap;
 
-		// If the chain referencers is the GUObjectReferencer then go one down to find a more unique object to use for deduplication
-		FGraphNode* ChainReferencer = Chain->Nodes.Last();
-		if (ChainReferencer == GCObjectReferencerNode && Chain->Nodes.Num() > 2)
+		TMap<FVertex, FReferenceChainSearch::FObjectReferenceInfo>* ReferenceInfoMap;
+
+		FReferenceInfoSearch(FPolicyUObjectHeap& InPolicy, const TMap<const UObject*, FGCObjectInfo*>& InObjectToInfoMap)
+			: Policy(InPolicy)
+			, ObjectToInfoMap(InObjectToInfoMap)
 		{
-			ChainReferencer = Chain->Nodes[Chain->Nodes.Num()-2];
 		}
 
-		FChainDedupKey ChainRootAndReferencer(Chain->TargetNode, ChainRoot, ChainReferencer);
+		void Reset(TMap<FVertex, FReferenceChainSearch::FObjectReferenceInfo>* InReferenceInfoMap)
+		{
+			ReferenceInfoMap = InReferenceInfoMap;
+		}
 
-		FReferenceChain** ExistingChain = UniqueChains.Find(ChainRootAndReferencer);
-		if (ExistingChain)
+		void HandleObjectReference(
+			int32 ThreadIndex, UObject* Object, UObject* ReferencingObject, UE::GC::FTokenId TokenIndex, EGCTokenType TokenType)
 		{
-			if ((*ExistingChain)->Nodes.Num() > Chain->Nodes.Num())
+			if (!Object || Object == ReferencingObject) // Skip self-references just in case
 			{
-				delete (*ExistingChain);
-				UniqueChains[ChainRootAndReferencer] = Chain;
+				return;
+			}
+
+			if (!ReferencingObject)
+			{
+				ReferencingObject = FGCObject::GGCObjectReferencer;
+			}
+
+			FReferenceChainSearch::FObjectReferenceInfo* RefInfo = ReferenceInfoMap->Find(Policy.ObjectToVertex(Object));
+			if (!RefInfo)
+			{
+				return; // Irrelevant reference
+			}
+			if (RefInfo->Object != nullptr)
+			{
+				return; // Already found a ref from this object to this target
+			}
+
+			FGCObjectInfo* ObjectInfo = ObjectToInfoMap.FindRef(Object);
+			*RefInfo = FReferenceChainSearch::FObjectReferenceInfo(ObjectInfo);
+			if (TokenIndex != UE::GC::ETokenlessId::Collector)
+			{
+				FTokenInfo TokenInfo = ReferencingObject->GetClass()->ReferenceTokens.GetTokenInfo(TokenIndex);
+				RefInfo->ReferencerName = TokenInfo.Name;
+				RefInfo->Type = FReferenceChainSearch::EReferenceType::Property;
 			}
 			else
 			{
-				delete Chain;
+				RefInfo->Type = FReferenceChainSearch::EReferenceType::AddReferencedObjects;
+				RefInfo->StackFrames.AddUninitialized(FReferenceChainSearch::FObjectReferenceInfo::MaxStackFrames);
+				int32 NumStackFrames = FPlatformStackWalk::CaptureStackBackTrace(RefInfo->StackFrames.GetData(), RefInfo->StackFrames.Num());
+				RefInfo->StackFrames.SetNum(NumStackFrames);
+
+				if (FGCObject::GGCObjectReferencer && (!ReferencingObject || ReferencingObject == FGCObject::GGCObjectReferencer))
+				{
+					FString RefName;
+					if (FGCObject::GGCObjectReferencer->GetReferencerName(Object, RefName, true))
+					{
+						// In case FGCObjects misbehave and give us long strings, truncate so we can make them an FName.
+						if (RefName.Len() >= NAME_SIZE)
+						{
+							RefName.LeftInline(NAME_SIZE - 1);
+						}
+						else
+						{
+							RefInfo->ReferencerName = FName(*RefName);
+						}
+					}
+					else if (FGCObject* ReferencingFGCObject = FGCObject::GGCObjectReferencer->GetCurrentlySerializingObject())
+					{
+						RefInfo->ReferencerName = *ReferencingFGCObject->GetReferencerName();
+					}
+				}
+				else if (ReferencingObject)
+				{
+					RefInfo->ReferencerName = *ReferencingObject->GetFullName();
+				}
 			}
 		}
-		else
+	};
+
+	void FPolicyUObjectHeap::GatherReferenceInfo(TMap<FVertex, TMap<FVertex, FReferenceChainSearch::FObjectReferenceInfo>>& InOutReferenceInfo)
+	{
+		FReferenceInfoSearch Search(*this, ObjectToInfoMap);
+		for (TPair<FVertex, TMap<FVertex, FReferenceChainSearch::FObjectReferenceInfo>>& SourcePair : InOutReferenceInfo)
 		{
-			UniqueChains.Add(ChainRootAndReferencer, Chain);
+			FVertex SourceVertex = SourcePair.Key;
+			UObject* Object = VertexToObject(SourceVertex);
+			Search.Reset(&SourcePair.Value);
+			Search.CollectReferencesFromObject(Object);
 		}
 	}
-	AllChains.Reset();
-	UniqueChains.GenerateValueArray(AllChains);
-}
 
-void FReferenceChainSearch::BuildReferenceChains(TConstArrayView<FGraphNode*> TargetNodes, TArray<FReferenceChain*>& Chains, EReferenceChainSearchMode SearchMode, FGraphNode* GCObjReferencerNode)
-{
-	int32 VisitCounter = 0;
-	TArray<FReferenceChain*> AllChains;
-	for (FGraphNode* TargetNode : TargetNodes)
+	void FPolicyGCHistory::GatherReferenceInfo(TMap<FVertex, TMap<FVertex, FReferenceChainSearch::FObjectReferenceInfo>>& InOutReferenceInfo)
 	{
-		TArray<FReferenceChain*> ThisTargetChains;
-
-		// Recursively construct reference chains
-		for (FGraphNode* ReferencedByNode : TargetNode->ReferencedByObjects)
+		FGCObjectInfo* GCObjReferencerInfo = Snapshot.ObjectToInfoMap.FindRef(FGCObject::GGCObjectReferencer);
+		FVertex GCObjReferencerVertex = ObjectToVertex(GCObjReferencerInfo);
+		for (TPair<FVertex, TMap<FVertex, FReferenceChainSearch::FObjectReferenceInfo>>& SourcePair : InOutReferenceInfo)
 		{
-			TargetNode->Visited = ++VisitCounter;
-		
-			AllChains.Reset();
-			const int32 MinChainDepth = 2; // The chain will contain at least the TargetNode and the ReferencedByNode
-			BuildReferenceChainsRecursive(ReferencedByNode, AllChains, MinChainDepth, VisitCounter, SearchMode, GCObjReferencerNode);
-
-			for (FReferenceChain* Chain : AllChains)
-			{
-				Chain->TargetNode = TargetNode; // Store the node that created this chain in case we later truncate the chain
-				Chain->InsertNode(TargetNode);
-			}
-
-			// Filter based on search mode	
-			if (!!(SearchMode & EReferenceChainSearchMode::ExternalOnly))
+			FVertex SourceVertex = SourcePair.Key;
+			TArray<FGCDirectReferenceInfo>* DirectReferences = Snapshot.DirectReferences.FindRef(VertexToObject(SourceVertex));
+			if (DirectReferences)
 			{
-				for (int32 ChainIndex = AllChains.Num() - 1; ChainIndex >= 0; --ChainIndex)
+				for (TPair<FVertex, FReferenceChainSearch::FObjectReferenceInfo>& TargetPair : SourcePair.Value)
 				{
-					FReferenceChain* Chain = AllChains[ChainIndex];	
-					if (!Chain->IsExternal())
+					FVertex TargetVertex = TargetPair.Key;
+					FGCObjectInfo* TargetInfo = VertexToObject(TargetVertex);
+					// Linear search is not ideal, maybe build a map if we're looking for many references
+					FGCDirectReferenceInfo* RefInfo = DirectReferences->FindByPredicate(
+						[TargetInfo](const FGCDirectReferenceInfo& RefInfo) { return RefInfo.ReferencedObjectInfo == TargetInfo; });
+					if (RefInfo)
 					{
-						// Discard the chain
-						delete Chain;
-						AllChains.RemoveAtSwap(ChainIndex);
+						FReferenceChainSearch::EReferenceType ReferenceType = FReferenceChainSearch::EReferenceType::Unknown;
+						if (SourceVertex == GCObjReferencerVertex || RefInfo->ReferencerName == NAME_None)
+						{
+							ReferenceType = FReferenceChainSearch::EReferenceType::AddReferencedObjects;
+						}
+						else
+						{
+							ReferenceType = FReferenceChainSearch::EReferenceType::Property;
+						}
+						TargetPair.Value = FReferenceChainSearch::FObjectReferenceInfo(TargetInfo, ReferenceType, RefInfo->ReferencerName);
 					}
 				}
 			}
-
-			ThisTargetChains.Append(AllChains);
 		}
+	}
 
-		// Sort all chains based on the search criteria
-		if (!(SearchMode & EReferenceChainSearchMode::Longest))
-		{
-			// Sort from the shortest to the longest chain
-			ThisTargetChains.Sort([](const FReferenceChain& LHS, const FReferenceChain& RHS) { return LHS.Num() < RHS.Num(); });
-		}
-		else
+	template<typename PolicyType>
+	void PopulateReferenceInfo(PolicyType& Policy,
+		const UE::Graph::FGraph& Graph,
+		EReferenceChainSearchMode Mode,
+		const TArray<FGraphPath>& GraphPaths,
+		TArray<FReferenceChainSearch::FReferenceChain*>& OutReferenceChains)
+	{
+		TRACE_CPUPROFILER_EVENT_SCOPE(UE::ReferenceChainSearch::PopulateReferenceInfo);
+		// Construct the node/info types the public API expects
+		for (const FGraphPath& Path : GraphPaths)
 		{
-			// Sort from the longest to the shortest chain
-			ThisTargetChains.Sort([](const FReferenceChain& LHS, const FReferenceChain& RHS) { return LHS.Num() > RHS.Num(); });
+			for (FVertex Vertex : Path.Vertices)
+			{
+				Policy.FindOrCreateObjectInfoAndGraphNode(Vertex);
+			}
 		}
 
-		if (!!(SearchMode & (EReferenceChainSearchMode::ShortestToGarbage)))
+		// Only gather callstacks/info for relevant edges
+		TMap<FVertex, TMap<FVertex, FReferenceChainSearch::FObjectReferenceInfo>> ReferenceInfoMap;
+		for (const FGraphPath& Path : GraphPaths)
 		{
-			RemoveDuplicateGarbageChains(ThisTargetChains, GCObjReferencerNode);
+			// From rooted object to target object
+			for (int32 i = Path.Vertices.Num() - 1; i > 0; --i)
+			{
+				FVertex From = Path.Vertices[i];
+				FVertex To = Path.Vertices[i - 1];
+				ReferenceInfoMap.FindOrAdd(From).FindOrAdd(To);
+			}
 		}
-		// Remove duplicated roots per target object
-		else if (!!(SearchMode & (EReferenceChainSearchMode::Longest | EReferenceChainSearchMode::Shortest)))
+
+		// Re-gather references for each object we care about to gather detailed info
+		Policy.GatherReferenceInfo(ReferenceInfoMap);
+		for (auto It = ReferenceInfoMap.CreateIterator(); It; ++It)
 		{
-			RemoveChainsWithDuplicatedRoots(ThisTargetChains, GCObjReferencerNode);
+			FVertex SourceVertex = It.Key();
+			FReferenceChainSearch::FGraphNode* SourceNode = Policy.VertexToGraphNode(SourceVertex);
+			const TMap<FVertex, FReferenceChainSearch::FObjectReferenceInfo>& ReferencedObjects = It.Value();
+			for (const TPair<FVertex, FReferenceChainSearch::FObjectReferenceInfo>& RefereePair : ReferencedObjects)
+			{
+				FVertex ToVertex = RefereePair.Key;
+				const FReferenceChainSearch::FObjectReferenceInfo& Info = RefereePair.Value;
+				if (ensureAlwaysMsgf(Info.Object,
+						TEXT("Failed to find expected reference from %s to %s"),
+						*Policy.VertexToObject(SourceVertex)->GetPathName(),
+						*Policy.VertexToObject(ToVertex)->GetPathName()))
+				{
+					FReferenceChainSearch::FGraphNode* TargetNode = Policy.VertexToGraphNode(ToVertex);
+					TargetNode->ReferencedByObjects.Add(SourceNode);
+					SourceNode->ReferencedObjects.Emplace(
+						FReferenceChainSearch::FNodeReferenceInfo(TargetNode, Info.Type, Info.ReferencerName, Info.StackFrames));
+				}
+			}
 		}
 
-		Chains.Append(ThisTargetChains);
-	}
-
-	// If we didn't remove dupe-root chains, remove overall duplicated chains on the entire set of results
-	if (!(SearchMode & (EReferenceChainSearchMode::Longest | EReferenceChainSearchMode::Shortest | EReferenceChainSearchMode::ShortestToGarbage)))
-	{
-		RemoveDuplicatedChains(Chains, GCObjReferencerNode);
-	}
+		for (const FGraphPath& Path : GraphPaths)
+		{
+			FReferenceChainSearch::FGraphNode* TargetNode = Policy.VertexToGraphNode(Path.Target);
+			TArray<FReferenceChainSearch::FGraphNode*> Nodes;
+			Nodes.Reserve(Path.Vertices.Num());
+			TArray<const FReferenceChainSearch::FNodeReferenceInfo*> ReferenceInfos;
+			ReferenceInfos.Reserve(Path.Vertices.Num());
 
-	// Finally, fill extended reference info for the remaining chains
-	for (FReferenceChain* Chain : Chains)
-	{
-		Chain->FillReferenceInfo();
-	}
-}
+			for (FVertex Vertex : Path.Vertices)
+			{
+				Nodes.Add(Policy.VertexToGraphNode(Vertex));
+			}
 
-void FReferenceChainSearch::BuildReferenceChainsForDirectReferences(TConstArrayView<FGraphNode*> TargetNodes, TArray<FReferenceChain*>& AllChains, EReferenceChainSearchMode SearchMode, FGraphNode* GCObjReferencerNode)
-{
-	for (FGraphNode* TargetNode : TargetNodes) 
-	{
-		for (FGraphNode* ReferencedByNode : TargetNode->ReferencedByObjects)
-		{
-			if (!(SearchMode & EReferenceChainSearchMode::ExternalOnly) || !ReferencedByNode->ObjectInfo->IsIn(TargetNode->ObjectInfo))
+			ReferenceInfos.SetNumUninitialized(Path.Vertices.Num());
+			ReferenceInfos[0] = nullptr; // Target references nothing
+			for (int32 i = Path.Vertices.Num() - 1; i > 0; --i)
 			{
-				FReferenceChain* Chain = new FReferenceChain();
-				Chain->TargetNode = TargetNode;
-				Chain->AddNode(TargetNode);
-				Chain->AddNode(ReferencedByNode);
-				Chain->FillReferenceInfo();
-				AllChains.Add(Chain);
+				const FReferenceChainSearch::FNodeReferenceInfo* Info =
+					Nodes[i]->ReferencedObjects.Find(FReferenceChainSearch::FNodeReferenceInfo(Nodes[i - 1]));
+				ReferenceInfos[i] = Info;
 			}
+			FReferenceChainSearch::FReferenceChain* Chain =
+				new FReferenceChainSearch::FReferenceChain(TargetNode, MoveTemp(Nodes), MoveTemp(ReferenceInfos));
+			OutReferenceChains.Add(Chain);
 		}
 	}
-}
+} // namespace UE::ReferenceChainSearch
 
 FString FReferenceChainSearch::GetObjectFlags(FGCObjectInfo* InObject)
 {
@@ -365,6 +1237,11 @@ FString FReferenceChainSearch::GetObjectFlags(FGCObjectInfo* InObject)
 		Flags += TEXT("(asyncloading) ");
 	}
 
+	if (InObject->HasAnyInternalFlags(EInternalObjectFlags::LoaderImport))
+	{
+		Flags += TEXT("(loaderimport) ");
+	}
+
 	if (InObject->IsDisregardForGC())
 	{
 		Flags += TEXT("(NeverGCed) ");
@@ -381,39 +1258,50 @@ FString FReferenceChainSearch::GetObjectFlags(FGCObjectInfo* InObject)
 	return Flags;
 }
 
-static void ConvertStackFramesToCallstack(const uint64* StackFrames, int32 NumStackFrames, int32 Indent, FOutputDevice& Out)
+static void ConvertStackFramesToCallstack(
+	const uint64* StackFrames,
+	int32 NumStackFrames,
+	int32 Indent,
+	FOutputDevice& Out,
+	TMap<uint64, FString>& Cache)
 {
 	// Convert the stack trace to text
+	FStringView View;
 	for (int32 Idx = 0; Idx < NumStackFrames; Idx++)
 	{
-		ANSICHAR Buffer[1024];
-		Buffer[0] = '\0';
-		FPlatformStackWalk::ProgramCounterToHumanReadableString(Idx, StackFrames[Idx], Buffer, sizeof(Buffer));
+		if (FString* Existing = Cache.Find(StackFrames[Idx]))
+		{
+			View = FStringView(**Existing, Existing->Len());
+		}
+		else
+		{
+			ANSICHAR Buffer[1024];
+			Buffer[0] = '\0';
+			FPlatformStackWalk::ProgramCounterToHumanReadableString(Idx, StackFrames[Idx], Buffer, sizeof(Buffer));
+			View = Cache.FindOrAdd(StackFrames[Idx], FString(Buffer));
+		}
 
-		if (FCStringAnsi::Strstr(Buffer, "TFastReferenceCollector") != nullptr)
+		if (View.Contains(TEXTVIEW("TFastReferenceCollector")))
 		{
 			break;
 		}
 
-		if (FCStringAnsi::Strstr(Buffer, "FWindowsPlatformStackWalk") == nullptr &&
-			FCStringAnsi::Strstr(Buffer, "FDirectReferenceProcessor") == nullptr)
+		if (!View.Contains(TEXTVIEW("FWindowsPlatformStackWalk")) && !View.Contains(TEXTVIEW("FDirectReferenceProcessor"))
+			&& !View.Contains(TEXTVIEW("FReferenceInfoProcessor")))
 		{
-			ANSICHAR* TrimmedBuffer = FCStringAnsi::Strstr(Buffer, "!");
-			if (!TrimmedBuffer)
-			{
-				TrimmedBuffer = Buffer;
-			}
-			else
+			if (int32 Index = View.Find(TEXTVIEW("!")); Index != INDEX_NONE)
 			{
-				TrimmedBuffer++;
+				View = View.RightChop(Index + 1);
 			}
-
-			Out.Logf(ELogVerbosity::Log, TEXT("%s   ^ %s"), FCString::Spc(Indent), *FString(TrimmedBuffer));
+			Out.Logf(ELogVerbosity::Log, TEXT("%*s   ^ %.*s"), Indent, TEXT(""), View.Len(), View.GetData());
 		}
 	}
 }
 
-void FReferenceChainSearch::DumpChain(FReferenceChainSearch::FReferenceChain* Chain, TFunctionRef<bool(FCallbackParams& Params)> ReferenceCallback, FOutputDevice& Out)
+void FReferenceChainSearch::DumpChain(FReferenceChainSearch::FReferenceChain* Chain,
+	TFunctionRef<bool(FCallbackParams& Params)> ReferenceCallback,
+	TMap<uint64, FString>& CallstackCache,
+	FOutputDevice& Out)
 {
 	if (Chain->Num())
 	{
@@ -449,7 +1337,7 @@ void FReferenceChainSearch::DumpChain(FReferenceChainSearch::FReferenceChain* Ch
 			Params.Indent = FMath::Min<int32>(TCStringSpcHelper<TCHAR>::MAX_SPACES, Chain->Num() - NodeIndex - 1);
 			Params.Out = &Out;
 
-			if (ReferenceInfo->Type == EReferenceType::Property)
+			if (ReferenceInfo && ReferenceInfo->Type == EReferenceType::Property)
 			{
 				FString ReferencingPropertyName;
 				UClass* ReferencerClass = Cast<UClass>(ReferencerObject->GetClass()->TryResolveObject());
@@ -471,7 +1359,7 @@ void FReferenceChainSearch::DumpChain(FReferenceChainSearch::FReferenceChain* Ch
 					// Handle base UObject referencer info (it's only exposed to the GC token stream and not to the reflection system)
 					static const FName ClassPropertyName(TEXT("Class"));
 					static const FName OuterPropertyName(TEXT("Outer"));
-					
+
 					FString ClassName;
 					if (ReferenceInfo->ReferencerName == ClassPropertyName || ReferenceInfo->ReferencerName == OuterPropertyName)
 					{
@@ -497,7 +1385,7 @@ void FReferenceChainSearch::DumpChain(FReferenceChainSearch::FReferenceChain* Ch
 					*GetObjectFlags(Object),
 					*Object->GetFullName());
 			}
-			else if (ReferenceInfo->Type == EReferenceType::AddReferencedObjects)
+			else if (ReferenceInfo && ReferenceInfo->Type == EReferenceType::AddReferencedObjects)
 			{
 				FString UObjectOrGCObjectName;
 				if (ReferenceInfo->ReferencerName.IsNone())
@@ -526,10 +1414,14 @@ void FReferenceChainSearch::DumpChain(FReferenceChainSearch::FReferenceChain* Ch
 
 				if (ReferenceInfo->StackFrames.Num())
 				{
-					ConvertStackFramesToCallstack(ReferenceInfo->StackFrames.GetData(), ReferenceInfo->StackFrames.Num(), Params.Indent, Out);
+					ConvertStackFramesToCallstack(ReferenceInfo->StackFrames.GetData(),
+						ReferenceInfo->StackFrames.Num(),
+						Params.Indent,
+						Out,
+						CallstackCache);
 				}
 			}
-			else if (ReferenceInfo->Type == EReferenceType::OuterChain)
+			else if (ReferenceInfo && ReferenceInfo->Type == EReferenceType::OuterChain)
 			{
 				Out.Logf(ELogVerbosity::Log, TEXT("%s-> %s = %s %s"),
 					FCString::Spc(Params.Indent),
@@ -549,40 +1441,12 @@ void FReferenceChainSearch::DumpChain(FReferenceChainSearch::FReferenceChain* Ch
 			bPostCallbackContinue = ReferenceCallback(Params);
 
 			ReferencerObject = Object;
-			ReferenceInfo = Chain->GetReferenceInfo(NodeIndex);			
+			ReferenceInfo = Chain->GetReferenceInfo(NodeIndex);
 		}
 		Out.Logf(ELogVerbosity::Log, TEXT("  "));
 	}
 }
 
-void FReferenceChainSearch::FReferenceChain::FillReferenceInfo()
-{
-	// The first entry is the object we were looking for references to so add an empty entry for it
-	ReferenceInfos.Emplace(nullptr);
-
-	// Iterate over all nodes and add reference info based on the next node (which is the object that referenced the current node)
-	for (int32 NodeIndex = 1; NodeIndex < Nodes.Num(); ++NodeIndex)
-	{
-		FGraphNode* PreviousNode = Nodes[NodeIndex - 1];
-		FGraphNode* CurrentNode = Nodes[NodeIndex];
-
-		// Found the PreviousNode in the list of objects referenced by the CurrentNode
-		const FNodeReferenceInfo* FoundInfo = nullptr;
-		for (const FNodeReferenceInfo& Info : CurrentNode->ReferencedObjects)
-		{
-			if (Info.Object == PreviousNode)
-			{
-				FoundInfo = &Info;
-				break;
-			}
-		}
-		check(FoundInfo); // because there must have been a reference since we created this chain using it
-		check(FoundInfo->Object == PreviousNode);
-		ReferenceInfos.Add(FoundInfo);
-	}
-	check(ReferenceInfos.Num() == Nodes.Num());
-}
-
 bool FReferenceChainSearch::FReferenceChain::IsExternal() const
 {
 	if (Nodes.Num() > 1)
@@ -596,189 +1460,43 @@ bool FReferenceChainSearch::FReferenceChain::IsExternal() const
 	}
 }
 
-/**
-* Handles UObject references found by TFastReferenceCollector
-*/
-class FDirectReferenceProcessor : public FSimpleReferenceProcessorBase
-{	
-protected:
-	TSet<FReferenceChainSearch::FObjectReferenceInfo> ReferencedObjects;
-	TMap<const UObject*, FGCObjectInfo*>& ObjectToInfoMap;
-
-public:
-
-	FDirectReferenceProcessor(TMap<const UObject*, FGCObjectInfo*>& InObjectToInfoMap)
-	: ObjectToInfoMap(InObjectToInfoMap)
-	{
-	}
-
-	void Reset()
-	{
-		ReferencedObjects.Reset();
-	}
-
-	const TSet<FReferenceChainSearch::FObjectReferenceInfo>& GetReferencedObjects() const
-	{
-		return ReferencedObjects;
-	}
-
-	void HandleTokenStreamObjectReference(FGCArrayStruct& ObjectsToSerializeStruct, UObject* ReferencingObject, UObject*& Object, UE::GC::FTokenId TokenIndex, EGCTokenType TokenType, bool bAllowReferenceElimination)
-	{		
-		if (Object && Object != ReferencingObject) // Skip self-references just in case 
-		{
-			FGCObjectInfo* ObjectInfo = FGCObjectInfo::FindOrAddInfoHelper(Object, ObjectToInfoMap);
-
-			FReferenceChainSearch::FObjectReferenceInfo RefInfo(ObjectInfo);
-			if (!ReferencedObjects.Contains(RefInfo))
-			{
-				if (TokenIndex != UE::GC::ETokenlessId::Collector)
-				{
-					FTokenInfo TokenInfo = ReferencingObject->GetClass()->ReferenceTokens.GetTokenInfo(TokenIndex);
-					RefInfo.ReferencerName = TokenInfo.Name;
-					RefInfo.Type = FReferenceChainSearch::EReferenceType::Property;
-				}
-				else
-				{
-					RefInfo.Type = FReferenceChainSearch::EReferenceType::AddReferencedObjects;
-					if (GAddReferencedObjectsStackTrace)
-					{
-						RefInfo.StackFrames.AddUninitialized(FReferenceChainSearch::FObjectReferenceInfo::MaxStackFrames);
-						int32 NumStackFrames = FPlatformStackWalk::CaptureStackBackTrace(RefInfo.StackFrames.GetData(), RefInfo.StackFrames.Num());
-						RefInfo.StackFrames.SetNum(NumStackFrames);
-					}
-
-					if (FGCObject::GGCObjectReferencer && (!ReferencingObject || ReferencingObject == FGCObject::GGCObjectReferencer))
-					{
-						FString RefName;
-						if (GFGCObjectPropertyNames && FGCObject::GGCObjectReferencer->GetReferencerName(Object, RefName, true))
-						{
-							// In case FGCObjects misbehave and give us long strings, truncate so we can make them an FName.
-							if (RefName.Len() >= NAME_SIZE)
-							{
-								RefName.LeftInline(NAME_SIZE-1);
-							}
-
-							RefInfo.ReferencerName = FName(*RefName);
-						}
-						else if (FGCObject* ReferencingFGCObject = FGCObject::GGCObjectReferencer->GetCurrentlySerializingObject())
-						{
-							RefInfo.ReferencerName = *ReferencingFGCObject->GetReferencerName();
-						}
-					}
-					else if (ReferencingObject)
-					{
-						RefInfo.ReferencerName = *ReferencingObject->GetFullName();
-					}
-				}
-
-				ReferencedObjects.Emplace(MoveTemp(RefInfo));
-			}
-		}
-	}
-};
-
-class FMinimalReferenceProcessor : public FDirectReferenceProcessor
-{	
-	TSet<UObject*>& TargetObjects;
-	bool bSearchInners = false;
-
-public:
-	FMinimalReferenceProcessor(
-		TSet<UObject*>& InTargetObjects,
-		TMap<const UObject*, FGCObjectInfo*>& InObjectToInfoMap
-	)
-	: FDirectReferenceProcessor(InObjectToInfoMap)
-	, TargetObjects(InTargetObjects)
-	{
-	}
-
-	void SetSearchInners(bool InValue)
-	{
-		bSearchInners = InValue;
-	}
-
-	void HandleTokenStreamObjectReference(FGCArrayStruct& ObjectsToSerializeStruct, UObject* ReferencingObject, UObject*& Object, UE::GC::FTokenId TokenIndex, EGCTokenType TokenType, bool bAllowReferenceElimination)
-	{		
-		if (Object && Object != ReferencingObject) // Skip self-references just in case 
-		{
-			FGCObjectInfo* ObjectInfo = nullptr;
-			if (bSearchInners)
-			{
-				// Walk the outer chain to see if this is the inner of an object we're looking for which would be referenced via this outer chain
-				for (UObject* OuterObject = Object->GetOuter(); OuterObject; OuterObject = OuterObject->GetOuter())
-				{
-					if (TargetObjects.Contains(OuterObject))
-					{
-						// Create a node to reference this inner object, which we will later connect to the target object
-						ObjectInfo = FGCObjectInfo::FindOrAddInfoHelper(Object, ObjectToInfoMap);
-						break;
-					}
-				}
-			}
-			else if (TargetObjects.Contains(Object))
-			{
-				ObjectInfo = FGCObjectInfo::FindOrAddInfoHelper(Object, ObjectToInfoMap);
-			}
-
-			// We don't care about this reference
-			if (!ObjectInfo) 
-			{
-				return;
-			}
-
-			FDirectReferenceProcessor::HandleTokenStreamObjectReference(ObjectsToSerializeStruct, ReferencingObject, Object, TokenIndex, TokenType, bAllowReferenceElimination);
-		}
-	}
-};
-
-using FDirectReferenceCollectorBase = UE::GC::TDefaultCollector<FDirectReferenceProcessor>;
-class FDirectReferenceCollector : public FDirectReferenceCollectorBase
-{
-public:
-	using FDirectReferenceCollectorBase::FDirectReferenceCollectorBase;
-
-	virtual bool MarkWeakObjectReferenceForClearing(UObject** WeakReference) override
-	{
-		// To avoid false positives we need to implement this method just like GC does
-		// as these references will be treated as weak and should not be reported
-		return true;
-	}
-};
-
-FReferenceChainSearch::FReferenceChainSearch(UObject* InObjectToFindReferencesTo, EReferenceChainSearchMode Mode /*= EReferenceChainSearchMode::PrintResults*/) 
-: FReferenceChainSearch(TConstArrayView<UObject*>(&InObjectToFindReferencesTo, 1), Mode)
+FReferenceChainSearch::FReferenceChainSearch(UObject* InObjectToFindReferencesTo,
+	EReferenceChainSearchMode Mode /*= EReferenceChainSearchMode::PrintResults*/)
+	: FReferenceChainSearch(TConstArrayView<UObject*>(&InObjectToFindReferencesTo, 1), Mode)
 {
 }
 
-FReferenceChainSearch::FReferenceChainSearch(TConstArrayView<UObject*> InObjectsToFindReferencesTo, EReferenceChainSearchMode Mode /*= EReferenceChainSearchMode::PrintResults*/)
-	: ObjectsToFindReferencesTo(InObjectsToFindReferencesTo)
-	, SearchMode(Mode)
+FReferenceChainSearch::FReferenceChainSearch(TConstArrayView<UObject*> InObjectsToFindReferencesTo,
+	EReferenceChainSearchMode Mode /*= EReferenceChainSearchMode::PrintResults*/)
+	: SearchMode(Mode)
 {
 	FSlowHeartBeatScope DisableHangDetection; // This function can be very slow
 
-	for( UObject* Obj : InObjectsToFindReferencesTo )
-	{
-		check(Obj);
-		ObjectInfosToFindReferencesTo.Add(FGCObjectInfo::FindOrAddInfoHelper(Obj, ObjectToInfoMap));
-	}
+	// Lock the global array so that nothing can add UObjects while we're iterating over it
+	GUObjectArray.LockInternalArray();
 
-	// First pass is to find all direct references for each object
-	if (!!(SearchMode & EReferenceChainSearchMode::Minimal))
+	for (const UObject* Object : InObjectsToFindReferencesTo)
 	{
-		FindMinimalDirectReferencesForObjects();
-	}
-	else
-	{
-		FindDirectReferencesForObjects();
+		ObjectInfosToFindReferencesTo.Add(FGCObjectInfo::FindOrAddInfoHelper(Object, ObjectToInfoMap));
 	}
 
-	// Second pass creates all reference chains
-	PerformSearch();
+	UE::ReferenceChainSearch::FPolicyUObjectHeap Policy(ObjectToInfoMap, AllNodes);
+
+	UE::Graph::FGraph ReferenceGraph;
+	UE::ReferenceChainSearch::PerformInitialGatherFromLiveUObjectHeap(Policy, InObjectsToFindReferencesTo, SearchMode, ReferenceGraph);
+
+	TArray<UE::ReferenceChainSearch::FGraphPath> Paths;
+	UE::ReferenceChainSearch::PerformSearch(Policy, InObjectsToFindReferencesTo, ReferenceGraph, Mode, Paths);
+	UE::ReferenceChainSearch::PopulateReferenceInfo(Policy, ReferenceGraph, Mode, Paths, ReferenceChains);
 
 	if (!!(Mode & (EReferenceChainSearchMode::PrintResults|EReferenceChainSearchMode::PrintAllResults)))
 	{
 		PrintResults(!!(Mode & EReferenceChainSearchMode::PrintAllResults));
 	}
+
+	UE_LOG(LogReferenceChain, Display, TEXT("Post-search memory usage: %.2f"), (ReferenceGraph.GetAllocatedSize() + Paths.GetAllocatedSize() + GetAllocatedSize()) / 1024.0 / 1024.0);
+
+	GUObjectArray.UnlockInternalArray();
 }
 
 FReferenceChainSearch::FReferenceChainSearch(EReferenceChainSearchMode Mode)
@@ -794,15 +1512,15 @@ FReferenceChainSearch::~FReferenceChainSearch()
 int64 FReferenceChainSearch::GetAllocatedSize() const
 {
 	int64 Size = 0;
-	Size += ObjectsToFindReferencesTo.GetAllocatedSize();
-	Size += ObjectInfosToFindReferencesTo.GetAllocatedSize();
+	// Size += ObjectsToFindReferencesTo.GetAllocatedSize();
+	// Size += ObjectInfosToFindReferencesTo.GetAllocatedSize();
 	Size += ReferenceChains.GetAllocatedSize();
 	for (const FReferenceChain* Chain : ReferenceChains)
 	{
 		Size += Chain->GetAllocatedSize();
 	}
 	Size += AllNodes.GetAllocatedSize();
-	Size += AllNodes.Num() * sizeof(FGCObjectInfo); // GC object infos are heap allocated and owned by this structure 
+	Size += AllNodes.Num() * sizeof(FGCObjectInfo); // GC object infos are heap allocated and owned by this structure
 	for (const TPair<FGCObjectInfo*, FGraphNode*> Pair : AllNodes)
 	{
 		// FGCObjectInfo makes no allocations
@@ -812,35 +1530,6 @@ int64 FReferenceChainSearch::GetAllocatedSize() const
 	return Size;
 }
 
-void FReferenceChainSearch::PerformSearch()
-{
-	TArray<FGraphNode*> GraphNodesToFind;
-	GraphNodesToFind.Reserve(ObjectsToFindReferencesTo.Num());
-	for (const UObject* Obj : ObjectsToFindReferencesTo)
-	{
-		FGraphNode* ObjectNodeToFindReferencesTo = FindOrAddNode(Obj);
-		check(ObjectNodeToFindReferencesTo);
-		GraphNodesToFind.Add(ObjectNodeToFindReferencesTo);
-	}
-
-	FGraphNode* GCObjectReferencerGraphNode = nullptr;
-	if (FGCObject::GGCObjectReferencer)
-	{
-		FGCObjectInfo* ObjectInfo = FGCObjectInfo::FindOrAddInfoHelper(FGCObject::GGCObjectReferencer, ObjectToInfoMap);
-		GCObjectReferencerGraphNode = AllNodes.FindRef(ObjectInfo);
-	}
-
-	// Now it's time to build the reference chain from all of the objects that reference the object to find references to
-	if (!(SearchMode & EReferenceChainSearchMode::Direct))
-	{		
-		BuildReferenceChains(GraphNodesToFind, ReferenceChains, SearchMode, GCObjectReferencerGraphNode);
-	}
-	else
-	{
-		BuildReferenceChainsForDirectReferences(GraphNodesToFind, ReferenceChains, SearchMode, GCObjectReferencerGraphNode);
-	}
-}
-
 #if ENABLE_GC_HISTORY
 void FReferenceChainSearch::PerformSearchFromGCSnapshot(UObject* InObjectToFindReferencesTo, FGCSnapshot& InSnapshot)
 {
@@ -853,220 +1542,50 @@ void FReferenceChainSearch::PerformSearchFromGCSnapshot(TConstArrayView<UObject*
 
 	Cleanup();
 
-	// Temporarily move the generated object info structs. We don't want to copy everything to minimize memory usage and save a few ms
-	ObjectToInfoMap = MoveTemp(InSnapshot.ObjectToInfoMap);
+	UE::ReferenceChainSearch::FPolicyGCHistory Policy(InSnapshot, AllNodes);
 
-	ObjectsToFindReferencesTo = InObjectsToFindReferencesTo;
-	ObjectInfosToFindReferencesTo.Reset();
 	for (const UObject* Obj : InObjectsToFindReferencesTo)
 	{
-		ObjectInfosToFindReferencesTo.Add(FGCObjectInfo::FindOrAddInfoHelper(Obj, ObjectToInfoMap));
-	}
-	FGCObjectInfo* GCObjectReferencerInfo = FGCObjectInfo::FindOrAddInfoHelper(FGCObject::GGCObjectReferencer, ObjectToInfoMap);
-
-	// We can avoid copying object infos but we need to regenerate direct reference infos
-	for (TPair<FGCObjectInfo*, TArray<FGCDirectReferenceInfo>*>& DirectReferencesInfo : InSnapshot.DirectReferences)
-	{
-		FGraphNode* ObjectNode = FindOrAddNode(DirectReferencesInfo.Key);
-		for (FGCDirectReferenceInfo& ReferenceInfo : *DirectReferencesInfo.Value)
+		FGCObjectInfo* ObjInfo = InSnapshot.ObjectToInfoMap.FindRef(Obj);
+		if (ObjInfo)
 		{
-			FGraphNode* ReferencedObjectNode = FindOrAddNode(ReferenceInfo.ReferencedObjectInfo);
-			EReferenceType ReferenceType = EReferenceType::Unknown;
-			if (GCObjectReferencerInfo == DirectReferencesInfo.Key || ReferenceInfo.ReferencerName == NAME_None)
-			{
-				ReferenceType = EReferenceType::AddReferencedObjects;
-			}
-			else
-			{
-				ReferenceType = EReferenceType::Property;
-			}
-			ObjectNode->ReferencedObjects.Emplace(FNodeReferenceInfo(ReferencedObjectNode, ReferenceType, ReferenceInfo.ReferencerName));
-			ReferencedObjectNode->ReferencedByObjects.Add(ObjectNode);
+			ObjectInfosToFindReferencesTo.Add(ObjInfo);
 		}
-	}
-
-	// Second pass creates all reference chains
-	PerformSearch();
-
-	if (!!(SearchMode & (EReferenceChainSearchMode::PrintResults | EReferenceChainSearchMode::PrintAllResults)))
-	{
-		PrintResults(!!(SearchMode & EReferenceChainSearchMode::PrintAllResults));
-	}
-
-	// Return the object info struct back to the snapshot
-	InSnapshot.ObjectToInfoMap = MoveTemp(ObjectToInfoMap);
-}
-#endif // ENABLE_GC_HISTORY
-
-template<typename ProcessorType>
-struct TReferenceSearchHelper
-{
-	ProcessorType Processor;
-	TArray<UObject*> ObjectsToProcess;
-
-	TReferenceSearchHelper(ProcessorType&& InProcessor)
-	: Processor(MoveTemp(InProcessor))
-	{}
-
-	static bool DoNotFilterObjects(UObject* Object) { return true; }
-
-	// HANDLE_REFERENCES is a function bool(const TSet<FReferenceChainSearch>&) function returning true to continue the search
-	// FILTER_OBJECTS is a function bool(UObject*) that returns true if the object's outgoing references should be collected
-	template<typename HANDLE_REFERENCES, typename FILTER_OBJECTS = decltype(DoNotFilterObjects)>
-	void CollectAllReferences(bool bGCOnly, HANDLE_REFERENCES HandleReferences, FILTER_OBJECTS FilterObjects = &DoNotFilterObjects)
-	{
-		for (FRawObjectIterator It; It; ++It)
+		else
 		{
-			FUObjectItem* ObjItem = *It;
-			UObject* Object = static_cast<UObject*>(ObjItem->Object);
-
-			// We can't ask the iterator for only GC objects because that would skip the GC Object referencer 
-			if (bGCOnly && GUObjectArray.IsDisregardForGC(Object) && (Object != FGCObject::GGCObjectReferencer))
-			{
-				continue;
-			}
-
-			if (!FilterObjects(Object))
-			{
-				continue;
-			}
-
-			// Find direct references
-			Processor.Reset();
-			UE::GC::FWorkerContext Context;
-			ObjectsToProcess = {Object};
-			Context.SetInitialObjectsUnpadded(ObjectsToProcess);
-			CollectReferences<FDirectReferenceCollector>(Processor, Context);
-
-			if (!HandleReferences(Object, Processor.GetReferencedObjects()))
-			{
-				break;
-			}
+			UE_LOG(LogReferenceChain, Warning, TEXT("Target object %s was not present in GC snapshot."), *Obj->GetPathName());
 		}
 	}
-};
-
-void FReferenceChainSearch::FindDirectReferencesForObjects()
-{
-	TReferenceSearchHelper<FDirectReferenceProcessor> Search{FDirectReferenceProcessor(ObjectToInfoMap)};
-	const bool bGCOnly = !!(SearchMode & EReferenceChainSearchMode::GCOnly);
-
-	// First pass, create nodes for any objects which have outgoing references 
-	Search.CollectAllReferences(bGCOnly, [&](UObject* SourceObject, const TSet<FObjectReferenceInfo>& References)
-	{
-		// Skip node creation if there are no outgoing references that we care about 
-		if (References.Num() > 0)
-		{
-			FindOrAddNode(FGCObjectInfo::FindOrAddInfoHelper(SourceObject, ObjectToInfoMap));
-		}
-		return true;
-	});
-
-	Search.CollectAllReferences(bGCOnly, [&](UObject* SourceObject, const TSet<FObjectReferenceInfo>& References)
-	{
-		if (References.Num() > 0)
-		{
-			FGCObjectInfo* SourceObjectInfo = FGCObjectInfo::FindOrAddInfoHelper(SourceObject, ObjectToInfoMap);
-
-			// Build direct reference tree filtering out leaf notes
-			for (const FObjectReferenceInfo& ReferenceInfo : References)
-			{
-				FGCObjectInfo* ReferencedObject = ReferenceInfo.Object;
-				if (AllNodes.Contains(ReferencedObject))
-				{
-					LinkNodes(SourceObjectInfo, ReferencedObject, ReferenceInfo.Type, ReferenceInfo.ReferencerName, ReferenceInfo.StackFrames);
-				}
-			}
-		}
-		return true;
-	});
-}
 
-void FReferenceChainSearch::FindMinimalDirectReferencesForObjects()
-{
-	TSet<UObject*> TargetObjects;
-	TReferenceSearchHelper<FMinimalReferenceProcessor> Search{FMinimalReferenceProcessor(TargetObjects, ObjectToInfoMap)};
-	const bool bGCOnly = !!(SearchMode & EReferenceChainSearchMode::GCOnly);
-
-	auto FilterObjects = [&TargetObjects](UObject* Object)
+	for (TPair<FGCObjectInfo*, TArray<FGCDirectReferenceInfo>*>& Pair : InSnapshot.DirectReferences)
 	{
-		// Skip objects which are in any of what we're looking for so we only report reference external to that entire group
-		for( UObject* OuterObject = Object; OuterObject; OuterObject = OuterObject->GetOuter())
+		for (const FGCDirectReferenceInfo& RefInfo : *Pair.Value)
 		{
-			if (TargetObjects.Contains(OuterObject))
+			if (ObjectInfosToFindReferencesTo.Contains(RefInfo.ReferencedObjectInfo))
 			{
-				return false;
+				UE_LOG(LogReferenceChain,
+					Display,
+					TEXT("Direct ref in GC history from %s to %s"),
+					*Pair.Key->GetPathName(),
+					*RefInfo.ReferencedObjectInfo->GetPathName());
 			}
 		}
-		return true;
-	};
-
-	for (UObject* Object: ObjectsToFindReferencesTo)
-	{
-		TargetObjects.Add(Object);
 	}
-	
-	// First pass, create nodes anything that directly references any of the target objects 
-	Search.CollectAllReferences(bGCOnly, [&](UObject* SourceObject, const TSet<FObjectReferenceInfo>& References)
-	{
-		if (References.Num() > 0)
-		{
-			FGCObjectInfo* SourceObjectInfo = FGCObjectInfo::FindOrAddInfoHelper(SourceObject, ObjectToInfoMap);
 
-			// Build direct reference tree filtering out leaf notes
-			for (const FObjectReferenceInfo& ReferenceInfo : References)
-			{
-				FGCObjectInfo* ReferencedObject = ReferenceInfo.Object;
-				LinkNodes(SourceObjectInfo, ReferencedObject, ReferenceInfo.Type, ReferenceInfo.ReferencerName, ReferenceInfo.StackFrames);
-			}
-		}
-		return true;
-	}, FilterObjects);
+	UE::Graph::FGraph ReferenceGraph;
+	UE::ReferenceChainSearch::PerformInitialGatherFromGCHistory(Policy, ReferenceGraph);
 
-	// If we didn't find a reference to any of the target objects in our first pass, loosen requirements slightly
-	TargetObjects.Reset();
-	for (FGCObjectInfo* Info : ObjectInfosToFindReferencesTo)
-	{
-		if (FGraphNode* Node = AllNodes.FindRef(Info); !Node || Node->ReferencedByObjects.Num() == 0)
-		{
-			TargetObjects.Add(Info->TryResolveObject());
-		}
-	}
+	TArray<UE::ReferenceChainSearch::FGraphPath> Paths;
+	UE::ReferenceChainSearch::PerformSearch(Policy, ObjectInfosToFindReferencesTo, ReferenceGraph, SearchMode, Paths);
+	UE::ReferenceChainSearch::PopulateReferenceInfo(Policy, ReferenceGraph, SearchMode, Paths, ReferenceChains);
 
-	// We failed to find direct references to some objects, try searching for direct references to one of their inners 
-	if (TargetObjects.Num())
+	if (!!(SearchMode & (EReferenceChainSearchMode::PrintResults | EReferenceChainSearchMode::PrintAllResults)))
 	{
-		Search.Processor.SetSearchInners(true);
-		Search.CollectAllReferences(bGCOnly, [&](UObject* SourceObject, const TSet<FObjectReferenceInfo>& References)
-		{
-			if (References.Num() > 0)
-			{
-				FGCObjectInfo* SourceObjectInfo = FGCObjectInfo::FindOrAddInfoHelper(SourceObject, ObjectToInfoMap);
-
-				// Build direct reference tree
-				for (const FObjectReferenceInfo& ReferenceInfo : References)
-				{
-					LinkNodes(SourceObjectInfo, ReferenceInfo.Object, ReferenceInfo.Type, ReferenceInfo.ReferencerName, ReferenceInfo.StackFrames);
-					UObject* ReferencedObject = ReferenceInfo.Object->TryResolveObject();
-					for (UObject* OuterObject = ReferencedObject; OuterObject; OuterObject = OuterObject->GetOuter())
-					{
-						if (TargetObjects.Contains(OuterObject))
-						{
-							LinkNodes(ReferencedObject, OuterObject, EReferenceType::OuterChain, NAME_None);
-							TargetObjects.Remove(OuterObject);
-						}
-					}
-				}
-
-				if (TargetObjects.Num() == 0)
-				{
-					return false;
-				}
-			}
-			return true;
-		}, FilterObjects);
+		PrintResults(!!(SearchMode & EReferenceChainSearchMode::PrintAllResults));
 	}
 }
+#endif // ENABLE_GC_HISTORY
+
 int32 FReferenceChainSearch::PrintResults(bool bDumpAllChains /*= false*/, UObject* TargetObject /*= nullptr*/) const
 {
 	return PrintResults([](FCallbackParams& Params) { return true; }, bDumpAllChains, TargetObject);
@@ -1079,6 +1598,7 @@ int32 FReferenceChainSearch::PrintResults(TFunctionRef<bool(FCallbackParams& Par
 	const int32 MaxChainsToPrint = 100;
 	int32 NumPrintedChains = 0;
 
+	TMap<uint64, FString> CallstackCache;
 	for (FReferenceChain* Chain : ReferenceChains)
 	{
 		if (TargetObject != nullptr && Chain->TargetNode->ObjectInfo->TryResolveObject() != TargetObject)
@@ -1088,7 +1608,7 @@ int32 FReferenceChainSearch::PrintResults(TFunctionRef<bool(FCallbackParams& Par
 
 		if (bDumpAllChains || NumPrintedChains < MaxChainsToPrint)
 		{
-			DumpChain(Chain, ReferenceCallback, *GLog);
+			DumpChain(Chain, ReferenceCallback, CallstackCache, *GLog);
 			NumPrintedChains++;
 		}
 		else
@@ -1098,20 +1618,34 @@ int32 FReferenceChainSearch::PrintResults(TFunctionRef<bool(FCallbackParams& Par
 		}
 	}
 
-	if(NumPrintedChains == 0)
+	if (NumPrintedChains == 0)
 	{
+		auto LogUnreachableObject = [this](FGCObjectInfo* ObjInfo) {
+			if (ObjInfo->HasAnyInternalFlags(EInternalObjectFlags::GarbageCollectionKeepFlags))
+			{
+				UE_LOG(LogReferenceChain, Log, TEXT("%s%s is not currently reachable but it does have some of EInternalObjectFlags::GarbageCollectionKeepFlags set."), *GetObjectFlags(ObjInfo), *ObjInfo->GetFullName());
+			}
+			else if (ObjInfo->HasAnyFlags(GARBAGE_COLLECTION_KEEPFLAGS))
+			{
+				UE_LOG(LogReferenceChain, Log, TEXT("%s%s is not currently reachable but it does have some of GARBAGE_COLLECTION_KEEPFLAGS set."), *GetObjectFlags(ObjInfo), *ObjInfo->GetFullName());
+			}
+			else
+			{
+				UE_LOG(LogReferenceChain, Log, TEXT("%s%s is not currently reachable. Try using GC history to debug transient leaks with 'gc.historysize 1'"), *GetObjectFlags(ObjInfo), *ObjInfo->GetFullName());
+			}
+		};
 		if (TargetObject)
 		{
 			FGCObjectInfo* ObjInfo = FGCObjectInfo::FindOrAddInfoHelper(TargetObject, const_cast<TMap<const UObject*, FGCObjectInfo*>&>(ObjectToInfoMap));
 			check(ObjInfo);
-			UE_LOG(LogReferenceChain, Log, TEXT("%s%s is not currently reachable."),
-				*GetObjectFlags(ObjInfo),
-				*ObjInfo->GetFullName()
-			);
+			LogUnreachableObject(ObjInfo);
 		}
 		else
 		{
-			UE_LOG(LogReferenceChain, Log, TEXT("No target objects are currently reachable."));
+			for (FGCObjectInfo* ObjInfo : ObjectInfosToFindReferencesTo)
+			{
+				LogUnreachableObject(ObjInfo);
+			}
 		}
 	}
 
@@ -1135,12 +1669,13 @@ FString FReferenceChainSearch::GetRootPath(TFunctionRef<bool(FCallbackParams& Pa
 			Chain = Index != INDEX_NONE ? ReferenceChains[Index] : nullptr;
 		}
 	}
-	
+
 	if (Chain)
 	{
 		FStringOutputDevice OutString;
 		OutString.SetAutoEmitLineTerminator(true);
-		DumpChain(Chain, ReferenceCallback, OutString);
+		TMap<uint64, FString> CallstackCache;
+		DumpChain(Chain, ReferenceCallback, CallstackCache, OutString);
 		return MoveTemp(OutString);
 	}
 	else
@@ -1162,6 +1697,7 @@ FString FReferenceChainSearch::GetRootPath(TFunctionRef<bool(FCallbackParams& Pa
 
 void FReferenceChainSearch::Cleanup()
 {
+	ObjectInfosToFindReferencesTo.Empty();
 	for (int32 ChainIndex = 0; ChainIndex < ReferenceChains.Num(); ++ChainIndex)
 	{
 		delete ReferenceChains[ChainIndex];
@@ -1179,4 +1715,4 @@ void FReferenceChainSearch::Cleanup()
 		delete ObjectToInfoPair.Value;
 	}
 	ObjectToInfoMap.Empty();
-}
+}
\ No newline at end of file
diff --git a/Engine/Source/Runtime/CoreUObject/Public/UObject/ReferenceChainSearch.h b/Engine/Source/Runtime/CoreUObject/Public/UObject/ReferenceChainSearch.h
--- a/Engine/Source/Runtime/CoreUObject/Public/UObject/ReferenceChainSearch.h
+++ b/Engine/Source/Runtime/CoreUObject/Public/UObject/ReferenceChainSearch.h
@@ -206,6 +206,13 @@ public:
 			Nodes.Reserve(ReserveDepth);
 		}
 
+		FReferenceChain(FGraphNode* InTargetNode, TArray<FGraphNode*> InNodes, TArray<const FNodeReferenceInfo*> InReferenceInfos)
+			: TargetNode(InTargetNode)
+			, Nodes(MoveTemp(InNodes))
+			, ReferenceInfos(MoveTemp(InReferenceInfos))
+		{
+		}
+
 		int64 GetAllocatedSize() const
 		{
 			return Nodes.GetAllocatedSize() + ReferenceInfos.GetAllocatedSize();
@@ -329,7 +336,6 @@ public:
 private:
 
 	/** The objects we're going to look for references to */
-	TArray<UObject*> ObjectsToFindReferencesTo;
 	TArray<FGCObjectInfo*> ObjectInfosToFindReferencesTo;
 
 	/** Search mode and options */
@@ -343,22 +349,9 @@ private:
 	/** Maps UObject pointers to object info structs */
 	TMap<const UObject*, FGCObjectInfo*> ObjectToInfoMap;
 
-	/** Performs the search */
-	void PerformSearch();
-
-	/** Constructs a complete graph of all references between objects */
-	void FindDirectReferencesForObjects();
-
-	/** Constructs a subset of the graph of references between objects to try and find the targets using minimal memory */
-	void FindMinimalDirectReferencesForObjects();
-
 	/** Frees memory */
 	void Cleanup();
 
-	/** Tries to find a node for an object and if it doesn't exists creates a new one and returns it */
-	FGraphNode* FindOrAddNode(const UObject* InObjectToFindNodeFor);
-	FGraphNode* FindOrAddNode(FGCObjectInfo* InObjectInfo);
-
 	/** Link two nodes together with the given reference info */
 	void LinkNodes(
 		UObject* From,
@@ -373,23 +366,11 @@ private:
 		FName PropertyName,
 		TConstArrayView<uint64> StackFrames = {});
 
-	/** Builds reference chains */
-	static int32 BuildReferenceChainsRecursive(FGraphNode* TargetNodes, TArray<FReferenceChain*>& ProducedChains, int32 ChainDepth, const int32 VisitCounter, EReferenceChainSearchMode SearchMode, FGraphNode* GCObjReferencerNode);
-	/** Builds reference chains */
-	static void BuildReferenceChains(TConstArrayView<FGraphNode*> TargetNodes, TArray<FReferenceChain*>& AllChains, EReferenceChainSearchMode SearchMode, FGraphNode* GCObjReferencerNode);
-	/** Builds reference chains for direct references only */
-	static void BuildReferenceChainsForDirectReferences(TConstArrayView<FGraphNode*> TargetNodes, TArray<FReferenceChain*>& AllChains, EReferenceChainSearchMode SearchMode, FGraphNode* GCObjReferencerNode);
-	/** Leaves only chains with unique root objects */
-	static void RemoveChainsWithDuplicatedRoots(TArray<FReferenceChain*>& AllChains, FGraphNode* GCObjReferencerNode);
-	/** Leaves only unique chains */
-	static void RemoveDuplicatedChains(TArray<FReferenceChain*>& AllChains, FGraphNode* GCObjReferencerNode);
-	/** Deduplicates chains that have the same root and first garbage object encountered */
-	static void RemoveDuplicateGarbageChains(TArray<FReferenceChain*>& AllChains, FGraphNode* GCObjReferencerNode);
-
 	/** Returns a string with all flags (we care about) set on an object */
 	static FString GetObjectFlags(FGCObjectInfo* InObject);
 	static FString GetObjectFlags(EInternalObjectFlags InternalFlags, EObjectFlags Flags);
 
 	/** Dumps a reference chain to log */
-	static void DumpChain(FReferenceChainSearch::FReferenceChain* Chain, TFunctionRef<bool(FCallbackParams& Params)> ReferenceCallback, FOutputDevice& Out);
+	static void DumpChain(FReferenceChainSearch::FReferenceChain* Chain, TFunctionRef<bool(FCallbackParams& Params)> ReferenceCallback,
+		TMap<uint64, FString>& CallstackCache, FOutputDevice& Out);
 };
