Fix building when ENABLE_GC_HISTORY is disabled.

diff --git a/Engine/Source/Runtime/CoreUObject/Private/UObject/ReferenceChainSearch.cpp b/Engine/Source/Runtime/CoreUObject/Private/UObject/ReferenceChainSearch.cpp
--- a/Engine/Source/Runtime/CoreUObject/Private/UObject/ReferenceChainSearch.cpp
+++ b/Engine/Source/Runtime/CoreUObject/Private/UObject/ReferenceChainSearch.cpp
@@ -108,6 +108,7 @@ namespace UE::ReferenceChainSearch
 		void GatherReferenceInfo(TMap<FVertex, TMap<FVertex, FReferenceChainSearch::FObjectReferenceInfo>>& InOutReferenceInfo);
 	};
 
+#if ENABLE_GC_HISTORY
 	struct FPolicyGCHistory
 	{
 		using ObjectType = FGCObjectInfo*;
@@ -204,6 +205,7 @@ namespace UE::ReferenceChainSearch
 		 */
 		void GatherReferenceInfo(TMap<FVertex, TMap<FVertex, FReferenceChainSearch::FObjectReferenceInfo>>& InOutReferenceInfo);
 	};
+#endif // ENABLE_GC_HISTORY
 
 	template<typename Derived>
 	struct TReferenceSearchBase
@@ -641,6 +643,7 @@ namespace UE::ReferenceChainSearch
 		}
 	}
 
+#if ENABLE_GC_HISTORY
 	void PerformInitialGatherFromGCHistory(const FPolicyGCHistory& Policy, FGraph& OutGraph)
 	{
 		TRACE_CPUPROFILER_EVENT_SCOPE(UE::ReferenceChainSearch::PerformInitialGatherFromGCHistory);
@@ -672,6 +675,7 @@ namespace UE::ReferenceChainSearch
 
 		OutGraph = UE::Graph::ConstructTransposeGraph(TempGraph.EdgeLists);
 	}
+#endif // ENABLE_GC_HISTORY
 
 	// Temporary storage for a reference chain before gathering information
 	struct FGraphPath
@@ -1079,6 +1083,7 @@ namespace UE::ReferenceChainSearch
 		}
 	}
 
+#if ENABLE_GC_HISTORY
 	void FPolicyGCHistory::GatherReferenceInfo(TMap<FVertex, TMap<FVertex, FReferenceChainSearch::FObjectReferenceInfo>>& InOutReferenceInfo)
 	{
 		FGCObjectInfo* GCObjReferencerInfo = Snapshot.ObjectToInfoMap.FindRef(FGCObject::GGCObjectReferencer);
@@ -1113,6 +1118,7 @@ namespace UE::ReferenceChainSearch
 			}
 		}
 	}
+#endif // ENABLE_GC_HISTORY
 
 	template<typename PolicyType>
 	void PopulateReferenceInfo(PolicyType& Policy,
