Fix truncation warning when calculating FPackageFileSummary::TotalHeaderSize

#rnx

- CL 23489396 introduced a potential truncation error when saving a package if TotalHeaderSize is larger than int32, similar to the truncation fixed in CL 20257672.
- Moved the fix from 20257672 to a utility function that can be called elsewhere (ideally we'd make more members in FPackageFileSummary private but as its in a public header that is difficult)
- As with the previous fix this is probably more work that we need to do as it is very unlikely that TotalHeaderSize will ever be > MAX_int32 and if we ever thought that was likely then we'd upgrade the type in FPackageFileSummary, HOWEVER if we do actually exceed MAX_int32 then it is better to catch it when saving the package rather than save bad data that will be hard to debug.

diff --git a/Engine/Source/Runtime/CoreUObject/Private/UObject/SavePackage2.cpp b/Engine/Source/Runtime/CoreUObject/Private/UObject/SavePackage2.cpp
--- a/Engine/Source/Runtime/CoreUObject/Private/UObject/SavePackage2.cpp
+++ b/Engine/Source/Runtime/CoreUObject/Private/UObject/SavePackage2.cpp
@@ -1726,6 +1726,33 @@ int64 WriteObjectDataResources(TArray<FObjectDataResource>& DataResources, FStru
 	return Linker.Tell() - Linker.Summary.DataResourceOffset;
 }
 
+/**
+ * Utility for safely setting the TotalHeaderSize member of FPackageFileSummary with a int64 value.
+ * 
+ * FPackageFileSummary uses int32 for a lot of offsets but a lot of our package writing code is
+ * capable of handling files that exceed MAX_int32 so the final size is calculated as a int64.
+ * If this value is truncated when storing in FPackageFileSummary then the package will not read
+ * in correctly and most likely cause a crash. Rather than allowing the user to save bad data we 
+ * can use this utility to catch the error and log it so that the user can take action.
+ */
+ESavePackageResult SetSummaryTotalHeaderSize(FSaveContext& SaveContext, int64 TotalHeaderSize)
+{
+	FLinkerSave* Linker = SaveContext.GetLinker();
+
+	if (TotalHeaderSize <= MAX_int32)
+	{
+		Linker->Summary.TotalHeaderSize = (int32)TotalHeaderSize;
+		return ESavePackageResult::Success;
+	}
+	else
+	{
+		UE_LOG(LogSavePackage, Error, TEXT("Package header for '%s' is too large (%" UINT64_FMT " bytes), some package file summary offsets will be truncated when stored as a int32"),
+			*SaveContext.GetPackage()->GetName(), TotalHeaderSize);
+
+		return ESavePackageResult::Error;
+	}
+}
+
 ESavePackageResult WritePackageHeader(FStructuredArchive::FRecord& StructuredArchiveRoot, FSaveContext& SaveContext)
 {
 	FLinkerSave* Linker = SaveContext.GetLinker();
@@ -1869,16 +1896,13 @@ ESavePackageResult WritePackageHeader(FStructuredArchive::FRecord& StructuredArc
 		SCOPED_SAVETIMER(UPackage_Save_PreloadDependencies);
 		SavePreloadDependencies(StructuredArchiveRoot, SaveContext);
 	}
-	Linker->Summary.TotalHeaderSize = (int32)Linker->Tell();
-
+	
 	// Rather than check if an offset is truncated every time we assign one, we can just check the final TotalHeaderSize to see if it is truncated.
 	// Checking every time an offset is assigned would let us fail quicker but a) relies that new code follows the convention b) bloats the code a fair bit.
-	if (Linker->Tell() > MAX_int32)
+	const ESavePackageResult Result = SetSummaryTotalHeaderSize(SaveContext, Linker->Tell());
+	if (Result != ESavePackageResult::Success)
 	{
-		UE_LOG(LogSavePackage, Error, TEXT("Package header for '%s' is too large (%" UINT64_FMT " bytes), some package file summary offsets will be truncated when stored as a int32"),
-			*SaveContext.GetPackage()->GetName(), Linker->Tell());
-
-		return ESavePackageResult::Error;
+		return Result;
 	}
 
 	return ReturnSuccessOrCancel();
@@ -2717,7 +2741,7 @@ ESavePackageResult SaveHarvestedRealms(FSaveContext& SaveContext, ESaveRealm Har
 
 		if (Linker.IsCooking())
 		{
-			// Write the exports into a seperate archive
+			// Write the exports into a separate archive
 			TUniquePtr<FLargeMemoryWriter> ExportsArchive = SaveContext.GetPackageWriter()->CreateLinkerExportsArchive(
 				SaveContext.GetPackage()->GetFName(), SaveContext.GetAsset());
 			ExportsArchive->SetSerializeContext(SaveContext.GetSerializeContext());
@@ -2729,7 +2753,13 @@ ESavePackageResult SaveHarvestedRealms(FSaveContext& SaveContext, ESaveRealm Har
 				const int64 DataResourceSize = WriteObjectDataResources(Linker.DataResourceMap, StructuredArchiveRoot, SaveContext);
 				check(DataResourceSize >= 0);
 
-				Linker.Summary.TotalHeaderSize += DataResourceSize;
+				// Check to make sure that the package header is not too large
+				SaveContext.Result = SetSummaryTotalHeaderSize(SaveContext, (int64)Linker.Summary.TotalHeaderSize + DataResourceSize);
+				if (SaveContext.Result != ESavePackageResult::Success)
+				{
+					return SaveContext.Result;
+				}
+
 				{
 					// Disables writing stack trace data when appending the exports data
 					FArchiveStackTraceDisabledScope _; 
