Change reinterpret cast from UObjectBase* to UObject* to static cast.

diff --git a/Engine/Source/Runtime/CoreUObject/Private/UObject/ReferenceChainSearch.cpp b/Engine/Source/Runtime/CoreUObject/Private/UObject/ReferenceChainSearch.cpp
--- a/Engine/Source/Runtime/CoreUObject/Private/UObject/ReferenceChainSearch.cpp
+++ b/Engine/Source/Runtime/CoreUObject/Private/UObject/ReferenceChainSearch.cpp
@@ -52,7 +52,7 @@ namespace UE::ReferenceChainSearch
 		}
 		ObjectType VertexToObject(FVertex Vertex) const
 		{
-			return reinterpret_cast<UObject*>(GUObjectArray.IndexToObject(Vertex)->Object);
+			return static_cast<UObject*>(GUObjectArray.IndexToObject(Vertex)->Object);
 		}
 
 		bool IsIn(ConstObjectType Outer, ConstObjectType PotentialInner) const
