Adjust SetLinker error message to be more accurate
Move the flush of the reinstancing queue back from the post load node to the end of game thread tick or flush for zenloader

diff --git a/Engine/Source/Runtime/CoreUObject/Private/Serialization/AsyncLoading2.cpp b/Engine/Source/Runtime/CoreUObject/Private/Serialization/AsyncLoading2.cpp
--- a/Engine/Source/Runtime/CoreUObject/Private/Serialization/AsyncLoading2.cpp
+++ b/Engine/Source/Runtime/CoreUObject/Private/Serialization/AsyncLoading2.cpp
@@ -6389,11 +6389,6 @@ EEventLoadNodeExecutionResult FAsyncPackage2::Event_PostLoadExportBundle(FAsyncL
 
 	FAsyncPackageScope2 PackageScope(Package);
 
-#if WITH_EDITOR
-	// This needs to happen after loading new blueprints in the editor, and this is handled in EndLoad for synchronous loads
-	FBlueprintSupport::FlushReinstancingQueue();
-#endif
-
 #if ALT2_ENABLE_LINKERLOAD_SUPPORT
 	if (Package->LinkerLoadState.IsSet())
 	{
@@ -7572,7 +7567,7 @@ void FAsyncLoadingThread2::CollectUnreachableObjects(
 					bLinkerEntryWasInvalid = ObjExport.bExportLoadFailed;
 				}
 #endif // WITH_EDITORONLY_DATA
-				Object->SetLinker(NULL, INDEX_NONE);
+				Object->SetLinker(nullptr, INDEX_NONE);
 #if WITH_EDITORONLY_DATA
 				if (bLinkerEntryWasInvalid)
 				{
@@ -8075,6 +8070,9 @@ void FAsyncLoadingThread2::ConditionalProcessEditorCallbacks()
 	{
 		return;
 	}
+
+	FBlueprintSupport::FlushReinstancingQueue();
+
 	while (!EditorCompletedUPackages.IsEmpty() || !EditorLoadedAssets.IsEmpty())
 	{
 		TArray<UObject*> LocalEditorLoadedAssets;
diff --git a/Engine/Source/Runtime/CoreUObject/Private/UObject/UObjectLinker.cpp b/Engine/Source/Runtime/CoreUObject/Private/UObject/UObjectLinker.cpp
--- a/Engine/Source/Runtime/CoreUObject/Private/UObject/UObjectLinker.cpp
+++ b/Engine/Source/Runtime/CoreUObject/Private/UObject/UObjectLinker.cpp
@@ -116,10 +116,10 @@ void UObject::SetLinker( FLinkerLoad* LinkerLoad, int32 LinkerIndex, bool bShoul
 	if( Existing.Linker && bShouldDetachExisting )
 	{
 		UE_CLOG(HasAnyFlags(RF_NeedLoad|RF_NeedPostLoad), LogUObjectLinker, Error,
-			TEXT("Detaching from existing linker for %s while object %s needs loading from linker %s."),
+			TEXT("Detaching from existing linker %s while object %s needs loading. Setting linker to %s."),
 			*Existing.Linker->GetArchiveName(),
 			*GetFullName(),
-			LinkerLoad ? *LinkerLoad->GetDebugName() : TEXT(""));
+			LinkerLoad ? *LinkerLoad->GetDebugName() : TEXT("nullptr"));
 		check(Existing.Linker->ExportMap[Existing.LinkerIndex].Object!=nullptr);
 		check(Existing.Linker->ExportMap[Existing.LinkerIndex].Object==this);
 		Existing.Linker->ExportMap[Existing.LinkerIndex].ResetObject();
