Fix for recursive async loading in legacy loader
Do not reset load indices in FinishObjects

diff --git a/Engine/Source/Runtime/CoreUObject/Private/Serialization/AsyncLoading.cpp b/Engine/Source/Runtime/CoreUObject/Private/Serialization/AsyncLoading.cpp
--- a/Engine/Source/Runtime/CoreUObject/Private/Serialization/AsyncLoading.cpp
+++ b/Engine/Source/Runtime/CoreUObject/Private/Serialization/AsyncLoading.cpp
@@ -5919,7 +5919,7 @@ EAsyncPackageState::Type FAsyncPackage::TickAsyncPackage(bool InbUseTimeLimit, b
 	// We want this check only with EDL enabled
 	check(!GEventDrivenLoaderEnabled || int32(AsyncPackageLoadingState) > int32(EAsyncPackageLoadingState::ProcessNewImportsAndExports));
 	// @note FH: We shouldn't be ticking this package from the async thread once we are have finished loading it
-	checkSlow(!HasFinishedLoading());
+	check(!HasFinishedLoading());
 
 	ReentryCount++;
 
@@ -5955,7 +5955,6 @@ EAsyncPackageState::Type FAsyncPackage::TickAsyncPackage(bool InbUseTimeLimit, b
 	}
 
 	FAsyncPackageScope PackageScope(this);
-
 	// Make sure we finish our work if there's no time limit. The loop is required as PostLoad
 	// might cause more objects to be loaded in which case we need to Preload them again.
 	do
@@ -5986,14 +5985,14 @@ EAsyncPackageState::Type FAsyncPackage::TickAsyncPackage(bool InbUseTimeLimit, b
 			}
 
 			// Load imports from linker import table asynchronously.
-			if (LoadingState == EAsyncPackageState::Complete)
+			if (LoadingState == EAsyncPackageState::Complete && !bLoadHasFinished)
 			{
 				SCOPED_LOADTIMER(Package_LoadImports);
 				LoadingState = LoadImports();
 			}
 
 			// Create imports from linker import table.
-			if (LoadingState == EAsyncPackageState::Complete)
+			if (LoadingState == EAsyncPackageState::Complete && !bLoadHasFinished)
 			{
 				SCOPED_LOADTIMER(Package_CreateImports);
 				LoadingState = CreateImports();
@@ -6001,7 +6000,7 @@ EAsyncPackageState::Type FAsyncPackage::TickAsyncPackage(bool InbUseTimeLimit, b
 
 #if WITH_EDITORONLY_DATA
 			// Create and preload the package meta-data
-			if (LoadingState == EAsyncPackageState::Complete)
+			if (LoadingState == EAsyncPackageState::Complete && !bLoadHasFinished)
 			{
 				SCOPED_LOADTIMER(Package_CreateMetaData);
 				LoadingState = CreateMetaData();
@@ -6009,20 +6008,20 @@ EAsyncPackageState::Type FAsyncPackage::TickAsyncPackage(bool InbUseTimeLimit, b
 #endif // WITH_EDITORONLY_DATA
 
 			// Create exports from linker export table and also preload them.
-			if (LoadingState == EAsyncPackageState::Complete)
+			if (LoadingState == EAsyncPackageState::Complete && !bLoadHasFinished)
 			{
 				SCOPED_LOADTIMER(Package_CreateExports);
 				LoadingState = CreateExports();
 			}
 
 			// Call Preload on the linker for all loaded objects which causes actual serialization.
-			if (LoadingState == EAsyncPackageState::Complete)
+			if (LoadingState == EAsyncPackageState::Complete && !bLoadHasFinished)
 			{
 				SCOPED_LOADTIMER(Package_PreLoadObjects);
 				LoadingState = PreLoadObjects();
 			}
 
-			if (LoadingState == EAsyncPackageState::Complete || bLoadHasFailed)
+			if ((LoadingState == EAsyncPackageState::Complete || bLoadHasFailed) && !bLoadHasFinished)
 			{
 				const bool bInternalCallbacks = true;
 				CallCompletionCallbacks(bInternalCallbacks, bLoadHasFailed ? EAsyncLoadingResult::Failed : EAsyncLoadingResult::Succeeded);
@@ -6042,7 +6041,7 @@ EAsyncPackageState::Type FAsyncPackage::TickAsyncPackage(bool InbUseTimeLimit, b
 			}
 		} // !GEventDrivenLoaderEnabled
 
-		if (LoadingState == EAsyncPackageState::Complete && !bLoadHasFailed)
+		if (LoadingState == EAsyncPackageState::Complete && !bLoadHasFailed && !bLoadHasFinished)
 		{
 			SCOPED_LOADTIMER(Package_ExternalReadDependencies);
 			LoadingState = FinishExternalReadDependencies();
@@ -6050,7 +6049,7 @@ EAsyncPackageState::Type FAsyncPackage::TickAsyncPackage(bool InbUseTimeLimit, b
 
 		// Call PostLoad on objects, this could cause new objects to be loaded that require
 		// another iteration of the PreLoad loop.
-		if (LoadingState == EAsyncPackageState::Complete && !bLoadHasFailed)
+		if (LoadingState == EAsyncPackageState::Complete && !bLoadHasFailed && !bLoadHasFinished)
 		{
 			SCOPED_LOADTIMER(Package_PostLoadObjects);
 			LoadingState = PostLoadObjects();
@@ -6079,12 +6078,12 @@ EAsyncPackageState::Type FAsyncPackage::TickAsyncPackage(bool InbUseTimeLimit, b
 		LinkerRoot->MarkAsFullyLoaded();
 	}
 
+	// Mark this package as loaded if everything completed. leave it true if it was already marked as such through recursivity
+	bLoadHasFinished = bLoadHasFinished || (LoadingState == EAsyncPackageState::Complete);
 	// We can't have a reference to a UObject.
 	LastObjectWorkWasPerformedOn = nullptr;
 	// Reset type of work performed.
 	LastTypeOfWorkPerformed = nullptr;
-	// Mark this package as loaded if everything completed.
-	bLoadHasFinished = (LoadingState == EAsyncPackageState::Complete);
 
 	if (bLoadHasFinished && GEventDrivenLoaderEnabled)
 	{
@@ -7217,10 +7216,6 @@ EAsyncPackageState::Type FAsyncPackage::FinishObjects()
 
 	// Simulate what EndLoad does.
 	FLinkerManager::Get().DissociateImportsAndForcedExports(); //@todo: this should be avoidable
-	PreLoadIndex = 0;
-	PreLoadSortIndex = 0;
-	PostLoadIndex = 0;
-	FinishExternalReadDependenciesIndex = 0;
 
 	// Keep the linkers to close until we finish loading and it's safe to close them too
 	LoadContext->MoveDelayedLinkerClosePackages(DelayedLinkerClosePackages);
diff --git a/Engine/Source/Runtime/CoreUObject/Private/Serialization/AsyncLoading.h b/Engine/Source/Runtime/CoreUObject/Private/Serialization/AsyncLoading.h
--- a/Engine/Source/Runtime/CoreUObject/Private/Serialization/AsyncLoading.h
+++ b/Engine/Source/Runtime/CoreUObject/Private/Serialization/AsyncLoading.h
@@ -478,7 +478,7 @@ struct FAsyncPackage : public FGCObject
 	/** Returns true if this package is already being loaded in the current callstack */
 	bool IsBeingProcessedRecursively() const
 	{
-		return ReentryCount > 1;
+		return ReentryCount > 0;
 	}
 
 	/** FGCObject Interface */
