ZenLoader: Use the same package dependency waiting logic for cooked and uncooked packages, i.e. always wait for all dependencies of a package to have processed their headers before starting to process its exports
#rnx

diff --git a/Engine/Source/Runtime/CoreUObject/Private/Serialization/AsyncLoading2.cpp b/Engine/Source/Runtime/CoreUObject/Private/Serialization/AsyncLoading2.cpp
--- a/Engine/Source/Runtime/CoreUObject/Private/Serialization/AsyncLoading2.cpp
+++ b/Engine/Source/Runtime/CoreUObject/Private/Serialization/AsyncLoading2.cpp
@@ -212,6 +212,10 @@ void FindAllRuntimeScriptPackages(TArray<UPackage*>& OutPackages)
 #define ALT2_LOG_VERBOSE DO_CHECK
 #endif
 
+#ifndef ALT2_ENABLE_PACKAGE_DEPENDENCY_DEBUGGING
+#define ALT2_ENABLE_PACKAGE_DEPENDENCY_DEBUGGING 0
+#endif
+
 static TSet<FPackageId> GAsyncLoading2_DebugPackageIds;
 static FString GAsyncLoading2_DebugPackageNamesString;
 static TSet<FPackageId> GAsyncLoading2_VerbosePackageIds;
@@ -1898,7 +1902,8 @@ enum class EAsyncPackageLoadingState2 : uint8
 	NewPackage,
 	WaitingForIo,
 	ProcessPackageSummary,
-	SetupDependencies,
+	WaitingForDependencies,
+	DependenciesReady,
 	ProcessExportBundles,
 	WaitingForExternalReads,
 	ExportsDone,
@@ -2204,7 +2209,7 @@ uint32 FAsyncLoadingThreadState2::TlsSlot;
 enum EEventLoadNode2 : uint8
 {
 	Package_ProcessSummary,
-	Package_SetupDependencies,
+	Package_DependenciesReady,
 	Package_ExportsSerialized,
 	Package_NumPhases,
 
@@ -2715,9 +2720,7 @@ private:
 			PackageToRemoveState.WaitingForPackage = nullptr;
 		}
 	};
-#if ALT2_ENABLE_LINKERLOAD_SUPPORT
 	FAllDependenciesState		AllDependenciesSetupState;
-#endif
 	FAllDependenciesState		AllDependenciesFullyLoadedState;
 
 	/** True if our load has failed */
@@ -2751,7 +2754,7 @@ public:
 
 	static EEventLoadNodeExecutionResult Event_ProcessExportBundle(FAsyncLoadingThreadState2& ThreadState, FAsyncPackage2* Package, int32 ExportBundleIndex);
 	static EEventLoadNodeExecutionResult Event_ProcessPackageSummary(FAsyncLoadingThreadState2& ThreadState, FAsyncPackage2* Package, int32);
-	static EEventLoadNodeExecutionResult Event_SetupDependencies(FAsyncLoadingThreadState2& ThreadState, FAsyncPackage2* Package, int32);
+	static EEventLoadNodeExecutionResult Event_DependenciesReady(FAsyncLoadingThreadState2& ThreadState, FAsyncPackage2* Package, int32);
 	static EEventLoadNodeExecutionResult Event_ExportsDone(FAsyncLoadingThreadState2& ThreadState, FAsyncPackage2* Package, int32);
 	static EEventLoadNodeExecutionResult Event_PostLoadExportBundle(FAsyncLoadingThreadState2& ThreadState, FAsyncPackage2* Package, int32 ExportBundleIndex);
 	static EEventLoadNodeExecutionResult Event_DeferredPostLoadExportBundle(FAsyncLoadingThreadState2& ThreadState, FAsyncPackage2* Package, int32 ExportBundleIndex);
@@ -2783,14 +2786,17 @@ private:
 	void CreateExportBundleNodes(const FAsyncLoadEventSpec* EventSpecs);
 	void SetupSerializedArcs(const FAsyncPackageHeaderData& Header, const TArrayView<FEventLoadNode2>& ExportBundleNodes, const TArrayView<FAsyncPackage2*>& ImportedAsyncPackages);
 	void SetupScriptDependencies();
-	bool HaveAllDependenciesReachedStateDebug(FAsyncLoadingThreadState2& ThreadState, FAsyncPackage2* Package, TSet<FAsyncPackage2*>& VisitedPackages, EAsyncPackageLoadingState2 WaitForPackageState, bool bExcludeZenPackages);
-	bool HaveAllDependenciesReachedState(FAsyncLoadingThreadState2& ThreadState, FAllDependenciesState FAsyncPackage2::* StateMemberPtr, EAsyncPackageLoadingState2 WaitForPackageState, uint32 CurrentTick, bool bExcludeZenPackages);
-	void UpdateDependenciesStateRecursive(FAsyncLoadingThreadState2& ThreadState, FAllDependenciesState FAsyncPackage2::* StateMemberPtr, EAsyncPackageLoadingState2 WaitForPackageState, uint32 CurrentTick, FAsyncPackage2* Root, bool bExcludeZenPackages);
-	void WaitForAllDependenciesToReachState(FAsyncLoadingThreadState2& ThreadState, FAllDependenciesState FAsyncPackage2::* StateMemberPtr, EAsyncPackageLoadingState2 WaitForPackageState, uint32& CurrentTickVariable, bool bExcludeZenPackages, TFunctionRef<void(FAsyncPackage2*)> OnStateReached);
+#if ALT2_ENABLE_PACKAGE_DEPENDENCY_DEBUGGING
+	bool HasDependencyToPackageDebug(FAsyncPackage2* Package);
+	void CheckThatAllDependenciesHaveReachedStateDebug(FAsyncLoadingThreadState2& ThreadState, EAsyncPackageLoadingState2 PackageState, EAsyncPackageLoadingState2 PackageStateForCircularDependencies);
+#endif
+	bool HaveAllDependenciesReachedState(FAsyncLoadingThreadState2& ThreadState, FAllDependenciesState FAsyncPackage2::* StateMemberPtr, EAsyncPackageLoadingState2 WaitForPackageState, uint32 CurrentTick);
+	void UpdateDependenciesStateRecursive(FAsyncLoadingThreadState2& ThreadState, FAllDependenciesState FAsyncPackage2::* StateMemberPtr, EAsyncPackageLoadingState2 WaitForPackageState, uint32 CurrentTick, FAsyncPackage2* Root);
+	void WaitForAllDependenciesToReachState(FAsyncLoadingThreadState2& ThreadState, FAllDependenciesState FAsyncPackage2::* StateMemberPtr, EAsyncPackageLoadingState2 WaitForPackageState, uint32& CurrentTickVariable, TFunctionRef<void(FAsyncPackage2*)> OnStateReached);
+	void ConditionalBeginProcessPackageExports(FAsyncLoadingThreadState2& ThreadState);
 	void ConditionalFinishLoading(FAsyncLoadingThreadState2& ThreadState);
 
 #if ALT2_ENABLE_LINKERLOAD_SUPPORT
-	void ConditionalProcessLinkerLoadPackageImportsAndExports(FAsyncLoadingThreadState2& ThreadState);
 	EEventLoadNodeExecutionResult ProcessLinkerLoadPackageSummary(FAsyncLoadingThreadState2& ThreadState);
 	bool ProcessLinkerLoadPackageImports(FAsyncLoadingThreadState2& ThreadState);
 	bool ProcessLinkerLoadPackageExports(FAsyncLoadingThreadState2& ThreadState);
@@ -2988,9 +2994,7 @@ private:
 	TUniquePtr<FAsyncLoadingThreadState2> GameThreadState;
 	TUniquePtr<FAsyncLoadingThreadState2> AsyncLoadingThreadState;
 
-#if ALT2_ENABLE_LINKERLOAD_SUPPORT
-	uint32 ConditionalProcessLinkerLoadPackageExportsTick = 0;
-#endif
+	uint32 ConditionalBeginProcessExportsTick = 0;
 	uint32 ConditionalFinishLoadingTick = 0;
 
 public:
@@ -3227,12 +3231,9 @@ public:
 		for (UClass* Class : Classes)
 		{
 			TArray<FEventLoadNode2*>& Nodes = PendingCDOs.FindOrAdd(Class);
-			for (int32 ExportBundleIndex = 0; ExportBundleIndex < Package->Data.ExportBundleCount; ++ExportBundleIndex)
-			{
-				FEventLoadNode2& BundleNode = Package->GetExportBundleNode(ExportBundle_Process, ExportBundleIndex);
-				BundleNode.AddBarrier();
-				Nodes.Add(&BundleNode);
-			}
+			FEventLoadNode2& Node = Package->GetPackageNode(Package_DependenciesReady);
+			Node.AddBarrier();
+			Nodes.Add(&Node);
 		}
 	}
 
@@ -4721,13 +4722,7 @@ void FAsyncPackage2::ImportPackagesRecursiveInner(FAsyncLoadingThreadState2& Thr
 		{
 			// When importing a linker load package from a zen package or vice versa we need to wait for all the exports in the imported package to be created
 			// and serialized before we can start processing our own exports
-			GetPackageNode(Package_SetupDependencies).DependsOn(&ImportedPackage->GetPackageNode(Package_ExportsSerialized));
-		}
-		else if (bIsZenPackage)
-		{
-			// When importing a zen package from another zen package we need the summary of the imported package to be ready before we can hook up our export
-			// bundle node dependenices to the imported package
-			GetPackageNode(Package_SetupDependencies).DependsOn(&ImportedPackage->GetPackageNode(Package_ProcessSummary));
+			GetPackageNode(Package_DependenciesReady).DependsOn(&ImportedPackage->GetPackageNode(Package_ExportsSerialized));
 		}
 
 		if (bInserted)
@@ -5025,18 +5020,6 @@ static void ReadAsyncPackageHeader(FAsyncPackageSerializationState& Serializatio
 }
 
 #if ALT2_ENABLE_LINKERLOAD_SUPPORT
-void FAsyncPackage2::ConditionalProcessLinkerLoadPackageImportsAndExports(FAsyncLoadingThreadState2& ThreadState)
-{
-	TRACE_CPUPROFILER_EVENT_SCOPE(ConditionalProcessLinkerLoadPackageExports);
-	constexpr bool bExcludeZenPackages = true;
-	WaitForAllDependenciesToReachState(ThreadState, &FAsyncPackage2::AllDependenciesSetupState, EAsyncPackageLoadingState2::ProcessExportBundles, AsyncLoadingThread.ConditionalProcessLinkerLoadPackageExportsTick, bExcludeZenPackages,
-		[&ThreadState](FAsyncPackage2* Package)
-		{
-			Package->AsyncPackageLoadingState = EAsyncPackageLoadingState2::ProcessExportBundles;
-			Package->GetExportBundleNode(ExportBundle_Process, 0).ReleaseBarrier(&ThreadState);
-		});
-}
-
 EEventLoadNodeExecutionResult FAsyncPackage2::ProcessLinkerLoadPackageSummary(FAsyncLoadingThreadState2& ThreadState)
 {
 	TRACE_CPUPROFILER_EVENT_SCOPE(ProcessLinkerLoadPackageSummary);
@@ -5218,8 +5201,12 @@ EEventLoadNodeExecutionResult FAsyncPackage2::ProcessLinkerLoadPackageSummary(FA
 		PackageRef.SetHasFailed();
 	}
 
-	AsyncPackageLoadingState = EAsyncPackageLoadingState2::SetupDependencies;
-	GetPackageNode(Package_SetupDependencies).ReleaseBarrier(&ThreadState);
+	AsyncPackageLoadingState = EAsyncPackageLoadingState2::WaitingForDependencies;
+	if (!AsyncLoadingThread.bHasRegisteredAllScriptObjects)
+	{
+		SetupScriptDependencies();
+	}
+	GetPackageNode(Package_DependenciesReady).ReleaseBarrier(&ThreadState);
 	return EEventLoadNodeExecutionResult::Complete;
 }
 
@@ -5241,9 +5228,9 @@ EEventLoadNodeExecutionResult FAsyncPackage2::ProcessLinkerLoadPackageImportsAnd
 		return EEventLoadNodeExecutionResult::Timeout;
 	}
 
+	check(AsyncPackageLoadingState == EAsyncPackageLoadingState2::ProcessExportBundles);
 	if (ExternalReadDependencies.Num() == 0)
 	{
-		AsyncPackageLoadingState = EAsyncPackageLoadingState2::ProcessExportBundles;
 		GetPackageNode(Package_ExportsSerialized).ReleaseBarrier(&ThreadState);
 	}
 	else
@@ -5256,7 +5243,7 @@ EEventLoadNodeExecutionResult FAsyncPackage2::ProcessLinkerLoadPackageImportsAnd
 
 bool FAsyncPackage2::ProcessLinkerLoadPackageImports(FAsyncLoadingThreadState2& ThreadState)
 {
-	check(AsyncPackageLoadingState == EAsyncPackageLoadingState2::ProcessExportBundles);
+	check(AsyncPackageLoadingState >= EAsyncPackageLoadingState2::DependenciesReady);
 	if (LinkerLoadState->bIsCurrentlyProcessingImports)
 	{
 		return true;
@@ -5268,15 +5255,15 @@ bool FAsyncPackage2::ProcessLinkerLoadPackageImports(FAsyncLoadingThreadState2&
 		FAsyncPackage2* ImportedPackage = Data.ImportedAsyncPackages[LinkerLoadState->ProcessingImportedPackageIndex];
 		if (ImportedPackage && ImportedPackage->LinkerLoadState.IsSet())
 		{
-			check(ImportedPackage->AsyncPackageLoadingState >= EAsyncPackageLoadingState2::ProcessExportBundles);
-			if (ImportedPackage->AsyncPackageLoadingState == EAsyncPackageLoadingState2::ProcessExportBundles)
+			check(ImportedPackage->AsyncPackageLoadingState >= EAsyncPackageLoadingState2::DependenciesReady);
+			if (ImportedPackage->AsyncPackageLoadingState <= EAsyncPackageLoadingState2::ProcessExportBundles)
 			{
 				if (!ImportedPackage->ProcessLinkerLoadPackageImports(ThreadState))
 				{
 					return false;
 				}
 			}
-			if (ImportedPackage->AsyncPackageLoadingState == EAsyncPackageLoadingState2::ProcessExportBundles)
+			if (ImportedPackage->AsyncPackageLoadingState <= EAsyncPackageLoadingState2::ProcessExportBundles)
 			{
 				if (!ImportedPackage->ProcessLinkerLoadPackageExports(ThreadState))
 				{
@@ -5341,7 +5328,7 @@ bool FAsyncPackage2::ProcessLinkerLoadPackageImports(FAsyncLoadingThreadState2&
 
 bool FAsyncPackage2::ProcessLinkerLoadPackageExports(FAsyncLoadingThreadState2& ThreadState)
 {
-	check(AsyncPackageLoadingState == EAsyncPackageLoadingState2::ProcessExportBundles);
+	check(AsyncPackageLoadingState >= EAsyncPackageLoadingState2::DependenciesReady);
 
 	// Create exports
 	const int32 ExportCount = LinkerLoadState->Linker->ExportMap.Num();
@@ -5558,51 +5545,23 @@ EEventLoadNodeExecutionResult FAsyncPackage2::Event_ProcessPackageSummary(FAsync
 		TRACE_LOADTIME_PACKAGE_SUMMARY(Package, Package->HeaderData.PackageName, Package->HeaderData.PackageSummary->HeaderSize, Package->HeaderData.ImportMap.Num(), Package->HeaderData.ExportMap.Num());
 	}
 
-	Package->AsyncPackageLoadingState = EAsyncPackageLoadingState2::SetupDependencies;
-	Package->GetPackageNode(Package_SetupDependencies).ReleaseBarrier(&ThreadState);
+	Package->AsyncPackageLoadingState = EAsyncPackageLoadingState2::WaitingForDependencies;
+	if (!Package->AsyncLoadingThread.bHasRegisteredAllScriptObjects)
+	{
+		Package->SetupScriptDependencies();
+	}
+	Package->GetPackageNode(Package_DependenciesReady).ReleaseBarrier(&ThreadState);
 	return EEventLoadNodeExecutionResult::Complete;
 }
 
-EEventLoadNodeExecutionResult FAsyncPackage2::Event_SetupDependencies(FAsyncLoadingThreadState2& ThreadState, FAsyncPackage2* Package, int32)
+EEventLoadNodeExecutionResult FAsyncPackage2::Event_DependenciesReady(FAsyncLoadingThreadState2& ThreadState, FAsyncPackage2* Package, int32)
 {
-	TRACE_CPUPROFILER_EVENT_SCOPE(Event_SetupDependencies);
+	TRACE_CPUPROFILER_EVENT_SCOPE(Event_DependenciesReadt);
 	UE_ASYNC_PACKAGE_DEBUG(Package->Desc);
-	check(Package->AsyncPackageLoadingState == EAsyncPackageLoadingState2::SetupDependencies);
+	check(Package->AsyncPackageLoadingState == EAsyncPackageLoadingState2::WaitingForDependencies);
 
-#if ALT2_ENABLE_LINKERLOAD_SUPPORT
-	if (Package->LinkerLoadState.IsSet())
-	{
-		if (!Package->AsyncLoadingThread.bHasRegisteredAllScriptObjects)
-		{
-			Package->SetupScriptDependencies();
-		}
-		Package->AsyncPackageLoadingState = EAsyncPackageLoadingState2::ProcessExportBundles;
-		Package->ConditionalProcessLinkerLoadPackageImportsAndExports(ThreadState);
-		return EEventLoadNodeExecutionResult::Complete;
-	}
-#endif
-
-	if (!Package->bLoadHasFailed)
-	{
-		if (!Package->AsyncLoadingThread.bHasRegisteredAllScriptObjects)
-		{
-			Package->SetupScriptDependencies();
-		}
-		TArrayView<FEventLoadNode2> ExportBundleNodesView = Package->Data.ExportBundleNodes;
-		TArrayView<FAsyncPackage2*> ImportedAsyncPackagesView = Package->Data.ImportedAsyncPackages;
-		Package->SetupSerializedArcs(Package->HeaderData, ExportBundleNodesView.Left(Package->HeaderData.ExportBundleHeaders.Num() * EEventLoadNode2::ExportBundle_NumPhases), ImportedAsyncPackagesView.Left(Package->HeaderData.ImportedPackageIds.Num()));
-#if WITH_EDITOR
-		if (Package->OptionalSegmentHeaderData.IsSet())
-		{
-			Package->SetupSerializedArcs(*Package->OptionalSegmentHeaderData, ExportBundleNodesView.Right(Package->OptionalSegmentHeaderData->ExportBundleHeaders.Num() * EEventLoadNode2::ExportBundle_NumPhases), ImportedAsyncPackagesView.Right(Package->OptionalSegmentHeaderData->ImportedPackageIds.Num()));
-		}
-#endif
-	}
-	Package->AsyncPackageLoadingState = EAsyncPackageLoadingState2::ProcessExportBundles;
-	for (int32 ExportBundleIndex = 0; ExportBundleIndex < Package->Data.ExportBundleCount; ++ExportBundleIndex)
-	{
-		Package->GetExportBundleNode(ExportBundle_Process, ExportBundleIndex).ReleaseBarrier(&ThreadState);
-	}
+	Package->AsyncPackageLoadingState = EAsyncPackageLoadingState2::DependenciesReady;
+	Package->ConditionalBeginProcessPackageExports(ThreadState);
 	return EEventLoadNodeExecutionResult::Complete;
 }
 
@@ -6170,43 +6129,75 @@ EEventLoadNodeExecutionResult FAsyncPackage2::Event_ExportsDone(FAsyncLoadingThr
 	return EEventLoadNodeExecutionResult::Complete;
 }
 
-bool FAsyncPackage2::HaveAllDependenciesReachedStateDebug(FAsyncLoadingThreadState2& ThreadState, FAsyncPackage2* Package, TSet<FAsyncPackage2*>& VisitedPackages, EAsyncPackageLoadingState2 WaitForPackageState, bool bExcludeZenPackages)
+#if ALT2_ENABLE_PACKAGE_DEPENDENCY_DEBUGGING
+bool FAsyncPackage2::HasDependencyToPackageDebug(FAsyncPackage2* Package)
 {
-	for (FAsyncPackage2* ImportedPackage : Package->Data.ImportedAsyncPackages)
+	TSet<FAsyncPackage2*> Visited;
+	TArray<FAsyncPackage2*> Stack;
+	for (FAsyncPackage2* ImportedPackage : Data.ImportedAsyncPackages)
 	{
 		if (ImportedPackage)
 		{
-#if ALT2_ENABLE_LINKERLOAD_SUPPORT
-			if (bExcludeZenPackages && !ImportedPackage->LinkerLoadState.IsSet())
-			{
-				continue;
-			}
-#endif
-			if (ThreadState.PackagesOnStack.Contains(ImportedPackage))
-			{
-				continue;
-			}
-			if (VisitedPackages.Contains(ImportedPackage))
+			Stack.Push(ImportedPackage);
+		}
+	}
+	while (!Stack.IsEmpty())
+	{
+		FAsyncPackage2* InnerPackage = Stack.Top();
+		Stack.Pop();
+		Visited.Add(InnerPackage);
+		if (InnerPackage == Package)
+		{
+			return true;
+		}
+		for (FAsyncPackage2* ImportedPackage : InnerPackage->Data.ImportedAsyncPackages)
+		{
+			if (ImportedPackage && !Visited.Contains(ImportedPackage))
 			{
-				continue;
+				Stack.Push(ImportedPackage);
 			}
-			VisitedPackages.Add(ImportedPackage);
+		}
+	}
+	return false;
+}
 
-			if (ImportedPackage->AsyncPackageLoadingState < WaitForPackageState)
+void FAsyncPackage2::CheckThatAllDependenciesHaveReachedStateDebug(FAsyncLoadingThreadState2& ThreadState, EAsyncPackageLoadingState2 PackageState, EAsyncPackageLoadingState2 PackageStateForCircularDependencies)
+{
+	TSet<FAsyncPackage2*> Visited;
+	TArray<TTuple<FAsyncPackage2*, TArray<FAsyncPackage2*>>> Stack;
+
+	TArray<FAsyncPackage2*> DependencyChain;
+	DependencyChain.Add(this);
+	Stack.Push(MakeTuple(this, DependencyChain));
+	while (!Stack.IsEmpty())
+	{
+		TTuple<FAsyncPackage2*, TArray<FAsyncPackage2*>> PackageAndDependencyChain = Stack.Top();
+		Stack.Pop();
+		FAsyncPackage2* Package = PackageAndDependencyChain.Get<0>();
+		DependencyChain = PackageAndDependencyChain.Get<1>();
+		Visited.Add(Package);
+		check(Package->AsyncPackageLoadingState >= PackageStateForCircularDependencies);
+		if (Package->AsyncPackageLoadingState < PackageState)
+		{
+			if (!ThreadState.PackagesOnStack.Contains(Package))
 			{
-				return false;
+				bool bHasCircularDependencyToPackage = Package->HasDependencyToPackageDebug(this);
+				check(bHasCircularDependencyToPackage);
 			}
-
-			if (!HaveAllDependenciesReachedStateDebug(ThreadState, ImportedPackage, VisitedPackages, WaitForPackageState, bExcludeZenPackages))
+		}
+		for (FAsyncPackage2* ImportedPackage : Package->Data.ImportedAsyncPackages)
+		{
+			if (ImportedPackage && !Visited.Contains(ImportedPackage))
 			{
-				return false;
+				DependencyChain.Add(ImportedPackage);
+				Stack.Push(MakeTuple(ImportedPackage, DependencyChain));
 			}
 		}
 	}
-	return true;
 }
+#endif
 
-bool FAsyncPackage2::HaveAllDependenciesReachedState(FAsyncLoadingThreadState2& ThreadState, FAllDependenciesState FAsyncPackage2::* StateMemberPtr, EAsyncPackageLoadingState2 WaitForPackageState, uint32 CurrentTick, bool bExcludeZenPackages)
+bool FAsyncPackage2::HaveAllDependenciesReachedState(FAsyncLoadingThreadState2& ThreadState, FAllDependenciesState FAsyncPackage2::* StateMemberPtr, EAsyncPackageLoadingState2 WaitForPackageState, uint32 CurrentTick)
 {
 	FAllDependenciesState& ThisState = this->*StateMemberPtr;
 	if (ThisState.bAllDone)
@@ -6218,12 +6209,12 @@ bool FAsyncPackage2::HaveAllDependenciesReachedState(FAsyncLoadingThreadState2&
 		return false;
 	}
 	ThisState.UpdateTick(CurrentTick);
-	UpdateDependenciesStateRecursive(ThreadState, StateMemberPtr, WaitForPackageState, CurrentTick, this, bExcludeZenPackages);
+	UpdateDependenciesStateRecursive(ThreadState, StateMemberPtr, WaitForPackageState, CurrentTick, this);
 	check(ThisState.bAllDone || (ThisState.WaitingForPackage && ThisState.WaitingForPackage->AsyncPackageLoadingState <= WaitForPackageState));
 	return ThisState.bAllDone;
 }
 
-void FAsyncPackage2::UpdateDependenciesStateRecursive(FAsyncLoadingThreadState2& ThreadState, FAllDependenciesState FAsyncPackage2::* StateMemberPtr, EAsyncPackageLoadingState2 WaitForPackageState, uint32 CurrentTick, FAsyncPackage2* Root, bool bExcludeZenPackages)
+void FAsyncPackage2::UpdateDependenciesStateRecursive(FAsyncLoadingThreadState2& ThreadState, FAllDependenciesState FAsyncPackage2::* StateMemberPtr, EAsyncPackageLoadingState2 WaitForPackageState, uint32 CurrentTick, FAsyncPackage2* Root)
 {
 	FAllDependenciesState& ThisState = this->*StateMemberPtr;
 
@@ -6247,7 +6238,7 @@ void FAsyncPackage2::UpdateDependenciesStateRecursive(FAsyncLoadingThreadState2&
 				WaitingForPackageState.UpdateTick(CurrentTick);
 				if (!WaitingForPackageState.bVisitedMark)
 				{
-					WaitingForPackage->UpdateDependenciesStateRecursive(ThreadState, StateMemberPtr, WaitForPackageState, CurrentTick, Root, bExcludeZenPackages);
+					WaitingForPackage->UpdateDependenciesStateRecursive(ThreadState, StateMemberPtr, WaitForPackageState, CurrentTick, Root);
 				}
 				if (WaitingForPackageState.bAnyNotDone)
 				{
@@ -6267,13 +6258,6 @@ void FAsyncPackage2::UpdateDependenciesStateRecursive(FAsyncLoadingThreadState2&
 			continue;
 		}
 
-#if ALT2_ENABLE_LINKERLOAD_SUPPORT
-		if (bExcludeZenPackages && !ImportedPackage->LinkerLoadState.IsSet())
-		{
-			continue;
-		}
-#endif
-
 		if (ThreadState.PackagesOnStack.Contains(ImportedPackage))
 		{
 			continue;
@@ -6294,7 +6278,7 @@ void FAsyncPackage2::UpdateDependenciesStateRecursive(FAsyncLoadingThreadState2&
 		}
 		else if (!ImportedPackageState.bVisitedMark)
 		{
-			ImportedPackage->UpdateDependenciesStateRecursive(ThreadState, StateMemberPtr, WaitForPackageState, CurrentTick, Root, bExcludeZenPackages);
+			ImportedPackage->UpdateDependenciesStateRecursive(ThreadState, StateMemberPtr, WaitForPackageState, CurrentTick, Root);
 		}
 
 		if (ImportedPackageState.bAnyNotDone)
@@ -6326,9 +6310,9 @@ void FAsyncPackage2::UpdateDependenciesStateRecursive(FAsyncLoadingThreadState2&
 	}
 }
 
-void FAsyncPackage2::WaitForAllDependenciesToReachState(FAsyncLoadingThreadState2& ThreadState, FAllDependenciesState FAsyncPackage2::* StateMemberPtr, EAsyncPackageLoadingState2 WaitForPackageState, uint32& CurrentTickVariable, bool bExcludeZenPackages, TFunctionRef<void(FAsyncPackage2*)> OnStateReached)
+void FAsyncPackage2::WaitForAllDependenciesToReachState(FAsyncLoadingThreadState2& ThreadState, FAllDependenciesState FAsyncPackage2::* StateMemberPtr, EAsyncPackageLoadingState2 WaitForPackageState, uint32& CurrentTickVariable, TFunctionRef<void(FAsyncPackage2*)> OnStateReached)
 {
-	if (HaveAllDependenciesReachedState(ThreadState, StateMemberPtr, WaitForPackageState, CurrentTickVariable++, bExcludeZenPackages))
+	if (HaveAllDependenciesReachedState(ThreadState, StateMemberPtr, WaitForPackageState, CurrentTickVariable++))
 	{
 		FAsyncPackage2* FirstPackageReadyToProceed = this;
 
@@ -6343,16 +6327,15 @@ void FAsyncPackage2::WaitForAllDependenciesToReachState(FAsyncLoadingThreadState
 				continue;
 			}
 
-#if DO_CHECK
-			TSet<FAsyncPackage2*> VisitedPackages;
-			check(HaveAllDependenciesReachedStateDebug(ThreadState, this, VisitedPackages, WaitForPackageState, bExcludeZenPackages));
+#if ALT2_ENABLE_PACKAGE_DEPENDENCY_DEBUGGING
+			PackageReadyToProceed->CheckThatAllDependenciesHaveReachedStateDebug(ThreadState, WaitForPackageState, WaitForPackageState);
 #endif
 
 			while (FAsyncPackage2* WaitingPackage = PackageReadyToProceedState.PackagesWaitingForThisHead)
 			{
 				FAllDependenciesState& WaitingPackageState = WaitingPackage->*StateMemberPtr;
 				check(WaitingPackageState.WaitingForPackage == PackageReadyToProceed);
-				if (WaitingPackage->HaveAllDependenciesReachedState(ThreadState, StateMemberPtr, WaitForPackageState, CurrentTickVariable++, bExcludeZenPackages))
+				if (WaitingPackage->HaveAllDependenciesReachedState(ThreadState, StateMemberPtr, WaitForPackageState, CurrentTickVariable++))
 				{
 					// If the waiting package wasn't ready to proceed it will remove itself from the waitlist (PackageReadyToProceedState.PackagesWaitingForThisHead) when its WaitingForPackage
 					// is replaced by another package
@@ -6369,11 +6352,39 @@ void FAsyncPackage2::WaitForAllDependenciesToReachState(FAsyncLoadingThreadState
 	}
 }
 
+void FAsyncPackage2::ConditionalBeginProcessPackageExports(FAsyncLoadingThreadState2& ThreadState)
+{
+	TRACE_CPUPROFILER_EVENT_SCOPE(ConditionalBeginProcessPackageExports);
+	WaitForAllDependenciesToReachState(ThreadState, &FAsyncPackage2::AllDependenciesSetupState, EAsyncPackageLoadingState2::DependenciesReady, AsyncLoadingThread.ConditionalBeginProcessExportsTick,
+		[&ThreadState](FAsyncPackage2* Package)
+		{
+			if (!Package->bLoadHasFailed
+#if ALT2_ENABLE_LINKERLOAD_SUPPORT
+			&& !Package->LinkerLoadState.IsSet()
+#endif
+				)
+			{
+				TArrayView<FEventLoadNode2> ExportBundleNodesView = Package->Data.ExportBundleNodes;
+				TArrayView<FAsyncPackage2*> ImportedAsyncPackagesView = Package->Data.ImportedAsyncPackages;
+				Package->SetupSerializedArcs(Package->HeaderData, ExportBundleNodesView.Left(Package->HeaderData.ExportBundleHeaders.Num() * EEventLoadNode2::ExportBundle_NumPhases), ImportedAsyncPackagesView.Left(Package->HeaderData.ImportedPackageIds.Num()));
+#if WITH_EDITOR
+				if (Package->OptionalSegmentHeaderData.IsSet())
+				{
+					Package->SetupSerializedArcs(*Package->OptionalSegmentHeaderData, ExportBundleNodesView.Right(Package->OptionalSegmentHeaderData->ExportBundleHeaders.Num() * EEventLoadNode2::ExportBundle_NumPhases), ImportedAsyncPackagesView.Right(Package->OptionalSegmentHeaderData->ImportedPackageIds.Num()));
+				}
+#endif
+			}
+			for (int32 ExportBundleIndex = 0; ExportBundleIndex < Package->Data.ExportBundleCount; ++ExportBundleIndex)
+			{
+				Package->GetExportBundleNode(EEventLoadNode2::ExportBundle_Process, ExportBundleIndex).ReleaseBarrier(&ThreadState);
+			}
+		});
+}
+
 void FAsyncPackage2::ConditionalFinishLoading(FAsyncLoadingThreadState2& ThreadState)
 {
 	TRACE_CPUPROFILER_EVENT_SCOPE(ConditionalFinishLoading);
-	constexpr bool bExcludeZenPackages = false;
-	WaitForAllDependenciesToReachState(ThreadState, &FAsyncPackage2::AllDependenciesFullyLoadedState, EAsyncPackageLoadingState2::DeferredPostLoadDone, AsyncLoadingThread.ConditionalFinishLoadingTick, bExcludeZenPackages,
+	WaitForAllDependenciesToReachState(ThreadState, &FAsyncPackage2::AllDependenciesFullyLoadedState, EAsyncPackageLoadingState2::DeferredPostLoadDone, AsyncLoadingThread.ConditionalFinishLoadingTick,
 		[](FAsyncPackage2* Package)
 		{
 			Package->AsyncLoadingThread.LoadedPackagesToProcess.Add(Package);
@@ -7113,7 +7124,7 @@ FAsyncLoadingThread2::FAsyncLoadingThread2(FIoDispatcher& InIoDispatcher, IAsync
 
 	EventSpecs.AddDefaulted(EEventLoadNode2::Package_NumPhases + EEventLoadNode2::ExportBundle_NumPhases);
 	EventSpecs[EEventLoadNode2::Package_ProcessSummary] = { &FAsyncPackage2::Event_ProcessPackageSummary, &EventQueue, false };
-	EventSpecs[EEventLoadNode2::Package_SetupDependencies] = { &FAsyncPackage2::Event_SetupDependencies, &EventQueue, false };
+	EventSpecs[EEventLoadNode2::Package_DependenciesReady] = { &FAsyncPackage2::Event_DependenciesReady, &EventQueue, false };
 	EventSpecs[EEventLoadNode2::Package_ExportsSerialized] = { &FAsyncPackage2::Event_ExportsDone, &EventQueue, true };
 
 	EventSpecs[EEventLoadNode2::Package_NumPhases + EEventLoadNode2::ExportBundle_Process] = { &FAsyncPackage2::Event_ProcessExportBundle, &EventQueue, false };
@@ -8248,14 +8259,6 @@ void FAsyncLoadingThread2::FlushLoading(int32 RequestId)
 						CurrentlyExecutingPackage->PostLoadGroup->PackagesWithExportsToSerializeCount = 1;
 					}
 				}
-
-#if ALT2_ENABLE_LINKERLOAD_SUPPORT
-				while (FAsyncPackage2* WaitingPackage = CurrentlyExecutingPackage->AllDependenciesSetupState.PackagesWaitingForThisHead)
-				{
-					FAsyncPackage2::FAllDependenciesState::RemoveFromWaitList(&FAsyncPackage2::AllDependenciesSetupState, CurrentlyExecutingPackage, WaitingPackage);
-					WaitingPackage->ConditionalProcessLinkerLoadPackageImportsAndExports(*GameThreadState);
-				}
-#endif
 			}
 		}
 
