Move FArchiveStackTraceWriter's default destructor to be implemented in the .cpp to allow ArchiveStackTrace.h to compile.

diff --git a/Engine/Source/Runtime/CoreUObject/Private/Serialization/ArchiveStackTrace.cpp b/Engine/Source/Runtime/CoreUObject/Private/Serialization/ArchiveStackTrace.cpp
--- a/Engine/Source/Runtime/CoreUObject/Private/Serialization/ArchiveStackTrace.cpp
+++ b/Engine/Source/Runtime/CoreUObject/Private/Serialization/ArchiveStackTrace.cpp
@@ -428,6 +428,8 @@ FArchiveStackTraceWriter::FArchiveStackTraceWriter(
 {
 }
 
+FArchiveStackTraceWriter::~FArchiveStackTraceWriter() = default;
+
 void FArchiveStackTraceWriter::Serialize(void* Data, int64 Length)
 {
 	static struct FBreakAtOffsetSettings
diff --git a/Engine/Source/Runtime/CoreUObject/Public/Serialization/ArchiveStackTrace.h b/Engine/Source/Runtime/CoreUObject/Public/Serialization/ArchiveStackTrace.h
--- a/Engine/Source/Runtime/CoreUObject/Public/Serialization/ArchiveStackTrace.h
+++ b/Engine/Source/Runtime/CoreUObject/Public/Serialization/ArchiveStackTrace.h
@@ -244,6 +244,8 @@ public:
 		const FArchiveDiffMap* InDiffMap = nullptr,
 		int64 InDiffMapStartOffset = 0);
 
+	virtual ~FArchiveStackTraceWriter() override;
+
 	FORCENOINLINE virtual void Serialize(void* Data, int64 Length) override; // FORCENOINLINE so it can be counted during StackTrace
 	virtual void SetSerializeContext(FUObjectSerializeContext* Context) override;
 	virtual FUObjectSerializeContext* GetSerializeContext() override;
