Silence static analysis warning about intentional null-relative pointers

diff --git a/Engine/Source/Runtime/Core/Public/Experimental/Graph/GraphConvert.h b/Engine/Source/Runtime/Core/Public/Experimental/Graph/GraphConvert.h
--- a/Engine/Source/Runtime/Core/Public/Experimental/Graph/GraphConvert.h
+++ b/Engine/Source/Runtime/Core/Public/Experimental/Graph/GraphConvert.h
@@ -192,7 +192,8 @@ inline FGraph ConvertToGraph(
 
 		// Store the vertex's offset into GraphBuffer, for later fixup
 		// Edge list is guaranteed to fit within 32 bits after removal of duplicates
-		OutGraph.EdgeLists.Emplace(EdgeListBase + InitialOffset, static_cast<int32>(VertexEdges.Num()));
+		OutGraph.EdgeLists.Emplace(EdgeListBase + InitialOffset, static_cast<int32>(VertexEdges.Num())); //-V769
+
 		++Vertex;
 	}
 	if (EnumHasAnyFlags(Options, EConvertToGraphOptions::Shrink))
diff --git a/Engine/Source/Runtime/CoreUObject/Private/UObject/ReferenceChainSearch.cpp b/Engine/Source/Runtime/CoreUObject/Private/UObject/ReferenceChainSearch.cpp
--- a/Engine/Source/Runtime/CoreUObject/Private/UObject/ReferenceChainSearch.cpp
+++ b/Engine/Source/Runtime/CoreUObject/Private/UObject/ReferenceChainSearch.cpp
@@ -325,7 +325,7 @@ namespace UE::ReferenceChainSearch
 					// During construction store edge lists relative to null to be fixed up later
 					FVertex* StartAddress = nullptr;
 					GraphBuffer.Append(EdgeList);
-					EdgeLists[SourceVertex - StartVertex] = TArrayView<FVertex>(StartAddress + BufferStartIndex, EdgeList.Num());
+					EdgeLists[SourceVertex - StartVertex] = TArrayView<FVertex>(StartAddress + BufferStartIndex, EdgeList.Num()); //-V769
 					CurrentEdgeList.Reset();
 				}
 			}
