Latest Clang toolchain fixes for FPrimaryAssetType and FPrimaryAssetId. deprecated-copy-with-user-provided-copy.

diff --git a/Engine/Source/Runtime/CoreUObject/Public/UObject/PrimaryAssetId.h b/Engine/Source/Runtime/CoreUObject/Public/UObject/PrimaryAssetId.h
--- a/Engine/Source/Runtime/CoreUObject/Public/UObject/PrimaryAssetId.h
+++ b/Engine/Source/Runtime/CoreUObject/Public/UObject/PrimaryAssetId.h
@@ -22,12 +22,17 @@ class UObject;
 struct FPrimaryAssetType
 {
 	/** Convert from FName */
-	FPrimaryAssetType() {}
+	FPrimaryAssetType() = default;
 	FPrimaryAssetType(FName InName) : Name(InName) {}
 	FPrimaryAssetType(EName InName) : Name(FName(InName)) {}
 	FPrimaryAssetType(const WIDECHAR* InName) : Name(FName(InName)) {}
 	FPrimaryAssetType(const ANSICHAR* InName) : Name(FName(InName)) {}
 
+	FPrimaryAssetType(const FPrimaryAssetType&) = default;
+	FPrimaryAssetType(FPrimaryAssetType&&) = default;
+	FPrimaryAssetType& operator=(const FPrimaryAssetType&) = default;
+	FPrimaryAssetType& operator=(FPrimaryAssetType&&) = default;
+
 	/** Convert to FName */
 	operator FName&() { return Name; }
 	operator const FName&() const { return Name; }
@@ -48,12 +53,6 @@ struct FPrimaryAssetType
 		return Name != Other.Name;
 	}
 
-	FPrimaryAssetType& operator=(const FPrimaryAssetType& Other)
-	{
-		Name = Other.Name;
-		return *this;
-	}
-
 	/** Returns true if this is a valid Type */
 	bool IsValid() const
 	{
@@ -115,8 +114,7 @@ struct FPrimaryAssetId
 	static COREUOBJECT_API const FName PrimaryAssetTypeTag;
 	static COREUOBJECT_API const FName PrimaryAssetNameTag;
 
-	FPrimaryAssetId() {}
-
+	FPrimaryAssetId() = default;
 	FPrimaryAssetId(FPrimaryAssetType InAssetType, FName InAssetName)
 		: PrimaryAssetType(InAssetType), PrimaryAssetName(InAssetName)
 	{}
@@ -132,6 +130,11 @@ struct FPrimaryAssetId
 		: FPrimaryAssetId(ParseTypeAndName(TypeAndName))
 	{}
 
+	FPrimaryAssetId(const FPrimaryAssetId&) = default;
+	FPrimaryAssetId(FPrimaryAssetId&&) = default;
+	FPrimaryAssetId& operator=(const FPrimaryAssetId&) = default;
+	FPrimaryAssetId& operator=(FPrimaryAssetId&&) = default;
+
 	/** Returns true if this is a valid identifier */
 	bool IsValid() const
 	{
@@ -173,13 +176,6 @@ struct FPrimaryAssetId
 		return PrimaryAssetType != Other.PrimaryAssetType || PrimaryAssetName != Other.PrimaryAssetName;
 	}
 
-	FPrimaryAssetId& operator=(const FPrimaryAssetId& Other)
-	{
-		PrimaryAssetType = Other.PrimaryAssetType;
-		PrimaryAssetName = Other.PrimaryAssetName;
-		return *this;
-	}
-
 	FORCEINLINE friend FArchive& operator<<(FArchive& Ar, FPrimaryAssetId& Other)
 	{
         Ar << Other.PrimaryAssetType;
