minor improvement: using `TScriptDelegate::operator==` instead of the custom comparison impl

diff --git a/Engine/Source/Runtime/CoreUObject/Private/UObject/PropertyDelegate.cpp b/Engine/Source/Runtime/CoreUObject/Private/UObject/PropertyDelegate.cpp
--- a/Engine/Source/Runtime/CoreUObject/Private/UObject/PropertyDelegate.cpp
+++ b/Engine/Source/Runtime/CoreUObject/Private/UObject/PropertyDelegate.cpp
@@ -62,25 +62,17 @@ void FDelegateProperty::InstanceSubobjects(void* Data, void const* DefaultData,
 
 bool FDelegateProperty::Identical( const void* A, const void* B, uint32 PortFlags ) const
 {
+	check(A);
+
 	const FScriptDelegate* DA = (const FScriptDelegate*)A;
 	const FScriptDelegate* DB = (const FScriptDelegate*)B;
-	
-	if (!DB)
-	{
-		return DA->GetFunctionName() == NAME_None;
-	}
 
-	if (DA->GetUObject() != DB->GetUObject())
-	{
-		return false;
-	}
-
-	if (DA->GetFunctionName() != DB->GetFunctionName())
+	if (!DB)
 	{
-		return false;
+		return !DA->IsBound();
 	}
-
-	return true;
+	
+	return *DA == *DB;
 }
 
 
