Making sure that bytecode FProperties resolved in PostLoad add their referenced objects to the bytecode owner UStruct's reference list

diff --git a/Engine/Source/Runtime/CoreUObject/Private/UObject/Class.cpp b/Engine/Source/Runtime/CoreUObject/Private/UObject/Class.cpp
--- a/Engine/Source/Runtime/CoreUObject/Private/UObject/Class.cpp
+++ b/Engine/Source/Runtime/CoreUObject/Private/UObject/Class.cpp
@@ -2000,6 +2000,10 @@ void UStruct::PostLoad()
 				check((int32)Script.Num() >= (int32)(MissingProperty.Value + sizeof(FField*)));
 				FField** TargetScriptPropertyPtr = (FField**)(Script.GetData() + MissingProperty.Value);
 				*TargetScriptPropertyPtr = ResolvedProperty;
+
+				// Collect UObjects referenced by this property so that its owner doesn't get GC'd leaving a stale FProperty reference in bytecode
+				FPropertyReferenceCollector Collector(this, ScriptAndPropertyObjectReferences);
+				ResolvedProperty->AddReferencedObjects(Collector);
 			}
 			else if (!MissingProperty.Key.IsPathToFieldEmpty())
 			{
