Fix 'gconly' reference chain search.
Properly create empty edge lists for skipped objects.
Properly scan all gc objects when skipping non-gc objects.

diff --git a/Engine/Source/Runtime/CoreUObject/Private/UObject/Obj.cpp b/Engine/Source/Runtime/CoreUObject/Private/UObject/Obj.cpp
--- a/Engine/Source/Runtime/CoreUObject/Private/UObject/Obj.cpp
+++ b/Engine/Source/Runtime/CoreUObject/Private/UObject/Obj.cpp
@@ -4359,6 +4359,10 @@ bool StaticExec( UWorld* InWorld, const TCHAR* Cmd, FOutputDevice& Ar )
 					{
 						SearchModeFlags |= EReferenceChainSearchMode::Minimal;
 					}
+					else if (FCString::Stricmp(*Tok, TEXT("gconly")) == 0)
+					{
+						SearchModeFlags |= EReferenceChainSearchMode::GCOnly;
+					}
 #if ENABLE_GC_HISTORY
 					else if (FParse::Value(Str, TEXT("history="), HistoryLevel))
 					{
diff --git a/Engine/Source/Runtime/CoreUObject/Private/UObject/ReferenceChainSearch.cpp b/Engine/Source/Runtime/CoreUObject/Private/UObject/ReferenceChainSearch.cpp
--- a/Engine/Source/Runtime/CoreUObject/Private/UObject/ReferenceChainSearch.cpp
+++ b/Engine/Source/Runtime/CoreUObject/Private/UObject/ReferenceChainSearch.cpp
@@ -31,17 +31,24 @@ namespace UE::ReferenceChainSearch
 		TMap<const UObject*, FGCObjectInfo*>& ObjectToInfoMap;
 		TMap<FGCObjectInfo*, FReferenceChainSearch::FGraphNode*>& AllNodes;
 
-		FPolicyUObjectHeap(TMap<const UObject*, FGCObjectInfo*>& InObjectToInfoMap,
-			TMap<FGCObjectInfo*, FReferenceChainSearch::FGraphNode*>& InAllNodes)
+		FPolicyUObjectHeap(
+			TMap<const UObject*, FGCObjectInfo*>& InObjectToInfoMap,
+			TMap<FGCObjectInfo*, FReferenceChainSearch::FGraphNode*>& InAllNodes,
+			EReferenceChainSearchMode Mode)
 			: ObjectToInfoMap(InObjectToInfoMap)
 			, AllNodes(InAllNodes)
 		{
+			bGCOnly = !!(Mode & EReferenceChainSearchMode::GCOnly);
 		}
 
 		FVertex GetGCObjectReferencerVertex() const
 		{
 			return FGCObject::GGCObjectReferencer ? ObjectToVertex(FGCObject::GGCObjectReferencer) : UE::Graph::InvalidVertex;
 		}
+		int32 GetFirstVertexIndex() const
+		{
+			return bGCOnly ? GUObjectArray.GetFirstGCIndex() : 0;
+		}
 		int32 GetNumVertices() const
 		{
 			return GUObjectArray.GetObjectArrayNum();
@@ -144,6 +151,10 @@ namespace UE::ReferenceChainSearch
 			return UE::Graph::InvalidVertex;
 		}
 
+		int32 GetFirstVertexIndex() const
+		{
+			return 0;
+		}
 		int32 GetNumVertices() const
 		{
 			return Snapshot.ObjectToInfoMap.GetMaxIndex() + 1;
@@ -383,7 +394,7 @@ namespace UE::ReferenceChainSearch
 												   : (MaxNumberOfObjects - (NumThreads - 1) * NumberOfObjectsPerThread);
 
 						// First cache all potential referencers
-						for (int32 ObjectIndex = 0; ObjectIndex < NumObjects && (FirstObjectIndex + ObjectIndex) < MaxNumberOfObjects; ++ObjectIndex)
+						for (int32 ObjectIndex = 0; ObjectIndex < NumObjects; ++ObjectIndex)
 						{
 							FUObjectItem& ObjectItem = GUObjectArray.GetObjectItemArrayUnsafe()[FirstObjectIndex + ObjectIndex];
 							UObject* Object = static_cast<UObject*>(ObjectItem.Object);
@@ -623,6 +634,10 @@ namespace UE::ReferenceChainSearch
 			Search.CollectAllReferences(bGCOnly);
 			TArray<TConstArrayView<FVertex>> MergedEdgeLists;
 			MergedEdgeLists.Reserve(Policy.GetNumVertices());
+			if (bGCOnly)
+			{
+				MergedEdgeLists.AddZeroed(Policy.GetFirstVertexIndex());
+			}
 			MergedEdgeLists.SetNumUninitialized(Policy.GetNumVertices());
 			Search.MergeGraph(MergedEdgeLists);
 
@@ -635,6 +650,10 @@ namespace UE::ReferenceChainSearch
 			Search.CollectAllReferences(bGCOnly);
 			TArray<TConstArrayView<FVertex>> MergedEdgeLists;
 			MergedEdgeLists.Reserve(Policy.GetNumVertices());
+			if (bGCOnly)
+			{
+				MergedEdgeLists.AddZeroed(Policy.GetFirstVertexIndex());
+			}
 			MergedEdgeLists.SetNumUninitialized(Policy.GetNumVertices());
 			Search.MergeGraph(MergedEdgeLists);
 
@@ -1486,7 +1505,7 @@ FReferenceChainSearch::FReferenceChainSearch(TConstArrayView<UObject*> InObjects
 		ObjectInfosToFindReferencesTo.Add(FGCObjectInfo::FindOrAddInfoHelper(Object, ObjectToInfoMap));
 	}
 
-	UE::ReferenceChainSearch::FPolicyUObjectHeap Policy(ObjectToInfoMap, AllNodes);
+	UE::ReferenceChainSearch::FPolicyUObjectHeap Policy(ObjectToInfoMap, AllNodes, Mode);
 
 	UE::Graph::FGraph ReferenceGraph;
 	UE::ReferenceChainSearch::PerformInitialGatherFromLiveUObjectHeap(Policy, InObjectsToFindReferencesTo, SearchMode, ReferenceGraph);
