Fix missing member declaration
(Accidentally reverted pre-submission)

diff --git a/Engine/Source/Runtime/CoreUObject/Private/UObject/ReferenceChainSearch.cpp b/Engine/Source/Runtime/CoreUObject/Private/UObject/ReferenceChainSearch.cpp
--- a/Engine/Source/Runtime/CoreUObject/Private/UObject/ReferenceChainSearch.cpp
+++ b/Engine/Source/Runtime/CoreUObject/Private/UObject/ReferenceChainSearch.cpp
@@ -30,6 +30,7 @@ namespace UE::ReferenceChainSearch
 
 		TMap<const UObject*, FGCObjectInfo*>& ObjectToInfoMap;
 		TMap<FGCObjectInfo*, FReferenceChainSearch::FGraphNode*>& AllNodes;
+		bool bGCOnly;
 
 		FPolicyUObjectHeap(
 			TMap<const UObject*, FGCObjectInfo*>& InObjectToInfoMap,
