IoStore/ZenLoader: Some cleanup to facilitate future changes
#rnx

diff --git a/Engine/Source/Developer/IoStoreUtilities/Internal/PackageStoreOptimizer.h b/Engine/Source/Developer/IoStoreUtilities/Internal/PackageStoreOptimizer.h
--- a/Engine/Source/Developer/IoStoreUtilities/Internal/PackageStoreOptimizer.h
+++ b/Engine/Source/Developer/IoStoreUtilities/Internal/PackageStoreOptimizer.h
@@ -157,7 +157,6 @@ private:
 		FString FullName;
 		FName FromPackageName;
 		int32 FromPackageNameLen = 0;
-		FPackageId FromPackageId;
 		bool bIsScriptImport = false;
 		bool bIsImportOfPackage = false;
 		bool bIsImportOptional = false;
@@ -209,6 +208,30 @@ private:
 		int32 IncomingEdgeCount = 0;
 	};
 
+	struct FImportedPackageRef
+	{
+		FImportedPackageRef(FName InName)
+			: Id(FPackageId::FromName(InName))
+			, Name(InName)
+		{
+		}
+
+		// Note: For Zen packages we currently don't know the package name but we also don't use it for anything
+		FImportedPackageRef(FPackageId InId)
+			: Id(InId)
+		{
+
+		}
+
+		bool operator<(const FImportedPackageRef& Other) const
+		{
+			return Id < Other.Id;
+		}
+
+		const FPackageId Id;
+		const FName Name;
+	};
+
 	FPackageId Id;
 	FName Name;
 
@@ -220,7 +243,7 @@ private:
 	TArray<FExportGraphNode> ExportGraphNodes;
 	FGraphData GraphData;
 
-	TArray<FPackageId> ImportedPackageIds;
+	TArray<FImportedPackageRef> ImportedPackages;
 	TArray<uint64> ImportedPublicExportHashes;
 	TArray<FBulkDataMapEntry> BulkDataEntries;
 
@@ -313,7 +336,7 @@ private:
 	void ProcessPreloadDependencies(const FZenPackageHeaderData& ZenHeaderData, FPackageStorePackage* Package) const;
 	void ProcessDataResources(const FCookedHeaderData& CookedHeaderData, FPackageStorePackage* Package) const;
 	TArray<FPackageStorePackage*> SortPackagesInLoadOrder(const TMap<FPackageId, FPackageStorePackage*>& PackagesMap) const;
-	void SerializeGraphData(const TArray<FPackageId>& ImportedPackageIds, FPackageStorePackage::FGraphData& GraphData, FBufferWriter& GraphArchive) const;
+	void SerializeGraphData(const TArray<FPackageStorePackage::FImportedPackageRef>& ImportedPackages, FPackageStorePackage::FGraphData& GraphData, FBufferWriter& GraphArchive) const;
 	TArray<FPackageStorePackage::FExportGraphNode*> SortExportGraphNodesInLoadOrder(FPackageStorePackage* Package, FExportGraphEdges& Edges) const;
 	TArray<FPackageStorePackage::FExportBundleGraphNode*> SortExportBundleGraphNodesInLoadOrder(const TArray<FPackageStorePackage*>& Packages, FExportBundleGraphEdges& Edges) const;
 	void CreateExportBundles(FPackageStorePackage* Package) const;
diff --git a/Engine/Source/Developer/IoStoreUtilities/Private/IoStoreUtilities.cpp b/Engine/Source/Developer/IoStoreUtilities/Private/IoStoreUtilities.cpp
--- a/Engine/Source/Developer/IoStoreUtilities/Private/IoStoreUtilities.cpp
+++ b/Engine/Source/Developer/IoStoreUtilities/Private/IoStoreUtilities.cpp
@@ -2289,8 +2289,7 @@ void InitializeContainerTargetsAndPackages(
 	TArray<FCookedPackage*>& Packages,
 	FPackageNameMap& PackageNameMap,
 	FPackageIdMap& PackageIdMap,
-	TArray<FContainerTargetSpec*>& ContainerTargets,
-	FPackageStoreOptimizer& PackageStoreOptimizer)
+	TArray<FContainerTargetSpec*>& ContainerTargets)
 {
 	auto CreateTargetFileFromCookedFile = [
 		&Arguments,
@@ -3533,7 +3532,7 @@ void CreateContainerHeader(FContainerTargetSpec& ContainerTarget, bool bIsOption
 	{
 		for (const FCookedPackage* Package : ContainerTarget.Packages)
 		{
-			if (Package->PackageStoreEntry.OptionalSegmentExportInfo.ExportCount)
+			if (Package->PackageStoreEntry.HasOptionalSegment())
 			{
 				if (Package->PackageStoreEntry.IsAutoOptional())
 				{
@@ -3609,7 +3608,7 @@ void CreateContainerHeader(FContainerTargetSpec& ContainerTarget, bool bIsOption
 		for (const FCookedPackage* Package : SortedPackages)
 		{
 			const FPackageStoreEntryResource& Entry = Package->PackageStoreEntry;
-			if (Entry.OptionalSegmentExportInfo.ExportCount)
+			if (Entry.HasOptionalSegment())
 			{
 				FStoreEntriesWriter* TargetEntriesWriter;
 				if (Entry.IsAutoOptional())
@@ -3730,7 +3729,7 @@ int32 CreateTarget(const FIoStoreArguments& Arguments, const FIoStoreWriterSetti
 	UE_LOG(LogIoStore, Display, TEXT("Creating container targets..."));
 	{
 		IOSTORE_CPU_SCOPE(CreateContainerTargets);
-		InitializeContainerTargetsAndPackages(Arguments, Packages, PackageNameMap, PackageIdMap, ContainerTargets, PackageStoreOptimizer);
+		InitializeContainerTargetsAndPackages(Arguments, Packages, PackageNameMap, PackageIdMap, ContainerTargets);
 	}
 
 	TUniquePtr<FIoStoreWriterContext> IoStoreWriterContext;
@@ -3821,8 +3820,52 @@ int32 CreateTarget(const FIoStoreArguments& Arguments, const FIoStoreWriterSetti
 	TArray<FShaderInfo*> Shaders;
 	ProcessShaderLibraries(Arguments, ContainerTargets, Shaders);
 
+	auto AppendTargetFileChunk = [&WriteRequestManager](FContainerTargetSpec* ContainerTarget, const FContainerTargetFile& TargetFile)
 	{
-		IOSTORE_CPU_SCOPE(AppendNonPackageChunks);
+		FIoWriteOptions WriteOptions;
+		WriteOptions.DebugName = *TargetFile.DestinationPath;
+		WriteOptions.bForceUncompressed = TargetFile.bForceUncompressed;
+		WriteOptions.bIsMemoryMapped = TargetFile.ChunkType == EContainerChunkType::MemoryMappedBulkData;
+		WriteOptions.FileName = TargetFile.DestinationPath;
+		FIoChunkId ChunkId = TargetFile.ChunkId;
+		bool bIsOptionalSegmentChunk = false;
+		switch (TargetFile.ChunkType)
+		{
+		case EContainerChunkType::OptionalSegmentPackageData:
+		{
+			if (TargetFile.Package->PackageStoreEntry.IsAutoOptional())
+			{
+				// Auto optional packages replace the non-optional part when the container is mounted
+				ChunkId = CreateIoChunkId(TargetFile.Package->GlobalPackageId.Value(), 0, EIoChunkType::ExportBundleData);
+			}
+			bIsOptionalSegmentChunk = true;
+			break;
+		}
+		case EContainerChunkType::OptionalSegmentBulkData:
+		{
+			if (TargetFile.Package->PackageStoreEntry.IsAutoOptional())
+			{
+				// Auto optional packages replace the non-optional part when the container is mounted
+				ChunkId = CreateIoChunkId(TargetFile.Package->GlobalPackageId.Value(), 0, EIoChunkType::BulkData);
+			}
+			bIsOptionalSegmentChunk = true;
+			break;
+		}
+		}
+
+		if (bIsOptionalSegmentChunk)
+		{
+			ContainerTarget->OptionalSegmentIoStoreWriter->Append(ChunkId, WriteRequestManager.Read(TargetFile), WriteOptions);
+		}
+		else
+		{
+			ContainerTarget->IoStoreWriter->Append(TargetFile.ChunkId, WriteRequestManager.Read(TargetFile), WriteOptions);
+		}
+	};
+
+	{
+		// Append all chunks that don't change after this point
+		IOSTORE_CPU_SCOPE(AppendChunks);
 		for (FContainerTargetSpec* ContainerTarget : ContainerTargets)
 		{
 			if (ContainerTarget->IoStoreWriter)
@@ -3831,25 +3874,7 @@ int32 CreateTarget(const FIoStoreArguments& Arguments, const FIoStoreWriterSetti
 				{
 					if (TargetFile.ChunkType != EContainerChunkType::PackageData && TargetFile.ChunkType != EContainerChunkType::OptionalSegmentPackageData)
 					{
-						FIoWriteOptions WriteOptions;
-						WriteOptions.DebugName = *TargetFile.DestinationPath;
-						WriteOptions.bForceUncompressed = TargetFile.bForceUncompressed;
-						WriteOptions.bIsMemoryMapped = TargetFile.ChunkType == EContainerChunkType::MemoryMappedBulkData;
-						WriteOptions.FileName = TargetFile.DestinationPath;
-						if (TargetFile.ChunkType == EContainerChunkType::OptionalSegmentBulkData)
-						{
-							FIoChunkId ChunkId = TargetFile.ChunkId;
-							if (TargetFile.Package->PackageStoreEntry.IsAutoOptional())
-							{
-								// Auto optional packages replace the non-optional part when the container is mounted
-								ChunkId = CreateIoChunkId(TargetFile.Package->GlobalPackageId.Value(), 0, EIoChunkType::BulkData);
-							}
-							ContainerTarget->OptionalSegmentIoStoreWriter->Append(ChunkId, WriteRequestManager.Read(TargetFile), WriteOptions);
-						}
-						else
-						{
-							ContainerTarget->IoStoreWriter->Append(TargetFile.ChunkId, WriteRequestManager.Read(TargetFile), WriteOptions);
-						}
+						AppendTargetFileChunk(ContainerTarget, TargetFile);
 					}
 				}
 			}
@@ -3897,7 +3922,7 @@ int32 CreateTarget(const FIoStoreArguments& Arguments, const FIoStoreWriterSetti
 	CreateDiskLayout(ContainerTargets, Packages, Arguments.OrderMaps, PackageIdMap, Arguments.bClusterByOrderFilePriority);
 
 	{
-		IOSTORE_CPU_SCOPE(AppendPackageChunks);
+		IOSTORE_CPU_SCOPE(AppendRemaningChunks);
 		for (FContainerTargetSpec* ContainerTarget : ContainerTargets)
 		{
 			if (ContainerTarget->IoStoreWriter)
@@ -3906,27 +3931,7 @@ int32 CreateTarget(const FIoStoreArguments& Arguments, const FIoStoreWriterSetti
 				{
 					if (TargetFile.ChunkType == EContainerChunkType::PackageData || TargetFile.ChunkType == EContainerChunkType::OptionalSegmentPackageData)
 					{
-						check(TargetFile.Package);
-						FIoWriteOptions WriteOptions;
-						WriteOptions.DebugName = *TargetFile.DestinationPath;
-						WriteOptions.bForceUncompressed = TargetFile.bForceUncompressed;
-						WriteOptions.FileName = TargetFile.DestinationPath;
-						if (TargetFile.ChunkType == EContainerChunkType::OptionalSegmentPackageData)
-						{
-							check(ContainerTarget->OptionalSegmentIoStoreWriter);
-							check(TargetFile.Package->OptimizedOptionalSegmentPackage);
-							FIoChunkId ChunkId = TargetFile.ChunkId;
-							if (TargetFile.Package->OptimizedOptionalSegmentPackage->HasEditorData())
-							{
-								// Auto optional packages replace the non-optional part when the container is mounted
-								ChunkId = CreateIoChunkId(TargetFile.Package->GlobalPackageId.Value(), 0, EIoChunkType::ExportBundleData);
-							}
-							ContainerTarget->OptionalSegmentIoStoreWriter->Append(ChunkId, WriteRequestManager.Read(TargetFile), WriteOptions);
-						}
-						else
-						{
-							ContainerTarget->IoStoreWriter->Append(TargetFile.ChunkId, WriteRequestManager.Read(TargetFile), WriteOptions);
-						}
+						AppendTargetFileChunk(ContainerTarget, TargetFile);
 					}
 				}
 
@@ -5217,7 +5222,6 @@ namespace DescribeUtils
 					const FPackageId& PackageId = ContainerHeader.PackageIds[PackageIndex++];
 					FPackageDesc* PackageDesc = new FPackageDesc();
 					PackageDesc->PackageId = PackageId;
-					PackageDesc->Exports.SetNum(ContainerEntry.ExportCount);
 					PackageDesc->ExportBundleCount = ContainerEntry.ExportBundleCount;
 					PackageDesc->ImportedPackageIds = TArrayView<FPackageId>(ContainerEntry.ImportedPackages.Data(), ContainerEntry.ImportedPackages.Num());
 					Job.Packages.Add(PackageDesc);
@@ -5338,6 +5342,7 @@ namespace DescribeUtils
 			}
 
 			const FExportMapEntry* ExportMap = reinterpret_cast<const FExportMapEntry*>(PackageSummaryData + PackageSummary->ExportMapOffset);
+			Job.PackageDesc->Exports.SetNum((PackageSummary->ExportBundleEntriesOffset - PackageSummary->ExportMapOffset) / sizeof(FExportMapEntry));
 			for (int32 ExportIndex = 0; ExportIndex < Job.PackageDesc->Exports.Num(); ++ExportIndex)
 			{
 				const FExportMapEntry& ExportMapEntry = ExportMap[ExportIndex];
diff --git a/Engine/Source/Developer/IoStoreUtilities/Private/PackageStoreOptimizer.cpp b/Engine/Source/Developer/IoStoreUtilities/Private/PackageStoreOptimizer.cpp
--- a/Engine/Source/Developer/IoStoreUtilities/Private/PackageStoreOptimizer.cpp
+++ b/Engine/Source/Developer/IoStoreUtilities/Private/PackageStoreOptimizer.cpp
@@ -374,7 +374,6 @@ void FPackageStoreOptimizer::ResolveImport(FPackageStorePackage::FUnresolvedImpo
 			FName PackageName = ObjectImport->ObjectName;
 			PackageName.AppendString(Import->FullName);
 			Import->FullName.ToLowerInline();
-			Import->FromPackageId = FPackageId::FromName(PackageName);
 			Import->FromPackageName = PackageName;
 			Import->FromPackageNameLen = Import->FullName.Len();
 			Import->bIsScriptImport = Import->FullName.StartsWith(TEXT("/Script/"));
@@ -391,7 +390,6 @@ void FPackageStoreOptimizer::ResolveImport(FPackageStorePackage::FUnresolvedImpo
 			Import->FullName.AppendChar(TEXT('/'));
 			ObjectImport->ObjectName.AppendString(Import->FullName);
 			Import->FullName.ToLowerInline();
-			Import->FromPackageId = OuterImport->FromPackageId;
 			Import->bIsImportOptional = ObjectImport->bImportOptional;
 			Import->FromPackageName = OuterImport->FromPackageName;
 			Import->FromPackageNameLen = OuterImport->FromPackageNameLen;
@@ -412,7 +410,7 @@ void FPackageStoreOptimizer::ProcessImports(const FCookedHeaderData& CookedHeade
 	UnresolvedImports.SetNum(ImportCount);
 	Package->Imports.SetNum(ImportCount);
 
-	TSet<FPackageId> ImportedPackageIds;
+	TSet<FName> ImportedPackageNames;
 	for (int32 ImportIndex = 0; ImportIndex < ImportCount; ++ImportIndex)
 	{
 		ResolveImport(UnresolvedImports.GetData(), CookedHeaderData.ObjectImports.GetData(), ImportIndex);
@@ -421,12 +419,16 @@ void FPackageStoreOptimizer::ProcessImports(const FCookedHeaderData& CookedHeade
 		{
 			if (UnresolvedImport.bIsImportOfPackage)
 			{
-				ImportedPackageIds.Add(UnresolvedImport.FromPackageId);
+				ImportedPackageNames.Add(UnresolvedImport.FromPackageName);
 			}
 		}
 	}
-	Package->ImportedPackageIds = ImportedPackageIds.Array();
-	Algo::Sort(Package->ImportedPackageIds);
+	Package->ImportedPackages.Reserve(ImportedPackageNames.Num());
+	for (FName ImportedPackageName : ImportedPackageNames)
+	{
+		Package->ImportedPackages.Emplace(ImportedPackageName);
+	}
+	Algo::Sort(Package->ImportedPackages);
 
 	for (int32 ImportIndex = 0; ImportIndex < ImportCount; ++ImportIndex)
 	{
@@ -443,10 +445,9 @@ void FPackageStoreOptimizer::ProcessImports(const FCookedHeaderData& CookedHeade
 		else if (!UnresolvedImport.bIsImportOfPackage)
 		{
 			bool bFoundPackageIndex = false;
-			for (uint32 PackageIndex = 0, PackageCount = static_cast<uint32>(Package->ImportedPackageIds.Num()); PackageIndex < PackageCount; ++PackageIndex)
+			for (uint32 PackageIndex = 0, PackageCount = static_cast<uint32>(Package->ImportedPackages.Num()); PackageIndex < PackageCount; ++PackageIndex)
 			{
-				FPackageId FromPackageId = UnresolvedImport.FromPackageId;
-				if (FromPackageId == Package->ImportedPackageIds[PackageIndex])
+				if (UnresolvedImport.FromPackageName == Package->ImportedPackages[PackageIndex].Name)
 				{
 					FStringView PackageRelativeName = FStringView(UnresolvedImport.FullName).RightChop(UnresolvedImport.FromPackageNameLen);
 					check(PackageRelativeName.Len());
@@ -465,7 +466,11 @@ void FPackageStoreOptimizer::ProcessImports(const FCookedHeaderData& CookedHeade
 
 void FPackageStoreOptimizer::ProcessImports(const FZenPackageHeaderData& ZenHeaderData, FPackageStorePackage* Package) const
 {
-	Package->ImportedPackageIds = ZenHeaderData.ImportedPackageIds;
+	Package->ImportedPackages.Reserve(ZenHeaderData.ImportedPackageIds.Num());
+	for (FPackageId ImportedPackageId : ZenHeaderData.ImportedPackageIds)
+	{
+		Package->ImportedPackages.Emplace(ImportedPackageId);
+	}
 	Package->ImportedPublicExportHashes = ZenHeaderData.ImportedPublicExportHashes;
 	Package->Imports = ZenHeaderData.Imports;
 }
@@ -659,9 +664,9 @@ TArray<FPackageStorePackage*> FPackageStoreOptimizer::SortPackagesInLoadOrder(co
 	TMap<FPackageStorePackage*, TArray<FPackageStorePackage*>> SortedEdges;
 	for (FPackageStorePackage* Package : Packages)
 	{
-		for (FPackageId ImportedPackageId : Package->ImportedPackageIds)
+		for (const FPackageStorePackage::FImportedPackageRef& ImportedPackage : Package->ImportedPackages)
 		{
-			FPackageStorePackage* FindImportedPackage = PackagesMap.FindRef(ImportedPackageId);
+			FPackageStorePackage* FindImportedPackage = PackagesMap.FindRef(ImportedPackage.Id);
 			if (FindImportedPackage)
 			{
 				TArray<FPackageStorePackage*>& SourceArray = SortedEdges.FindOrAdd(FindImportedPackage);
@@ -817,7 +822,7 @@ void FPackageStoreOptimizer::OptimizeExportBundles(const TMap<FPackageId, FPacka
 					FPackageObjectIndex FromImport = Package->Imports[ExternalDependency.ImportIndex];
 					check(FromImport.IsPackageImport());
 					FPackageImportReference FromPackageImportRef = FromImport.ToPackageImportRef();
-					FPackageId FromPackageId = Package->ImportedPackageIds[FromPackageImportRef.GetImportedPackageIndex()];
+					FPackageId FromPackageId = Package->ImportedPackages[FromPackageImportRef.GetImportedPackageIndex()].Id;
 					uint64 FromPublicExportHash = Package->ImportedPublicExportHashes[FromPackageImportRef.GetImportedPublicExportHashIndex()];
 					FPackageStorePackage* FindFromPackage = PackagesMap.FindRef(FromPackageId);
 					if (FindFromPackage)
@@ -1100,7 +1105,7 @@ void FPackageStoreOptimizer::CreateExportBundles(FPackageStorePackage* Package)
 	}
 }
 
-void FPackageStoreOptimizer::SerializeGraphData(const TArray<FPackageId>& ImportedPackageIds, FPackageStorePackage::FGraphData& GraphData, FBufferWriter& GraphArchive) const
+void FPackageStoreOptimizer::SerializeGraphData(const TArray<FPackageStorePackage::FImportedPackageRef>& ImportedPackages, FPackageStorePackage::FGraphData& GraphData, FBufferWriter& GraphArchive) const
 {
 	uint32 ExportBundleEntryIndex = 0;
 	for (const FPackageStorePackage::FExportBundle& ExportBundle : GraphData.ExportBundles)
@@ -1126,9 +1131,9 @@ void FPackageStoreOptimizer::SerializeGraphData(const TArray<FPackageId>& Import
 		GraphArchive << InternalArc.ToExportBundleIndex;
 	}
 
-	for (FPackageId ImportedPackageId : ImportedPackageIds)
+	for (const FPackageStorePackage::FImportedPackageRef& ImportedPackage : ImportedPackages)
 	{
-		TArray<FPackageStorePackage::FExternalArc>* FindArcsFromImportedPackage = GraphData.ExternalArcs.Find(ImportedPackageId);
+		TArray<FPackageStorePackage::FExternalArc>* FindArcsFromImportedPackage = GraphData.ExternalArcs.Find(ImportedPackage.Id);
 		if (!FindArcsFromImportedPackage)
 		{
 			int32 ExternalArcCount = 0;
@@ -1216,7 +1221,7 @@ void FPackageStoreOptimizer::FinalizePackageHeader(FPackageStorePackage* Package
 	uint64 ExportBundleEntriesSize = ExportBundleEntriesArchive.Tell();
 
 	FBufferWriter GraphArchive(nullptr, 0, EBufferWriterFlags::AllowResize | EBufferWriterFlags::TakeOwnership);
-	SerializeGraphData(Package->ImportedPackageIds, Package->GraphData, GraphArchive);
+	SerializeGraphData(Package->ImportedPackages, Package->GraphData, GraphArchive);
 	uint64 GraphDataSize = GraphArchive.Tell();
 
 	Package->NameMapBuilder.MarkNameAsReferenced(Package->Name);
@@ -1320,7 +1325,7 @@ void FPackageStoreOptimizer::FinalizePackage(FPackageStorePackage* Package) cons
 			const FPackageObjectIndex& Import = Package->Imports[ExternalDependency.ImportIndex];
 			check(Import.IsPackageImport());
 			FPackageImportReference PackageImportRef = Import.ToPackageImportRef();
-			TArray<FPackageStorePackage::FExternalArc>& ArcsFromImportedPackage = Package->GraphData.ExternalArcs.FindOrAdd(Package->ImportedPackageIds[PackageImportRef.GetImportedPackageIndex()]);
+			TArray<FPackageStorePackage::FExternalArc>& ArcsFromImportedPackage = Package->GraphData.ExternalArcs.FindOrAdd(Package->ImportedPackages[PackageImportRef.GetImportedPackageIndex()].Id);
 			FPackageStorePackage::FExternalArc Arc;
 			Arc.FromImportIndex = ExternalDependency.ImportIndex;
 			Arc.FromCommandType = ExternalDependency.ExportBundleCommandType;
@@ -1556,23 +1561,35 @@ FPackageStoreEntryResource FPackageStoreOptimizer::CreatePackageStoreEntry(const
 	FPackageStoreEntryResource Result;
 	Result.Flags = EPackageStoreEntryFlags::None;
 
-	if (OptionalSegmentPackage && OptionalSegmentPackage->HasEditorData())
+	if (OptionalSegmentPackage)
 	{
-		// AutoOptional packages are saved with editor data included
-		Result.Flags |= EPackageStoreEntryFlags::AutoOptional;
+		Result.Flags |= EPackageStoreEntryFlags::OptionalSegment;
+		if (OptionalSegmentPackage->HasEditorData())
+		{
+			// AutoOptional packages are saved with editor data included
+			Result.Flags |= EPackageStoreEntryFlags::AutoOptional;
+		}
 	}
 	
 	Result.PackageName = Package->Name;
 	Result.PackageId = FPackageId::FromName(Package->Name);
 	Result.ExportInfo.ExportCount = Package->Exports.Num();
 	Result.ExportInfo.ExportBundleCount = Package->GraphData.ExportBundles.Num();
-	Result.ImportedPackageIds = Package->ImportedPackageIds;
+	Result.ImportedPackageIds.Reserve(Package->ImportedPackages.Num());
+	for (const FPackageStorePackage::FImportedPackageRef& ImportedPackage : Package->ImportedPackages)
+	{
+		Result.ImportedPackageIds.Add(ImportedPackage.Id);
+	}
 	
 	if (OptionalSegmentPackage)
 	{
 		Result.OptionalSegmentExportInfo.ExportCount = OptionalSegmentPackage->Exports.Num();
 		Result.OptionalSegmentExportInfo.ExportBundleCount = OptionalSegmentPackage->GraphData.ExportBundles.Num();
-		Result.OptionalSegmentImportedPackageIds = OptionalSegmentPackage->ImportedPackageIds;
+		Result.OptionalSegmentImportedPackageIds.Reserve(OptionalSegmentPackage->ImportedPackages.Num());
+		for (const FPackageStorePackage::FImportedPackageRef& ImportedPackage : OptionalSegmentPackage->ImportedPackages)
+		{
+			Result.OptionalSegmentImportedPackageIds.Add(ImportedPackage.Id);
+		}
 	}
 	return Result;
 }
diff --git a/Engine/Source/Runtime/Core/Public/IO/PackageStore.h b/Engine/Source/Runtime/Core/Public/IO/PackageStore.h
--- a/Engine/Source/Runtime/Core/Public/IO/PackageStore.h
+++ b/Engine/Source/Runtime/Core/Public/IO/PackageStore.h
@@ -68,6 +68,7 @@ enum class EPackageStoreEntryFlags : uint32
 {
 	None		= 0,
 	AutoOptional= 0x02,
+	OptionalSegment = 0x04,
 };
 ENUM_CLASS_FLAGS(EPackageStoreEntryFlags);
 
@@ -108,6 +109,11 @@ struct FPackageStoreEntryResource
 		return EnumHasAnyFlags(Flags, EPackageStoreEntryFlags::AutoOptional);
 	}
 
+	bool HasOptionalSegment() const
+	{
+		return EnumHasAnyFlags(Flags, EPackageStoreEntryFlags::OptionalSegment);
+	}
+
 	CORE_API friend FArchive& operator<<(FArchive& Ar, FPackageStoreEntryResource& PackageStoreEntry);
 	
 	CORE_API friend FCbWriter& operator<<(FCbWriter& Writer, const FPackageStoreEntryResource& PackageStoreEntry);
diff --git a/Engine/Source/Runtime/CoreUObject/Private/Serialization/AsyncLoading2.cpp b/Engine/Source/Runtime/CoreUObject/Private/Serialization/AsyncLoading2.cpp
--- a/Engine/Source/Runtime/CoreUObject/Private/Serialization/AsyncLoading2.cpp
+++ b/Engine/Source/Runtime/CoreUObject/Private/Serialization/AsyncLoading2.cpp
@@ -1361,7 +1361,8 @@ public:
 struct FAsyncPackageHeaderData
 {
 	uint32 CookedHeaderSize = 0;
-	uint32 ExportCount = 0; // Need to keep this count around until after ExportMap is being cleared
+	uint32 ExportCount = 0; // Need to keep this count around after ExportMap is cleared
+	uint32 ExportBundleCount = 0; // Need to keep this count around after ExportBundleHeaders is cleared
 	TOptional<FZenPackageVersioningInfo> VersioningInfo;
 	FNameMap NameMap;
 	FName PackageName;
@@ -1372,11 +1373,17 @@ struct FAsyncPackageHeaderData
 	TArrayView<const FExportMapEntry> ExportMap;
 	TArrayView<const uint8> ArcsData;
 	TArrayView<const FBulkDataMapEntry> BulkDataMap;
+	TArrayView<const FExportBundleHeader> ExportBundleHeaders;
+	TArrayView<const FExportBundleEntry> ExportBundleEntries;
+
 	// Backed by allocation in FAsyncPackageData
 	TArrayView<FPackageId> ImportedPackageIds;
-	TArrayView<FName> ImportedPackageNames;
-	TArrayView<FExportBundleHeader> ExportBundleHeaders;
-	TArrayView<FExportBundleEntry> ExportBundleEntries;
+	TArrayView<FAsyncPackage2*> ImportedAsyncPackagesView;
+	TArrayView<FExportObject> ExportsView;
+	TArrayView<FExportBundleHeader> ExportBundleHeadersCopyForPostLoad; // TODO: Can we use ConstructedObjects or Exports instead for postloading?
+	TArrayView<FExportBundleEntry> ExportBundleEntriesCopyForPostLoad;
+
+	TArray<FName> ImportedPackageNames;
 
 	void OnReleaseHeaderBuffer()
 	{
@@ -1386,6 +1393,8 @@ struct FAsyncPackageHeaderData
 		ExportMap = TArrayView<const FExportMapEntry>();
 		ArcsData = TArrayView<const uint8>();
 		BulkDataMap = TArrayView<const FBulkDataMapEntry>();
+		ExportBundleHeaders = TArrayView<const FExportBundleHeader>();
+		ExportBundleEntries = TArrayView<const FExportBundleEntry>();
 	}
 };
 
@@ -1394,7 +1403,7 @@ struct FAsyncPackageLinkerLoadHeaderData
 {
 	TArray<uint64> ImportedPublicExportHashes;
 	TArray<FPackageObjectIndex> ImportMap;
-	TArray<uint64> PublicExportHashes;
+	TArray<FExportMapEntry> ExportMap; // Note: Currently only using the public export hash field
 };
 #endif
 
@@ -1659,7 +1668,7 @@ public:
 	{
 		UE_ASYNC_PACKAGE_LOG(Fatal, *PackageDesc, TEXT("ObjectSerializationError"),
 			TEXT("%s: Bad export index %d/%d."),
-			CurrentExport ? *CurrentExport->GetFullName() : TEXT("null"), ExportIndex, Exports.Num());
+			CurrentExport ? *CurrentExport->GetFullName() : TEXT("null"), ExportIndex, HeaderData->ExportsView.Num());
 
 		Object = nullptr;
 	}
@@ -1685,9 +1694,9 @@ public:
 		else if (Index.IsExport())
 		{
 			const int32 ExportIndex = Index.ToExport();
-			if (ExportIndex < Exports.Num())
+			if (ExportIndex < HeaderData->ExportsView.Num())
 			{
-				Object = Exports[ExportIndex].Object;
+				Object = HeaderData->ExportsView[ExportIndex].Object;
 
 #if ALT2_LOG_VERBOSE
 				const FExportMapEntry& Export = HeaderData->ExportMap[ExportIndex];
@@ -1879,7 +1888,6 @@ private:
 	TArray<FExternalReadCallback>* ExternalReadDependencies;
 	const FAsyncPackageHeaderData* HeaderData = nullptr;
 	const FLinkerInstancingContext* InstanceContext = nullptr;
-	TArrayView<const FExportObject> Exports;
 	UObject* CurrentExport = nullptr;
 	uint64 CookedSerialOffset = 0;
 	uint64 CookedSerialSize = 0;
@@ -2227,13 +2235,14 @@ struct FAsyncPackageExportToBundleMapping
 
 struct FAsyncPackageData
 {
-	uint8* MemoryBuffer = nullptr;
+	uint8* MemoryBuffer0 = nullptr;
+	uint8* MemoryBuffer1 = nullptr;
 	TArrayView<FExportObject> Exports;
 	TArrayView<FAsyncPackage2*> ImportedAsyncPackages;
 	TArrayView<FEventLoadNode2> ExportBundleNodes;
 	TArrayView<const FSHAHash> ShaderMapHashes;
 	TArrayView<FAsyncPackageExportToBundleMapping> ExportToBundleMappings;
-	int32 ExportBundleCount = 0;
+	int32 TotalExportBundleCount = 0;
 };
 
 struct FAsyncPackageSerializationState
@@ -2575,7 +2584,7 @@ struct FAsyncPackage2
 #endif
 
 private:
-	void ImportPackagesRecursiveInner(FAsyncLoadingThreadState2& ThreadState, FIoBatch& IoBatch, FPackageStore& PackageStore, const TArrayView<FPackageId>& ImportedPackageIds, const TArrayView<FName>& ImportedPackageNames, int32& ImportedPackageIndex);
+	void ImportPackagesRecursiveInner(FAsyncLoadingThreadState2& ThreadState, FIoBatch& IoBatch, FPackageStore& PackageStore, FAsyncPackageHeaderData& Header);
 
 	uint8 PackageNodesMemory[EEventLoadNode2::Package_NumPhases * sizeof(FEventLoadNode2)];
 	TArrayView<FEventLoadNode2> PackageNodes;
@@ -2759,14 +2768,14 @@ public:
 	static EEventLoadNodeExecutionResult Event_PostLoadExportBundle(FAsyncLoadingThreadState2& ThreadState, FAsyncPackage2* Package, int32 ExportBundleIndex);
 	static EEventLoadNodeExecutionResult Event_DeferredPostLoadExportBundle(FAsyncLoadingThreadState2& ThreadState, FAsyncPackage2* Package, int32 ExportBundleIndex);
 
-	void EventDrivenCreateExport(const FAsyncPackageHeaderData& Header, const TArrayView<FExportObject>& Exports, int32 LocalExportIndex);
-	bool EventDrivenSerializeExport(const FAsyncPackageHeaderData& Header, const TArrayView<FExportObject>& Exports, int32 LocalExportIndex, FExportArchive& Ar);
+	void EventDrivenCreateExport(const FAsyncPackageHeaderData& Header, int32 LocalExportIndex);
+	bool EventDrivenSerializeExport(const FAsyncPackageHeaderData& Header, int32 LocalExportIndex, FExportArchive& Ar);
 
-	UObject* EventDrivenIndexToObject(const FAsyncPackageHeaderData& Header, const TArrayView<const FExportObject>& Exports, FPackageObjectIndex Index, bool bCheckSerialized);
+	UObject* EventDrivenIndexToObject(const FAsyncPackageHeaderData& Header, FPackageObjectIndex Index, bool bCheckSerialized);
 	template<class T>
-	T* CastEventDrivenIndexToObject(const FAsyncPackageHeaderData& Header, const TArrayView<const FExportObject>& Exports, FPackageObjectIndex Index, bool bCheckSerialized)
+	T* CastEventDrivenIndexToObject(const FAsyncPackageHeaderData& Header, FPackageObjectIndex Index, bool bCheckSerialized)
 	{
-		UObject* Result = EventDrivenIndexToObject(Header, Exports, Index, bCheckSerialized);
+		UObject* Result = EventDrivenIndexToObject(Header, Index, bCheckSerialized);
 		if (!Result)
 		{
 			return nullptr;
@@ -2782,6 +2791,7 @@ public:
 	void CallCompletionCallbacks(EAsyncLoadingResult::Type LoadingResult);
 
 private:
+	void InitializeExportArchive(FExportArchive& Ar, bool bIsOptionalSegment);
 	void CreatePackageNodes(const FAsyncLoadEventSpec* EventSpecs);
 	void CreateExportBundleNodes(const FAsyncLoadEventSpec* EventSpecs);
 	void SetupSerializedArcs(const FAsyncPackageHeaderData& Header, const TArrayView<FEventLoadNode2>& ExportBundleNodes, const TArrayView<FAsyncPackage2*>& ImportedAsyncPackages);
@@ -3369,62 +3379,35 @@ private:
 		return new FAsyncPackage2(Desc, *this, GraphAllocator, EventSpecs.GetData());
 	}
 
-	void InitializeAsyncPackageFromPackageStore(FAsyncLoadingThreadState2& ThreadState, FIoBatch& IoBatch, FAsyncPackage2* AsyncPackage, const FPackageStoreEntry& PackageStoreEntry, const TArrayView<FName>& ImportedPackageNames)
+	void InitializeAsyncPackageFromPackageStore(FAsyncLoadingThreadState2& ThreadState, FIoBatch* IoBatch, FAsyncPackage2* AsyncPackage, const FPackageStoreEntry& PackageStoreEntry)
 	{
 		TRACE_CPUPROFILER_EVENT_SCOPE(InitializeAsyncPackageFromPackageStore);
 		UE_ASYNC_PACKAGE_DEBUG(AsyncPackage->Desc);
 		
 		FAsyncPackageData& Data = AsyncPackage->Data;
-		FAsyncPackageHeaderData& HeaderData = AsyncPackage->HeaderData;
 
-		const int32 ExportCount = PackageStoreEntry.ExportInfo.ExportCount;
-		const int32 ExportBundleCount = PackageStoreEntry.ExportInfo.ExportBundleCount;
-		const uint64 ExportBundleHeadersMemSize = Align(sizeof(FExportBundleHeader) * ExportBundleCount, 8);
-		const int32 ExportBundleEntriesCount = ExportCount * FExportBundleEntry::ExportCommandType_Count;
-		const uint64 ExportBundleEntriesMemSize = Align(sizeof(FExportBundleEntry) * ExportBundleEntriesCount, 8);
 		const int32 ImportedPackagesCount = PackageStoreEntry.ImportedPackageIds.Num();
 		const uint64 ImportedPackageIdsMemSize = Align(sizeof(FPackageId) * ImportedPackagesCount, 8);
-		const uint64 ImportedPackageNamesMemSize = Align(sizeof(FName) * ImportedPackageNames.Num(), 8);
 #if WITH_EDITOR
-		const int32 OptionalSegmentExportCount = PackageStoreEntry.OptionalSegmentExportInfo.ExportCount;
-		const int32 OptionalSegmentExportBundleCount = PackageStoreEntry.OptionalSegmentExportInfo.ExportBundleCount;
-		const uint64 OptionalSegmentExportBundleHeadersMemSize = Align(sizeof(FExportBundleHeader) * OptionalSegmentExportBundleCount, 8);
-		const int32 OptionalSegmentExportBundleEntriesCount = OptionalSegmentExportCount * FExportBundleEntry::ExportCommandType_Count;
-		const uint64 OptionalSegmentExportBundleEntriesMemSize = Align(sizeof(FExportBundleEntry) * OptionalSegmentExportCount * FExportBundleEntry::ExportCommandType_Count, 8);
+		const bool bHasOptionalSegment = PackageStoreEntry.OptionalSegmentExportInfo.ExportBundleCount > 0;
 		const int32 OptionalSegmentImportedPackagesCount = PackageStoreEntry.OptionalSegmentImportedPackageIds.Num();
 		const uint64 OptionalSegmentImportedPackageIdsMemSize = Align(sizeof(FPackageId) * OptionalSegmentImportedPackagesCount, 8);
 
 		const int32 TotalImportedPackagesCount = ImportedPackagesCount + OptionalSegmentImportedPackagesCount;
-		const int32 TotalExportCount = ExportCount + OptionalSegmentExportCount;
-		AsyncPackage->Data.ExportBundleCount = ExportBundleCount + OptionalSegmentExportBundleCount;
 #else
 		const int32 TotalImportedPackagesCount = ImportedPackagesCount;
-		const int32 TotalExportCount = ExportCount;
-		AsyncPackage->Data.ExportBundleCount = ExportBundleCount;
 #endif
-		const int32 ExportBundleNodeCount = AsyncPackage->Data.ExportBundleCount * EEventLoadNode2::ExportBundle_NumPhases;
 		const int32 ShaderMapHashesCount = PackageStoreEntry.ShaderMapHashes.Num();
 
 		const uint64 ImportedPackagesMemSize = Align(sizeof(FAsyncPackage2*) * TotalImportedPackagesCount, 8);
-		const uint64 ExportsMemSize = Align(sizeof(FExportObject) * TotalExportCount, 8);
-		const uint64 ExportBundleNodesMemSize = Align(sizeof(FEventLoadNode2) * ExportBundleNodeCount, 8);
 		const uint64 ShaderMapHashesMemSize = Align(sizeof(FSHAHash) * ShaderMapHashesCount, 8);
-		const uint64 ExportToBundleMappingMemSize = Align(sizeof(FAsyncPackageExportToBundleMapping) * TotalExportCount, 8);
 		const uint64 MemoryBufferSize =
 #if WITH_EDITOR
-			OptionalSegmentExportBundleHeadersMemSize +
-			OptionalSegmentExportBundleEntriesMemSize +
 			OptionalSegmentImportedPackageIdsMemSize +
 #endif
-			ExportBundleHeadersMemSize +
-			ExportBundleEntriesMemSize +
 			ImportedPackageIdsMemSize +
-			ImportedPackageNamesMemSize +
 			ImportedPackagesMemSize +
-			ExportsMemSize +
-			ExportBundleNodesMemSize +
-			ShaderMapHashesMemSize +
-			ExportToBundleMappingMemSize;
+			ShaderMapHashesMemSize;
 
 #if PLATFORM_32BITS
 		if (MemoryBufferSize > MAX_Int32)
@@ -3433,52 +3416,131 @@ private:
 			return;
 		}
 #endif
-		Data.MemoryBuffer = reinterpret_cast<uint8*>(FMemory::Malloc(MemoryBufferSize));
+		Data.MemoryBuffer0 = reinterpret_cast<uint8*>(FMemory::Malloc(MemoryBufferSize));
 
-		uint8* DataPtr = Data.MemoryBuffer;
+		uint8* DataPtr = Data.MemoryBuffer0;
 
-		Data.Exports = MakeArrayView(reinterpret_cast<FExportObject*>(DataPtr), TotalExportCount);
-		DataPtr += ExportsMemSize;
-		Data.ExportBundleNodes = MakeArrayView(reinterpret_cast<FEventLoadNode2*>(DataPtr), ExportBundleNodeCount);
-		DataPtr += ExportBundleNodesMemSize;
 		Data.ShaderMapHashes = MakeArrayView(reinterpret_cast<const FSHAHash*>(DataPtr), ShaderMapHashesCount);
 		FMemory::Memcpy((void*)Data.ShaderMapHashes.GetData(), PackageStoreEntry.ShaderMapHashes.GetData(), sizeof(FSHAHash) * ShaderMapHashesCount);
 		DataPtr += ShaderMapHashesMemSize;
-		Data.ImportedAsyncPackages = MakeArrayView(reinterpret_cast<FAsyncPackage2**>(DataPtr), 0);
+		Data.ImportedAsyncPackages = MakeArrayView(reinterpret_cast<FAsyncPackage2**>(DataPtr), TotalImportedPackagesCount);
+		FMemory::Memzero(DataPtr, ImportedPackagesMemSize);
 		DataPtr += ImportedPackagesMemSize;
-		Data.ExportToBundleMappings = MakeArrayView(reinterpret_cast<FAsyncPackageExportToBundleMapping*>(DataPtr), TotalExportCount);
-		DataPtr += ExportToBundleMappingMemSize;
 
-		HeaderData.ExportCount = ExportCount;
-		HeaderData.ExportBundleHeaders = MakeArrayView(reinterpret_cast<FExportBundleHeader*>(DataPtr), ExportBundleCount);
-		DataPtr += ExportBundleHeadersMemSize;
-		HeaderData.ExportBundleEntries = MakeArrayView(reinterpret_cast<FExportBundleEntry*>(DataPtr), ExportBundleEntriesCount);
-		DataPtr += ExportBundleEntriesMemSize;
+		FAsyncPackageHeaderData& HeaderData = AsyncPackage->HeaderData;
 		HeaderData.ImportedPackageIds = MakeArrayView(reinterpret_cast<FPackageId*>(DataPtr), ImportedPackagesCount);
 		FMemory::Memcpy((void*)HeaderData.ImportedPackageIds.GetData(), PackageStoreEntry.ImportedPackageIds.GetData(), sizeof(FPackageId) * ImportedPackagesCount);
 		DataPtr += ImportedPackageIdsMemSize;
-		if (ImportedPackageNamesMemSize)
-		{
-			HeaderData.ImportedPackageNames = MakeArrayView(reinterpret_cast<FName*>(DataPtr), ImportedPackageNames.Num());
-			FMemory::Memcpy((void*)HeaderData.ImportedPackageNames.GetData(), ImportedPackageNames.GetData(), sizeof(FName) * ImportedPackageNames.Num());
-			DataPtr += ImportedPackageNamesMemSize;
-		}
+
+		Data.TotalExportBundleCount = PackageStoreEntry.ExportInfo.ExportBundleCount;
+		HeaderData.ExportBundleCount = PackageStoreEntry.ExportInfo.ExportBundleCount;
+		HeaderData.ImportedAsyncPackagesView = Data.ImportedAsyncPackages;
 #if WITH_EDITOR
-		if (OptionalSegmentExportCount)
+		if (bHasOptionalSegment)
 		{
 			AsyncPackage->OptionalSegmentSerializationState.Emplace();
 			FAsyncPackageHeaderData& OptionalSegmentHeaderData = AsyncPackage->OptionalSegmentHeaderData.Emplace();
-			OptionalSegmentHeaderData.ExportCount = OptionalSegmentExportCount;
-			OptionalSegmentHeaderData.ExportBundleHeaders = MakeArrayView(reinterpret_cast<FExportBundleHeader*>(DataPtr), OptionalSegmentExportBundleCount);
-			DataPtr += OptionalSegmentExportBundleHeadersMemSize;
-			OptionalSegmentHeaderData.ExportBundleEntries = MakeArrayView(reinterpret_cast<FExportBundleEntry*>(DataPtr), OptionalSegmentExportBundleEntriesCount);
-			DataPtr += OptionalSegmentExportBundleEntriesMemSize;
 			OptionalSegmentHeaderData.ImportedPackageIds = MakeArrayView(reinterpret_cast<FPackageId*>(DataPtr), OptionalSegmentImportedPackagesCount);
 			FMemory::Memcpy((void*)OptionalSegmentHeaderData.ImportedPackageIds.GetData(), PackageStoreEntry.OptionalSegmentImportedPackageIds.GetData(), sizeof(FPackageId) * OptionalSegmentImportedPackagesCount);
 			DataPtr += OptionalSegmentImportedPackageIdsMemSize;
+
+			Data.TotalExportBundleCount += PackageStoreEntry.OptionalSegmentExportInfo.ExportBundleCount;
+			OptionalSegmentHeaderData.ExportBundleCount = PackageStoreEntry.OptionalSegmentExportInfo.ExportBundleCount;
+			HeaderData.ImportedAsyncPackagesView = Data.ImportedAsyncPackages.Left(ImportedPackagesCount);
+			OptionalSegmentHeaderData.ImportedAsyncPackagesView = Data.ImportedAsyncPackages.Right(OptionalSegmentImportedPackagesCount);
+		}
+#endif
+		check(DataPtr - Data.MemoryBuffer0 == MemoryBufferSize);
+
+#if ALT2_ENABLE_LINKERLOAD_SUPPORT
+		const bool bIsZenPackage = !AsyncPackage->LinkerLoadState.IsSet();
+		bool bCanImportPackagesWithIdsOnly = bIsZenPackage;
+#else
+		const bool bCanImportPackagesWithIdsOnly = true;
+#endif
+		if (bCanImportPackagesWithIdsOnly)
+		{
+			check(IoBatch);
+			AsyncPackage->ImportPackagesRecursive(ThreadState, *IoBatch, PackageStore);
+		}
+	}
+
+	void FinishInitializeAsyncPackage(FAsyncLoadingThreadState2& ThreadState, FAsyncPackage2* AsyncPackage)
+	{
+		FAsyncPackageData& Data = AsyncPackage->Data;
+		FAsyncPackageHeaderData& HeaderData = AsyncPackage->HeaderData;
+		int32 TotalExportCount = HeaderData.ExportMap.Num();
+		const uint64 ExportBundleHeadersCopyMemSize = Align(HeaderData.ExportBundleHeaders.Num() * sizeof(FExportBundleHeader), 8);
+		const uint64 ExportBundleEntriesCopyMemSize = Align(HeaderData.ExportBundleEntries.Num() * sizeof(FExportBundleEntry), 8);
+#if WITH_EDITOR
+		FAsyncPackageHeaderData* OptionalSegmentHeaderData = AsyncPackage->OptionalSegmentHeaderData.GetPtrOrNull();
+		uint64 OptionalSegmentExportBundleHeadersCopyMemSize = 0;
+		uint64 OptionalSegmentExportBundleEntriesCopyMemSize = 0;
+		if (OptionalSegmentHeaderData)
+		{
+			TotalExportCount += OptionalSegmentHeaderData->ExportMap.Num();
+			OptionalSegmentExportBundleHeadersCopyMemSize = Align(OptionalSegmentHeaderData->ExportBundleHeaders.Num() * sizeof(FExportBundleHeader), 8);
+			OptionalSegmentExportBundleEntriesCopyMemSize = Align(OptionalSegmentHeaderData->ExportBundleEntries.Num() * sizeof(FExportBundleEntry), 8);
+		}
+#endif
+		const int32 ExportBundleNodeCount = Data.TotalExportBundleCount * EEventLoadNode2::ExportBundle_NumPhases;
+		const uint64 ExportBundleNodesMemSize = Align(sizeof(FEventLoadNode2) * ExportBundleNodeCount, 8);
+		const uint64 ExportsMemSize = Align(sizeof(FExportObject) * TotalExportCount, 8);
+		const uint64 ExportToBundleMappingMemSize = Align(sizeof(FAsyncPackageExportToBundleMapping) * TotalExportCount, 8);
+
+		const uint64 MemoryBufferSize =
+			ExportBundleNodesMemSize +
+			ExportsMemSize +
+#if WITH_EDITOR
+			OptionalSegmentExportBundleHeadersCopyMemSize +
+			OptionalSegmentExportBundleEntriesCopyMemSize +
+#endif
+			ExportBundleHeadersCopyMemSize +
+			ExportBundleEntriesCopyMemSize +
+			ExportToBundleMappingMemSize;
+
+#if PLATFORM_32BITS
+		if (MemoryBufferSize > MAX_Int32)
+		{
+			UE_LOG(LogStreaming, Fatal, TEXT("Memory buffer size overflow"));
+			return;
 		}
 #endif
-		check(DataPtr - Data.MemoryBuffer == MemoryBufferSize);
+		Data.MemoryBuffer1 = reinterpret_cast<uint8*>(FMemory::Malloc(MemoryBufferSize));
+		
+		uint8* DataPtr = Data.MemoryBuffer1;
+
+		Data.Exports = MakeArrayView(reinterpret_cast<FExportObject*>(DataPtr), TotalExportCount);
+		DataPtr += ExportsMemSize;
+		Data.ExportBundleNodes = MakeArrayView(reinterpret_cast<FEventLoadNode2*>(DataPtr), ExportBundleNodeCount);
+		DataPtr += ExportBundleNodesMemSize;
+		Data.ExportToBundleMappings = MakeArrayView(reinterpret_cast<FAsyncPackageExportToBundleMapping*>(DataPtr), TotalExportCount);
+		DataPtr += ExportToBundleMappingMemSize;
+		HeaderData.ExportBundleHeadersCopyForPostLoad = MakeArrayView(reinterpret_cast<FExportBundleHeader*>(DataPtr), HeaderData.ExportBundleHeaders.Num());
+		FMemory::Memcpy(DataPtr, HeaderData.ExportBundleHeaders.GetData(), HeaderData.ExportBundleHeaders.Num() * sizeof(FExportBundleHeader));
+		DataPtr += ExportBundleHeadersCopyMemSize;
+		HeaderData.ExportBundleEntriesCopyForPostLoad = MakeArrayView(reinterpret_cast<FExportBundleEntry*>(DataPtr), HeaderData.ExportBundleEntries.Num());
+		FMemory::Memcpy(DataPtr, HeaderData.ExportBundleEntries.GetData(), HeaderData.ExportBundleEntries.Num() * sizeof(FExportBundleEntry));
+		DataPtr += ExportBundleEntriesCopyMemSize;
+
+		HeaderData.ExportsView = Data.Exports;
+
+#if WITH_EDITOR
+		if (OptionalSegmentHeaderData)
+		{
+			OptionalSegmentHeaderData->ExportBundleHeadersCopyForPostLoad = MakeArrayView(reinterpret_cast<FExportBundleHeader*>(DataPtr), OptionalSegmentHeaderData->ExportBundleHeaders.Num());
+			FMemory::Memcpy(DataPtr, OptionalSegmentHeaderData->ExportBundleHeaders.GetData(), OptionalSegmentHeaderData->ExportBundleHeaders.Num() * sizeof(FExportBundleHeader));
+			DataPtr += OptionalSegmentExportBundleHeadersCopyMemSize;
+			OptionalSegmentHeaderData->ExportBundleEntriesCopyForPostLoad = MakeArrayView(reinterpret_cast<FExportBundleEntry*>(DataPtr), OptionalSegmentHeaderData->ExportBundleEntries.Num());
+			FMemory::Memcpy(DataPtr, OptionalSegmentHeaderData->ExportBundleEntries.GetData(), OptionalSegmentHeaderData->ExportBundleEntries.Num() * sizeof(FExportBundleEntry));
+			DataPtr += OptionalSegmentExportBundleEntriesCopyMemSize;
+
+			HeaderData.ExportsView = Data.Exports.Left(HeaderData.ExportCount);
+			OptionalSegmentHeaderData->ExportsView = Data.Exports.Right(OptionalSegmentHeaderData->ExportCount);
+		}
+#endif
+
+		check(DataPtr - Data.MemoryBuffer1 == MemoryBufferSize);
 		AsyncPackage->CreateExportBundleNodes(EventSpecs.GetData());
 
 		AsyncPackage->ConstructedObjects.Reserve(Data.Exports.Num() + 1); // +1 for UPackage, may grow dynamically beoynd that
@@ -3487,7 +3549,48 @@ private:
 			Export = FExportObject();
 		}
 
-		AsyncPackage->ImportPackagesRecursive(ThreadState, IoBatch, PackageStore);
+		// Map publc export hashes to export bundles
+		if (!AsyncPackage->bLoadHasFailed
+#if ALT2_ENABLE_LINKERLOAD_SUPPORT
+			&& !AsyncPackage->LinkerLoadState.IsSet()
+#endif
+			)
+		{
+			for (int32 ExportBundleIndex = 0; ExportBundleIndex < AsyncPackage->Data.TotalExportBundleCount; ++ExportBundleIndex)
+			{
+				const FAsyncPackageHeaderData* TargetHeaderData = &AsyncPackage->HeaderData;
+				int32 LocalExportBundleIndex = ExportBundleIndex;
+				int32 ExportIndexOffset = 0;
+#if WITH_EDITOR
+				if (ExportBundleIndex >= AsyncPackage->HeaderData.ExportBundleHeaders.Num())
+				{
+					check(OptionalSegmentHeaderData);
+					TargetHeaderData = OptionalSegmentHeaderData;
+					LocalExportBundleIndex -= AsyncPackage->HeaderData.ExportBundleHeaders.Num();
+					ExportIndexOffset = AsyncPackage->HeaderData.ExportMap.Num();
+				}
+#endif
+				const FExportBundleHeader& ExportBundle = TargetHeaderData->ExportBundleHeaders[LocalExportBundleIndex];
+				for (int32 EntryIndex = ExportBundle.FirstEntryIndex, EntryEnd = ExportBundle.FirstEntryIndex + ExportBundle.EntryCount; EntryIndex < EntryEnd; ++EntryIndex)
+				{
+					const FExportBundleEntry& BundleEntry = TargetHeaderData->ExportBundleEntries[EntryIndex];
+					const FExportMapEntry& ExportMapEntry = TargetHeaderData->ExportMap[BundleEntry.LocalExportIndex];
+					FAsyncPackageExportToBundleMapping& ExportToBundleMapping = AsyncPackage->Data.ExportToBundleMappings[BundleEntry.LocalExportIndex + ExportIndexOffset];
+					ExportToBundleMapping.ExportHash = ExportMapEntry.PublicExportHash;
+					ExportToBundleMapping.BundleIndex[BundleEntry.CommandType] = ExportBundleIndex;
+				}
+			}
+		}
+
+		if (!AsyncPackage->bHasStartedImportingPackages)
+		{
+			FIoBatch IoBatch = IoDispatcher.NewBatch();
+			{
+				FPackageStoreReadScope _(PackageStore);
+				AsyncPackage->ImportPackagesRecursive(ThreadState, IoBatch, PackageStore);
+			}
+			IoBatch.Issue();
+		}
 	}
 
 	void DeleteAsyncPackage(FAsyncPackage2* Package)
@@ -3658,7 +3761,7 @@ void FAsyncLoadingThread2::ConditionalBeginPostLoad(FAsyncLoadingThreadState2& T
 			check(Package->AsyncPackageLoadingState == EAsyncPackageLoadingState2::ExportsDone);
 			Package->AsyncPackageLoadingState = EAsyncPackageLoadingState2::PostLoad;
 			Package->PostLoadGroup = nullptr;
-			for (int32 ExportBundleIndex = 0; ExportBundleIndex < Package->Data.ExportBundleCount; ++ExportBundleIndex)
+			for (int32 ExportBundleIndex = 0; ExportBundleIndex < Package->Data.TotalExportBundleCount; ++ExportBundleIndex)
 			{
 				Package->GetExportBundleNode(EEventLoadNode2::ExportBundle_PostLoad, ExportBundleIndex).ReleaseBarrier(&ThreadState);
 			}
@@ -3737,9 +3840,9 @@ FAsyncPackage2* FAsyncLoadingThread2::FindOrInsertPackage(FAsyncLoadingThreadSta
 		if (Package->AsyncPackageLoadingState >= EAsyncPackageLoadingState2::ExportsDone)
 		{
 			// The imported package has started postloading, wait for it to finish postloading before serializing any exports
-			for (int32 DependentExportBundleIndex = 0; DependentExportBundleIndex < ImportedByPackage->Data.ExportBundleCount; ++DependentExportBundleIndex)
+			for (int32 DependentExportBundleIndex = 0; DependentExportBundleIndex < ImportedByPackage->Data.TotalExportBundleCount; ++DependentExportBundleIndex)
 			{
-				for (int32 DependsOnExportBundleIndex = 0; DependsOnExportBundleIndex < Package->Data.ExportBundleCount; ++DependsOnExportBundleIndex)
+				for (int32 DependsOnExportBundleIndex = 0; DependsOnExportBundleIndex < Package->Data.TotalExportBundleCount; ++DependsOnExportBundleIndex)
 				{
 					ImportedByPackage->GetExportBundleNode(ExportBundle_Process, DependentExportBundleIndex).DependsOn(&Package->GetExportBundleNode(ExportBundle_DeferredPostLoad, DependsOnExportBundleIndex));
 				}
@@ -3811,7 +3914,7 @@ bool FAsyncLoadingThread2::CreateAsyncPackagesFromQueue(FAsyncLoadingThreadState
 		EPackageStoreEntryStatus PendingPackageStatus = PackageStore.GetPackageStoreEntry(PendingPackage->Desc.PackageIdToLoad, PackageEntry);
 		if (PendingPackageStatus == EPackageStoreEntryStatus::Ok)
 		{
-			InitializeAsyncPackageFromPackageStore(ThreadState, IoBatch, PendingPackage, PackageEntry, TArrayView<FName>());
+			InitializeAsyncPackageFromPackageStore(ThreadState, &IoBatch, PendingPackage, PackageEntry);
 			PendingPackage->StartLoading(ThreadState, IoBatch);
 			It.RemoveCurrent();
 		}
@@ -3819,9 +3922,8 @@ bool FAsyncLoadingThread2::CreateAsyncPackagesFromQueue(FAsyncLoadingThreadState
 		{
 			// Initialize package with a fake package store entry
 			FPackageStoreEntry FakePackageEntry;
-			FakePackageEntry.ExportInfo.ExportCount = 1;
 			FakePackageEntry.ExportInfo.ExportBundleCount = 1;
-			InitializeAsyncPackageFromPackageStore(ThreadState, IoBatch, PendingPackage, FakePackageEntry, TArrayView<FName>());
+			InitializeAsyncPackageFromPackageStore(ThreadState, &IoBatch, PendingPackage, FakePackageEntry);
 			// Simulate StartLoading() getting back a failed IoRequest and let it go through all package states
 			PendingPackage->AsyncPackageLoadingState = EAsyncPackageLoadingState2::WaitingForIo;
 			PendingPackage->bLoadHasFailed = true;
@@ -3964,7 +4066,7 @@ bool FAsyncLoadingThread2::CreateAsyncPackagesFromQueue(FAsyncLoadingThreadState
 						else
 #endif
 						{
-							InitializeAsyncPackageFromPackageStore(ThreadState, IoBatch, Package, PackageEntry, TArrayView<FName>());
+							InitializeAsyncPackageFromPackageStore(ThreadState, &IoBatch, Package, PackageEntry);
 						}
 						Package->StartLoading(ThreadState, IoBatch);
 					}
@@ -4348,7 +4450,7 @@ void FAsyncPackage2::SetupSerializedArcs(const FAsyncPackageHeaderData& Header,
 		ArcsArchive << FromExportBundleIndex;
 		int32 ToExportBundleIndex;
 		ArcsArchive << ToExportBundleIndex;
-		if (FromExportBundleIndex >= Data.ExportBundleCount || ToExportBundleIndex >= Data.ExportBundleCount)
+		if (FromExportBundleIndex >= Data.TotalExportBundleCount || ToExportBundleIndex >= Data.TotalExportBundleCount)
 		{
 			UE_LOG(LogStreaming, Fatal, TEXT("Corrupt package data when serializing arcs"));
 			return;
@@ -4382,7 +4484,7 @@ void FAsyncPackage2::SetupSerializedArcs(const FAsyncPackageHeaderData& Header,
 			{
 				if (FromImportIndex >= Header.ImportMap.Num() ||
 					FromCommandType >= FExportBundleEntry::ExportCommandType_Count ||
-					ToExportBundleIndex >= Data.ExportBundleCount)
+					ToExportBundleIndex >= Data.TotalExportBundleCount)
 				{
 					UE_LOG(LogStreaming, Fatal, TEXT("Corrupt package data when serializing arcs"));
 					return;
@@ -4508,13 +4610,15 @@ void FGlobalImportStore::RegistrationComplete()
 	ScriptObjects.Shrink();
 }
 
-void FAsyncPackage2::ImportPackagesRecursiveInner(FAsyncLoadingThreadState2& ThreadState, FIoBatch& IoBatch, FPackageStore& PackageStore, const TArrayView<FPackageId>& ImportedPackageIds, const TArrayView<FName>& ImportedPackageNames, int32& ImportedPackageIndex)
+void FAsyncPackage2::ImportPackagesRecursiveInner(FAsyncLoadingThreadState2& ThreadState, FIoBatch& IoBatch, FPackageStore& PackageStore, FAsyncPackageHeaderData& Header)
 {
+	const TArrayView<FPackageId>& ImportedPackageIds = Header.ImportedPackageIds;
 	const int32 ImportedPackageCount = ImportedPackageIds.Num();
 	if (!ImportedPackageCount)
 	{
 		return;
 	}
+	TArray<FName>& ImportedPackageNames = Header.ImportedPackageNames;
 	bool bHasImportedPackageNames = !ImportedPackageNames.IsEmpty();
 	check(ImportedPackageNames.Num() == 0 || ImportedPackageNames.Num() == ImportedPackageCount);
 	for (int32 LocalImportedPackageIndex = 0; LocalImportedPackageIndex < ImportedPackageCount; ++LocalImportedPackageIndex)
@@ -4666,7 +4770,6 @@ void FAsyncPackage2::ImportPackagesRecursiveInner(FAsyncLoadingThreadState2& Thr
 				UE_ASYNC_PACKAGE_LOG(Warning, Desc, TEXT("ImportPackages: SkipPackage"),
 					TEXT("Failed to load uncooked imported package with id '0x%llX' ('%s')"), ImportedPackageId.Value(), *ImportedPackageEntry.UncookedPackageName.ToString());
 			}
-			Data.ImportedAsyncPackages[ImportedPackageIndex++] = nullptr;
 			continue;
 		}
 #endif
@@ -4678,7 +4781,6 @@ void FAsyncPackage2::ImportPackagesRecursiveInner(FAsyncLoadingThreadState2& Thr
 			ImportedPackage = AsyncLoadingThread.FindAsyncPackage(ImportedPackageId);
 			if (!ImportedPackage)
 			{
-				Data.ImportedAsyncPackages[ImportedPackageIndex++] = nullptr;
 				continue;
 			}
 			bInserted = false;
@@ -4691,7 +4793,6 @@ void FAsyncPackage2::ImportPackagesRecursiveInner(FAsyncLoadingThreadState2& Thr
 					TEXT("Skipping non mounted imported package %s (0x%llX)"), *ImportedPackageNameToLoad.ToString(), ImportedPackageId.Value());
 				ImportedPackageRef.SetIsMissingPackage();
 			}
-			Data.ImportedAsyncPackages[ImportedPackageIndex++] = nullptr;
 			continue;
 		}
 		else
@@ -4716,7 +4817,7 @@ void FAsyncPackage2::ImportPackagesRecursiveInner(FAsyncLoadingThreadState2& Thr
 				TEXT("Imported package with id '0x%llX' is already being loaded."), ImportedPackageId.ValueForDebugging());
 		}
 		ImportedPackage->AddRef();
-		Data.ImportedAsyncPackages[ImportedPackageIndex++] = ImportedPackage;
+		Header.ImportedAsyncPackagesView[LocalImportedPackageIndex] = ImportedPackage;
 
 		if (bIsZenPackage != bIsZenPackageImport)
 		{
@@ -4742,7 +4843,7 @@ void FAsyncPackage2::ImportPackagesRecursiveInner(FAsyncLoadingThreadState2& Thr
 				else
 #endif
 				{
-					AsyncLoadingThread.InitializeAsyncPackageFromPackageStore(ThreadState, IoBatch, ImportedPackage, ImportedPackageEntry, TArrayView<FName>());
+					AsyncLoadingThread.InitializeAsyncPackageFromPackageStore(ThreadState, &IoBatch, ImportedPackage, ImportedPackageEntry);
 				}
 				ImportedPackage->StartLoading(ThreadState, IoBatch);
 			}
@@ -4758,27 +4859,16 @@ void FAsyncPackage2::ImportPackagesRecursive(FAsyncLoadingThreadState2& ThreadSt
 	}
 	bHasStartedImportingPackages = true;
 
-	int32 ImportedPackageCount = HeaderData.ImportedPackageIds.Num();
-#if WITH_EDITOR
-	if (OptionalSegmentHeaderData.IsSet())
-	{
-		ImportedPackageCount += OptionalSegmentHeaderData->ImportedPackageIds.Num();
-	}
-#endif
-	if (!ImportedPackageCount)
+	if (Data.ImportedAsyncPackages.IsEmpty())
 	{
 		return;
 	}
 
-	int32 ImportedPackageIndex = 0;
-
-	Data.ImportedAsyncPackages = MakeArrayView(Data.ImportedAsyncPackages.GetData(), ImportedPackageCount);
-	
-	ImportPackagesRecursiveInner(ThreadState, IoBatch, PackageStore, HeaderData.ImportedPackageIds, HeaderData.ImportedPackageNames, ImportedPackageIndex);
+	ImportPackagesRecursiveInner(ThreadState, IoBatch, PackageStore, HeaderData);
 #if WITH_EDITOR
 	if (OptionalSegmentHeaderData.IsSet())
 	{
-		ImportPackagesRecursiveInner(ThreadState, IoBatch, PackageStore, OptionalSegmentHeaderData->ImportedPackageIds, OptionalSegmentHeaderData->ImportedPackageNames, ImportedPackageIndex);
+		ImportPackagesRecursiveInner(ThreadState, IoBatch, PackageStore, OptionalSegmentHeaderData.GetValue());
 	}
 #endif
 
@@ -5004,17 +5094,29 @@ static void ReadAsyncPackageHeader(FAsyncPackageSerializationState& Serializatio
 	HeaderData.ExportMap = TArrayView<const FExportMapEntry>(
 		reinterpret_cast<const FExportMapEntry*>(PackageHeaderDataPtr + PackageSummary->ExportMapOffset),
 		(PackageSummary->ExportBundleEntriesOffset - PackageSummary->ExportMapOffset) / sizeof(FExportMapEntry));
-	check(HeaderData.ExportMap.Num() == HeaderData.ExportCount);
+	HeaderData.ExportCount = HeaderData.ExportMap.Num();
 
 	const uint64 ExportBundleHeadersOffset = PackageSummary->GraphDataOffset;
+	const uint64 ExportBundleEntriesSize = ExportBundleHeadersOffset - PackageSummary->ExportBundleEntriesOffset;
+	const int32 ExportBundleEntriesCount = ExportBundleEntriesSize / sizeof(FExportBundleEntry);
+
+	if (ExportBundleEntriesCount != HeaderData.ExportCount * FExportBundleEntry::ExportCommandType_Count)
+	{
+		UE_LOG(LogStreaming, Fatal, TEXT("Corrupt Zen header in package %s"), *HeaderData.PackageName.ToString());
+		return;
+	}
+
+	HeaderData.ExportBundleHeaders = TArrayView<const FExportBundleHeader>(
+		reinterpret_cast<const FExportBundleHeader*>(PackageHeaderDataPtr + ExportBundleHeadersOffset),
+		HeaderData.ExportBundleCount);
+	HeaderData.ExportBundleEntries = TArrayView<const FExportBundleEntry>(
+		reinterpret_cast<const FExportBundleEntry*>(PackageHeaderDataPtr + PackageSummary->ExportBundleEntriesOffset),
+		ExportBundleEntriesCount);
+
 	const uint64 ExportBundleHeadersSize = sizeof(FExportBundleHeader) * HeaderData.ExportBundleHeaders.Num();
-	const uint64 ExportBundleEntriesSize = sizeof(FExportBundleEntry) * HeaderData.ExportBundleEntries.Num();
 	const uint64 ArcsDataOffset = ExportBundleHeadersOffset + ExportBundleHeadersSize;
 	const uint64 ArcsDataSize = PackageSummary->HeaderSize - ArcsDataOffset;
-	check(ExportBundleEntriesSize == PackageSummary->GraphDataOffset - PackageSummary->ExportBundleEntriesOffset);
 	HeaderData.ArcsData = TArrayView64<const uint8>(PackageHeaderDataPtr + ArcsDataOffset, ArcsDataSize);
-	FMemory::Memcpy(HeaderData.ExportBundleHeaders.GetData(), PackageHeaderDataPtr + ExportBundleHeadersOffset, ExportBundleHeadersSize);
-	FMemory::Memcpy(HeaderData.ExportBundleEntries.GetData(), PackageHeaderDataPtr + PackageSummary->ExportBundleEntriesOffset, ExportBundleEntriesSize);
 
 	SerializationState.AllExportDataPtr = PackageHeaderDataPtr + PackageSummary->HeaderSize;
 }
@@ -5123,7 +5225,7 @@ EEventLoadNodeExecutionResult FAsyncPackage2::ProcessLinkerLoadPackageSummary(FA
 		}
 	}
 
-	LinkerLoadState->LinkerLoadHeaderData.PublicExportHashes.SetNum(LinkerLoadState->Linker->ExportMap.Num());
+	LinkerLoadState->LinkerLoadHeaderData.ExportMap.SetNum(LinkerLoadState->Linker->ExportMap.Num());
 	for (int32 ExportIndex = 0; ExportIndex < LinkerLoadState->Linker->ExportMap.Num(); ++ExportIndex)
 	{
 		FObjectExport& ObjectExport = LinkerLoadState->Linker->ExportMap[ExportIndex];
@@ -5150,29 +5252,26 @@ EEventLoadNodeExecutionResult FAsyncPackage2::ProcessLinkerLoadPackageSummary(FA
 				PackageRelativeExportPath.Append(FStringView(NameStr, NameLen));
 			}
 			uint64 PublicExportHash = CityHash64(reinterpret_cast<const char*>(PackageRelativeExportPath.GetData() + 1), (PackageRelativeExportPath.Len() - 1) * sizeof(TCHAR));
-			LinkerLoadState->LinkerLoadHeaderData.PublicExportHashes[ExportIndex] = PublicExportHash;
+			LinkerLoadState->LinkerLoadHeaderData.ExportMap[ExportIndex].PublicExportHash = PublicExportHash;
 		}
 	}
 
 	FPackageStoreEntry PackageStoreEntry;
-	PackageStoreEntry.ExportInfo.ExportCount = LinkerLoadState->Linker->ExportMap.Num();
 	PackageStoreEntry.ExportInfo.ExportBundleCount = 1;
 	PackageStoreEntry.ImportedPackageIds = ImportedPackageIds;
+	AsyncLoadingThread.InitializeAsyncPackageFromPackageStore(ThreadState, nullptr, this, PackageStoreEntry);
 
-	FIoBatch IoBatch = AsyncLoadingThread.IoDispatcher.NewBatch();
-	{
-		FPackageStoreReadScope _(AsyncLoadingThread.PackageStore);
-		AsyncLoadingThread.InitializeAsyncPackageFromPackageStore(ThreadState, IoBatch, this, PackageStoreEntry, ImportedPackageNames);
-	}
-	IoBatch.Issue();
-
+	HeaderData.ImportedPackageNames = ImportedPackageNames;
 	HeaderData.ImportedPublicExportHashes = LinkerLoadState->LinkerLoadHeaderData.ImportedPublicExportHashes;
 	HeaderData.ImportMap = LinkerLoadState->LinkerLoadHeaderData.ImportMap;
+	HeaderData.ExportMap = LinkerLoadState->LinkerLoadHeaderData.ExportMap;
+
+	AsyncLoadingThread.FinishInitializeAsyncPackage(ThreadState, this);
 
 	FLoadedPackageRef& PackageRef = AsyncLoadingThread.GlobalImportStore.FindPackageRefChecked(Desc.UPackageId, Desc.UPackageName);
 	if (!bLoadHasFailed)
 	{
-		PackageRef.ReserveSpaceForPublicExports(LinkerLoadState->LinkerLoadHeaderData.PublicExportHashes.Num());
+		PackageRef.ReserveSpaceForPublicExports(LinkerLoadState->LinkerLoadHeaderData.ExportMap.Num());
 #if WITH_EDITORONLY_DATA
 		// Create metadata object, this needs to happen before any other package wants to use our exports
 		LinkerLoadState->MetaDataIndex = LinkerLoadState->Linker->LoadMetaDataFromExportMap(false);
@@ -5349,7 +5448,7 @@ bool FAsyncPackage2::ProcessLinkerLoadPackageExports(FAsyncLoadingThreadState2&
 			ExportObject.Object = Object;
 			ExportObject.bWasFoundInMemory = true; // Make sure that the async flags are cleared in ClearConstructedObjects
 			EInternalObjectFlags FlagsToSet = EInternalObjectFlags::Async;
-			uint64 PublicExportHash = LinkerLoadState->LinkerLoadHeaderData.PublicExportHashes[ExportIndex];
+			uint64 PublicExportHash = LinkerLoadState->LinkerLoadHeaderData.ExportMap[ExportIndex].PublicExportHash;
 			if (Desc.bCanBeImported && PublicExportHash)
 			{
 				FlagsToSet |= EInternalObjectFlags::LoaderImport;
@@ -5481,31 +5580,6 @@ EEventLoadNodeExecutionResult FAsyncPackage2::Event_ProcessPackageSummary(FAsync
 			}
 		}
 
-		for (int32 ExportBundleIndex = 0; ExportBundleIndex < Package->Data.ExportBundleCount; ++ExportBundleIndex)
-		{
-			const FAsyncPackageHeaderData* HeaderData = &Package->HeaderData;
-			int32 LocalExportBundleIndex = ExportBundleIndex;
-			int32 ExportIndexOffset = 0;
-#if WITH_EDITOR
-			if (ExportBundleIndex >= Package->HeaderData.ExportBundleHeaders.Num())
-			{
-				check(OptionalSegmentHeaderData);
-				HeaderData = OptionalSegmentHeaderData;
-				LocalExportBundleIndex -= Package->HeaderData.ExportBundleHeaders.Num();
-				ExportIndexOffset = Package->HeaderData.ExportMap.Num();
-			}
-#endif
-			const FExportBundleHeader& ExportBundle = HeaderData->ExportBundleHeaders[LocalExportBundleIndex];
-			for (int32 EntryIndex = ExportBundle.FirstEntryIndex, EntryEnd = ExportBundle.FirstEntryIndex + ExportBundle.EntryCount; EntryIndex < EntryEnd; ++EntryIndex)
-			{
-				const FExportBundleEntry& BundleEntry = HeaderData->ExportBundleEntries[EntryIndex];
-				const FExportMapEntry& ExportMapEntry = HeaderData->ExportMap[BundleEntry.LocalExportIndex];
-				FAsyncPackageExportToBundleMapping& ExportToBundleMapping = Package->Data.ExportToBundleMappings[BundleEntry.LocalExportIndex + ExportIndexOffset];
-				ExportToBundleMapping.ExportHash = ExportMapEntry.PublicExportHash;
-				ExportToBundleMapping.BundleIndex[BundleEntry.CommandType] = ExportBundleIndex;
-			}
-		}
-
 		check(Package->Desc.PackageIdToLoad == FPackageId::FromName(Package->HeaderData.PackageName));
 		if (Package->Desc.PackagePathToLoad.IsEmpty())
 		{
@@ -5545,6 +5619,8 @@ EEventLoadNodeExecutionResult FAsyncPackage2::Event_ProcessPackageSummary(FAsync
 		TRACE_LOADTIME_PACKAGE_SUMMARY(Package, Package->HeaderData.PackageName, Package->HeaderData.PackageSummary->HeaderSize, Package->HeaderData.ImportMap.Num(), Package->HeaderData.ExportMap.Num());
 	}
 
+	Package->AsyncLoadingThread.FinishInitializeAsyncPackage(ThreadState, Package);
+
 	Package->AsyncPackageLoadingState = EAsyncPackageLoadingState2::WaitingForDependencies;
 	if (!Package->AsyncLoadingThread.bHasRegisteredAllScriptObjects)
 	{
@@ -5565,6 +5641,47 @@ EEventLoadNodeExecutionResult FAsyncPackage2::Event_DependenciesReady(FAsyncLoad
 	return EEventLoadNodeExecutionResult::Complete;
 }
 
+void FAsyncPackage2::InitializeExportArchive(FExportArchive& Ar, bool bIsOptionalSegment)
+{
+	Ar.SetUEVer(LinkerRoot->GetLinkerPackageVersion());
+	Ar.SetLicenseeUEVer(LinkerRoot->GetLinkerLicenseeVersion());
+	// Ar.SetEngineVer(Summary.SavedByEngineVersion); // very old versioning scheme
+	if (!LinkerRoot->GetLinkerCustomVersions().GetAllVersions().IsEmpty())
+	{
+		Ar.SetCustomVersions(LinkerRoot->GetLinkerCustomVersions());
+	}
+	Ar.SetUseUnversionedPropertySerialization((LinkerRoot->GetPackageFlags() & PKG_UnversionedProperties) != 0);
+	Ar.SetIsLoadingFromCookedPackage((LinkerRoot->GetPackageFlags() & PKG_Cooked) != 0);
+	Ar.SetIsLoading(true);
+	Ar.SetIsPersistent(true);
+	if (LinkerRoot->GetPackageFlags() & PKG_FilterEditorOnly)
+	{
+		Ar.SetFilterEditorOnly(true);
+	}
+	Ar.ArAllowLazyLoading = true;
+
+	// FExportArchive special fields
+	Ar.PackageDesc = &Desc;
+	Ar.HeaderData = &HeaderData;
+#if WITH_EDITOR
+	if (bIsOptionalSegment)
+	{
+		Ar.HeaderData = OptionalSegmentHeaderData.GetPtrOrNull();
+		check(Ar.HeaderData);
+		Ar.ExportBundleChunkData = OptionalSegmentSerializationState->IoRequest.GetResultOrDie().Data();
+	}
+	else
+#endif
+	{
+		Ar.ExportBundleChunkData = SerializationState.IoRequest.GetResultOrDie().Data();
+	}
+	Ar.ImportStore = &ImportStore;
+	Ar.ExternalReadDependencies = &ExternalReadDependencies;
+	Ar.InstanceContext = InstanceContext.Get();
+	Ar.bIsOptionalSegment = bIsOptionalSegment;
+	Ar.bExportsCookedToSeparateArchive = Ar.UEVer() >= EUnrealEngineObjectUE5Version::DATA_RESOURCES;
+}
+
 EEventLoadNodeExecutionResult FAsyncPackage2::Event_ProcessExportBundle(FAsyncLoadingThreadState2& ThreadState, FAsyncPackage2* Package, int32 InExportBundleIndex)
 {
 	TRACE_CPUPROFILER_EVENT_SCOPE(Event_ProcessExportBundle);
@@ -5580,33 +5697,7 @@ EEventLoadNodeExecutionResult FAsyncPackage2::Event_ProcessExportBundle(FAsyncLo
 	}
 #endif
 
-	auto FilterExport = [](const EExportFilterFlags FilterFlags) -> bool
-	{
-#if WITH_EDITOR
-		return false;
-#elif UE_SERVER
-		return !!(static_cast<uint32>(FilterFlags) & static_cast<uint32>(EExportFilterFlags::NotForServer));
-#elif !WITH_SERVER_CODE
-		return !!(static_cast<uint32>(FilterFlags) & static_cast<uint32>(EExportFilterFlags::NotForClient));
-#else
-		static const bool bIsDedicatedServer = !GIsClient && GIsServer;
-		static const bool bIsClientOnly = GIsClient && !GIsServer;
-
-		if (bIsDedicatedServer && static_cast<uint32>(FilterFlags) & static_cast<uint32>(EExportFilterFlags::NotForServer))
-		{
-			return true;
-		}
-
-		if (bIsClientOnly && static_cast<uint32>(FilterFlags) & static_cast<uint32>(EExportFilterFlags::NotForClient))
-		{
-			return true;
-		}
-
-		return false;
-#endif
-	};
-
-	check(InExportBundleIndex < Package->Data.ExportBundleCount);
+	check(InExportBundleIndex < Package->Data.TotalExportBundleCount);
 	
 	if (!Package->bLoadHasFailed)
 	{
@@ -5615,7 +5706,6 @@ EEventLoadNodeExecutionResult FAsyncPackage2::Event_ProcessExportBundle(FAsyncLo
 		const FAsyncPackageHeaderData* HeaderData;
 		const FExportBundleHeader* ExportBundle;
 		FAsyncPackageSerializationState* SerializationState;
-		TArrayView<FExportObject> Exports = Package->Data.Exports;
 		bIsOptionalSegment = InExportBundleIndex >= Package->HeaderData.ExportBundleHeaders.Num();
 
 		if (bIsOptionalSegment)
@@ -5623,7 +5713,6 @@ EEventLoadNodeExecutionResult FAsyncPackage2::Event_ProcessExportBundle(FAsyncLo
 			HeaderData = Package->OptionalSegmentHeaderData.GetPtrOrNull();
 			check(HeaderData);
 			ExportBundle = &HeaderData->ExportBundleHeaders[InExportBundleIndex - Package->HeaderData.ExportBundleHeaders.Num()];
-			Exports.RightInline(HeaderData->ExportCount);
 			SerializationState = Package->OptionalSegmentSerializationState.GetPtrOrNull();
 			check(SerializationState);
 		}
@@ -5632,13 +5721,11 @@ EEventLoadNodeExecutionResult FAsyncPackage2::Event_ProcessExportBundle(FAsyncLo
 			HeaderData = &Package->HeaderData;
 			ExportBundle = &Package->HeaderData.ExportBundleHeaders[InExportBundleIndex];
 			SerializationState = &Package->SerializationState;
-			Exports.LeftInline(Package->HeaderData.ExportCount);
 		}
 #else
 		const FAsyncPackageHeaderData* HeaderData = &Package->HeaderData;
 		const FExportBundleHeader* ExportBundle = &Package->HeaderData.ExportBundleHeaders[InExportBundleIndex];
 		FAsyncPackageSerializationState* SerializationState = &Package->SerializationState;
-		const TArrayView<FExportObject>& Exports = Package->Data.Exports;
 #endif
 		const FIoBuffer& IoBuffer = SerializationState->IoRequest.GetResultOrDie();
 		const uint64 AllExportDataSize = IoBuffer.DataSize() - (SerializationState->AllExportDataPtr - IoBuffer.Data());
@@ -5647,35 +5734,7 @@ EEventLoadNodeExecutionResult FAsyncPackage2::Event_ProcessExportBundle(FAsyncLo
 			SerializationState->CurrentExportDataPtr = SerializationState->AllExportDataPtr + ExportBundle->SerialOffset;
 		}
 		FExportArchive Ar(SerializationState->AllExportDataPtr, SerializationState->CurrentExportDataPtr, AllExportDataSize);
-		{
-			Ar.SetUEVer(Package->LinkerRoot->GetLinkerPackageVersion());
-			Ar.SetLicenseeUEVer(Package->LinkerRoot->GetLinkerLicenseeVersion());
-			// Ar.SetEngineVer(Summary.SavedByEngineVersion); // very old versioning scheme
-			if (!Package->LinkerRoot->GetLinkerCustomVersions().GetAllVersions().IsEmpty())
-			{
-				Ar.SetCustomVersions(Package->LinkerRoot->GetLinkerCustomVersions());
-			}
-			Ar.SetUseUnversionedPropertySerialization((Package->LinkerRoot->GetPackageFlags() & PKG_UnversionedProperties) != 0);
-			Ar.SetIsLoadingFromCookedPackage((Package->LinkerRoot->GetPackageFlags() & PKG_Cooked) != 0);
-			Ar.SetIsLoading(true);
-			Ar.SetIsPersistent(true);
-			if (Package->LinkerRoot->GetPackageFlags() & PKG_FilterEditorOnly)
-			{
-				Ar.SetFilterEditorOnly(true);
-			}
-			Ar.ArAllowLazyLoading = true;
-
-			// FExportArchive special fields
-			Ar.PackageDesc = &Package->Desc;
-			Ar.HeaderData = HeaderData;
-			Ar.ImportStore = &Package->ImportStore;
-			Ar.Exports = Exports;
-			Ar.ExternalReadDependencies = &Package->ExternalReadDependencies;
-			Ar.InstanceContext = Package->InstanceContext.Get();
-			Ar.bIsOptionalSegment = bIsOptionalSegment;
-			Ar.bExportsCookedToSeparateArchive = Ar.UEVer() >= EUnrealEngineObjectUE5Version::DATA_RESOURCES;			
-			Ar.ExportBundleChunkData = SerializationState->IoRequest.GetResultOrDie().Data();
-		}
+		Package->InitializeExportArchive(Ar, bIsOptionalSegment);
 
 		while (Package->ExportBundleEntryIndex < int32(ExportBundle->EntryCount))
 		{
@@ -5685,12 +5744,11 @@ EEventLoadNodeExecutionResult FAsyncPackage2::Event_ProcessExportBundle(FAsyncLo
 				return EEventLoadNodeExecutionResult::Timeout;
 			}
 			const FExportMapEntry& ExportMapEntry = HeaderData->ExportMap[BundleEntry.LocalExportIndex];
-			FExportObject& Export = Exports[BundleEntry.LocalExportIndex];
-			Export.bFiltered = FilterExport(ExportMapEntry.FilterFlags);
+			FExportObject& Export = HeaderData->ExportsView[BundleEntry.LocalExportIndex];
 
 			if (BundleEntry.CommandType == FExportBundleEntry::ExportCommandType_Create)
 			{
-				Package->EventDrivenCreateExport(*HeaderData, Exports, BundleEntry.LocalExportIndex);
+				Package->EventDrivenCreateExport(*HeaderData, BundleEntry.LocalExportIndex);
 			}
 			else
 			{
@@ -5710,7 +5768,7 @@ EEventLoadNodeExecutionResult FAsyncPackage2::Event_ProcessExportBundle(FAsyncLo
 					TEXT("%s: Serial size mismatch: Expected read size %d, Remaining archive size: %d"),
 					Object ? *Object->GetFullName() : TEXT("null"), CookedSerialSize, uint64(Ar.TotalSize() - Pos));
 
-				const bool bSerialized = Package->EventDrivenSerializeExport(*HeaderData, Exports, BundleEntry.LocalExportIndex, Ar);
+				const bool bSerialized = Package->EventDrivenSerializeExport(*HeaderData, BundleEntry.LocalExportIndex, Ar);
 				if (!bSerialized)
 				{
 					Ar.Skip(CookedSerialSize);
@@ -5732,7 +5790,7 @@ EEventLoadNodeExecutionResult FAsyncPackage2::Event_ProcessExportBundle(FAsyncLo
 	
 	Package->ExportBundleEntryIndex = 0;
 
-	if (++Package->ProcessedExportBundlesCount == Package->Data.ExportBundleCount)
+	if (++Package->ProcessedExportBundlesCount == Package->Data.TotalExportBundleCount)
 	{
 		Package->ProcessedExportBundlesCount = 0;
 		Package->HeaderData.OnReleaseHeaderBuffer();
@@ -5761,7 +5819,7 @@ EEventLoadNodeExecutionResult FAsyncPackage2::Event_ProcessExportBundle(FAsyncLo
 	return EEventLoadNodeExecutionResult::Complete;
 }
 
-UObject* FAsyncPackage2::EventDrivenIndexToObject(const FAsyncPackageHeaderData& Header, const TArrayView<const FExportObject>& Exports, FPackageObjectIndex Index, bool bCheckSerialized)
+UObject* FAsyncPackage2::EventDrivenIndexToObject(const FAsyncPackageHeaderData& Header, FPackageObjectIndex Index, bool bCheckSerialized)
 {
 	UObject* Result = nullptr;
 	if (Index.IsNull())
@@ -5770,7 +5828,7 @@ UObject* FAsyncPackage2::EventDrivenIndexToObject(const FAsyncPackageHeaderData&
 	}
 	if (Index.IsExport())
 	{
-		Result = Exports[Index.ToExport()].Object;
+		Result = Header.ExportsView[Index.ToExport()].Object;
 		UE_CLOG(!Result, LogStreaming, Warning, TEXT("Missing Dependency, missing export 0x%llX in package %s"),
 			Index.Value(),
 			*Desc.PackagePathToLoad.GetPackageFName().ToString());
@@ -5803,13 +5861,12 @@ UObject* FAsyncPackage2::EventDrivenIndexToObject(const FAsyncPackageHeaderData&
 	return Result;
 }
 
-
-void FAsyncPackage2::EventDrivenCreateExport(const FAsyncPackageHeaderData& Header, const TArrayView<FExportObject>& Exports, int32 LocalExportIndex)
+void FAsyncPackage2::EventDrivenCreateExport(const FAsyncPackageHeaderData& Header, int32 LocalExportIndex)
 {
 	TRACE_CPUPROFILER_EVENT_SCOPE(CreateExport);
 
 	const FExportMapEntry& Export = Header.ExportMap[LocalExportIndex];
-	FExportObject& ExportObject = Exports[LocalExportIndex];
+	FExportObject& ExportObject = Header.ExportsView[LocalExportIndex];
 	UObject*& Object = ExportObject.Object;
 	check(!Object);
 
@@ -5821,6 +5878,33 @@ void FAsyncPackage2::EventDrivenCreateExport(const FAsyncPackageHeaderData& Head
 		ObjectName = Header.NameMap.GetName(Export.ObjectName);
 	}
 
+	auto FilterExport = [](const EExportFilterFlags FilterFlags) -> bool
+	{
+#if WITH_EDITOR
+		return false;
+#elif UE_SERVER
+		return !!(static_cast<uint32>(FilterFlags) & static_cast<uint32>(EExportFilterFlags::NotForServer));
+#elif !WITH_SERVER_CODE
+		return !!(static_cast<uint32>(FilterFlags) & static_cast<uint32>(EExportFilterFlags::NotForClient));
+#else
+		static const bool bIsDedicatedServer = !GIsClient && GIsServer;
+		static const bool bIsClientOnly = GIsClient && !GIsServer;
+
+		if (bIsDedicatedServer && static_cast<uint32>(FilterFlags) & static_cast<uint32>(EExportFilterFlags::NotForServer))
+		{
+			return true;
+		}
+
+		if (bIsClientOnly && static_cast<uint32>(FilterFlags) & static_cast<uint32>(EExportFilterFlags::NotForClient))
+		{
+			return true;
+		}
+
+		return false;
+#endif
+	};
+
+	ExportObject.bFiltered = FilterExport(Export.FilterFlags);
 	if (ExportObject.bFiltered | ExportObject.bExportLoadFailed)
 	{
 		if (ExportObject.bExportLoadFailed)
@@ -5835,8 +5919,8 @@ void FAsyncPackage2::EventDrivenCreateExport(const FAsyncPackageHeaderData& Head
 	}
 
 	bool bIsCompleteyLoaded = false;
-	UClass* LoadClass = Export.ClassIndex.IsNull() ? UClass::StaticClass() : CastEventDrivenIndexToObject<UClass>(Header, Exports, Export.ClassIndex, true);
-	UObject* ThisParent = Export.OuterIndex.IsNull() ? LinkerRoot : EventDrivenIndexToObject(Header, Exports, Export.OuterIndex, false);
+	UClass* LoadClass = Export.ClassIndex.IsNull() ? UClass::StaticClass() : CastEventDrivenIndexToObject<UClass>(Header, Export.ClassIndex, true);
+	UObject* ThisParent = Export.OuterIndex.IsNull() ? LinkerRoot : EventDrivenIndexToObject(Header, Export.OuterIndex, false);
 
 	if (!LoadClass)
 	{
@@ -5853,7 +5937,7 @@ void FAsyncPackage2::EventDrivenCreateExport(const FAsyncPackageHeaderData& Head
 	check(!dynamic_cast<UObjectRedirector*>(ThisParent));
 	if (!Export.SuperIndex.IsNull())
 	{
-		ExportObject.SuperObject = EventDrivenIndexToObject(Header, Exports, Export.SuperIndex, false);
+		ExportObject.SuperObject = EventDrivenIndexToObject(Header, Export.SuperIndex, false);
 		if (!ExportObject.SuperObject)
 		{
 			UE_ASYNC_PACKAGE_LOG(Error, Desc, TEXT("CreateExport"), TEXT("Could not find SuperStruct object for %s"), *ObjectName.ToString());
@@ -5863,7 +5947,7 @@ void FAsyncPackage2::EventDrivenCreateExport(const FAsyncPackageHeaderData& Head
 	}
 	// Find the Archetype object for the one we are loading.
 	check(!Export.TemplateIndex.IsNull());
-	ExportObject.TemplateObject = EventDrivenIndexToObject(Header, Exports, Export.TemplateIndex, true);
+	ExportObject.TemplateObject = EventDrivenIndexToObject(Header, Export.TemplateIndex, true);
 	if (!ExportObject.TemplateObject)
 	{
 		UE_ASYNC_PACKAGE_LOG(Error, Desc, TEXT("CreateExport"), TEXT("Could not find template object for %s"), *ObjectName.ToString());
@@ -6007,13 +6091,13 @@ void FAsyncPackage2::EventDrivenCreateExport(const FAsyncPackageHeaderData& Head
 	Object->SetInternalFlags(FlagsToSet);
 }
 
-bool FAsyncPackage2::EventDrivenSerializeExport(const FAsyncPackageHeaderData& Header, const TArrayView<FExportObject>& Exports, int32 LocalExportIndex, FExportArchive& Ar)
+bool FAsyncPackage2::EventDrivenSerializeExport(const FAsyncPackageHeaderData& Header, int32 LocalExportIndex, FExportArchive& Ar)
 {
 	LLM_SCOPE(ELLMTag::UObject);
 	TRACE_CPUPROFILER_EVENT_SCOPE(SerializeExport);
 
 	const FExportMapEntry& Export = Header.ExportMap[LocalExportIndex];
-	FExportObject& ExportObject = Exports[LocalExportIndex];
+	FExportObject& ExportObject = Header.ExportsView[LocalExportIndex];
 	UObject* Object = ExportObject.Object;
 	check(Object || (ExportObject.bFiltered | ExportObject.bExportLoadFailed));
 
@@ -6052,7 +6136,7 @@ bool FAsyncPackage2::EventDrivenSerializeExport(const FAsyncPackageHeaderData& H
 		}
 	}
 
-	const UClass* LoadClass = Export.ClassIndex.IsNull() ? UClass::StaticClass() : CastEventDrivenIndexToObject<UClass>(Header, Exports, Export.ClassIndex, true);
+	const UClass* LoadClass = Export.ClassIndex.IsNull() ? UClass::StaticClass() : CastEventDrivenIndexToObject<UClass>(Header, Export.ClassIndex, true);
 	UE_TRACE_METADATA_SCOPE_ASSET(Object, LoadClass);
 	LLM_SCOPED_TAG_WITH_OBJECT_IN_SET(Object, ELLMTagSet::Assets);
 	LLM_SCOPED_TAG_WITH_OBJECT_IN_SET(LoadClass, ELLMTagSet::AssetClasses);
@@ -6374,7 +6458,7 @@ void FAsyncPackage2::ConditionalBeginProcessPackageExports(FAsyncLoadingThreadSt
 				}
 #endif
 			}
-			for (int32 ExportBundleIndex = 0; ExportBundleIndex < Package->Data.ExportBundleCount; ++ExportBundleIndex)
+			for (int32 ExportBundleIndex = 0; ExportBundleIndex < Package->Data.TotalExportBundleCount; ++ExportBundleIndex)
 			{
 				Package->GetExportBundleNode(EEventLoadNode2::ExportBundle_Process, ExportBundleIndex).ReleaseBarrier(&ThreadState);
 			}
@@ -6431,7 +6515,7 @@ EEventLoadNodeExecutionResult FAsyncPackage2::Event_PostLoadExportBundle(FAsyncL
 		}
 	}*/
 	
-	check(InExportBundleIndex < Package->Data.ExportBundleCount);
+	check(InExportBundleIndex < Package->Data.TotalExportBundleCount);
 
 	EEventLoadNodeExecutionResult LoadingState = EEventLoadNodeExecutionResult::Complete;
 
@@ -6451,30 +6535,26 @@ EEventLoadNodeExecutionResult FAsyncPackage2::Event_PostLoadExportBundle(FAsyncL
 #if WITH_EDITOR
 		const FAsyncPackageHeaderData* HeaderData;
 		const FExportBundleHeader* ExportBundle;
-		TArrayView<FExportObject> Exports = Package->Data.Exports;
-		if (InExportBundleIndex >= Package->HeaderData.ExportBundleHeaders.Num())
+		if (InExportBundleIndex >= Package->HeaderData.ExportBundleHeadersCopyForPostLoad.Num())
 		{
 			HeaderData = Package->OptionalSegmentHeaderData.GetPtrOrNull();
 			check(HeaderData);
-			ExportBundle = &HeaderData->ExportBundleHeaders[InExportBundleIndex - Package->HeaderData.ExportBundleHeaders.Num()];
-			Exports.RightInline(HeaderData->ExportCount);
+			ExportBundle = &HeaderData->ExportBundleHeadersCopyForPostLoad[InExportBundleIndex - Package->HeaderData.ExportBundleHeadersCopyForPostLoad.Num()];
 		}
 		else
 		{
 			HeaderData = &Package->HeaderData;
-			ExportBundle = &Package->HeaderData.ExportBundleHeaders[InExportBundleIndex];
-			Exports.LeftInline(Package->HeaderData.ExportCount);
+			ExportBundle = &Package->HeaderData.ExportBundleHeadersCopyForPostLoad[InExportBundleIndex];
 		}
 #else
 		const FAsyncPackageHeaderData* HeaderData = &Package->HeaderData;
-		const FExportBundleHeader* ExportBundle = &Package->HeaderData.ExportBundleHeaders[InExportBundleIndex];
-		const TArrayView<FExportObject>& Exports = Package->Data.Exports;
+		const FExportBundleHeader* ExportBundle = &Package->HeaderData.ExportBundleHeadersCopyForPostLoad[InExportBundleIndex];
 #endif
 
 		TRACE_LOADTIME_POSTLOAD_SCOPE;
 		while (Package->ExportBundleEntryIndex < int32(ExportBundle->EntryCount))
 		{
-			const FExportBundleEntry& BundleEntry = HeaderData->ExportBundleEntries[ExportBundle->FirstEntryIndex + Package->ExportBundleEntryIndex];
+			const FExportBundleEntry& BundleEntry = HeaderData->ExportBundleEntriesCopyForPostLoad[ExportBundle->FirstEntryIndex + Package->ExportBundleEntryIndex];
 			if (ThreadState.IsTimeLimitExceeded(TEXT("Event_PostLoadExportBundle")))
 			{
 				LoadingState = EEventLoadNodeExecutionResult::Timeout;
@@ -6485,7 +6565,7 @@ EEventLoadNodeExecutionResult FAsyncPackage2::Event_PostLoadExportBundle(FAsyncL
 			{
 				do
 				{
-					FExportObject& Export = Exports[BundleEntry.LocalExportIndex];
+					FExportObject& Export = HeaderData->ExportsView[BundleEntry.LocalExportIndex];
 					if (Export.bFiltered | Export.bExportLoadFailed)
 					{
 						break;
@@ -6522,7 +6602,7 @@ EEventLoadNodeExecutionResult FAsyncPackage2::Event_PostLoadExportBundle(FAsyncL
 
 	Package->ExportBundleEntryIndex = 0;
 
-	if (++Package->ProcessedExportBundlesCount == Package->Data.ExportBundleCount)
+	if (++Package->ProcessedExportBundlesCount == Package->Data.TotalExportBundleCount)
 	{
 		Package->ProcessedExportBundlesCount = 0;
 		if (Package->LinkerRoot && !Package->bLoadHasFailed)
@@ -6536,7 +6616,7 @@ EEventLoadNodeExecutionResult FAsyncPackage2::Event_PostLoadExportBundle(FAsyncL
 
 		check(Package->AsyncPackageLoadingState == EAsyncPackageLoadingState2::PostLoad);
 		Package->AsyncPackageLoadingState = EAsyncPackageLoadingState2::DeferredPostLoad;
-		for (int32 ExportBundleIndex = 0; ExportBundleIndex < Package->Data.ExportBundleCount; ++ExportBundleIndex)
+		for (int32 ExportBundleIndex = 0; ExportBundleIndex < Package->Data.TotalExportBundleCount; ++ExportBundleIndex)
 		{
 			Package->GetExportBundleNode(ExportBundle_DeferredPostLoad, ExportBundleIndex).ReleaseBarrier(&ThreadState);
 		}
@@ -6562,7 +6642,7 @@ EEventLoadNodeExecutionResult FAsyncPackage2::Event_DeferredPostLoadExportBundle
 	}
 #endif
 
-	check(InExportBundleIndex < Package->Data.ExportBundleCount);
+	check(InExportBundleIndex < Package->Data.TotalExportBundleCount);
 	EEventLoadNodeExecutionResult LoadingState = EEventLoadNodeExecutionResult::Complete;
 
 	if (Package->bLoadHasFailed)
@@ -6577,30 +6657,26 @@ EEventLoadNodeExecutionResult FAsyncPackage2::Event_DeferredPostLoadExportBundle
 #if WITH_EDITOR
 		const FAsyncPackageHeaderData* HeaderData;
 		const FExportBundleHeader* ExportBundle;
-		TArrayView<FExportObject> Exports = Package->Data.Exports;
-		if (InExportBundleIndex >= Package->HeaderData.ExportBundleHeaders.Num())
+		if (InExportBundleIndex >= Package->HeaderData.ExportBundleHeadersCopyForPostLoad.Num())
 		{
 			HeaderData = Package->OptionalSegmentHeaderData.GetPtrOrNull();
 			check(HeaderData);
-			ExportBundle = &HeaderData->ExportBundleHeaders[InExportBundleIndex - Package->HeaderData.ExportBundleHeaders.Num()];
-			Exports.RightInline(HeaderData->ExportCount);
+			ExportBundle = &HeaderData->ExportBundleHeadersCopyForPostLoad[InExportBundleIndex - Package->HeaderData.ExportBundleHeadersCopyForPostLoad.Num()];
 		}
 		else
 		{
 			HeaderData = &Package->HeaderData;
-			ExportBundle = &Package->HeaderData.ExportBundleHeaders[InExportBundleIndex];
-			Exports.LeftInline(Package->HeaderData.ExportCount);
+			ExportBundle = &Package->HeaderData.ExportBundleHeadersCopyForPostLoad[InExportBundleIndex];
 		}
 #else
 		const FAsyncPackageHeaderData* HeaderData = &Package->HeaderData;
-		const FExportBundleHeader* ExportBundle = &Package->HeaderData.ExportBundleHeaders[InExportBundleIndex];
-		const TArrayView<FExportObject>& Exports = Package->Data.Exports;
+		const FExportBundleHeader* ExportBundle = &Package->HeaderData.ExportBundleHeadersCopyForPostLoad[InExportBundleIndex];
 #endif
 
 		TRACE_LOADTIME_POSTLOAD_SCOPE;
 		while (Package->ExportBundleEntryIndex < int32(ExportBundle->EntryCount))
 		{
-			const FExportBundleEntry& BundleEntry = HeaderData->ExportBundleEntries[ExportBundle->FirstEntryIndex + Package->ExportBundleEntryIndex];
+			const FExportBundleEntry& BundleEntry = HeaderData->ExportBundleEntriesCopyForPostLoad[ExportBundle->FirstEntryIndex + Package->ExportBundleEntryIndex];
 			if (ThreadState.IsTimeLimitExceeded(TEXT("Event_DeferredPostLoadExportBundle")))
 			{
 				LoadingState = EEventLoadNodeExecutionResult::Timeout;
@@ -6611,7 +6687,7 @@ EEventLoadNodeExecutionResult FAsyncPackage2::Event_DeferredPostLoadExportBundle
 			{
 				do
 				{
-					FExportObject& Export = Exports[BundleEntry.LocalExportIndex];
+					FExportObject& Export = HeaderData->ExportsView[BundleEntry.LocalExportIndex];
 					if (Export.bFiltered | Export.bExportLoadFailed)
 					{
 						break;
@@ -6642,7 +6718,7 @@ EEventLoadNodeExecutionResult FAsyncPackage2::Event_DeferredPostLoadExportBundle
 
 	Package->ExportBundleEntryIndex = 0;
 
-	if (++Package->ProcessedExportBundlesCount == Package->Data.ExportBundleCount)
+	if (++Package->ProcessedExportBundlesCount == Package->Data.TotalExportBundleCount)
 	{
 		Package->ProcessedExportBundlesCount = 0;
 		check(Package->AsyncPackageLoadingState == EAsyncPackageLoadingState2::DeferredPostLoad);
@@ -6661,7 +6737,7 @@ FEventLoadNode2& FAsyncPackage2::GetPackageNode(EEventLoadNode2 Phase)
 
 FEventLoadNode2& FAsyncPackage2::GetExportBundleNode(EEventLoadNode2 Phase, uint32 ExportBundleIndex)
 {
-	check(ExportBundleIndex < uint32(Data.ExportBundleCount));
+	check(ExportBundleIndex < uint32(Data.TotalExportBundleCount));
 	uint32 ExportBundleNodeIndex = ExportBundleIndex * EEventLoadNode2::ExportBundle_NumPhases + Phase;
 	return Data.ExportBundleNodes[ExportBundleNodeIndex];
 }
@@ -7783,7 +7859,7 @@ void FAsyncPackage2::CreatePackageNodes(const FAsyncLoadEventSpec* EventSpecs)
 void FAsyncPackage2::CreateExportBundleNodes(const FAsyncLoadEventSpec* EventSpecs)
 {
 	const int32 BarrierCount = 1;
-	for (int32 ExportBundleIndex = 0; ExportBundleIndex < Data.ExportBundleCount; ++ExportBundleIndex)
+	for (int32 ExportBundleIndex = 0; ExportBundleIndex < Data.TotalExportBundleCount; ++ExportBundleIndex)
 	{
 		uint32 NodeIndex = EEventLoadNode2::ExportBundle_NumPhases * ExportBundleIndex;
 		for (int32 Phase = 0; Phase < EEventLoadNode2::ExportBundle_NumPhases; ++Phase)
@@ -7813,7 +7889,8 @@ FAsyncPackage2::~FAsyncPackage2()
 	checkf(ConstructedObjects.Num() == 0, TEXT("ClearConstructedObjects() has not been called for package %s"),
 		*Desc.PackagePathToLoad.GetPackageFName().ToString());
 
-	FMemory::Free(Data.MemoryBuffer);
+	FMemory::Free(Data.MemoryBuffer0);
+	FMemory::Free(Data.MemoryBuffer1);
 }
 
 void FAsyncPackage2::ReleaseRef()
