AsyncLoading2: Reset the ImportedAsyncPackages member array before releasing the async package references.
Fix rare check failure when HaveAllDependenciesReachedStateDebug returns false due to reading a stale package pointer.

#rnx

diff --git a/Engine/Source/Runtime/CoreUObject/Private/Serialization/AsyncLoading2.cpp b/Engine/Source/Runtime/CoreUObject/Private/Serialization/AsyncLoading2.cpp
--- a/Engine/Source/Runtime/CoreUObject/Private/Serialization/AsyncLoading2.cpp
+++ b/Engine/Source/Runtime/CoreUObject/Private/Serialization/AsyncLoading2.cpp
@@ -7908,14 +7908,18 @@ void FAsyncPackage2::ReleaseRef()
 void FAsyncPackage2::ClearImportedPackages()
 {
 	TRACE_CPUPROFILER_EVENT_SCOPE(ClearImportedPackages);
-	for (FAsyncPackage2* ImportedAsyncPackage : Data.ImportedAsyncPackages)
+	// Reset the ImportedAsyncPackages member array before releasing the async package references.
+	// ReleaseRef may queue up the ImportedAsyncPackage for deletion on the ALT.
+	// ImportedAsyncPackages of this package can still be accessed by both the GT and the ALT.
+	TArrayView<FAsyncPackage2*> LocalImportedAsyncPackages = Data.ImportedAsyncPackages;
+	Data.ImportedAsyncPackages = MakeArrayView(Data.ImportedAsyncPackages.GetData(), 0);
+	for (FAsyncPackage2* ImportedAsyncPackage : LocalImportedAsyncPackages)
 	{
 		if (ImportedAsyncPackage)
 		{
 			ImportedAsyncPackage->ReleaseRef();
 		}
 	}
-	Data.ImportedAsyncPackages = MakeArrayView(Data.ImportedAsyncPackages.GetData(), 0);
 }
 
 void FAsyncPackage2::ClearConstructedObjects()
