fix for FReferenceCollectorArchive not handling unresolved FObjectPtr
stop gap fix to prevent crashing on unresolved pointers. matches the behaviour in reference collector where it skip unresolved references

https://p4-swarm.epicgames.net/reviews/24661544

diff --git a/Engine/Source/Runtime/CoreUObject/Private/UObject/UObjectGlobals.cpp b/Engine/Source/Runtime/CoreUObject/Private/UObject/UObjectGlobals.cpp
--- a/Engine/Source/Runtime/CoreUObject/Private/UObject/UObjectGlobals.cpp
+++ b/Engine/Source/Runtime/CoreUObject/Private/UObject/UObjectGlobals.cpp
@@ -132,7 +132,7 @@ CSV_DEFINE_CATEGORY(UObject, false);
 namespace LoadPackageStats
 {
 	static double LoadPackageTimeSec = 0.0;
-	static int NumPackagesLoaded = 0;
+	int NumPackagesLoaded = 0;
 	static FCookStatsManager::FAutoRegisterCallback RegisterCookStats([](FCookStatsManager::AddStatFuncRef AddStat)
 	{
 		AddStat(TEXT("Package.Load"), FCookStatsManager::CreateKeyValueArray(
@@ -4960,7 +4960,10 @@ FArchive& FReferenceCollectorArchive::operator<<(UObject*& Object)
 
 FArchive& FReferenceCollectorArchive::operator<<(FObjectPtr& Object)
 {
-	Collector.AddStableReference(reinterpret_cast<UObject**>(&Object));
+	if (Object.IsResolved())
+	{
+		Collector.AddStableReference(reinterpret_cast<UObject**>(&Object));
+	}
 	return *this;
 }
 
diff --git a/Engine/Source/Runtime/CoreUObject/Tests/Serialization/ReferenceCollectorTest.cpp b/Engine/Source/Runtime/CoreUObject/Tests/Serialization/ReferenceCollectorTest.cpp
new file mode 100644
--- /dev/null
+++ b/Engine/Source/Runtime/CoreUObject/Tests/Serialization/ReferenceCollectorTest.cpp
@@ -0,0 +1,66 @@
+// Copyright Epic Games, Inc. All Rights Reserved.
+
+#if WITH_LOW_LEVEL_TESTS
+#include "TestHarness.h"
+#include "UObject/Package.h"
+#include "../ObjectPtrTestClass.h"
+#include "Misc/ScopeExit.h"
+#include "../ObjectRefTrackingTestBase.h"
+
+
+class FTestReferencecollector : public FReferenceCollector
+{
+public:
+	virtual bool IsIgnoringArchetypeRef() const override
+	{
+		return true;
+	}
+
+	virtual bool IsIgnoringTransient() const override
+	{
+		return true;
+	}
+
+	virtual void HandleObjectReference(UObject*& InObject, const UObject* InReferencingObject, const FProperty* InReferencingProperty) override
+	{
+		++Count;
+	}
+
+	int Count = 0;
+};
+
+TEST_CASE("CoreUObject::ReferenceCollector::TestUnResolved")
+{
+	const FName TestPackage1Name(TEXT("/Engine/Test/RefCollector/Transient"));
+	UPackage* TestPackage1 = NewObject<UPackage>(nullptr, TestPackage1Name, RF_Transient);
+	TestPackage1->AddToRoot();
+	UObjectPtrTestClassWithRef* TestObject1 = NewObject<UObjectPtrTestClassWithRef>(TestPackage1, TEXT("TestObject1"));
+	UObjectPtrTestClass* Inner = NewObject<UObjectPtrTestClass>(TestObject1, TEXT("TestObject1"));
+	ON_SCOPE_EXIT
+	{
+		TestPackage1->RemoveFromRoot();
+	};
+
+#if UE_WITH_OBJECT_HANDLE_LATE_RESOLVE
+	{
+		FObjectPtr ObjPtr(MakeUnresolvedHandle(Inner));
+		TObjectPtr<UObjectPtrTestClass> Ptr = *reinterpret_cast<TObjectPtr<UObjectPtrTestClass>*>(&ObjPtr);
+		REQUIRE(!Ptr.IsResolved());
+		TestObject1->ObjectPtr = Ptr;
+		REQUIRE(!TestObject1->ObjectPtr.IsResolved());
+	}
+#endif
+
+	FTestReferencecollector Collector;
+	FReferenceCollectorArchive Ar(nullptr, Collector);
+
+	TestObject1->Serialize(Ar);
+
+#if UE_WITH_OBJECT_HANDLE_LATE_RESOLVE
+	CHECK(Collector.Count == 1);
+#else
+	CHECK(Collector.Count == 2);
+#endif
+	
+}
+#endif
\ No newline at end of file
