Latest Clang toolchain fixes for FSoftObjectPath, FSoftClassPath, FSoftObjectPtr, TSoftObjectPtr and TWeakInterfacePtr. deprecated-copy-with-user-provided-copy.

diff --git a/Engine/Source/Runtime/CoreUObject/Public/UObject/SoftObjectPath.h b/Engine/Source/Runtime/CoreUObject/Public/UObject/SoftObjectPath.h
--- a/Engine/Source/Runtime/CoreUObject/Public/UObject/SoftObjectPath.h
+++ b/Engine/Source/Runtime/CoreUObject/Public/UObject/SoftObjectPath.h
@@ -40,21 +40,12 @@ struct FUObjectSerializeContext;
  */
 struct COREUOBJECT_API FSoftObjectPath
 {
-	FSoftObjectPath() {}
-
-	/** Construct from another soft object path */
-	FSoftObjectPath(const FSoftObjectPath& Other)
-		: AssetPath(Other.AssetPath)
-		, SubPathString(Other.SubPathString)
-	{
-	}
-
-	/** Construct from a moveable soft object path */
-	FSoftObjectPath(FSoftObjectPath&& Other)
-		: AssetPath(Other.AssetPath)
-		, SubPathString(MoveTemp(Other.SubPathString))
-	{
-	}
+	FSoftObjectPath() = default;
+	FSoftObjectPath(const FSoftObjectPath& Other) = default;
+	FSoftObjectPath(FSoftObjectPath&& Other) = default;
+	~FSoftObjectPath() = default;
+	FSoftObjectPath& operator=(const FSoftObjectPath& Path) = default;
+	FSoftObjectPath& operator=(FSoftObjectPath&& Path) = default;
 
 	/** Construct from a path string. Non-explicit for backwards compatibility. */
 	FSoftObjectPath(const FString& Path)						{ SetPath(FStringView(Path)); }
@@ -106,11 +97,6 @@ PRAGMA_ENABLE_DEPRECATION_WARNINGS;
 		}
 	}
 
-	~FSoftObjectPath() {}
-PRAGMA_DISABLE_DEPRECATION_WARNINGS	
-	FSoftObjectPath& operator=(const FSoftObjectPath& Path)	= default;
-	FSoftObjectPath& operator=(FSoftObjectPath&& Path) = default;
-PRAGMA_ENABLE_DEPRECATION_WARNINGS
 	FSoftObjectPath& operator=(const FTopLevelAssetPath Path)			{ SetPath(Path); return *this; }
 	FSoftObjectPath& operator=(const FString& Path)						{ SetPath(FStringView(Path)); return *this; }
 	FSoftObjectPath& operator=(FWideStringView Path)					{ SetPath(Path); return *this; }
@@ -397,12 +383,12 @@ inline FStringBuilderBase& operator<<(FStringBuilderBase& Builder, const FSoftOb
  */
 struct COREUOBJECT_API FSoftClassPath : public FSoftObjectPath
 {
-	FSoftClassPath()
-	{ }
-
-	FSoftClassPath(FSoftClassPath const& Other)
-		: FSoftObjectPath(Other)
-	{ }
+	FSoftClassPath() = default;
+	FSoftClassPath(const FSoftClassPath& Other) = default;
+	FSoftClassPath(FSoftClassPath&& Other) = default;
+	~FSoftClassPath() = default;
+	FSoftClassPath& operator=(const FSoftClassPath& Path) = default;
+	FSoftClassPath& operator=(FSoftClassPath&& Path) = default;
 
 	/**
 	 * Construct from a path string
diff --git a/Engine/Source/Runtime/CoreUObject/Public/UObject/SoftObjectPtr.h b/Engine/Source/Runtime/CoreUObject/Public/UObject/SoftObjectPtr.h
--- a/Engine/Source/Runtime/CoreUObject/Public/UObject/SoftObjectPtr.h
+++ b/Engine/Source/Runtime/CoreUObject/Public/UObject/SoftObjectPtr.h
@@ -21,24 +21,18 @@
 struct FSoftObjectPtr : public TPersistentObjectPtr<FSoftObjectPath>
 {
 public:	
-	/** Default constructor, will be null */
-	FORCEINLINE FSoftObjectPtr()
-	{
-	}
-
-	/** Construct from another soft pointer */
-	FORCEINLINE FSoftObjectPtr(const FSoftObjectPtr& Other)
-	{
-		(*this)=Other;
-	}
+	FORCEINLINE FSoftObjectPtr() = default;
+	FORCEINLINE FSoftObjectPtr(const FSoftObjectPtr& Other) = default;
+	FORCEINLINE FSoftObjectPtr(FSoftObjectPtr&& Other) = default;
+	FORCEINLINE ~FSoftObjectPtr() = default;
+	FORCEINLINE FSoftObjectPtr& operator=(const FSoftObjectPtr& Other) = default;
+	FORCEINLINE FSoftObjectPtr& operator=(FSoftObjectPtr&& Other) = default;
 
-	/** Construct from a soft object path */
 	explicit FORCEINLINE FSoftObjectPtr(const FSoftObjectPath& ObjectPath)
 		: TPersistentObjectPtr<FSoftObjectPath>(ObjectPath)
 	{
 	}
 
-	/** Construct from an object already in memory */
 	explicit FORCEINLINE FSoftObjectPtr(const UObject* Object)
 	{
 		(*this)=Object;
@@ -115,10 +109,12 @@ struct TSoftObjectPtr
 public:
 	using ElementType = T;
 	
-	/** Default constructor, will be null */
-	FORCEINLINE TSoftObjectPtr()
-	{
-	}
+	FORCEINLINE TSoftObjectPtr() = default;
+	FORCEINLINE TSoftObjectPtr(const TSoftObjectPtr& Other) = default;
+	FORCEINLINE TSoftObjectPtr(TSoftObjectPtr&& Other) = default;
+	FORCEINLINE ~TSoftObjectPtr() = default;
+	FORCEINLINE TSoftObjectPtr& operator=(const TSoftObjectPtr& Other) = default;
+	FORCEINLINE TSoftObjectPtr& operator=(TSoftObjectPtr&& Other) = default;
 	
 	/** Construct from another soft pointer */
 	template <class U, class = decltype(ImplicitConv<T*>((U*)nullptr))>
diff --git a/Engine/Source/Runtime/CoreUObject/Public/UObject/WeakInterfacePtr.h b/Engine/Source/Runtime/CoreUObject/Public/UObject/WeakInterfacePtr.h
--- a/Engine/Source/Runtime/CoreUObject/Public/UObject/WeakInterfacePtr.h
+++ b/Engine/Source/Runtime/CoreUObject/Public/UObject/WeakInterfacePtr.h
@@ -18,13 +18,12 @@ struct TWeakInterfacePtr
 {
 	using ElementType = T;
 	
-	/**
-	 * Construct a new default weak pointer, pointing to null object.
-	 */
-	TWeakInterfacePtr() 
-		: InterfaceInstance(nullptr) 
-	{
-	}
+	FORCEINLINE TWeakInterfacePtr() = default;
+	FORCEINLINE TWeakInterfacePtr(const TWeakInterfacePtr& Other) = default;
+	FORCEINLINE TWeakInterfacePtr(TWeakInterfacePtr&& Other) = default;
+	FORCEINLINE ~TWeakInterfacePtr() = default;
+	FORCEINLINE TWeakInterfacePtr& operator=(const TWeakInterfacePtr& Other) = default;
+	FORCEINLINE TWeakInterfacePtr& operator=(TWeakInterfacePtr&& Other) = default;
 
 	/**
 	 * Construct from an object pointer
@@ -48,7 +47,6 @@ struct TWeakInterfacePtr
 	 * @param Interface The interface pointer to create a weak pointer to. There must be a UObject behind the interface.
 	 */
 	TWeakInterfacePtr(T* Interface)
-		: InterfaceInstance(nullptr)
 	{
 		ObjectInstance = Cast<UObject>(Interface);
 		if (ObjectInstance != nullptr)
@@ -59,7 +57,6 @@ struct TWeakInterfacePtr
 
 	UE_DEPRECATED(4.27, "Please use the constructor that takes a pointer")
 	TWeakInterfacePtr(T& Interface)
-		: InterfaceInstance(nullptr)
 	{
 		ObjectInstance = Cast<UObject>(&Interface);
 		if (ObjectInstance != nullptr)
@@ -144,16 +141,6 @@ struct TWeakInterfacePtr
 		return *this;
 	}
 
-	/**
-	 * Assign from another weak pointer.
-	 */
-	FORCEINLINE TWeakInterfacePtr<T>& operator=(const TWeakInterfacePtr<T>& Other)
-	{
-		ObjectInstance = Other.ObjectInstance;
-		InterfaceInstance = Other.InterfaceInstance;
-		return *this;
-	}
-
 	/**
 	 * Assign from a script interface.
 	 */
@@ -198,5 +185,5 @@ struct TWeakInterfacePtr
 
 private:
 	TWeakObjectPtr<UObject> ObjectInstance;
-	T* InterfaceInstance;
+	T* InterfaceInstance = nullptr;
 };
