Latest Clang toolchain fixes forTSoftClassPtr. deprecated-copy-with-user-provided-copy.

diff --git a/Engine/Source/Runtime/CoreUObject/Public/UObject/SoftObjectPtr.h b/Engine/Source/Runtime/CoreUObject/Public/UObject/SoftObjectPtr.h
--- a/Engine/Source/Runtime/CoreUObject/Public/UObject/SoftObjectPtr.h
+++ b/Engine/Source/Runtime/CoreUObject/Public/UObject/SoftObjectPtr.h
@@ -439,10 +439,12 @@ class TSoftClassPtr
 public:
 	using ElementType = TClass;
 	
-	/** Default constructor, will be null */
-	FORCEINLINE TSoftClassPtr()
-	{
-	}
+	FORCEINLINE TSoftClassPtr() = default;
+	FORCEINLINE TSoftClassPtr(const TSoftClassPtr& Other) = default;
+	FORCEINLINE TSoftClassPtr(TSoftClassPtr&& Other) = default;
+	FORCEINLINE ~TSoftClassPtr() = default;
+	FORCEINLINE TSoftClassPtr& operator=(const TSoftClassPtr& Other) = default;
+	FORCEINLINE TSoftClassPtr& operator=(TSoftClassPtr&& Other) = default;
 		
 	/** Construct from another soft pointer */
 	template <class TClassA, class = decltype(ImplicitConv<TClass*>((TClassA*)nullptr))>
