Core: Move UObject::ImplementsGetWorld to WITH_EDITOR and make it virtual

This allows classes to override it and force it to true to workaround the nuances of how the default implementation will return false even if a class overrides the function but forwards its GetWorld call onto something else that doesn't (ie, it's outer)

diff --git a/Engine/Source/Runtime/CoreUObject/Private/UObject/Obj.cpp b/Engine/Source/Runtime/CoreUObject/Private/UObject/Obj.cpp
--- a/Engine/Source/Runtime/CoreUObject/Private/UObject/Obj.cpp
+++ b/Engine/Source/Runtime/CoreUObject/Private/UObject/Obj.cpp
@@ -869,10 +869,10 @@ FString UObject::GetDetailedInfo() const
 
 #if WITH_ENGINE
 
-#if DO_CHECK
+#if DO_CHECK || WITH_EDITOR
 // Used to check to see if a derived class actually implemented GetWorld() or not
 thread_local bool bGetWorldOverridden = false;
-#endif
+#endif // #if DO_CHECK || WITH_EDITOR
 
 class UWorld* UObject::GetWorld() const
 {
@@ -881,7 +881,7 @@ class UWorld* UObject::GetWorld() const
 		return Outer->GetWorld();
 	}
 
-#if DO_CHECK
+#if DO_CHECK || WITH_EDITOR
 	bGetWorldOverridden = false;
 #endif
 	return nullptr;
@@ -928,18 +928,17 @@ class UWorld* UObject::GetWorldChecked(bool& bSupported) const
 	return World;
 }
 
+#if WITH_EDITOR
+
 bool UObject::ImplementsGetWorld() const
 {
-#if DO_CHECK
 	bGetWorldOverridden = true;
 	GetWorld();
 	return bGetWorldOverridden;
-#else
-	return true;
-#endif	
 }
 
-#endif
+#endif // #if WITH_EDITOR
+#endif // #if WITH_ENGINE
 
 #define PROFILE_ConditionalBeginDestroy (0)
 
diff --git a/Engine/Source/Runtime/CoreUObject/Public/UObject/Object.h b/Engine/Source/Runtime/CoreUObject/Public/UObject/Object.h
--- a/Engine/Source/Runtime/CoreUObject/Public/UObject/Object.h
+++ b/Engine/Source/Runtime/CoreUObject/Public/UObject/Object.h
@@ -698,9 +698,12 @@ public:
 	/** Internal function used by UEngine::GetWorldFromContextObject() */
 	class UWorld* GetWorldChecked(bool& bSupported) const;
 
-	/** Checks to see if GetWorld() is implemented on a specific class */
-	bool ImplementsGetWorld() const;
-#endif
+#if WITH_EDITOR
+	/** Checks to see if GetWorld() is implemented on a specific class. Optionally overridable in derived classes. */
+	virtual bool ImplementsGetWorld() const;
+#endif // #if WITH_EDITOR
+
+#endif // #if WITH_ENGINE
 
 	/**
 	 * Callback for retrieving a textual representation of natively serialized properties.  Child classes should implement this method if they wish
