Skip ReferenceCollectorTest
#rnx

diff --git a/Engine/Source/Runtime/CoreUObject/Tests/Serialization/ReferenceCollectorTest.cpp b/Engine/Source/Runtime/CoreUObject/Tests/Serialization/ReferenceCollectorTest.cpp
--- a/Engine/Source/Runtime/CoreUObject/Tests/Serialization/ReferenceCollectorTest.cpp
+++ b/Engine/Source/Runtime/CoreUObject/Tests/Serialization/ReferenceCollectorTest.cpp
@@ -29,7 +29,7 @@ public:
 	int Count = 0;
 };
 
-TEST_CASE("CoreUObject::ReferenceCollector::TestUnResolved")
+TEST_CASE("CoreUObject::ReferenceCollector::TestUnResolved", "[.]")
 {
 	const FName TestPackage1Name(TEXT("/Engine/Test/RefCollector/Transient"));
 	UPackage* TestPackage1 = NewObject<UPackage>(nullptr, TestPackage1Name, RF_Transient);
