Fixing non-unity builds

diff --git a/Engine/Source/Runtime/CoreUObject/Private/UObject/PropertyBaseObject.cpp b/Engine/Source/Runtime/CoreUObject/Private/UObject/PropertyBaseObject.cpp
--- a/Engine/Source/Runtime/CoreUObject/Private/UObject/PropertyBaseObject.cpp
+++ b/Engine/Source/Runtime/CoreUObject/Private/UObject/PropertyBaseObject.cpp
@@ -13,6 +13,7 @@
 #include "UObject/LinkerPlaceholderClass.h"
 #include "UObject/LinkerPlaceholderExportObject.h"
 #include "UObject/LinkerLoadImportBehavior.h"
+#include "UObject/UObjectThreadContext.h"
 #include "Misc/StringBuilder.h"
 
 /*-----------------------------------------------------------------------------
