Added possibility to disable subobject logic in transaction

diff --git a/Engine/Source/Runtime/CoreUObject/Private/UObject/Obj.cpp b/Engine/Source/Runtime/CoreUObject/Private/UObject/Obj.cpp
--- a/Engine/Source/Runtime/CoreUObject/Private/UObject/Obj.cpp
+++ b/Engine/Source/Runtime/CoreUObject/Private/UObject/Obj.cpp
@@ -1297,6 +1297,11 @@ bool UObject::Modify( bool bAlwaysMarkDirty/*=true*/ )
 
 	return bSavedToTransactionBuffer;
 }
+
+bool UObject::IsCapturingAsRootObjectForTransaction() const
+{
+	return false;
+}
 #endif
 
 bool UObject::IsSelected() const
diff --git a/Engine/Source/Runtime/CoreUObject/Public/UObject/Object.h b/Engine/Source/Runtime/CoreUObject/Public/UObject/Object.h
--- a/Engine/Source/Runtime/CoreUObject/Public/UObject/Object.h
+++ b/Engine/Source/Runtime/CoreUObject/Public/UObject/Object.h
@@ -286,6 +286,7 @@ public:
 	 */
 	virtual void PreSave(FObjectPreSaveContext SaveContext);
 
+#if WITH_EDITOR
 	/**
 	 * Note that the object will be modified.  If we are currently recording into the 
 	 * transaction buffer (undo/redo), save a copy of this object into the buffer and 
@@ -295,13 +296,16 @@ public:
 	 *								currently recording an active undo/redo transaction
 	 * @return true if the object was saved to the transaction buffer
 	 */
-#if WITH_EDITOR
 	virtual bool Modify( bool bAlwaysMarkDirty=true );
 
+	/** Method to disable the sub object special handling in the transaction buffer for this object*/
+	virtual bool IsCapturingAsRootObjectForTransaction() const;
+
 	/** Utility to allow overrides of Modify to avoid doing work if this object cannot be safely modified */
 	bool CanModify() const;
 #else
 	FORCEINLINE bool Modify(bool bAlwaysMarkDirty = true) { return false; }
+	FORCEINLINE bool IsCapturingAsRootObjectForTransaction() const { return false; };
 #endif
 
 #if WITH_EDITOR
diff --git a/Engine/Source/Runtime/Engine/Private/TransactionCommon.cpp b/Engine/Source/Runtime/Engine/Private/TransactionCommon.cpp
--- a/Engine/Source/Runtime/Engine/Private/TransactionCommon.cpp
+++ b/Engine/Source/Runtime/Engine/Private/TransactionCommon.cpp
@@ -14,7 +14,7 @@ FPersistentObjectRef::FPersistentObjectRef(UObject* InObject)
 	{
 		auto UseOuter = [](const UObject* Obj)
 		{
-			if (Obj == nullptr)
+			if (Obj == nullptr || Obj->IsCapturingAsRootObjectForTransaction())
 			{
 				return false;
 			}
