Expose the generic EAxis enum as a blueprint type.

#rnx

diff --git a/Engine/Source/Runtime/CoreUObject/Public/UObject/NoExportTypes.h b/Engine/Source/Runtime/CoreUObject/Public/UObject/NoExportTypes.h
--- a/Engine/Source/Runtime/CoreUObject/Public/UObject/NoExportTypes.h
+++ b/Engine/Source/Runtime/CoreUObject/Public/UObject/NoExportTypes.h
@@ -116,7 +116,7 @@ namespace ELogTimes
 }
 
 /** Generic axis enum (mirrored for native use in Axis.h). */
-UENUM(meta=(ScriptName="AxisType"))
+UENUM(BlueprintType, meta=(ScriptName="AxisType"))
 namespace EAxis
 {
 	enum Type : int
