ZenLoader in editor: Fixed issue where imports could fail to resolve after the owning linker had been detached
#rnx

diff --git a/Engine/Source/Runtime/CoreUObject/Private/Serialization/AsyncLoading2.cpp b/Engine/Source/Runtime/CoreUObject/Private/Serialization/AsyncLoading2.cpp
--- a/Engine/Source/Runtime/CoreUObject/Private/Serialization/AsyncLoading2.cpp
+++ b/Engine/Source/Runtime/CoreUObject/Private/Serialization/AsyncLoading2.cpp
@@ -4776,7 +4776,16 @@ void FAsyncPackage2::ImportPackagesRecursiveInner(FAsyncLoadingThreadState2& Thr
 
 		FAsyncPackage2* ImportedPackage = nullptr;
 		bool bInserted = false;
-		if (ImportedPackageRef.AreAllPublicExportsLoaded())
+		bool bIsFullyLoaded = ImportedPackageRef.AreAllPublicExportsLoaded();
+#if ALT2_ENABLE_LINKERLOAD_SUPPORT
+		if (!bIsZenPackageImport && (!ImportedPackageRef.HasPackage() || !ImportedPackageRef.GetPackage()->GetLinker()))
+		{
+			// If we're importing a linker load package and it doesn't have its linker we need to reload it, otherwise we can't reliably link to its imports
+			// Note: Legacy loader path appears to do this only for uncooked packages in the editor?
+			bIsFullyLoaded = false;
+		}
+#endif
+		if (bIsFullyLoaded)
 		{
 			ImportedPackage = AsyncLoadingThread.FindAsyncPackage(ImportedPackageId);
 			if (!ImportedPackage)
