Redo 24852317:

Making sure non-nullable object properties have a default object value assigned to them in case of serialization errors

diff --git a/Engine/Source/Runtime/CoreUObject/Private/UObject/PropertyBaseObject.cpp b/Engine/Source/Runtime/CoreUObject/Private/UObject/PropertyBaseObject.cpp
--- a/Engine/Source/Runtime/CoreUObject/Private/UObject/PropertyBaseObject.cpp
+++ b/Engine/Source/Runtime/CoreUObject/Private/UObject/PropertyBaseObject.cpp
@@ -626,6 +626,52 @@ bool FObjectPropertyBase::AllowObjectTypeReinterpretationTo(const FObjectPropert
 	return false;
 }
 
+UObject* FObjectPropertyBase::ConstructDefaultObjectValueIfNecessary(UObject* ExistingValue) const
+{
+	UObject* NewDefaultObjectValue = nullptr;
+
+	FUObjectSerializeContext* SerializeContext = FUObjectThreadContext::Get().GetSerializeContext();
+	UObject* Outer = SerializeContext ? SerializeContext->SerializedObject : nullptr;
+	if (!Outer)
+	{
+		Outer = GetTransientPackage();
+	}
+
+	if (ExistingValue)
+	{
+		UClass* ExistingValueClass = ExistingValue->GetClass();
+		// Sanity check to make sure the existing value class matches the property class
+		if (ExistingValueClass && (ExistingValueClass->IsChildOf(PropertyClass) || ExistingValueClass->GetAuthoritativeClass()->IsChildOf(PropertyClass)))
+		{
+			if (ExistingValue->IsTemplate())
+			{
+				// Existing value is a template so we can construct a new value with it as the archetype
+				// We probably got here because an object value failed to load (missing import class) and the property is left with a template of default subobject
+				NewDefaultObjectValue = NewObject<UObject>(Outer, ExistingValue->GetClass(), ExistingValue->GetFName(), RF_NoFlags, ExistingValue);
+			}
+			else
+			{
+				// Existing value is not a template so we can use it directly
+				// Similar to the above condition but the property was not referencing an instanced value in which case it's ok to leave the CDO default here
+				NewDefaultObjectValue = ExistingValue;
+			}
+		}
+	}
+
+	if (!NewDefaultObjectValue && !PropertyClass->HasAnyClassFlags(CLASS_Abstract))
+	{
+		// Existing value did not exist or it could not be used as a template
+		// Existing value may be null in case we were serializing an array of UObjects that failed to load (missing import class). Since the array is first pre-allocated with null values
+		// it will not have any existing objects to instantiate
+		NewDefaultObjectValue = NewObject<UObject>(Outer, PropertyClass);
+	}
+
+	// Final sanity check. We still may end up with a null object if the property class is abstract and the previous object value was missing or was not compatible
+	UE_CLOG(!NewDefaultObjectValue, LogProperty, Fatal, TEXT("Failed to create default object value for property %s. Previous value: %s"), *GetFullName(), *ExistingValue->GetFullName());
+
+	return NewDefaultObjectValue;
+}
+
 void FObjectPropertyBase::CheckValidObject(void* ValueAddress, TObjectPtr<UObject> OldValue) const
 {
 	const TObjectPtr<UObject> Object = GetObjectPtrPropertyValue(ValueAddress);
@@ -682,16 +728,18 @@ void FObjectPropertyBase::CheckValidObject(void* ValueAddress, TObjectPtr<UObjec
 			}
 			else
 			{
-				checkf(OldValue, TEXT("CheckValidObject(\"%s\") trying to assign null object value to non-nullable property \"%s\""), *Object.GetFullName(), *GetFullName());
+				UObject* DefaultValue = ConstructDefaultObjectValueIfNecessary(OldValue);
+
 				UE_LOG(LogProperty, Warning,
-					TEXT("Serialized %s for a property of %s. Reference will be reverted back to %s.\n    Property = %s\n    Item = %s"),
+					TEXT("Serialized %s for a non-nullable property of %s. Reference will be defaulted to %s.\n    Property = %s\n    Item = %s"),
 					*ObjectClass->GetFullName(),
 					*PropertyClass->GetFullName(),
-					OldValue ? *OldValue.GetFullName() : TEXT("None"),
+					*DefaultValue->GetFullName(),
 					*GetFullName(),
 					*Object.GetFullName()
 				);
-				SetObjectPtrPropertyValue(ValueAddress, OldValue);
+
+				SetObjectPropertyValue(ValueAddress, DefaultValue);
 			}
 		}
 	}
diff --git a/Engine/Source/Runtime/CoreUObject/Private/UObject/PropertyObject.cpp b/Engine/Source/Runtime/CoreUObject/Private/UObject/PropertyObject.cpp
--- a/Engine/Source/Runtime/CoreUObject/Private/UObject/PropertyObject.cpp
+++ b/Engine/Source/Runtime/CoreUObject/Private/UObject/PropertyObject.cpp
@@ -120,6 +120,23 @@ void FObjectProperty::SerializeItem( FStructuredArchive::FSlot Slot, void* Value
 		check(CurrentValuePtr.IsResolved());
 		UObject* CurrentValue = UE::CoreUObject::Private::ReadObjectHandlePointerNoCheck(CurrentValuePtr.GetHandle());
 
+		// Make sure non-nullable properties don't end up with null values
+		if (!ObjectValue && HasAnyPropertyFlags(CPF_NonNullable) &&
+			!UnderlyingArchive.IsSerializingDefaults() && // null values when Serializing CDOs are allowed, they will be fixed up later
+			!UnderlyingArchive.IsSaving()) // Constructing new objects when saving may confuse package saving code (new import/export created after import/export collection pass)
+		{
+			UObject* DefaultValue = ConstructDefaultObjectValueIfNecessary(CurrentValue);
+
+			UE_LOG(LogProperty, Warning,
+				TEXT("Failed to serialize value for non-nullable property %s. Reference will be defaulted to %s."),
+				*GetFullName(),
+				*DefaultValue->GetFullName()
+			);
+
+			SetObjectPropertyValue(Value, DefaultValue);
+			ObjectValue = DefaultValue;
+		}
+
 		if (ObjectValue != CurrentValue)
 		{
 			SetObjectPropertyValue(Value, ObjectValue);
diff --git a/Engine/Source/Runtime/CoreUObject/Public/UObject/UnrealType.h b/Engine/Source/Runtime/CoreUObject/Public/UObject/UnrealType.h
--- a/Engine/Source/Runtime/CoreUObject/Public/UObject/UnrealType.h
+++ b/Engine/Source/Runtime/CoreUObject/Public/UObject/UnrealType.h
@@ -2690,6 +2690,14 @@ protected:
 	virtual bool AllowObjectTypeReinterpretationTo(const FObjectPropertyBase* Other) const;
 	// End of FObjectPropertyBase interface
 
+	/**
+	 * Constructs a new object if the existing one is missing or is not compatible with the property class
+	 * Used for making sure non-nullable properties have valid values.
+	 * @param ExistingValue Previous object value (can be null)
+	 * @return Non-null object that was assigned to this property value address
+	 */
+	UObject* ConstructDefaultObjectValueIfNecessary(UObject* ExistingValue) const;
+
 	/* Helper functions for UObject property types that wrap the object pointer in a smart pointer */
 	template <typename T>
 	void GetWrappedUObjectPtrValues(UObject** OutObjects, const void* SrcAddress, EPropertyMemoryAccess SrcAccess, int32 ArrayIndex, int32 ArrayCount) const
diff --git a/Engine/Source/Runtime/CoreUObject/Tests/ObjectPropertyTest.cpp b/Engine/Source/Runtime/CoreUObject/Tests/ObjectPropertyTest.cpp
--- a/Engine/Source/Runtime/CoreUObject/Tests/ObjectPropertyTest.cpp
+++ b/Engine/Source/Runtime/CoreUObject/Tests/ObjectPropertyTest.cpp
@@ -115,7 +115,7 @@ TEST_CASE("UE::CoreUObject::FObjectProperty::CheckValidAddressNonNullable")
 
 	UE::Testing::FWarnFilterScope _([](const TCHAR* Message, ELogVerbosity::Type Verbosity, const FName& Category)
 		{
-			if (Category == TEXT("LogProperty") && FCString::Strstr(Message, TEXT("Reference will be reverted back to")) && Verbosity == ELogVerbosity::Type::Warning)
+			if (Category == TEXT("LogProperty") && FCString::Strstr(Message, TEXT("Reference will be defaulted to")) && Verbosity == ELogVerbosity::Type::Warning)
 			{
 				return true;
 			}
@@ -128,8 +128,10 @@ TEST_CASE("UE::CoreUObject::FObjectProperty::CheckValidAddressNonNullable")
 	Obj->ObjectPtrNonNullable = reinterpret_cast<UObjectPtrTestClass*>(Obj);
 	CHECK(Obj->ObjectPtrNonNullable != nullptr);
 
-	//old value is required for non nullable properties
-	REQUIRE_CHECK(Property->CheckValidObject(&Obj->ObjectPtrNonNullable, nullptr));
+	//new value is required for non nullable properties
+	Property->CheckValidObject(&Obj->ObjectPtrNonNullable, nullptr);
+	CHECK(Obj->ObjectPtrNonNullable != nullptr);
+	CHECK(Obj->ObjectPtrNonNullable->IsA(UObjectPtrTestClass::StaticClass()));
 }
 
 class FMockArchive : public FArchive
