IoStore: Include imported package names in the header data. This removes the need to keep track of uncooked packages in the package store for editor builds.
#rnx

diff --git a/Engine/Plugins/MemoryUsageQueries/Source/MemoryUsageQueries/Private/MemoryUsageQueries.cpp b/Engine/Plugins/MemoryUsageQueries/Source/MemoryUsageQueries/Private/MemoryUsageQueries.cpp
--- a/Engine/Plugins/MemoryUsageQueries/Source/MemoryUsageQueries/Private/MemoryUsageQueries.cpp
+++ b/Engine/Plugins/MemoryUsageQueries/Source/MemoryUsageQueries/Private/MemoryUsageQueries.cpp
@@ -1455,7 +1455,6 @@ class FPackageDependenciesLazyDatabase final
 				// Add leafs
 				if (ImportedPackageIds.IsEmpty()
 #if WITH_EDITOR
-					&& PackageEntry.UncookedPackageName.IsNone()
 					&& OptionalImportedPackageIds.IsEmpty()
 #endif
 				)
diff --git a/Engine/Source/Developer/IoStoreUtilities/Internal/PackageStoreOptimizer.h b/Engine/Source/Developer/IoStoreUtilities/Internal/PackageStoreOptimizer.h
--- a/Engine/Source/Developer/IoStoreUtilities/Internal/PackageStoreOptimizer.h
+++ b/Engine/Source/Developer/IoStoreUtilities/Internal/PackageStoreOptimizer.h
@@ -216,13 +216,6 @@ private:
 		{
 		}
 
-		// Note: For Zen packages we currently don't know the package name but we also don't use it for anything
-		FImportedPackageRef(FPackageId InId)
-			: Id(InId)
-		{
-
-		}
-
 		bool operator<(const FImportedPackageRef& Other) const
 		{
 			return Id < Other.Id;
@@ -308,7 +301,7 @@ private:
 	{
 		FZenPackageSummary Summary;
 		TOptional<FZenPackageVersioningInfo> VersioningInfo;
-		TArray<FPackageId> ImportedPackageIds;
+		TArray<FName> ImportedPackageNames;
 		TArray<uint64> ImportedPublicExportHashes;
 		TArray<FDisplayNameEntryId> NameMap;
 		TArray<FBulkDataMapEntry> BulkDataEntries;
@@ -325,7 +318,7 @@ private:
 
 	static uint64 GetPublicExportHash(FStringView PackageRelativeExportPath);
 	FCookedHeaderData LoadCookedHeader(const FIoBuffer& CookedHeaderBuffer) const;
-	FZenPackageHeaderData LoadZenPackageHeader(const FIoBuffer& HeaderBuffer, int32 ExportBundleCount, const TArrayView<const FPackageId>& ImportedPackageIds) const;
+	FZenPackageHeaderData LoadZenPackageHeader(const FIoBuffer& HeaderBuffer, int32 ExportBundleCount) const;
 	void ResolveImport(FPackageStorePackage::FUnresolvedImport* Imports, const FObjectImport* ObjectImports, int32 LocalImportIndex) const;
 	void ProcessImports(const FCookedHeaderData& CookedHeaderData, FPackageStorePackage* Package, TArray<FPackageStorePackage::FUnresolvedImport>& UnresolvedImports) const;
 	void ProcessImports(const FZenPackageHeaderData& ZenHeaderData, FPackageStorePackage* Package) const;
diff --git a/Engine/Source/Developer/IoStoreUtilities/Private/PackageStoreOptimizer.cpp b/Engine/Source/Developer/IoStoreUtilities/Private/PackageStoreOptimizer.cpp
--- a/Engine/Source/Developer/IoStoreUtilities/Private/PackageStoreOptimizer.cpp
+++ b/Engine/Source/Developer/IoStoreUtilities/Private/PackageStoreOptimizer.cpp
@@ -138,7 +138,7 @@ FPackageStorePackage* FPackageStoreOptimizer::CreatePackageFromZenPackageHeader(
 
 	Package->Name = Name;
 	
-	FZenPackageHeaderData ZenHeaderData = LoadZenPackageHeader(Buffer, ExportBundleCount, ImportedPackageIds);
+	FZenPackageHeaderData ZenHeaderData = LoadZenPackageHeader(Buffer, ExportBundleCount);
 	if (ZenHeaderData.VersioningInfo.IsSet())
 	{
 		Package->VersioningInfo.Emplace(ZenHeaderData.VersioningInfo.GetValue());
@@ -265,7 +265,7 @@ FPackageStoreOptimizer::FCookedHeaderData FPackageStoreOptimizer::LoadCookedHead
 }
 PRAGMA_ENABLE_DEPRECATION_WARNINGS
 
-FPackageStoreOptimizer::FZenPackageHeaderData FPackageStoreOptimizer::LoadZenPackageHeader(const FIoBuffer& HeaderBuffer, int32 ExportBundleCount, const TArrayView<const FPackageId>& ImportedPackageIds) const
+FPackageStoreOptimizer::FZenPackageHeaderData FPackageStoreOptimizer::LoadZenPackageHeader(const FIoBuffer& HeaderBuffer, int32 ExportBundleCount) const
 {
 	FZenPackageHeaderData ZenPackageHeaderData;
 
@@ -296,8 +296,11 @@ FPackageStoreOptimizer::FZenPackageHeaderData FPackageStoreOptimizer::LoadZenPac
 		ZenPackageHeaderData.BulkDataEntries = MakeArrayView(reinterpret_cast<const FBulkDataMapEntry*>(BulkDataMapData), BulkDataMapSize / sizeof(FBulkDataMapEntry));
 	}
 
-	ZenPackageHeaderData.ImportedPackageIds = ImportedPackageIds;
-
+	TArrayView<const uint8> ImportedPackageNamesDataView(HeaderData + Summary.ImportedPackageNamesOffset, Summary.HeaderSize - Summary.ImportedPackageNamesOffset);
+	FMemoryReaderView ImportedPackageNamesDataReader(ImportedPackageNamesDataView);
+	FZenPackageImportedPackageNamesContainer ImportedPackageNamesContainer;
+	ImportedPackageNamesDataReader << ImportedPackageNamesContainer;
+	ZenPackageHeaderData.ImportedPackageNames = MoveTemp(ImportedPackageNamesContainer.Names);
 
 	ZenPackageHeaderData.ImportedPublicExportHashes =
 		MakeArrayView<const uint64>(
@@ -340,7 +343,7 @@ FPackageStoreOptimizer::FZenPackageHeaderData FPackageStoreOptimizer::LoadZenPac
 		ArcsAr << InternalArc.ToExportBundleIndex;
 	}
 
-	for (FPackageId ImportedPackageId : ZenPackageHeaderData.ImportedPackageIds)
+	for (FName ImportedPackageName : ZenPackageHeaderData.ImportedPackageNames)
 	{
 		int32 ExternalArcsCount = 0;
 		ArcsAr << ExternalArcsCount;
@@ -466,10 +469,10 @@ void FPackageStoreOptimizer::ProcessImports(const FCookedHeaderData& CookedHeade
 
 void FPackageStoreOptimizer::ProcessImports(const FZenPackageHeaderData& ZenHeaderData, FPackageStorePackage* Package) const
 {
-	Package->ImportedPackages.Reserve(ZenHeaderData.ImportedPackageIds.Num());
-	for (FPackageId ImportedPackageId : ZenHeaderData.ImportedPackageIds)
+	Package->ImportedPackages.Reserve(ZenHeaderData.ImportedPackageNames.Num());
+	for (FName ImportedPackageName : ZenHeaderData.ImportedPackageNames)
 	{
-		Package->ImportedPackages.Emplace(ImportedPackageId);
+		Package->ImportedPackages.Emplace(ImportedPackageName);
 	}
 	Package->ImportedPublicExportHashes = ZenHeaderData.ImportedPublicExportHashes;
 	Package->Imports = ZenHeaderData.Imports;
@@ -1227,6 +1230,16 @@ void FPackageStoreOptimizer::FinalizePackageHeader(FPackageStorePackage* Package
 	Package->NameMapBuilder.MarkNameAsReferenced(Package->Name);
 	FMappedName MappedPackageName = Package->NameMapBuilder.MapName(Package->Name);
 
+	FZenPackageImportedPackageNamesContainer ImportedPackageNamesContainer;
+	ImportedPackageNamesContainer.Names.Reserve(Package->ImportedPackages.Num());
+	for (const FPackageStorePackage::FImportedPackageRef& ImportedPackage : Package->ImportedPackages)
+	{
+		ImportedPackageNamesContainer.Names.Add(ImportedPackage.Name);
+	}
+	FBufferWriter ImportedPackagesArchive(nullptr, 0, EBufferWriterFlags::AllowResize | EBufferWriterFlags::TakeOwnership);
+	ImportedPackagesArchive << ImportedPackageNamesContainer;
+	uint64 ImportedPackagesSize = ImportedPackagesArchive.Tell();
+
 	FBufferWriter NameMapArchive(nullptr, 0, EBufferWriterFlags::AllowResize | EBufferWriterFlags::TakeOwnership);
 	SaveNameBatch(Package->NameMapBuilder.GetNameMap(), NameMapArchive);
 	uint64 NameMapSize = NameMapArchive.Tell();
@@ -1255,6 +1268,7 @@ void FPackageStoreOptimizer::FinalizePackageHeader(FPackageStorePackage* Package
 		+ ExportMapSize
 		+ ExportBundleEntriesSize
 		+ GraphDataSize
+		+ ImportedPackagesSize
 		+ BulkDataMapSize + sizeof(int64);
 
 	Package->HeaderBuffer = FIoBuffer(Package->HeaderSize);
@@ -1293,6 +1307,8 @@ void FPackageStoreOptimizer::FinalizePackageHeader(FPackageStorePackage* Package
 	HeaderArchive.Serialize(ExportBundleEntriesArchive.GetWriterData(), ExportBundleEntriesArchive.Tell());
 	PackageSummary->GraphDataOffset = HeaderArchive.Tell();
 	HeaderArchive.Serialize(GraphArchive.GetWriterData(), GraphArchive.Tell());
+	PackageSummary->ImportedPackageNamesOffset = HeaderArchive.Tell();
+	HeaderArchive.Serialize(ImportedPackagesArchive.GetWriterData(), ImportedPackagesArchive.Tell());
 	check(HeaderArchive.Tell() == PackageSummary->HeaderSize)
 }
 
diff --git a/Engine/Source/Runtime/Core/Public/IO/PackageStore.h b/Engine/Source/Runtime/Core/Public/IO/PackageStore.h
--- a/Engine/Source/Runtime/Core/Public/IO/PackageStore.h
+++ b/Engine/Source/Runtime/Core/Public/IO/PackageStore.h
@@ -54,8 +54,6 @@ struct FPackageStoreEntry
 	TArrayView<const FPackageId> ImportedPackageIds;
 	TArrayView<const FSHAHash> ShaderMapHashes;
 #if WITH_EDITOR
-	FName UncookedPackageName;
-	uint8 UncookedPackageHeaderExtension; // TODO: Can't include PackagePath.h
 	FPackageStoreExportInfo OptionalSegmentExportInfo;
 	TArrayView<const FPackageId> OptionalSegmentImportedPackageIds;
 #endif
diff --git a/Engine/Source/Runtime/CoreUObject/Private/Serialization/AsyncLoading2.cpp b/Engine/Source/Runtime/CoreUObject/Private/Serialization/AsyncLoading2.cpp
--- a/Engine/Source/Runtime/CoreUObject/Private/Serialization/AsyncLoading2.cpp
+++ b/Engine/Source/Runtime/CoreUObject/Private/Serialization/AsyncLoading2.cpp
@@ -89,6 +89,42 @@ FArchive& operator<<(FArchive& Ar, FZenPackageVersioningInfo& VersioningInfo)
 	return Ar;
 }
 
+FArchive& operator<<(FArchive& Ar, FZenPackageImportedPackageNamesContainer& Container)
+{
+	TArray<FDisplayNameEntryId> NameEntries;
+	if (Ar.IsSaving())
+	{
+#if ALLOW_NAME_BATCH_SAVING
+		NameEntries.Reserve(Container.Names.Num());
+		for (FName ImportedPackageName : Container.Names)
+		{
+			NameEntries.Emplace(ImportedPackageName);
+		}
+		SaveNameBatch(NameEntries, Ar);
+		for (FName ImportedPackageName : Container.Names)
+		{
+			int32 Number = ImportedPackageName.GetNumber();
+			Ar << Number;
+		}
+#else
+		check(false);
+#endif
+	}
+	else
+	{
+		NameEntries = LoadNameBatch(Ar);
+		Container.Names.SetNum(NameEntries.Num());
+		for (int32 Index = 0; Index < NameEntries.Num(); ++Index)
+		{
+			int32 Number;
+			Ar << Number;
+			Container.Names[Index] = NameEntries[Index].ToName(Number);
+		}
+	}
+
+	return Ar;
+}
+
 FArchive& operator<<(FArchive& Ar, FExportBundleEntry& ExportBundleEntry)
 {
 	Ar << ExportBundleEntry.LocalExportIndex;
@@ -3452,7 +3488,9 @@ private:
 #endif
 		check(DataPtr - Data.MemoryBuffer0 == MemoryBufferSize);
 
-#if ALT2_ENABLE_LINKERLOAD_SUPPORT
+#if WITH_EDITOR
+		const bool bCanImportPackagesWithIdsOnly = false;
+#elif ALT2_ENABLE_LINKERLOAD_SUPPORT
 		const bool bIsZenPackage = !AsyncPackage->LinkerLoadState.IsSet();
 		bool bCanImportPackagesWithIdsOnly = bIsZenPackage;
 #else
@@ -4007,11 +4045,6 @@ bool FAsyncLoadingThread2::CreateAsyncPackagesFromQueue(FAsyncLoadingThreadState
 				bIsZenPackage = false;
 				PackageStatus = EPackageStoreEntryStatus::Ok;
 			}
-			if (bIsZenPackage && PackageStatus == EPackageStoreEntryStatus::Ok && !PackageEntry.UncookedPackageName.IsNone())
-			{
-				// This means that the package was added as an uncooked package to the store at one point and was later deleted from disk
-				PackageStatus = EPackageStoreEntryStatus::Missing;
-			}
 #endif
 
 			// Fixup CustomName to handle any input string that can be converted to a long package name.
@@ -4677,19 +4710,9 @@ void FAsyncPackage2::ImportPackagesRecursiveInner(FAsyncLoadingThreadState2& Thr
 			}
 		}
 		ImportedPackageStatus = PackageStore.GetPackageStoreEntry(ImportedPackageIdToLoad, ImportedPackageEntry);
-#if WITH_EDITOR
-		if (ImportedPackageStatus == EPackageStoreEntryStatus::Ok && !ImportedPackageEntry.UncookedPackageName.IsNone())
-		{
-			ImportedPackageNameToLoad = ImportedPackageEntry.UncookedPackageName;
-			if (ImportedPackageUPackageName.IsNone())
-			{
-				ImportedPackageUPackageName = ImportedPackageNameToLoad;
-			}
-		}
-#endif
 
 		FPackagePath ImportedPackagePath;
-#if ALT2_ENABLE_LINKERLOAD_SUPPORT
+#if ALT2_ENABLE_LINKERLOAD_SUPPORT || WITH_EDITOR
 		bool bIsZenPackage = !LinkerLoadState.IsSet();
 		bool bIsZenPackageImport = true;
 		if (AsyncLoadingThread.ShouldLoadPackageFromFileSystem(ImportedPackageNameToLoad, ImportedPackagePath))
@@ -4697,11 +4720,6 @@ void FAsyncPackage2::ImportPackagesRecursiveInner(FAsyncLoadingThreadState2& Thr
 			bIsZenPackageImport = false;
 			ImportedPackageStatus = EPackageStoreEntryStatus::Ok;
 		}
-		if (bIsZenPackageImport && ImportedPackageStatus == EPackageStoreEntryStatus::Ok && !ImportedPackageEntry.UncookedPackageName.IsNone())
-		{
-			// This means that the package was added as an uncooked package to the store at one point and was later deleted from disk
-			ImportedPackageStatus = EPackageStoreEntryStatus::Missing;
-		}
 #else
 		constexpr bool bIsZenPackage = true;
 		constexpr bool bIsZenPackageImport = true;
@@ -4709,16 +4727,14 @@ void FAsyncPackage2::ImportPackagesRecursiveInner(FAsyncLoadingThreadState2& Thr
 
 		FLoadedPackageRef& ImportedPackageRef = ImportStore.AddImportedPackageReference(ImportedPackageId, ImportedPackageUPackageName);
 #if WITH_EDITOR
-		if (AsyncLoadingThread.UncookedPackageLoader && ImportedPackageStatus == EPackageStoreEntryStatus::Ok && !ImportedPackageEntry.UncookedPackageName.IsNone())
+		if (AsyncLoadingThread.UncookedPackageLoader && ImportedPackageStatus == EPackageStoreEntryStatus::Ok && !bIsZenPackageImport)
 		{
 			UPackage* UncookedPackage = ImportedPackageRef.GetPackage();
 			if (!ImportedPackageRef.AreAllPublicExportsLoaded())
 			{
-				UE_ASYNC_PACKAGE_LOG(Verbose, Desc, TEXT("ImportPackages: LoadUncookedImport"), TEXT("Loading imported uncooked package '%s' '0x%llX'"), *ImportedPackageEntry.UncookedPackageName.ToString(), ImportedPackageId.ValueForDebugging());
+				UE_ASYNC_PACKAGE_LOG(Verbose, Desc, TEXT("ImportPackages: LoadUncookedImport"), TEXT("Loading imported uncooked package '%s' '0x%llX'"), *ImportedPackageNameToLoad.ToString(), ImportedPackageId.ValueForDebugging());
 				check(IsInGameThread());
 				IoBatch.Issue(); // The batch might already contain requests for packages being imported from the uncooked one we're going to load so make sure that those are started before blocking
-				ImportedPackagePath = FPackagePath::FromPackageNameUnchecked(ImportedPackageEntry.UncookedPackageName);
-				ImportedPackagePath.SetHeaderExtension(static_cast<EPackageExtension>(ImportedPackageEntry.UncookedPackageHeaderExtension));
 				int32 ImportRequestId = AsyncLoadingThread.UncookedPackageLoader->LoadPackage(ImportedPackagePath, NAME_None, FLoadPackageAsyncDelegate(), PKG_None, INDEX_NONE, 0, nullptr, LOAD_None);
 				AsyncLoadingThread.UncookedPackageLoader->FlushLoading(ImportRequestId);
 				UncookedPackage = FindObjectFast<UPackage>(nullptr, ImportedPackagePath.GetPackageFName());
@@ -4768,7 +4784,7 @@ void FAsyncPackage2::ImportPackagesRecursiveInner(FAsyncLoadingThreadState2& Thr
 			{
 				ImportedPackageRef.SetHasFailed();
 				UE_ASYNC_PACKAGE_LOG(Warning, Desc, TEXT("ImportPackages: SkipPackage"),
-					TEXT("Failed to load uncooked imported package with id '0x%llX' ('%s')"), ImportedPackageId.Value(), *ImportedPackageEntry.UncookedPackageName.ToString());
+					TEXT("Failed to load uncooked imported package with id '0x%llX' ('%s')"), ImportedPackageId.Value(), *ImportedPackageNameToLoad.ToString());
 			}
 			continue;
 		}
@@ -5124,9 +5140,18 @@ static void ReadAsyncPackageHeader(FAsyncPackageSerializationState& Serializatio
 
 	const uint64 ExportBundleHeadersSize = sizeof(FExportBundleHeader) * HeaderData.ExportBundleHeaders.Num();
 	const uint64 ArcsDataOffset = ExportBundleHeadersOffset + ExportBundleHeadersSize;
-	const uint64 ArcsDataSize = PackageSummary->HeaderSize - ArcsDataOffset;
+	const uint64 ArcsDataSize = PackageSummary->ImportedPackageNamesOffset - ArcsDataOffset;
 	HeaderData.ArcsData = TArrayView64<const uint8>(PackageHeaderDataPtr + ArcsDataOffset, ArcsDataSize);
 
+#if WITH_EDITOR
+	// Only read the imported package names in editor builds where we need to interact with uncooked packages
+	TArrayView<const uint8> ImportedPackageNamesDataView(PackageHeaderDataPtr + PackageSummary->ImportedPackageNamesOffset, PackageSummary->HeaderSize - PackageSummary->ImportedPackageNamesOffset);
+	FMemoryReaderView ImportedPackageNamesDataReader(ImportedPackageNamesDataView);
+	FZenPackageImportedPackageNamesContainer Container;
+	ImportedPackageNamesDataReader << Container;
+	HeaderData.ImportedPackageNames = MoveTemp(Container.Names);
+#endif
+
 	SerializationState.AllExportDataPtr = PackageHeaderDataPtr + PackageSummary->HeaderSize;
 }
 
diff --git a/Engine/Source/Runtime/CoreUObject/Public/Serialization/AsyncLoading2.h b/Engine/Source/Runtime/CoreUObject/Public/Serialization/AsyncLoading2.h
--- a/Engine/Source/Runtime/CoreUObject/Public/Serialization/AsyncLoading2.h
+++ b/Engine/Source/Runtime/CoreUObject/Public/Serialization/AsyncLoading2.h
@@ -246,6 +246,7 @@ enum class EZenPackageVersion : uint32
 {
 	Initial,
 	DataResourceTable,
+	ImportedPackageNames,
 
 	LatestPlusOne,
 	Latest = LatestPlusOne - 1
@@ -261,6 +262,13 @@ struct FZenPackageVersioningInfo
 	COREUOBJECT_API friend FArchive& operator<<(FArchive& Ar, FZenPackageVersioningInfo& ExportBundleEntry);
 };
 
+struct FZenPackageImportedPackageNamesContainer
+{
+	TArray<FName> Names;
+
+	COREUOBJECT_API friend FArchive& operator<<(FArchive& Ar, FZenPackageImportedPackageNamesContainer& Container);
+};
+
 /**
  * Package summary.
  */
@@ -276,6 +284,7 @@ struct FZenPackageSummary
 	int32 ExportMapOffset;
 	int32 ExportBundleEntriesOffset;
 	int32 GraphDataOffset;
+	int32 ImportedPackageNamesOffset;
 };
 
 /**
diff --git a/Engine/Source/Runtime/PakFile/Private/FilePackageStore.cpp b/Engine/Source/Runtime/PakFile/Private/FilePackageStore.cpp
--- a/Engine/Source/Runtime/PakFile/Private/FilePackageStore.cpp
+++ b/Engine/Source/Runtime/PakFile/Private/FilePackageStore.cpp
@@ -29,61 +29,16 @@ namespace FilePackageStore
 
 FFilePackageStoreBackend::FFilePackageStoreBackend()
 {
-#if WITH_EDITOR
-	OnContentPathMountedDelegateHandle = FPackageName::OnContentPathMounted().AddLambda([this](const FString& InAssetPath, const FString& InFilesystemPath)
-		{
-			{
-				FScopeLock _(&UncookedPackageRootsLock);
-				PendingAddUncookedPackageRoots.Add(InFilesystemPath);
-			}
-			{
-				FWriteScopeLock _(EntriesLock);
-				bNeedsUncookedPackagesUpdate = true;
-			}
-		});
-
-	OnContentPathDismountedDelegateHandle = FPackageName::OnContentPathDismounted().AddLambda([this](const FString& InAssetPath, const FString& InFilesystemPath)
-		{
-			{
-				FScopeLock _(&UncookedPackageRootsLock);
-				PendingRemoveUncookedPackageRoots.Add(InAssetPath);
-				PendingAddUncookedPackageRoots.Remove(InFilesystemPath);
-			}
-			{
-				FWriteScopeLock _(EntriesLock);
-				bNeedsUncookedPackagesUpdate = true;
-			}
-		});
-
-	TArray<FString> RootPaths;
-	FPackageName::QueryRootContentPaths(RootPaths);
-	{
-		FScopeLock _(&UncookedPackageRootsLock);
-		PendingAddUncookedPackageRoots.Append(RootPaths);
-	}
-	{
-		FWriteScopeLock _(EntriesLock);
-		bNeedsUncookedPackagesUpdate = true;
-	}
-#endif //if WITH_EDITOR
 }
 
 FFilePackageStoreBackend::~FFilePackageStoreBackend()
 {
-#if WITH_EDITOR
-	FPackageName::OnContentPathMounted().Remove(OnContentPathMountedDelegateHandle);
-	FPackageName::OnContentPathDismounted().Remove(OnContentPathDismountedDelegateHandle);
-#endif
 }
 
 void FFilePackageStoreBackend::BeginRead()
 {
 	EntriesLock.ReadLock();
-	if (bNeedsContainerUpdate
-#if WITH_EDITOR
-		|| bNeedsUncookedPackagesUpdate
-#endif
-		)
+	if (bNeedsContainerUpdate)
 	{
 		Update();
 	}
@@ -96,15 +51,6 @@ void FFilePackageStoreBackend::EndRead()
 
 EPackageStoreEntryStatus FFilePackageStoreBackend::GetPackageStoreEntry(FPackageId PackageId, FPackageStoreEntry& OutPackageStoreEntry)
 {
-#if WITH_EDITOR
-	const FUncookedPackage* FindUncookedPackage = UncookedPackagesMap.Find(PackageId);
-	if (FindUncookedPackage)
-	{
-		OutPackageStoreEntry.UncookedPackageName = FindUncookedPackage->PackageName;
-		OutPackageStoreEntry.UncookedPackageHeaderExtension = static_cast<uint8>(FindUncookedPackage->HeaderExtension);
-		return EPackageStoreEntryStatus::Ok;
-	}
-#endif
 	const FFilePackageStoreEntry* FindEntry = StoreEntriesMap.FindRef(PackageId);
 	if (FindEntry)
 	{
@@ -257,80 +203,6 @@ void FFilePackageStoreBackend::Update()
 		}
 		bNeedsContainerUpdate = false;
 	}
-
-#if WITH_EDITOR
-	if (bNeedsUncookedPackagesUpdate)
-	{
-		TRACE_CPUPROFILER_EVENT_SCOPE(UpdateUncookedPackages);
-		FScopeLock _(&UncookedPackageRootsLock);
-
-		if (!PendingRemoveUncookedPackageRoots.IsEmpty())
-		{
-			UE_LOG(LogFilePackageStore, Display, TEXT("Removing uncooked packages from %d old roots..."), PendingRemoveUncookedPackageRoots.Num());
-			uint64 TotalRemovedCount = RemoveUncookedPackagesFromRoot(PendingRemoveUncookedPackageRoots);
-			PendingRemoveUncookedPackageRoots.Empty();
-			UE_LOG(LogFilePackageStore, Display, TEXT("Removed %lld uncooked packages"), TotalRemovedCount);
-		}
-
-		if (!PendingAddUncookedPackageRoots.IsEmpty())
-		{
-			UE_LOG(LogFilePackageStore, Display, TEXT("Adding uncooked packages from %d new roots..."), PendingAddUncookedPackageRoots.Num());
-			uint64 TotalAddedCount = 0;
-			for (const FString& RootPath : PendingAddUncookedPackageRoots)
-			{
-				TotalAddedCount += AddUncookedPackagesFromRoot(RootPath);
-			}
-			PendingAddUncookedPackageRoots.Empty();
-			UE_LOG(LogFilePackageStore, Display, TEXT("Added %lld uncooked packages"), TotalAddedCount);
-		}
-		bNeedsUncookedPackagesUpdate = false;
-	}
-#endif //if WITH_EDITOR
 }
 
-#if WITH_EDITOR
-uint64 FFilePackageStoreBackend::AddUncookedPackagesFromRoot(const FString& RootPath)
-{
-	uint64 Count = 0;
-	FPackageName::IteratePackagesInDirectory(RootPath, [this, &RootPath, &Count](const TCHAR* InPackageFileName) -> bool
-		{
-			FPackagePath PackagePath = FPackagePath::FromLocalPath(InPackageFileName);
-			FName PackageName = PackagePath.GetPackageFName();
-			if (!PackageName.IsNone())
-			{
-				FPackageId PackageId = FPackageId::FromName(PackageName);
-				FUncookedPackage& UncookedPackage = UncookedPackagesMap.FindOrAdd(PackageId);
-				UncookedPackage.PackageName = PackageName;
-				UncookedPackage.HeaderExtension = PackagePath.GetHeaderExtension();
-				++Count;
-			}
-			return true;
-		});
-	return Count;
-}
-
-uint64 FFilePackageStoreBackend::RemoveUncookedPackagesFromRoot(const TSet<FString>& RootPaths)
-{
-	uint64 Count = 0;
-	TMap<FPackageId, FUncookedPackage> RemainingUncookedPackagesMap;
-	RemainingUncookedPackagesMap.Reserve(UncookedPackagesMap.Num());
-	for (auto& Pair : UncookedPackagesMap)
-	{
-		const FNameBuilder PackageName(Pair.Value.PackageName);
-		const FStringView PackageRootPath = FilePackageStore::GetRootPathFromPackageName(PackageName);
-		if (RootPaths.ContainsByHash(GetTypeHash(PackageRootPath), PackageRootPath))
-		{
-			++Count;
-		}
-		else
-		{
-			RemainingUncookedPackagesMap.Add(Pair.Key, MoveTemp(Pair.Value));
-		}
-	}
-	UncookedPackagesMap = MoveTemp(RemainingUncookedPackagesMap);
-	return Count;
-}
-#endif //if WITH_EDITOR
-
-
 //PRAGMA_ENABLE_OPTIMIZATION
\ No newline at end of file
