workaround for clang on windows error saying:
- 'ThreadLocalContext' cannot be thread local when declared 'dllexport'

diff --git a/Engine/Source/Runtime/CoreUObject/Private/UObject/ScriptCore.cpp b/Engine/Source/Runtime/CoreUObject/Private/UObject/ScriptCore.cpp
--- a/Engine/Source/Runtime/CoreUObject/Private/UObject/ScriptCore.cpp
+++ b/Engine/Source/Runtime/CoreUObject/Private/UObject/ScriptCore.cpp
@@ -177,13 +177,21 @@ FBlueprintContext::FBlueprintContext()
 	ensure(BlueprintContextVirtualStackAllocatorDecommitMode >= 0 && BlueprintContextVirtualStackAllocatorDecommitMode < (int)EVirtualStackAllocatorDecommitMode::NumModes);
 }
 
+// pulled the thread_local into a separate function to workaround
+// a compile error complaining about having it live local to the
+// lambda below
+FBlueprintContext* FBlueprintContextGetThreadSingletonImpl()
+{
+	static thread_local FBlueprintContext ThreadLocalContext;
+	return &ThreadLocalContext;
+}
+
 FBlueprintContext* FBlueprintContext::GetThreadSingleton()
 {
 	FBlueprintContext* Result;
 	AutoRTFM::Open([&Result]
 	{
-		static thread_local FBlueprintContext ThreadLocalContext;
-		Result = &ThreadLocalContext;
+		Result = FBlueprintContextGetThreadSingletonImpl();
 	});
 	return Result;
 }
