Fixing static analysis warnings.

Obj->GetFullName() does check 'this' agains nullptr but maybe SA doesn't know this

diff --git a/Engine/Source/Runtime/CoreUObject/Private/UObject/PropertyBaseObject.cpp b/Engine/Source/Runtime/CoreUObject/Private/UObject/PropertyBaseObject.cpp
--- a/Engine/Source/Runtime/CoreUObject/Private/UObject/PropertyBaseObject.cpp
+++ b/Engine/Source/Runtime/CoreUObject/Private/UObject/PropertyBaseObject.cpp
@@ -667,7 +667,9 @@ UObject* FObjectPropertyBase::ConstructDefaultObjectValueIfNecessary(UObject* Ex
 	}
 
 	// Final sanity check. We still may end up with a null object if the property class is abstract and the previous object value was missing or was not compatible
-	UE_CLOG(!NewDefaultObjectValue, LogProperty, Fatal, TEXT("Failed to create default object value for property %s. Previous value: %s"), *GetFullName(), *ExistingValue->GetFullName());
+	UE_CLOG(!NewDefaultObjectValue, LogProperty, Fatal, TEXT("Failed to create default object value for property %s. Previous value: %s"), 
+		*GetFullName(), 
+		ExistingValue ? *ExistingValue->GetFullName() : TEXT("None"));
 
 	return NewDefaultObjectValue;
 }
@@ -734,7 +736,7 @@ void FObjectPropertyBase::CheckValidObject(void* ValueAddress, TObjectPtr<UObjec
 					TEXT("Serialized %s for a non-nullable property of %s. Reference will be defaulted to %s.\n    Property = %s\n    Item = %s"),
 					*ObjectClass->GetFullName(),
 					*PropertyClass->GetFullName(),
-					*DefaultValue->GetFullName(),
+					DefaultValue ? *DefaultValue->GetFullName() : TEXT("None"),
 					*GetFullName(),
 					*Object.GetFullName()
 				);
