Fix TSAN warning by reordering parameters so we don't read globals from ALT

#rnx

diff --git a/Engine/Source/Runtime/CoreUObject/Private/UObject/UObjectBaseUtility.cpp b/Engine/Source/Runtime/CoreUObject/Private/UObject/UObjectBaseUtility.cpp
--- a/Engine/Source/Runtime/CoreUObject/Private/UObject/UObjectBaseUtility.cpp
+++ b/Engine/Source/Runtime/CoreUObject/Private/UObject/UObjectBaseUtility.cpp
@@ -317,7 +317,7 @@ bool UObjectBaseUtility::MarkPackageDirty() const
 			// we explicitly disable the ability to dirty a package or map during load/undo/redo.  Commandlets can still
 			// set the dirty state on load.
 			if( IsRunningCommandlet() || 
-				(GIsEditor && !GIsEditorLoadingPackage && !GIsCookerLoadingPackage && !GIsPlayInEditorWorld && !IsInAsyncLoadingThread() && !IsReloadActive()
+				(!IsInAsyncLoadingThread() && GIsEditor && !GIsEditorLoadingPackage && !GIsCookerLoadingPackage && !GIsPlayInEditorWorld && !IsReloadActive()
 #if WITH_EDITORONLY_DATA
 				&& !GIsTransacting
 				&& !Package->bIsCookedForEditor // Cooked packages can't be modified nor marked as dirty
