Make FBitReader and FNetBitReader Src buffer parameter const.

diff --git a/Engine/Source/Runtime/Core/Private/Serialization/BitReader.cpp b/Engine/Source/Runtime/Core/Private/Serialization/BitReader.cpp
--- a/Engine/Source/Runtime/Core/Private/Serialization/BitReader.cpp
+++ b/Engine/Source/Runtime/Core/Private/Serialization/BitReader.cpp
@@ -140,7 +140,7 @@ void appBitsCpy( uint8* Dest, int32 DestBit, uint8* Src, int32 SrcBit, int32 Bit
 //
 // Reads bitstreams.
 //
-FBitReader::FBitReader(uint8* Src, int64 CountBits)
+FBitReader::FBitReader(const uint8* Src, int64 CountBits)
 	: Num(CountBits)
 	, Pos(0)
 {
diff --git a/Engine/Source/Runtime/Core/Public/Serialization/BitReader.h b/Engine/Source/Runtime/Core/Public/Serialization/BitReader.h
--- a/Engine/Source/Runtime/Core/Public/Serialization/BitReader.h
+++ b/Engine/Source/Runtime/Core/Public/Serialization/BitReader.h
@@ -26,7 +26,7 @@ struct CORE_API FBitReader : public FBitArchive
 	friend struct FBitReaderMark;
 
 public:
-	FBitReader( uint8* Src = nullptr, int64 CountBits = 0 );
+	FBitReader( const uint8* Src = nullptr, int64 CountBits = 0 );
 
 	FBitReader(const FBitReader&) = default;
     FBitReader& operator=(const FBitReader&) = default;
diff --git a/Engine/Source/Runtime/CoreUObject/Private/UObject/CoreNet.cpp b/Engine/Source/Runtime/CoreUObject/Private/UObject/CoreNet.cpp
--- a/Engine/Source/Runtime/CoreUObject/Private/UObject/CoreNet.cpp
+++ b/Engine/Source/Runtime/CoreUObject/Private/UObject/CoreNet.cpp
@@ -498,7 +498,7 @@ void FNetBitWriter::CountMemory(FArchive& Ar) const
 // ----------------------------------------------------------------
 //	FNetBitReader
 // ----------------------------------------------------------------
-FNetBitReader::FNetBitReader(UPackageMap* InPackageMap, uint8* Src, int64 CountBits)
+FNetBitReader::FNetBitReader(UPackageMap* InPackageMap, const uint8* Src, int64 CountBits)
 	: FBitReader(Src, CountBits)
 	, PackageMap( InPackageMap )
 {
diff --git a/Engine/Source/Runtime/CoreUObject/Public/UObject/CoreNet.h b/Engine/Source/Runtime/CoreUObject/Public/UObject/CoreNet.h
--- a/Engine/Source/Runtime/CoreUObject/Public/UObject/CoreNet.h
+++ b/Engine/Source/Runtime/CoreUObject/Public/UObject/CoreNet.h
@@ -406,7 +406,7 @@ public:
 
 
 public:
-	FNetBitReader(UPackageMap* InPackageMap=nullptr, uint8* Src=nullptr, int64 CountBits=0);
+	FNetBitReader(UPackageMap* InPackageMap=nullptr, const uint8* Src=nullptr, int64 CountBits=0);
 
 	virtual FArchive& operator<<(FName& Name) override;
 	virtual FArchive& operator<<(UObject*& Object) override;
