Added the option to dump per class UObject counts to log when maximum number of UObjects limit is reached

diff --git a/Engine/Source/Runtime/CoreUObject/Private/UObject/UObjectArray.cpp b/Engine/Source/Runtime/CoreUObject/Private/UObject/UObjectArray.cpp
--- a/Engine/Source/Runtime/CoreUObject/Private/UObject/UObjectArray.cpp
+++ b/Engine/Source/Runtime/CoreUObject/Private/UObject/UObjectArray.cpp
@@ -5,6 +5,7 @@
 =============================================================================*/
 
 #include "UObject/UObjectArray.h"
+#include "HAL/IConsoleManager.h"
 #include "Misc/ScopeLock.h"
 #include "UObject/UObjectAllocator.h"
 #include "UObject/Class.h"
@@ -438,3 +439,97 @@ void FUObjectArray::ShutdownUObjectArray()
 		UE_CLOG(UObjectCreateListeners.Num(), LogUObjectArray, Fatal, TEXT("All UObject delete listeners should be unregistered when shutting down the UObject array"));
 	}
 }
+
+void FUObjectArray::DumpUObjectCountsToLog() const
+{
+	UE_LOG(LogUObjectArray, Display, TEXT("Dumping allocated UObject counts to log:"));
+	struct FClassEntry
+	{
+		UClass* Class = nullptr;
+		int32 NumInstances = 0;
+	};
+	int32 NumClasses = 0;
+	int32 NumUObjects = 0;
+	for (int32 ObjectIndex = 0; ObjectIndex < GetObjectArrayNum(); ++ObjectIndex)
+	{
+		const FUObjectItem& ObjectItem = GetObjectItemArrayUnsafe()[ObjectIndex];
+		UObject* Object = (UObject*)ObjectItem.Object;
+		if (Cast<UClass>(Object))
+		{
+			NumClasses++;
+		}
+	}
+
+	TMap<UClass*, FClassEntry> ClassCountMap;
+	ClassCountMap.Reserve(NumClasses);
+
+	for (int32 ObjectIndex = 0; ObjectIndex < GetObjectArrayNum(); ++ObjectIndex)
+	{
+		const FUObjectItem& ObjectItem = GetObjectItemArrayUnsafe()[ObjectIndex];
+		if (ObjectItem.Object)
+		{
+			UObject* Object = (UObject*)ObjectItem.Object;
+			UClass* ObjectClass = Object->GetClass();
+			FClassEntry& ClassEntry = ClassCountMap.FindOrAdd(ObjectClass);
+			ClassEntry.Class = ObjectClass;
+			ClassEntry.NumInstances++;
+			NumUObjects++;
+		}
+	}
+
+	TArray<FClassEntry> ClassArray;
+	ClassCountMap.GenerateValueArray(ClassArray);
+
+	ClassArray.Sort([](const FClassEntry& A, const FClassEntry& B) { return A.NumInstances > B.NumInstances; });
+
+	const int32 MinInstanceNum = 10; // Don't print classes with fewer than the specified number of instances
+	const double MaxPrintedInstancePercent = 0.95; // Finish printing when the specified percent of instances has already been printed
+	int32 NumClassesSkipped = 0;
+	int32 NumInstancesSkipped = 0;
+	int32 NumInstancesPrinted = 0;
+	double PercentOfInstancesPrinted = 0.0;
+
+	for (const FClassEntry& ClassEntry : ClassArray)
+	{		
+		if (ClassEntry.NumInstances > MinInstanceNum && PercentOfInstancesPrinted <= MaxPrintedInstancePercent)
+		{
+			UE_LOG(LogUObjectArray, Display, TEXT("%8d instances of %s"), ClassEntry.NumInstances, *ClassEntry.Class->GetPathName());
+			NumInstancesPrinted += ClassEntry.NumInstances;
+			PercentOfInstancesPrinted = (double)NumInstancesPrinted / NumUObjects;
+		}
+		else
+		{
+			NumClassesSkipped++;
+			NumInstancesSkipped += ClassEntry.NumInstances;
+		}
+	}
+	if (NumInstancesSkipped > 0)
+	{
+		if (PercentOfInstancesPrinted > MaxPrintedInstancePercent)
+		{
+			UE_LOG(LogUObjectArray, Display, TEXT("%8d instances in the remaining %.3f%% of instances of %d classes"), NumInstancesSkipped, (1.0f - PercentOfInstancesPrinted) * 100.0f, NumClassesSkipped);
+		}
+		else
+		{
+			UE_LOG(LogUObjectArray, Display, TEXT("%8d instances of %d classes with less than %d instances per class"), NumInstancesSkipped, NumClassesSkipped, MinInstanceNum);
+		}
+	}
+	UE_LOG(LogUObjectArray, Display, TEXT("%d total UObjects (%d classes)"), NumUObjects, NumClasses);
+}
+
+static int32 GVarDumpObjectCountsToLogWhenMaxObjectLimitExceeded = 0;
+static FAutoConsoleVariableRef CDumpObjectCountsToLogWhenMaxObjectLimitExceeded(
+	TEXT("gc.DumpObjectCountsToLogWhenMaxObjectLimitExceeded"),
+	GVarDumpObjectCountsToLogWhenMaxObjectLimitExceeded,
+	TEXT("If not 0 dumps UObject counts to log when maximum object count limit has been reached."),
+	ECVF_Default
+);
+
+void UE::UObjectArrayPrivate::FailMaxUObjectCountExceeded(const int32 MaxUObjects, const int32 NewUObjectCount)
+{
+	if (GVarDumpObjectCountsToLogWhenMaxObjectLimitExceeded)
+	{
+		GUObjectArray.DumpUObjectCountsToLog();
+	}
+	UE_LOG(LogUObjectArray, Fatal, TEXT("Maximum number of UObjects (%d) exceeded when trying to add %d object(s), make sure you update MaxObjectsInGame/MaxObjectsInEditor/MaxObjectsInProgram in project settings."), MaxUObjects, NewUObjectCount);
+}
diff --git a/Engine/Source/Runtime/CoreUObject/Public/UObject/UObjectArray.h b/Engine/Source/Runtime/CoreUObject/Public/UObject/UObjectArray.h
--- a/Engine/Source/Runtime/CoreUObject/Public/UObject/UObjectArray.h
+++ b/Engine/Source/Runtime/CoreUObject/Public/UObject/UObjectArray.h
@@ -222,6 +222,22 @@ private:
 	}
 };
 
+namespace UE::UObjectArrayPrivate
+{
+	COREUOBJECT_API void FailMaxUObjectCountExceeded(const int32 MaxUObjects, const int32 NewUObjectCount);
+
+	FORCEINLINE void CheckUObjectLimitReached(const int32 NumUObjects, const int32 MaxUObjects, const int32 NewUObjectCount)
+	{
+#if DO_CHECK
+		if ((NumUObjects + NewUObjectCount) > MaxUObjects)
+		{
+			FailMaxUObjectCountExceeded(MaxUObjects, NewUObjectCount);
+		}
+#endif // DO_CHECK
+	}
+};
+
+
 /**
 * Fixed size UObject array.
 */
@@ -262,7 +278,7 @@ public:
 	int32 AddSingle() TSAN_SAFE
 	{
 		int32 Result = NumElements;
-		checkf(NumElements + 1 <= MaxElements, TEXT("Maximum number of UObjects (%d) exceeded, make sure you update MaxObjectsInGame/MaxObjectsInEditor/MaxObjectsInProgram in project settings."), MaxElements);
+		UE::UObjectArrayPrivate::CheckUObjectLimitReached(NumElements, MaxElements, 1);
 		check(Result == NumElements);
 		++NumElements;
 		FPlatformMisc::MemoryBarrier();
@@ -273,7 +289,7 @@ public:
 	int32 AddRange(int32 Count) TSAN_SAFE
 	{
 		int32 Result = NumElements + Count - 1;
-		checkf(NumElements + Count <= MaxElements, TEXT("Maximum number of UObjects (%d) exceeded, make sure you update MaxObjectsInGame/MaxObjectsInEditor/MaxObjectsInProgram in project settings."), MaxElements);
+		UE::UObjectArrayPrivate::CheckUObjectLimitReached(NumElements, MaxElements, Count);
 		check(Result == (NumElements + Count - 1));
 		NumElements += Count;
 		FPlatformMisc::MemoryBarrier();
@@ -546,7 +562,7 @@ public:
 	int32 AddRange(int32 NumToAdd) TSAN_SAFE
 	{
 		int32 Result = NumElements;
-		checkf(Result + NumToAdd <= MaxElements, TEXT("Maximum number of UObjects (%d) exceeded, make sure you update MaxObjectsInGame/MaxObjectsInEditor/MaxObjectsInProgram in project settings."), MaxElements);
+		UE::UObjectArrayPrivate::CheckUObjectLimitReached(Result, MaxElements, NumToAdd);
 		ExpandChunksToIndex(Result + NumToAdd - 1);
 		NumElements += NumToAdd;
 		return Result;
@@ -1130,10 +1146,17 @@ public:
 		return ObjObjects;
 	}
     
+	const TUObjectArray& GetObjectItemArrayUnsafe() const
+	{
+		return ObjObjects;
+	}
+
     int64 GetAllocatedSize() const
     {
         return ObjObjects.GetAllocatedSize();
     }
+
+	void DumpUObjectCountsToLog() const;
 };
 
 /** UObject cluster. Groups UObjects into a single unit for GC. */
diff --git a/Engine/Source/Runtime/Engine/Classes/Engine/CoreSettings.h b/Engine/Source/Runtime/Engine/Classes/Engine/CoreSettings.h
--- a/Engine/Source/Runtime/Engine/Classes/Engine/CoreSettings.h
+++ b/Engine/Source/Runtime/Engine/Classes/Engine/CoreSettings.h
@@ -224,11 +224,17 @@ protected:
 		ConsoleVariable = "gc.VerifyUObjectsAreNotFGCObjects", DisplayName = "Verify UObjects Are Not FGCObjects",
 		ToolTip = "If true, the engine will throw a warning when it detects a UObject-derived class which also derives from FGCObject or any of its members is derived from FGCObject."))
 	uint32 VerifyUObjectsAreNotFGCObjects : 1;
+
 	UPROPERTY(EditAnywhere, config, Category = Optimization, meta = (
 		ConsoleVariable = "gc.PendingKillEnabled", DisplayName = "Pending Kill Enabled",
 		ToolTip = "If true, objects marked as PendingKill will be automatically nulled and destroyed by Garbage Collector."))
 	uint32 PendingKillEnabled : 1;
 
+	UPROPERTY(EditAnywhere, config, Category = Optimization, meta = (
+		ConsoleVariable = "gc.DumpObjectCountsToLogWhenMaxObjectLimitExceeded", DisplayName = "Log UObject counts when UObject limit has been exceeded",
+		ToolTip = "If true, the engine will dump UObject counts when maximum number of UObjects limit has been exceeded."))
+	uint32 DumpObjectCountsToLogWhenMaxObjectLimitExceeded : 1;
+
 	UPROPERTY(EditAnywhere, config, Category = Optimization, meta = (
 		ConsoleVariable = "gc.MinGCClusterSize", DisplayName = "Minimum GC Cluster size",
 		ToolTip = "Minimum GC cluster size."))
diff --git a/Engine/Source/Runtime/Engine/Private/CoreSettings.cpp b/Engine/Source/Runtime/Engine/Private/CoreSettings.cpp
--- a/Engine/Source/Runtime/Engine/Private/CoreSettings.cpp
+++ b/Engine/Source/Runtime/Engine/Private/CoreSettings.cpp
@@ -204,6 +204,7 @@ UGarbageCollectionSettings::UGarbageCollectionSettings()
 	VerifyGCObjectNames = true;
 	VerifyUObjectsAreNotFGCObjects = true;
 	PendingKillEnabled = false;
+	DumpObjectCountsToLogWhenMaxObjectLimitExceeded = false;
 }
 
 void UGarbageCollectionSettings::PostInitProperties()
