Packing FProperty fields to reduce memory footprint

diff --git a/Engine/Source/Runtime/CoreUObject/Private/UObject/Field.cpp b/Engine/Source/Runtime/CoreUObject/Private/UObject/Field.cpp
--- a/Engine/Source/Runtime/CoreUObject/Private/UObject/Field.cpp
+++ b/Engine/Source/Runtime/CoreUObject/Private/UObject/Field.cpp
@@ -113,9 +113,9 @@ FArchive& operator << (FArchive& Ar, FFieldClass*& InOutFieldClass)
 
 FFieldVariant FFieldVariant::GetOwnerVariant() const
 {
-	if (bIsUObject)
+	if (IsUObject())
 	{
-		return Container.Object->GetOuter();
+		return ToUObjectUnsafe()->GetOuter();
 	}
 	else
 	{
@@ -125,19 +125,19 @@ FFieldVariant FFieldVariant::GetOwnerVariant() const
 
 bool FFieldVariant::IsA(const UClass* InClass) const
 {
-	return bIsUObject && Container.Object && Container.Object->IsA(InClass);
+	return IsUObject() && ToUObjectUnsafe() && ToUObjectUnsafe()->IsA(InClass);
 }
 bool FFieldVariant::IsA(const FFieldClass* InClass) const
 {
-	return !bIsUObject && Container.Field && Container.Field->IsA(InClass);
+	return !IsUObject() && Container.Field && Container.Field->IsA(InClass);
 }
 
 UClass* FFieldVariant::GetOwnerClass() const
 {
 	check(Container.Object);
-	if (bIsUObject)
+	if (IsUObject())
 	{
-		return CastChecked<UField>(Container.Object)->GetOwnerClass();
+		return CastChecked<UField>(ToUObjectUnsafe())->GetOwnerClass();
 	}
 	else
 	{
@@ -147,9 +147,9 @@ UClass* FFieldVariant::GetOwnerClass() const
 
 FString FFieldVariant::GetFullName() const
 {
-	if (bIsUObject)
+	if (IsUObject())
 	{
-		return Container.Object->GetFullName();
+		return ToUObjectUnsafe()->GetFullName();
 	}
 	else
 	{
@@ -159,9 +159,9 @@ FString FFieldVariant::GetFullName() const
 
 FString FFieldVariant::GetPathName() const
 {
-	if (bIsUObject)
+	if (IsUObject())
 	{
-		return Container.Object->GetPathName();
+		return ToUObjectUnsafe()->GetPathName();
 	}
 	else
 	{
@@ -171,9 +171,9 @@ FString FFieldVariant::GetPathName() const
 
 FString FFieldVariant::GetName() const
 {
-	if (bIsUObject)
+	if (IsUObject())
 	{
-		return Container.Object->GetName();
+		return ToUObjectUnsafe()->GetName();
 	}
 	else
 	{
@@ -183,9 +183,9 @@ FString FFieldVariant::GetName() const
 
 FName FFieldVariant::GetFName() const
 {
-	if (bIsUObject)
+	if (IsUObject())
 	{
-		return Container.Object->GetFName();
+		return ToUObjectUnsafe()->GetFName();
 	}
 	else
 	{
@@ -196,9 +196,9 @@ FName FFieldVariant::GetFName() const
 FString FFieldVariant::GetClassName() const
 {
 	check(Container.Object);
-	if (bIsUObject)
+	if (IsUObject())
 	{
-		return Container.Object->GetClass()->GetName();
+		return ToUObjectUnsafe()->GetClass()->GetName();
 	}
 	else
 	{
@@ -209,9 +209,9 @@ FString FFieldVariant::GetClassName() const
 bool FFieldVariant::IsNative() const
 {
 	check(Container.Object);
-	if (bIsUObject)
+	if (IsUObject())
 	{
-		return Container.Object->IsNative();
+		return ToUObjectUnsafe()->IsNative();
 	}
 	else
 	{
@@ -222,9 +222,9 @@ bool FFieldVariant::IsNative() const
 UPackage* FFieldVariant::GetOutermost() const
 {
 	check(Container.Object);
-	if (bIsUObject)
+	if (IsUObject())
 	{
-		return Container.Object->GetOutermost();
+		return ToUObjectUnsafe()->GetOutermost();
 	}
 	else
 	{
@@ -235,9 +235,9 @@ UPackage* FFieldVariant::GetOutermost() const
 bool FFieldVariant::IsValidLowLevel() const
 {
 	check(Container.Object);
-	if (bIsUObject)
+	if (IsUObject())
 	{
-		return Container.Object->IsValidLowLevel();
+		return ToUObjectUnsafe()->IsValidLowLevel();
 	}
 	else
 	{
@@ -249,9 +249,9 @@ bool FFieldVariant::IsValidLowLevel() const
 bool FFieldVariant::HasMetaData(const FName& Key) const
 {
 	check(Container.Object);
-	if (bIsUObject)
+	if (IsUObject())
 	{
-		return CastChecked<UField>(Container.Object)->HasMetaData(Key);
+		return CastChecked<UField>(ToUObjectUnsafe())->HasMetaData(Key);
 	}
 	else
 	{
@@ -260,25 +260,6 @@ bool FFieldVariant::HasMetaData(const FName& Key) const
 }
 #endif // WITH_EDITORONLY_DATA
 
-FArchive& operator << (FArchive& Ar, FFieldVariant& InOutField)
-{
-	Ar << InOutField.bIsUObject;
-	if (InOutField.bIsUObject)
-	{
-		Ar << InOutField.Container.Object;
-	}
-	else
-	{
-		TFieldPath<FField> FieldRef(InOutField.Container.Field);
-		Ar << FieldRef;
-		if (Ar.IsLoading())
-		{
-			InOutField.Container.Field = FieldRef.Get();
-		}
-	}
-	return Ar;
-}
-
 /*-----------------------------------------------------------------------------
 FField implementation.
 -----------------------------------------------------------------------------*/
diff --git a/Engine/Source/Runtime/CoreUObject/Private/UObject/Property.cpp b/Engine/Source/Runtime/CoreUObject/Private/UObject/Property.cpp
--- a/Engine/Source/Runtime/CoreUObject/Private/UObject/Property.cpp
+++ b/Engine/Source/Runtime/CoreUObject/Private/UObject/Property.cpp
@@ -763,6 +763,10 @@ bool FGCStackSizeHelper::ConvertPathToProperties(UClass* ObjectClass, const FNam
 	FProperty implementation.
 -----------------------------------------------------------------------------*/
 
+#if UE_GAME
+	static_assert(sizeof(FProperty) <= 104, "FProperty was optimized to reduce its size so most of the classes that inherent from it will fall withing 112 bytes bin of MallocBinned3");
+#endif
+
 IMPLEMENT_FIELD(FProperty)
 
 //
diff --git a/Engine/Source/Runtime/CoreUObject/Public/UObject/Field.h b/Engine/Source/Runtime/CoreUObject/Public/UObject/Field.h
--- a/Engine/Source/Runtime/CoreUObject/Public/UObject/Field.h
+++ b/Engine/Source/Runtime/CoreUObject/Public/UObject/Field.h
@@ -266,20 +266,19 @@ class COREUOBJECT_API FFieldVariant
 		UObject* Object;
 	} Container;
 
-	bool bIsUObject;
+	static constexpr uintptr_t UObjectMask = 0x1;
 
 public:
 
 	FFieldVariant()
-		: bIsUObject(false)
 	{
 		Container.Field = nullptr;
 	}
 
 	FFieldVariant(const FField* InField)
-		: bIsUObject(false)
 	{
 		Container.Field = const_cast<FField*>(InField);
+		check(!IsUObject());
 	}
 
 	template <
@@ -287,9 +286,9 @@ public:
 		decltype(ImplicitConv<const UObject*>(std::declval<T>()))* = nullptr
 	>
 	FFieldVariant(T&& InObject)
-		: bIsUObject(true)
 	{
 		Container.Object = const_cast<UObject*>(ImplicitConv<const UObject*>(InObject));
+		Container.Object = (UObject*)((uintptr_t)Container.Object | UObjectMask);
 	}
 
 	FFieldVariant(TYPE_OF_NULLPTR)
@@ -299,11 +298,11 @@ public:
 
 	inline bool IsUObject() const
 	{
-		return bIsUObject;
+		return (uintptr_t)Container.Object & UObjectMask;
 	}
 	inline bool IsValid() const
 	{
-		return !!Container.Object;
+		return !!ToUObjectUnsafe();
 	}
 	bool IsValidLowLevel() const;
 	inline operator bool() const
@@ -327,7 +326,7 @@ public:
 		{
 			if constexpr (std::is_base_of_v<UObject, T>)
 			{
-				return static_cast<T*>(Container.Object);
+				return static_cast<T*>(ToUObjectUnsafe());
 			}
 			else
 			{
@@ -339,9 +338,9 @@ public:
 
 	UObject* ToUObject() const
 	{
-		if (bIsUObject)
+		if (IsUObject())
 		{
-			return Container.Object;
+			return ToUObjectUnsafe();
 		}
 		else
 		{
@@ -350,7 +349,7 @@ public:
 	}
 	FField* ToField() const
 	{
-		if (!bIsUObject)
+		if (!IsUObject())
 		{
 			return Container.Field;
 		}
@@ -367,7 +366,7 @@ public:
 	/** FOR INTERNAL USE ONLY: Function that returns the owner as UObject without checking if it's actually a UObject */
 	FORCEINLINE UObject* ToUObjectUnsafe() const
 	{
-		return Container.Object;
+		return (UObject*)((uintptr_t)Container.Object & ~UObjectMask);
 	}
 
 	void* GetRawPointer() const
@@ -402,8 +401,6 @@ public:
 	{
 		return GetTypeHash(InFieldVariant.GetRawPointer());
 	}
-
-	COREUOBJECT_API friend FArchive& operator << (FArchive& Ar, FFieldVariant& InOutField);
 };
 
 /**
