Reverted "Packing FProperty fields to reduce memory footprint"

diff --git a/Engine/Source/Runtime/CoreUObject/Private/UObject/Field.cpp b/Engine/Source/Runtime/CoreUObject/Private/UObject/Field.cpp
--- a/Engine/Source/Runtime/CoreUObject/Private/UObject/Field.cpp
+++ b/Engine/Source/Runtime/CoreUObject/Private/UObject/Field.cpp
@@ -113,9 +113,9 @@ FArchive& operator << (FArchive& Ar, FFieldClass*& InOutFieldClass)
 
 FFieldVariant FFieldVariant::GetOwnerVariant() const
 {
-	if (IsUObject())
+	if (bIsUObject)
 	{
-		return ToUObjectUnsafe()->GetOuter();
+		return Container.Object->GetOuter();
 	}
 	else
 	{
@@ -125,19 +125,19 @@ FFieldVariant FFieldVariant::GetOwnerVariant() const
 
 bool FFieldVariant::IsA(const UClass* InClass) const
 {
-	return IsUObject() && ToUObjectUnsafe() && ToUObjectUnsafe()->IsA(InClass);
+	return bIsUObject && Container.Object && Container.Object->IsA(InClass);
 }
 bool FFieldVariant::IsA(const FFieldClass* InClass) const
 {
-	return !IsUObject() && Container.Field && Container.Field->IsA(InClass);
+	return !bIsUObject && Container.Field && Container.Field->IsA(InClass);
 }
 
 UClass* FFieldVariant::GetOwnerClass() const
 {
 	check(Container.Object);
-	if (IsUObject())
+	if (bIsUObject)
 	{
-		return CastChecked<UField>(ToUObjectUnsafe())->GetOwnerClass();
+		return CastChecked<UField>(Container.Object)->GetOwnerClass();
 	}
 	else
 	{
@@ -147,9 +147,9 @@ UClass* FFieldVariant::GetOwnerClass() const
 
 FString FFieldVariant::GetFullName() const
 {
-	if (IsUObject())
+	if (bIsUObject)
 	{
-		return ToUObjectUnsafe()->GetFullName();
+		return Container.Object->GetFullName();
 	}
 	else
 	{
@@ -159,9 +159,9 @@ FString FFieldVariant::GetFullName() const
 
 FString FFieldVariant::GetPathName() const
 {
-	if (IsUObject())
+	if (bIsUObject)
 	{
-		return ToUObjectUnsafe()->GetPathName();
+		return Container.Object->GetPathName();
 	}
 	else
 	{
@@ -171,9 +171,9 @@ FString FFieldVariant::GetPathName() const
 
 FString FFieldVariant::GetName() const
 {
-	if (IsUObject())
+	if (bIsUObject)
 	{
-		return ToUObjectUnsafe()->GetName();
+		return Container.Object->GetName();
 	}
 	else
 	{
@@ -183,9 +183,9 @@ FString FFieldVariant::GetName() const
 
 FName FFieldVariant::GetFName() const
 {
-	if (IsUObject())
+	if (bIsUObject)
 	{
-		return ToUObjectUnsafe()->GetFName();
+		return Container.Object->GetFName();
 	}
 	else
 	{
@@ -196,9 +196,9 @@ FName FFieldVariant::GetFName() const
 FString FFieldVariant::GetClassName() const
 {
 	check(Container.Object);
-	if (IsUObject())
+	if (bIsUObject)
 	{
-		return ToUObjectUnsafe()->GetClass()->GetName();
+		return Container.Object->GetClass()->GetName();
 	}
 	else
 	{
@@ -209,9 +209,9 @@ FString FFieldVariant::GetClassName() const
 bool FFieldVariant::IsNative() const
 {
 	check(Container.Object);
-	if (IsUObject())
+	if (bIsUObject)
 	{
-		return ToUObjectUnsafe()->IsNative();
+		return Container.Object->IsNative();
 	}
 	else
 	{
@@ -222,9 +222,9 @@ bool FFieldVariant::IsNative() const
 UPackage* FFieldVariant::GetOutermost() const
 {
 	check(Container.Object);
-	if (IsUObject())
+	if (bIsUObject)
 	{
-		return ToUObjectUnsafe()->GetOutermost();
+		return Container.Object->GetOutermost();
 	}
 	else
 	{
@@ -235,9 +235,9 @@ UPackage* FFieldVariant::GetOutermost() const
 bool FFieldVariant::IsValidLowLevel() const
 {
 	check(Container.Object);
-	if (IsUObject())
+	if (bIsUObject)
 	{
-		return ToUObjectUnsafe()->IsValidLowLevel();
+		return Container.Object->IsValidLowLevel();
 	}
 	else
 	{
@@ -249,9 +249,9 @@ bool FFieldVariant::IsValidLowLevel() const
 bool FFieldVariant::HasMetaData(const FName& Key) const
 {
 	check(Container.Object);
-	if (IsUObject())
+	if (bIsUObject)
 	{
-		return CastChecked<UField>(ToUObjectUnsafe())->HasMetaData(Key);
+		return CastChecked<UField>(Container.Object)->HasMetaData(Key);
 	}
 	else
 	{
@@ -260,6 +260,25 @@ bool FFieldVariant::HasMetaData(const FName& Key) const
 }
 #endif // WITH_EDITORONLY_DATA
 
+FArchive& operator << (FArchive& Ar, FFieldVariant& InOutField)
+{
+	Ar << InOutField.bIsUObject;
+	if (InOutField.bIsUObject)
+	{
+		Ar << InOutField.Container.Object;
+	}
+	else
+	{
+		TFieldPath<FField> FieldRef(InOutField.Container.Field);
+		Ar << FieldRef;
+		if (Ar.IsLoading())
+		{
+			InOutField.Container.Field = FieldRef.Get();
+		}
+	}
+	return Ar;
+}
+
 /*-----------------------------------------------------------------------------
 FField implementation.
 -----------------------------------------------------------------------------*/
diff --git a/Engine/Source/Runtime/CoreUObject/Private/UObject/Property.cpp b/Engine/Source/Runtime/CoreUObject/Private/UObject/Property.cpp
--- a/Engine/Source/Runtime/CoreUObject/Private/UObject/Property.cpp
+++ b/Engine/Source/Runtime/CoreUObject/Private/UObject/Property.cpp
@@ -763,10 +763,6 @@ bool FGCStackSizeHelper::ConvertPathToProperties(UClass* ObjectClass, const FNam
 	FProperty implementation.
 -----------------------------------------------------------------------------*/
 
-#if UE_GAME
-	static_assert(sizeof(FProperty) <= 104, "FProperty was optimized to reduce its size so most of the classes that inherent from it will fall withing 112 bytes bin of MallocBinned3");
-#endif
-
 IMPLEMENT_FIELD(FProperty)
 
 //
diff --git a/Engine/Source/Runtime/CoreUObject/Public/UObject/Field.h b/Engine/Source/Runtime/CoreUObject/Public/UObject/Field.h
--- a/Engine/Source/Runtime/CoreUObject/Public/UObject/Field.h
+++ b/Engine/Source/Runtime/CoreUObject/Public/UObject/Field.h
@@ -266,19 +266,20 @@ class COREUOBJECT_API FFieldVariant
 		UObject* Object;
 	} Container;
 
-	static constexpr uintptr_t UObjectMask = 0x1;
+	bool bIsUObject;
 
 public:
 
 	FFieldVariant()
+		: bIsUObject(false)
 	{
 		Container.Field = nullptr;
 	}
 
 	FFieldVariant(const FField* InField)
+		: bIsUObject(false)
 	{
 		Container.Field = const_cast<FField*>(InField);
-		check(!IsUObject());
 	}
 
 	template <
@@ -286,9 +287,9 @@ public:
 		decltype(ImplicitConv<const UObject*>(std::declval<T>()))* = nullptr
 	>
 	FFieldVariant(T&& InObject)
+		: bIsUObject(true)
 	{
 		Container.Object = const_cast<UObject*>(ImplicitConv<const UObject*>(InObject));
-		Container.Object = (UObject*)((uintptr_t)Container.Object | UObjectMask);
 	}
 
 	FFieldVariant(TYPE_OF_NULLPTR)
@@ -298,11 +299,11 @@ public:
 
 	inline bool IsUObject() const
 	{
-		return (uintptr_t)Container.Object & UObjectMask;
+		return bIsUObject;
 	}
 	inline bool IsValid() const
 	{
-		return !!ToUObjectUnsafe();
+		return !!Container.Object;
 	}
 	bool IsValidLowLevel() const;
 	inline operator bool() const
@@ -326,7 +327,7 @@ public:
 		{
 			if constexpr (std::is_base_of_v<UObject, T>)
 			{
-				return static_cast<T*>(ToUObjectUnsafe());
+				return static_cast<T*>(Container.Object);
 			}
 			else
 			{
@@ -338,9 +339,9 @@ public:
 
 	UObject* ToUObject() const
 	{
-		if (IsUObject())
+		if (bIsUObject)
 		{
-			return ToUObjectUnsafe();
+			return Container.Object;
 		}
 		else
 		{
@@ -349,7 +350,7 @@ public:
 	}
 	FField* ToField() const
 	{
-		if (!IsUObject())
+		if (!bIsUObject)
 		{
 			return Container.Field;
 		}
@@ -366,7 +367,7 @@ public:
 	/** FOR INTERNAL USE ONLY: Function that returns the owner as UObject without checking if it's actually a UObject */
 	FORCEINLINE UObject* ToUObjectUnsafe() const
 	{
-		return (UObject*)((uintptr_t)Container.Object & ~UObjectMask);
+		return Container.Object;
 	}
 
 	void* GetRawPointer() const
@@ -401,6 +402,8 @@ public:
 	{
 		return GetTypeHash(InFieldVariant.GetRawPointer());
 	}
+
+	COREUOBJECT_API friend FArchive& operator << (FArchive& Ar, FFieldVariant& InOutField);
 };
 
 /**
