ZenLoader in editor:
Fixed issue with sync load context not being propagated correctly when merging postload groups
Fixed crash caused by a package recursively being completed while it was processed as an imported package
#rnx

diff --git a/Engine/Source/Runtime/CoreUObject/Private/Serialization/AsyncLoading2.cpp b/Engine/Source/Runtime/CoreUObject/Private/Serialization/AsyncLoading2.cpp
--- a/Engine/Source/Runtime/CoreUObject/Private/Serialization/AsyncLoading2.cpp
+++ b/Engine/Source/Runtime/CoreUObject/Private/Serialization/AsyncLoading2.cpp
@@ -3816,6 +3816,10 @@ void FAsyncLoadingThread2::MergePostLoadGroups(FAsyncLoadingThreadState2& Thread
 	}
 	TRACE_CPUPROFILER_EVENT_SCOPE(MergePostLoadGroups);
 	check(ThreadState.bCanAccessAsyncLoadingThreadData);
+	if (Source->Packages.Num() > Target->Packages.Num())
+	{
+		Swap(Source, Target);
+	}
 	for (FAsyncPackage2* Package : Source->Packages)
 	{
 		check(Package->PostLoadGroup == Source);
@@ -3823,6 +3827,21 @@ void FAsyncLoadingThread2::MergePostLoadGroups(FAsyncLoadingThreadState2& Thread
 	}
 	Target->Packages.Append(MoveTemp(Source->Packages));
 	Target->PackagesWithExportsToSerializeCount += Source->PackagesWithExportsToSerializeCount;
+
+	const uint64 SyncLoadContextId = FMath::Max(Source->SyncLoadContextId, Target->SyncLoadContextId);
+	if (SyncLoadContextId)
+	{
+		Target->SyncLoadContextId = SyncLoadContextId;
+		for (FAsyncPackage2* Package : Target->Packages)
+		{
+			Package->SyncLoadContextId = SyncLoadContextId;
+			if (Package->Desc.Priority < MAX_int32)
+			{
+				Package->Desc.Priority = MAX_int32;
+				UpdatePackagePriority(ThreadState, Package);
+			}
+		}
+	}
 	delete Source;
 }
 
@@ -3874,8 +3893,7 @@ FAsyncPackage2* FAsyncLoadingThread2::FindOrInsertPackage(FAsyncLoadingThreadSta
 	else if (ImportedByPackage)
 	{
 		// Importing a package that was already being loaded
-		check(Package->PostLoadGroup);
-		if (Package->AsyncPackageLoadingState >= EAsyncPackageLoadingState2::ExportsDone)
+		if (!Package->PostLoadGroup)
 		{
 			// The imported package has started postloading, wait for it to finish postloading before serializing any exports
 			for (int32 DependentExportBundleIndex = 0; DependentExportBundleIndex < ImportedByPackage->Data.TotalExportBundleCount; ++DependentExportBundleIndex)
@@ -5385,7 +5403,8 @@ bool FAsyncPackage2::ProcessLinkerLoadPackageImports(FAsyncLoadingThreadState2&
 	const int32 ImportedPackagesCount = Data.ImportedAsyncPackages.Num();
 	while (LinkerLoadState->ProcessingImportedPackageIndex < ImportedPackagesCount)
 	{
-		FAsyncPackage2* ImportedPackage = Data.ImportedAsyncPackages[LinkerLoadState->ProcessingImportedPackageIndex];
+		const int32 ImportIndex = LinkerLoadState->ProcessingImportedPackageIndex++;
+		FAsyncPackage2* ImportedPackage = Data.ImportedAsyncPackages[ImportIndex];
 		if (ImportedPackage && ImportedPackage->LinkerLoadState.IsSet())
 		{
 			check(ImportedPackage->AsyncPackageLoadingState >= EAsyncPackageLoadingState2::DependenciesReady);
@@ -5404,12 +5423,10 @@ bool FAsyncPackage2::ProcessLinkerLoadPackageImports(FAsyncLoadingThreadState2&
 				}
 			}
 		}
-		++LinkerLoadState->ProcessingImportedPackageIndex;
-	}
-
-	if (AsyncPackageLoadingState > EAsyncPackageLoadingState2::ProcessExportBundles)
-	{
-		return true;
+		if (AsyncPackageLoadingState > EAsyncPackageLoadingState2::ProcessExportBundles)
+		{
+			return true;
+		}
 	}
 
 	const int32 ImportCount = LinkerLoadState->Linker->ImportMap.Num();
@@ -8358,10 +8375,8 @@ void FAsyncLoadingThread2::FlushLoading(int32 RequestId)
 			}
 			if (GameThreadState->bCanAccessAsyncLoadingThreadData)
 			{
-				FAsyncLoadingPostLoadGroup* PostLoadGroup = CurrentlyExecutingPackage->PostLoadGroup;
-				if (CurrentlyExecutingPackage->AsyncPackageLoadingState < EAsyncPackageLoadingState2::ExportsDone)
+				if (FAsyncLoadingPostLoadGroup* PostLoadGroup = CurrentlyExecutingPackage->PostLoadGroup)
 				{
-					check(PostLoadGroup);
 					check(PostLoadGroup->Packages.Contains(CurrentlyExecutingPackage));
 					check(PostLoadGroup->PackagesWithExportsToSerializeCount > 0);
 					if (PostLoadGroup->Packages.Num() > 1)
@@ -8370,6 +8385,7 @@ void FAsyncLoadingThread2::FlushLoading(int32 RequestId)
 						--PostLoadGroup->PackagesWithExportsToSerializeCount;
 						ConditionalBeginPostLoad(*GameThreadState, PostLoadGroup);
 						CurrentlyExecutingPackage->PostLoadGroup = new FAsyncLoadingPostLoadGroup();
+						CurrentlyExecutingPackage->PostLoadGroup->SyncLoadContextId = CurrentlyExecutingPackage->SyncLoadContextId;
 						CurrentlyExecutingPackage->PostLoadGroup->Packages.Add(CurrentlyExecutingPackage);
 						CurrentlyExecutingPackage->PostLoadGroup->PackagesWithExportsToSerializeCount = 1;
 					}
