Fixing non-unity builds

diff --git a/Engine/Source/Runtime/CoreUObject/Private/UObject/UObjectArray.cpp b/Engine/Source/Runtime/CoreUObject/Private/UObject/UObjectArray.cpp
--- a/Engine/Source/Runtime/CoreUObject/Private/UObject/UObjectArray.cpp
+++ b/Engine/Source/Runtime/CoreUObject/Private/UObject/UObjectArray.cpp
@@ -454,7 +454,7 @@ void FUObjectArray::DumpUObjectCountsToLog() const
 	{
 		const FUObjectItem& ObjectItem = GetObjectItemArrayUnsafe()[ObjectIndex];
 		UObject* Object = (UObject*)ObjectItem.Object;
-		if (Cast<UClass>(Object))
+		if (Object && Object->IsA(UClass::StaticClass()))
 		{
 			NumClasses++;
 		}
