autortfm: change to begin/end macros to ensure we completely compile away when not in use

diff --git a/Engine/Source/Runtime/Core/Private/Delegates/DelegateHandle.cpp b/Engine/Source/Runtime/Core/Private/Delegates/DelegateHandle.cpp
--- a/Engine/Source/Runtime/Core/Private/Delegates/DelegateHandle.cpp
+++ b/Engine/Source/Runtime/Core/Private/Delegates/DelegateHandle.cpp
@@ -15,7 +15,7 @@ uint64 FDelegateHandle::GenerateNewID()
 	// Just increment a counter to generate an ID.
 	uint64 Result;
 	
-	AutoRTFM::Open([&Result]()
+	UE_AUTORTFM_OPEN_BEGIN
 	{
 		Result = ++UE::Delegates::Private::GNextID;
 
@@ -25,7 +25,8 @@ uint64 FDelegateHandle::GenerateNewID()
 			// Increment it again - it might not be zero, so don't just assign it to 1.
 			Result = ++UE::Delegates::Private::GNextID;
 		}
-	});
+	}
+	UE_AUTORTFM_OPEN_END
 
 	return Result;
 }
diff --git a/Engine/Source/Runtime/Core/Private/HAL/ThreadingBase.cpp b/Engine/Source/Runtime/Core/Private/HAL/ThreadingBase.cpp
--- a/Engine/Source/Runtime/Core/Private/HAL/ThreadingBase.cpp
+++ b/Engine/Source/Runtime/Core/Private/HAL/ThreadingBase.cpp
@@ -1341,9 +1341,10 @@ FQueuedThread::Run()
 
 FTlsAutoCleanup* FThreadSingletonInitializer::Get( TFunctionRef<FTlsAutoCleanup*()> CreateInstance, uint32& InOutTlsSlot )
 {
-	uint32 TlsSlot = (uint32)FPlatformAtomics::AtomicRead_Relaxed((int32*)&InOutTlsSlot);
-	AutoRTFM::Open([&TlsSlot, &InOutTlsSlot]
+	uint32 TlsSlot;
+	UE_AUTORTFM_OPEN_BEGIN
 	{
+		TlsSlot = (uint32)FPlatformAtomics::AtomicRead_Relaxed((int32*)&InOutTlsSlot);
 		if (TlsSlot == 0xFFFFFFFF)
 		{
 			const uint32 ThisTlsSlot = FPlatformTLS::AllocTlsSlot();
@@ -1359,7 +1360,8 @@ FTlsAutoCleanup* FThreadSingletonInitializer::Get( TFunctionRef<FTlsAutoCleanup*
 				TlsSlot = ThisTlsSlot;
 			}
 		}
-	});
+	}
+	UE_AUTORTFM_OPEN_END
 	
 	FTlsAutoCleanup* ThreadSingleton = (FTlsAutoCleanup*)FPlatformTLS::GetTlsValue( TlsSlot );
 	if( !ThreadSingleton )
diff --git a/Engine/Source/Runtime/Core/Private/UObject/UnrealNames.cpp b/Engine/Source/Runtime/Core/Private/UObject/UnrealNames.cpp
--- a/Engine/Source/Runtime/Core/Private/UObject/UnrealNames.cpp
+++ b/Engine/Source/Runtime/Core/Private/UObject/UnrealNames.cpp
@@ -1167,13 +1167,14 @@ public:
 	FNameEntryId Find(const FNameValue<Sensitivity>& Value) const
 	{
 		FNameEntryId Result;
-		AutoRTFM::Open([&Result, this, &Value]
+		UE_AUTORTFM_OPEN_BEGIN
 		{
 			FRWScopeLock _(Lock, FRWScopeLockType::SLT_ReadOnly);
 
 			FNameSlot& Slot = Probe(Value);
 			Result = Slot.GetId();
-		});
+		}
+		UE_AUTORTFM_OPEN_END
 		return Result;
 	}
 
diff --git a/Engine/Source/Runtime/Core/Public/AutoRTFM/AutoRTFM.h b/Engine/Source/Runtime/Core/Public/AutoRTFM/AutoRTFM.h
--- a/Engine/Source/Runtime/Core/Public/AutoRTFM/AutoRTFM.h
+++ b/Engine/Source/Runtime/Core/Public/AutoRTFM/AutoRTFM.h
@@ -513,6 +513,27 @@ struct FRegisterOpenFunction
     }
 };
 
+// Macro-based variants so we completely compile away when not in use, even in debug builds
+#if UE_AUTORTFM
+#define UE_AUTORTFM_OPEN_BEGIN AutoRTFM::Open([&]()
+#define UE_AUTORTFM_OPEN_END );
+
+#define UE_AUTORTFM_OPENCOMMIT_BEGIN AutoRTFM::OpenCommit([&]()
+#define UE_AUTORTFM_OPENCOMMIT_END );
+
+#define UE_AUTORTFM_TRANSACT_BEGIN AutoRTFM::Transact([&]()
+#define UE_AUTORTFM_TRANSACT_END );
+#else
+#define UE_AUTORTFM_OPEN_BEGIN
+#define UE_AUTORTFM_OPEN_END
+
+#define UE_AUTORTFM_OPENCOMMIT_BEGIN
+#define UE_AUTORTFM_OPENCOMMIT_END
+
+#define UE_AUTORTFM_TRANSACT_BEGIN
+#define UE_AUTORTFM_TRANSACT_END
+#endif
+
 #define UE_AUTORTFM_CONCAT_IMPL(A, B) A ## B
 #define UE_AUTORTFM_CONCAT(A, B) UE_AUTORTFM_CONCAT_IMPL(A, B)
 
diff --git a/Engine/Source/Runtime/Core/Public/HAL/FMemory.inl b/Engine/Source/Runtime/Core/Public/HAL/FMemory.inl
--- a/Engine/Source/Runtime/Core/Public/HAL/FMemory.inl
+++ b/Engine/Source/Runtime/Core/Public/HAL/FMemory.inl
@@ -24,7 +24,7 @@ FMEMORY_INLINE_FUNCTION_DECORATOR void* FMemory::Malloc(SIZE_T Count, uint32 Ali
 	// we don't end up keeping track of the writes to the allocator's internal data structures.
 	// This is because allocators are already transactional - malloc can be rolled back by
 	// calling free.
-	AutoRTFM::Open([&Ptr, Count, Alignment]
+	UE_AUTORTFM_OPEN_BEGIN
 	{
 		if (!FMEMORY_INLINE_GMalloc)
 		{
@@ -38,7 +38,8 @@ FMEMORY_INLINE_FUNCTION_DECORATOR void* FMemory::Malloc(SIZE_T Count, uint32 Ali
 		}
 		// optional tracking of every allocation
 		LLM_IF_ENABLED(FLowLevelMemTracker::Get().OnLowLevelAlloc(ELLMTracker::Default, Ptr, Count, ELLMTag::Untagged, ELLMAllocType::FMalloc));
-	});
+	}
+	UE_AUTORTFM_OPEN_END
 
 	// AutoRTFM: This is a no-op for non-transactional code.
 	// For transactional code, this defers a call to Free if the transaction aborts,
@@ -119,7 +120,7 @@ FMEMORY_INLINE_FUNCTION_DECORATOR void FMemory::Free(void* Original)
 	// AutoRTFM: For transactional code, in order to support the transaction 
 	// aborting and needing to 'roll back' the Free, we defer the actual
 	// free until commit time.
-	AutoRTFM::OpenCommit([Original]
+	UE_AUTORTFM_OPENCOMMIT_BEGIN
 	{
 		// optional tracking of every allocation
 		LLM_IF_ENABLED(FLowLevelMemTracker::Get().OnLowLevelFree(ELLMTracker::Default, Original, ELLMAllocType::FMalloc));
@@ -132,13 +133,14 @@ FMEMORY_INLINE_FUNCTION_DECORATOR void FMemory::Free(void* Original)
 		DoGamethreadHook(2);
 		FScopedMallocTimer Timer(2);
 		FMEMORY_INLINE_GMalloc->Free(Original);
-	});
+	}
+	UE_AUTORTFM_OPENCOMMIT_END
 }
 
 FMEMORY_INLINE_FUNCTION_DECORATOR SIZE_T FMemory::GetAllocSize(void* Original)
 {
 	SIZE_T Result;
-	AutoRTFM::Open([&Result, Original]
+	UE_AUTORTFM_OPEN_BEGIN
 	{
 		if (!FMEMORY_INLINE_GMalloc)
 		{
@@ -149,7 +151,8 @@ FMEMORY_INLINE_FUNCTION_DECORATOR SIZE_T FMemory::GetAllocSize(void* Original)
 			SIZE_T Size = 0;
 			Result = FMEMORY_INLINE_GMalloc->GetAllocationSize(Original, Size) ? Size : 0;
 		}
-	});
+	}
+	UE_AUTORTFM_OPEN_END
 
 	return Result;
 }
@@ -157,7 +160,7 @@ FMEMORY_INLINE_FUNCTION_DECORATOR SIZE_T FMemory::GetAllocSize(void* Original)
 FMEMORY_INLINE_FUNCTION_DECORATOR SIZE_T FMemory::QuantizeSize(SIZE_T Count, uint32 Alignment)
 {
 	SIZE_T Result;
-	AutoRTFM::Open([&Result, Count, Alignment]
+	UE_AUTORTFM_OPEN_BEGIN
 	{
 		if (!FMEMORY_INLINE_GMalloc)
 		{
@@ -167,7 +170,8 @@ FMEMORY_INLINE_FUNCTION_DECORATOR SIZE_T FMemory::QuantizeSize(SIZE_T Count, uin
 		{
 			Result = FMEMORY_INLINE_GMalloc->QuantizeSize(Count, Alignment);
 		}
-	});
+	}
+	UE_AUTORTFM_OPEN_END
 
 	return Result;
 }
diff --git a/Engine/Source/Runtime/Core/Public/Templates/SharedPointerInternals.h b/Engine/Source/Runtime/Core/Public/Templates/SharedPointerInternals.h
--- a/Engine/Source/Runtime/Core/Public/Templates/SharedPointerInternals.h
+++ b/Engine/Source/Runtime/Core/Public/Templates/SharedPointerInternals.h
@@ -187,7 +187,7 @@ namespace SharedPointerInternals
 		{
 			if constexpr (Mode == ESPMode::ThreadSafe)
 			{
-				AutoRTFM::OpenCommit([this]
+				UE_AUTORTFM_OPENCOMMIT_BEGIN
 				{
 					// std::memory_order_acq_rel is used here so that, if we do end up executing the destructor, it's not possible
 					// for side effects from executing the destructor end up being visible before we've determined that the shared
@@ -204,7 +204,8 @@ namespace SharedPointerInternals
 						// reference count reaches zero, this object will be deleted.
 						ReleaseWeakReference();
 					}
-				});
+				}
+				UE_AUTORTFM_OPENCOMMIT_END
 			}
 			else
 			{
diff --git a/Engine/Source/Runtime/CoreUObject/Private/UObject/Class.cpp b/Engine/Source/Runtime/CoreUObject/Private/UObject/Class.cpp
--- a/Engine/Source/Runtime/CoreUObject/Private/UObject/Class.cpp
+++ b/Engine/Source/Runtime/CoreUObject/Private/UObject/Class.cpp
@@ -5868,7 +5868,7 @@ UFunction* UClass::FindFunctionByName(FName InName, EIncludeSuperFlag::Type Incl
 
 	UFunction* Result = nullptr;
 
-	AutoRTFM::Open([&Result, this, &InName, IncludeSuper]
+	UE_AUTORTFM_OPEN_BEGIN
 	{
 		{
 			FReadScopeLock ScopeLock(FuncMapLock);
@@ -5911,7 +5911,8 @@ UFunction* UClass::FindFunctionByName(FName InName, EIncludeSuperFlag::Type Incl
 				}
 			}
 		}
-	});
+	}
+	UE_AUTORTFM_OPEN_END
 
 	return Result;
 }
diff --git a/Engine/Source/Runtime/CoreUObject/Private/UObject/ObjectHandleTracking.cpp b/Engine/Source/Runtime/CoreUObject/Private/UObject/ObjectHandleTracking.cpp
--- a/Engine/Source/Runtime/CoreUObject/Private/UObject/ObjectHandleTracking.cpp
+++ b/Engine/Source/Runtime/CoreUObject/Private/UObject/ObjectHandleTracking.cpp
@@ -20,68 +20,73 @@ namespace UE::CoreUObject
 
 			void OnHandleRead(TArrayView<const UObject* const> Objects)
 			{
-				AutoRTFM::Open([this, Objects]
+				UE_AUTORTFM_OPEN_BEGIN
 				{
 					FReadScopeLock _(HandleLock);
 					for (auto&& Pair : ReadHandleCallbacks)
 					{
 						Pair.Value(Objects);
 					}
-				});
+				}
+				UE_AUTORTFM_OPEN_END
 			}
 
 			void OnClassReferenceResolved(const FObjectRef& ObjectRef, UPackage* Package, UClass* Class)
 			{
-				AutoRTFM::Open([this, &ObjectRef, Package, Class]
+				UE_AUTORTFM_OPEN_BEGIN
 				{
 					FReadScopeLock _(HandleLock);
 					for (auto&& Pair : ClassResolvedCallbacks)
 					{
 						Pair.Value(ObjectRef, Package, Class);
 					}
-				});
+				}
+				UE_AUTORTFM_OPEN_END
 			}
 
 			void OnReferenceResolved(const FObjectRef& ObjectRef, UPackage* Package, UObject* Object)
 			{
-				AutoRTFM::Open([this, &ObjectRef, Package, Object]
+				UE_AUTORTFM_OPEN_BEGIN
 				{
 					FReadScopeLock _(HandleLock);
 					for (auto&& Pair : HandleResolvedCallbacks)
 					{
 						Pair.Value(ObjectRef, Package, Object);
 					}
-				});
+				}
+				UE_AUTORTFM_OPEN_END
 			}
 
 			void OnReferenceLoaded(const FObjectRef& ObjectRef, UPackage* Package, UObject* Object)
 			{
-				AutoRTFM::Open([this, &ObjectRef, Package, Object]
+				UE_AUTORTFM_OPEN_BEGIN
 				{
 					FReadScopeLock _(HandleLock);
 					for (auto&& Pair : HandleLoadedCallbacks)
 					{
 						Pair.Value(ObjectRef, Package, Object);
 					}
-				});
+				}
+				UE_AUTORTFM_OPEN_END
 			}
 
 			FObjectHandleTrackingCallbackId AddObjectHandleReadCallback(FObjectHandleReadFunc Func)
 			{
 				FObjectHandleTrackingCallbackId Result;
-				AutoRTFM::Open([&Result, this, &Func]
+				UE_AUTORTFM_OPEN_BEGIN
 				{
 					FWriteScopeLock _(HandleLock);
 					NextHandleId++;
 					ReadHandleCallbacks.Add({ NextHandleId, Func });
 					Result = FObjectHandleTrackingCallbackId{ NextHandleId };
-				});
+				}
+				UE_AUTORTFM_OPEN_END
 				return Result;
 			}
 
 			void RemoveObjectHandleReadCallback(FObjectHandleTrackingCallbackId Handle)
 			{
-				AutoRTFM::Open([this, Handle]
+				UE_AUTORTFM_OPEN_BEGIN
 				{
 					FWriteScopeLock _(HandleLock);
 					for (int32 i = ReadHandleCallbacks.Num() - 1; i >= 0; --i)
@@ -92,25 +97,27 @@ namespace UE::CoreUObject
 							ReadHandleCallbacks.RemoveAt(i);
 						}
 					}
-				});
+				}
+				UE_AUTORTFM_OPEN_END
 			}
 
 			FObjectHandleTrackingCallbackId AddObjectHandleClassResolvedCallback(FObjectHandleClassResolvedFunc Func)
 			{
 				FObjectHandleTrackingCallbackId Result;
-				AutoRTFM::Open([&Result, this, &Func]
+				UE_AUTORTFM_OPEN_BEGIN
 				{
 					FWriteScopeLock _(HandleLock);
 					NextHandleId++;
 					ClassResolvedCallbacks.Add({ NextHandleId, Func });
 					Result = FObjectHandleTrackingCallbackId{ NextHandleId };
-				});
+				}
+				UE_AUTORTFM_OPEN_END
 				return Result;
 			}
 
 			void RemoveObjectHandleClassResolvedCallback(FObjectHandleTrackingCallbackId Handle)
 			{
-				AutoRTFM::Open([this, Handle]
+				UE_AUTORTFM_OPEN_BEGIN
 				{
 					FWriteScopeLock _(HandleLock);
 					for (int32 i = ClassResolvedCallbacks.Num() - 1; i >= 0; --i)
@@ -121,25 +128,27 @@ namespace UE::CoreUObject
 							ClassResolvedCallbacks.RemoveAt(i);
 						}
 					}
-				});
+				}
+				UE_AUTORTFM_OPEN_END
 			}
 
 			FObjectHandleTrackingCallbackId AddObjectHandleReferenceResolvedCallback(FObjectHandleReferenceResolvedFunc Func)
 			{
 				FObjectHandleTrackingCallbackId Result;
-				AutoRTFM::Open([&Result, this, &Func]
+				UE_AUTORTFM_OPEN_BEGIN
 				{
 					FWriteScopeLock _(HandleLock);
 					NextHandleId++;
 					HandleResolvedCallbacks.Add({ NextHandleId, Func });
 					Result = FObjectHandleTrackingCallbackId{ NextHandleId };
-				});
+				}
+				UE_AUTORTFM_OPEN_END
 				return Result;
 			}
 
 			void RemoveObjectHandleReferenceResolvedCallback(FObjectHandleTrackingCallbackId Handle)
 			{
-				AutoRTFM::Open([this, Handle]
+				UE_AUTORTFM_OPEN_BEGIN
 				{
 					FWriteScopeLock _(HandleLock);
 					for (int32 i = HandleResolvedCallbacks.Num() - 1; i >= 0; --i)
@@ -150,25 +159,27 @@ namespace UE::CoreUObject
 							HandleResolvedCallbacks.RemoveAt(i);
 						}
 					}
-				});
+				}
+				UE_AUTORTFM_OPEN_END
 			}
 
 			FObjectHandleTrackingCallbackId AddObjectHandleReferenceLoadedCallback(FObjectHandleReferenceLoadedFunc Func)
 			{
 				FObjectHandleTrackingCallbackId Result;
-				AutoRTFM::Open([&Result, this, &Func]
+				UE_AUTORTFM_OPEN_BEGIN
 				{
 					FWriteScopeLock _(HandleLock);
 					NextHandleId++;
 					HandleLoadedCallbacks.Add({ NextHandleId, Func });
 					Result = FObjectHandleTrackingCallbackId{ NextHandleId };
-				});
+				}
+				UE_AUTORTFM_OPEN_END
 				return Result;
 			}
 
 			void RemoveObjectHandleReferenceLoadedCallback(FObjectHandleTrackingCallbackId Handle)
 			{
-				AutoRTFM::Open([this, Handle]
+				UE_AUTORTFM_OPEN_BEGIN
 				{
 					FWriteScopeLock _(HandleLock);
 					for (int32 i = HandleLoadedCallbacks.Num() - 1; i >= 0; --i)
@@ -179,7 +190,8 @@ namespace UE::CoreUObject
 							HandleLoadedCallbacks.RemoveAt(i);
 						}
 					}
-				});
+				}
+				UE_AUTORTFM_OPEN_END
 			}
 
 		private:
diff --git a/Engine/Source/Runtime/CoreUObject/Private/UObject/ScriptCore.cpp b/Engine/Source/Runtime/CoreUObject/Private/UObject/ScriptCore.cpp
--- a/Engine/Source/Runtime/CoreUObject/Private/UObject/ScriptCore.cpp
+++ b/Engine/Source/Runtime/CoreUObject/Private/UObject/ScriptCore.cpp
@@ -189,10 +189,7 @@ FBlueprintContext* FBlueprintContextGetThreadSingletonImpl()
 FBlueprintContext* FBlueprintContext::GetThreadSingleton()
 {
 	FBlueprintContext* Result;
-	AutoRTFM::Open([&Result]
-	{
-		Result = FBlueprintContextGetThreadSingletonImpl();
-	});
+	UE_AUTORTFM_OPEN_BEGIN { Result = FBlueprintContextGetThreadSingletonImpl(); } UE_AUTORTFM_OPEN_END
 	return Result;
 }
 
diff --git a/Engine/Source/Runtime/CoreUObject/Private/UObject/UObjectGlobals.cpp b/Engine/Source/Runtime/CoreUObject/Private/UObject/UObjectGlobals.cpp
--- a/Engine/Source/Runtime/CoreUObject/Private/UObject/UObjectGlobals.cpp
+++ b/Engine/Source/Runtime/CoreUObject/Private/UObject/UObjectGlobals.cpp
@@ -3854,10 +3854,11 @@ void FObjectInitializer::PostConstructInit()
 	{
 		// autortfm todo: if this transaction aborts and we are in a transaction's open nest,
 		// we need to have a way of propagating out that abort
-		AutoRTFM::Transact([this] 
+		UE_AUTORTFM_TRANSACT_BEGIN
 		{
 			PropertyInitCallback();
-		});
+		}
+		UE_AUTORTFM_TRANSACT_END
 	}
 	// After the call to `PropertyInitCallback` to allow the callback to modify the instancing graph
 	if (bNeedInstancing || bNeedSubobjectInstancing)
