Fix potential truncation issue when creating a FLinkerLoader

#rnx

- TotalSize is returned as a int64 and only ever used as an int64 so it should be safe to store locally as an int64.

diff --git a/Engine/Source/Runtime/CoreUObject/Private/UObject/LinkerLoad.cpp b/Engine/Source/Runtime/CoreUObject/Private/UObject/LinkerLoad.cpp
--- a/Engine/Source/Runtime/CoreUObject/Private/UObject/LinkerLoad.cpp
+++ b/Engine/Source/Runtime/CoreUObject/Private/UObject/LinkerLoad.cpp
@@ -1225,7 +1225,7 @@ FLinkerLoad::ELinkerStatus FLinkerLoad::CreateLoader(
 		{
 			// force preload into memory if file has an SHA entry
 			// Serialize data from memory instead of from disk.
-			uint32	BufferSize = Loader->TotalSize();
+			const int64	BufferSize = Loader->TotalSize();
 			void* Buffer = FMemory::Malloc(BufferSize);
 			Loader->Serialize(Buffer, BufferSize);
 			DestroyLoader();
