ZenLoader: Rework package dependency waiting logic to prevent getting stuck waiting for a package outside of the current sync load context and also make all packages be processed leaf first once the required state is reached
#rnx

diff --git a/Engine/Source/Runtime/CoreUObject/Private/Serialization/AsyncLoading2.cpp b/Engine/Source/Runtime/CoreUObject/Private/Serialization/AsyncLoading2.cpp
--- a/Engine/Source/Runtime/CoreUObject/Private/Serialization/AsyncLoading2.cpp
+++ b/Engine/Source/Runtime/CoreUObject/Private/Serialization/AsyncLoading2.cpp
@@ -2683,17 +2683,17 @@ private:
 		FAsyncPackage2* PrevLink = nullptr;
 		FAsyncPackage2* NextLink = nullptr;
 		uint32 LastTick = 0;
+		int32 PreOrderNumber = -1;
+		bool bAssignedToStronglyConnectedComponent = false;
 		bool bAllDone = false;
-		bool bAnyNotDone = false;
-		bool bVisitedMark = false;
 
 		void UpdateTick(int32 CurrentTick)
 		{
 			if (LastTick != CurrentTick)
 			{
 				LastTick = CurrentTick;
-				bAnyNotDone = false;
-				bVisitedMark = false;
+				PreOrderNumber = -1;
+				bAssignedToStronglyConnectedComponent = false;
 			}
 		}
 
@@ -2701,6 +2701,7 @@ private:
 		{
 			check(WaitListPackage);
 			check(PackageToAdd);
+			check(WaitListPackage != PackageToAdd);
 			FAllDependenciesState& WaitListPackageState = WaitListPackage->*StateMemberPtr;
 			FAllDependenciesState& PackageToAddState = PackageToAdd->*StateMemberPtr;
 			
@@ -2836,8 +2837,26 @@ private:
 	bool HasDependencyToPackageDebug(FAsyncPackage2* Package);
 	void CheckThatAllDependenciesHaveReachedStateDebug(FAsyncLoadingThreadState2& ThreadState, EAsyncPackageLoadingState2 PackageState, EAsyncPackageLoadingState2 PackageStateForCircularDependencies);
 #endif
-	bool HaveAllDependenciesReachedState(FAsyncLoadingThreadState2& ThreadState, FAllDependenciesState FAsyncPackage2::* StateMemberPtr, EAsyncPackageLoadingState2 WaitForPackageState, uint32 CurrentTick);
-	void UpdateDependenciesStateRecursive(FAsyncLoadingThreadState2& ThreadState, FAllDependenciesState FAsyncPackage2::* StateMemberPtr, EAsyncPackageLoadingState2 WaitForPackageState, uint32 CurrentTick, FAsyncPackage2* Root);
+	struct FUpdateDependenciesStateRecursiveContext
+	{
+		FUpdateDependenciesStateRecursiveContext(FAllDependenciesState FAsyncPackage2::* InStateMemberPtr, EAsyncPackageLoadingState2 InWaitForPackageState, uint32 InCurrentTick, TFunctionRef<void(FAsyncPackage2*)> InOnStateReached)
+			: StateMemberPtr(InStateMemberPtr)
+			, WaitForPackageState(InWaitForPackageState)
+			, OnStateReached(InOnStateReached)
+			, CurrentTick(InCurrentTick)
+		{
+
+		}
+
+		FAllDependenciesState FAsyncPackage2::* StateMemberPtr;
+		EAsyncPackageLoadingState2 WaitForPackageState;
+		TFunctionRef<void(FAsyncPackage2*)> OnStateReached;
+		TArray<FAsyncPackage2*, TInlineAllocator<512>> S;
+		TArray<FAsyncPackage2*, TInlineAllocator<512>> P;
+		uint32 CurrentTick;
+		int32 C = 0;
+	};
+	FAsyncPackage2* UpdateDependenciesStateRecursive(FAsyncLoadingThreadState2& ThreadState, FUpdateDependenciesStateRecursiveContext& Context);
 	void WaitForAllDependenciesToReachState(FAsyncLoadingThreadState2& ThreadState, FAllDependenciesState FAsyncPackage2::* StateMemberPtr, EAsyncPackageLoadingState2 WaitForPackageState, uint32& CurrentTickVariable, TFunctionRef<void(FAsyncPackage2*)> OnStateReached);
 	void ConditionalBeginProcessPackageExports(FAsyncLoadingThreadState2& ThreadState);
 	void ConditionalFinishLoading(FAsyncLoadingThreadState2& ThreadState);
@@ -6310,82 +6329,59 @@ void FAsyncPackage2::CheckThatAllDependenciesHaveReachedStateDebug(FAsyncLoading
 		Stack.Pop();
 		FAsyncPackage2* Package = PackageAndDependencyChain.Get<0>();
 		DependencyChain = PackageAndDependencyChain.Get<1>();
-		Visited.Add(Package);
-		check(Package->AsyncPackageLoadingState >= PackageStateForCircularDependencies);
-		if (Package->AsyncPackageLoadingState < PackageState)
-		{
-			if (!ThreadState.PackagesOnStack.Contains(Package))
-			{
-				bool bHasCircularDependencyToPackage = Package->HasDependencyToPackageDebug(this);
-				check(bHasCircularDependencyToPackage);
-			}
-		}
+		
 		for (FAsyncPackage2* ImportedPackage : Package->Data.ImportedAsyncPackages)
 		{
-			if (ImportedPackage && !Visited.Contains(ImportedPackage))
+			if (ImportedPackage && !Visited.Contains(ImportedPackage) && !ThreadState.PackagesOnStack.Contains(Package))
 			{
-				DependencyChain.Add(ImportedPackage);
-				Stack.Push(MakeTuple(ImportedPackage, DependencyChain));
+				TArray<FAsyncPackage2*> NextDependencyChain = DependencyChain;
+				NextDependencyChain.Add(ImportedPackage);
+				
+				check(ImportedPackage->AsyncPackageLoadingState >= PackageStateForCircularDependencies);
+				if (ImportedPackage->AsyncPackageLoadingState < PackageState)
+				{
+					bool bHasCircularDependencyToPackage = ImportedPackage->HasDependencyToPackageDebug(this);
+					check(bHasCircularDependencyToPackage);
+				}
+
+				Visited.Add(ImportedPackage);
+				Stack.Push(MakeTuple(ImportedPackage, NextDependencyChain));
 			}
 		}
 	}
 }
 #endif
 
-bool FAsyncPackage2::HaveAllDependenciesReachedState(FAsyncLoadingThreadState2& ThreadState, FAllDependenciesState FAsyncPackage2::* StateMemberPtr, EAsyncPackageLoadingState2 WaitForPackageState, uint32 CurrentTick)
+FAsyncPackage2* FAsyncPackage2::UpdateDependenciesStateRecursive(FAsyncLoadingThreadState2& ThreadState, FUpdateDependenciesStateRecursiveContext& Context)
 {
-	FAllDependenciesState& ThisState = this->*StateMemberPtr;
+	FAllDependenciesState& ThisState = this->*Context.StateMemberPtr;
+	
+	check(ThisState.PreOrderNumber < 0);
+	
 	if (ThisState.bAllDone)
 	{
-		return true;
-	}
-	if (AsyncPackageLoadingState < WaitForPackageState)
-	{
-		return false;
+		return nullptr;
 	}
-	ThisState.UpdateTick(CurrentTick);
-	UpdateDependenciesStateRecursive(ThreadState, StateMemberPtr, WaitForPackageState, CurrentTick, this);
-	check(ThisState.bAllDone || (ThisState.WaitingForPackage && ThisState.WaitingForPackage->AsyncPackageLoadingState <= WaitForPackageState));
-	return ThisState.bAllDone;
-}
 
-void FAsyncPackage2::UpdateDependenciesStateRecursive(FAsyncLoadingThreadState2& ThreadState, FAllDependenciesState FAsyncPackage2::* StateMemberPtr, EAsyncPackageLoadingState2 WaitForPackageState, uint32 CurrentTick, FAsyncPackage2* Root)
-{
-	FAllDependenciesState& ThisState = this->*StateMemberPtr;
-
-	check(!ThisState.bVisitedMark);
-	check(!ThisState.bAllDone);
-	check(!ThisState.bAnyNotDone);
-
-	ThisState.bVisitedMark = true;
-
-	if (FAsyncPackage2* WaitingForPackage = ThisState.WaitingForPackage)
+	FAsyncPackage2* WaitingForPackage = ThisState.WaitingForPackage;
+	if (WaitingForPackage)
 	{
+		if (WaitingForPackage->AsyncPackageLoadingState >= Context.WaitForPackageState)
 		{
-			FAllDependenciesState& WaitingForPackageState = WaitingForPackage->*StateMemberPtr;
-			if (WaitingForPackage->AsyncPackageLoadingState < WaitForPackageState)
-			{
-				ThisState.bAnyNotDone = true;
-				return;
-			}
-			else if (!WaitingForPackageState.bAllDone)
-			{
-				WaitingForPackageState.UpdateTick(CurrentTick);
-				if (!WaitingForPackageState.bVisitedMark)
-				{
-					WaitingForPackage->UpdateDependenciesStateRecursive(ThreadState, StateMemberPtr, WaitForPackageState, CurrentTick, Root);
-				}
-				if (WaitingForPackageState.bAnyNotDone)
-				{
-					ThisState.bAnyNotDone = true;
-					return;
-				}
-			}
+			FAllDependenciesState::RemoveFromWaitList(Context.StateMemberPtr, WaitingForPackage, this);
+			WaitingForPackage = nullptr;
+		}
+		else
+		{
+			return WaitingForPackage;
 		}
 	}
 
-	bool bAllDone = true;
-	FAsyncPackage2* WaitingForPackage = nullptr;
+	ThisState.PreOrderNumber = Context.C;
+	++Context.C;
+	Context.S.Push(this);
+	Context.P.Push(this);
+
 	for (FAsyncPackage2* ImportedPackage : Data.ImportedAsyncPackages)
 	{
 		if (!ImportedPackage)
@@ -6398,92 +6394,96 @@ void FAsyncPackage2::UpdateDependenciesStateRecursive(FAsyncLoadingThreadState2&
 			continue;
 		}
 
-		FAllDependenciesState& ImportedPackageState = ImportedPackage->*StateMemberPtr;
-
+		FAllDependenciesState& ImportedPackageState = ImportedPackage->*Context.StateMemberPtr;
 		if (ImportedPackageState.bAllDone)
 		{
 			continue;
 		}
 
-		ImportedPackageState.UpdateTick(CurrentTick);
-
-		if (ImportedPackage->AsyncPackageLoadingState < WaitForPackageState)
+		if (ImportedPackage->AsyncPackageLoadingState < Context.WaitForPackageState)
 		{
-			ImportedPackageState.bAnyNotDone = true;
+			WaitingForPackage = ImportedPackage;
+			break;
 		}
-		else if (!ImportedPackageState.bVisitedMark)
+
+		ImportedPackageState.UpdateTick(Context.CurrentTick);
+		if (ImportedPackageState.PreOrderNumber < 0)
 		{
-			ImportedPackage->UpdateDependenciesStateRecursive(ThreadState, StateMemberPtr, WaitForPackageState, CurrentTick, Root);
+			WaitingForPackage = ImportedPackage->UpdateDependenciesStateRecursive(ThreadState, Context);
+			if (WaitingForPackage)
+			{
+				break;
+			}
 		}
-
-		if (ImportedPackageState.bAnyNotDone)
+		else if (!ImportedPackageState.bAssignedToStronglyConnectedComponent)
 		{
-			ThisState.bAnyNotDone = true;
-			WaitingForPackage = ImportedPackage;
-			break;
+			while ((Context.P.Top()->*Context.StateMemberPtr).PreOrderNumber > ImportedPackageState.PreOrderNumber)
+			{
+				Context.P.Pop();
+			}
 		}
-		else if (!ImportedPackageState.bAllDone)
+		if (ImportedPackageState.WaitingForPackage)
 		{
-			bAllDone = false;
+			WaitingForPackage = ImportedPackageState.WaitingForPackage;
+			break;
 		}
 	}
-	if (WaitingForPackage)
-	{
-		check(WaitingForPackage != this);
-		FAllDependenciesState::AddToWaitList(StateMemberPtr, WaitingForPackage, this);
-	}
-	else if (bAllDone || this == Root)
-	{
-		// If we're the root an not waiting for any package we're done
-		ThisState.bAllDone = true;
-	}
-	else
+
+	if (Context.P.Top() == this)
 	{
-		// We didn't find any imported package that was not done but we could have a circular dependency back to the root which could either be done or end up waiting
-		// for another package. Make us wait for the root so that we are ticked when it completes.
-		FAllDependenciesState::AddToWaitList(StateMemberPtr, Root, this);
+		FAsyncPackage2* InStronglyConnectedComponent;
+		do
+		{
+			InStronglyConnectedComponent = Context.S.Pop();
+			FAllDependenciesState& InStronglyConnectedComponentState = InStronglyConnectedComponent->*Context.StateMemberPtr;
+			InStronglyConnectedComponentState.bAssignedToStronglyConnectedComponent = true;
+			check(InStronglyConnectedComponent->AsyncPackageLoadingState >= Context.WaitForPackageState);
+			if (WaitingForPackage)
+			{
+#if ALT2_ENABLE_PACKAGE_DEPENDENCY_DEBUGGING
+				check(HasDependencyToPackageDebug(WaitingForPackage));
+#endif
+				FAllDependenciesState::AddToWaitList(Context.StateMemberPtr, WaitingForPackage, InStronglyConnectedComponent);
+			}
+			else
+			{
+				check(InStronglyConnectedComponent->AsyncPackageLoadingState == Context.WaitForPackageState);
+				InStronglyConnectedComponentState.bAllDone = true;
+				InStronglyConnectedComponent->AsyncPackageLoadingState = static_cast<EAsyncPackageLoadingState2>(static_cast<uint32>(Context.WaitForPackageState) + 1);
+#if ALT2_ENABLE_PACKAGE_DEPENDENCY_DEBUGGING
+				InStronglyConnectedComponent->CheckThatAllDependenciesHaveReachedStateDebug(ThreadState, InStronglyConnectedComponent->AsyncPackageLoadingState, Context.WaitForPackageState);
+#endif
+				Context.OnStateReached(InStronglyConnectedComponent);
+			}
+		} while (InStronglyConnectedComponent != this);
+		Context.P.Pop();
 	}
+	
+	return WaitingForPackage;
 }
 
 void FAsyncPackage2::WaitForAllDependenciesToReachState(FAsyncLoadingThreadState2& ThreadState, FAllDependenciesState FAsyncPackage2::* StateMemberPtr, EAsyncPackageLoadingState2 WaitForPackageState, uint32& CurrentTickVariable, TFunctionRef<void(FAsyncPackage2*)> OnStateReached)
 {
-	if (HaveAllDependenciesReachedState(ThreadState, StateMemberPtr, WaitForPackageState, CurrentTickVariable++))
-	{
-		FAsyncPackage2* FirstPackageReadyToProceed = this;
+	check(AsyncPackageLoadingState == WaitForPackageState);
+	++CurrentTickVariable;
 
-		while (FirstPackageReadyToProceed)
-		{
-			FAsyncPackage2* PackageReadyToProceed = FirstPackageReadyToProceed;
-			FAllDependenciesState& PackageReadyToProceedState = PackageReadyToProceed->*StateMemberPtr;
-			FirstPackageReadyToProceed = PackageReadyToProceedState.NextLink;
+	FUpdateDependenciesStateRecursiveContext Context(StateMemberPtr, WaitForPackageState, CurrentTickVariable, OnStateReached);
 
-			if (PackageReadyToProceed->AsyncPackageLoadingState > WaitForPackageState)
-			{
-				continue;
-			}
-
-#if ALT2_ENABLE_PACKAGE_DEPENDENCY_DEBUGGING
-			PackageReadyToProceed->CheckThatAllDependenciesHaveReachedStateDebug(ThreadState, WaitForPackageState, WaitForPackageState);
-#endif
+	FAllDependenciesState& ThisState = this->*StateMemberPtr;
+	check(!ThisState.bAllDone);
+	ThisState.UpdateTick(CurrentTickVariable);
+	UpdateDependenciesStateRecursive(ThreadState, Context);
+	check(ThisState.bAllDone || (ThisState.WaitingForPackage && ThisState.WaitingForPackage->AsyncPackageLoadingState < WaitForPackageState));
 
-			while (FAsyncPackage2* WaitingPackage = PackageReadyToProceedState.PackagesWaitingForThisHead)
-			{
-				FAllDependenciesState& WaitingPackageState = WaitingPackage->*StateMemberPtr;
-				check(WaitingPackageState.WaitingForPackage == PackageReadyToProceed);
-				if (WaitingPackage->HaveAllDependenciesReachedState(ThreadState, StateMemberPtr, WaitForPackageState, CurrentTickVariable++))
-				{
-					// If the waiting package wasn't ready to proceed it will remove itself from the waitlist (PackageReadyToProceedState.PackagesWaitingForThisHead) when its WaitingForPackage
-					// is replaced by another package
-					FAllDependenciesState::RemoveFromWaitList(StateMemberPtr, PackageReadyToProceed, WaitingPackage);
-					WaitingPackageState.NextLink = FirstPackageReadyToProceed;
-					FirstPackageReadyToProceed = WaitingPackage;
-				}
-			}
-			check(!PackageReadyToProceedState.PackagesWaitingForThisTail);
-			check(PackageReadyToProceed->AsyncPackageLoadingState == WaitForPackageState);
-			PackageReadyToProceed->AsyncPackageLoadingState = static_cast<EAsyncPackageLoadingState2>(static_cast<uint32>(WaitForPackageState) + 1);
-			OnStateReached(PackageReadyToProceed);
+	while (FAsyncPackage2* WaitingPackage = ThisState.PackagesWaitingForThisHead)
+	{
+		FAllDependenciesState& WaitingPackageState = WaitingPackage->*StateMemberPtr;
+		WaitingPackageState.UpdateTick(CurrentTickVariable);
+		if (WaitingPackageState.PreOrderNumber < 0)
+		{
+			WaitingPackage->UpdateDependenciesStateRecursive(ThreadState, Context);
 		}
+		check(WaitingPackageState.bAllDone || (WaitingPackageState.WaitingForPackage && WaitingPackageState.WaitingForPackage->AsyncPackageLoadingState < WaitForPackageState));
 	}
 }
 
