Fixing bad revert... I do not understand what happen with the revert tool!

diff --git a/Engine/Source/Runtime/CoreUObject/Private/UObject/PropertyBaseObject.cpp b/Engine/Source/Runtime/CoreUObject/Private/UObject/PropertyBaseObject.cpp
--- a/Engine/Source/Runtime/CoreUObject/Private/UObject/PropertyBaseObject.cpp
+++ b/Engine/Source/Runtime/CoreUObject/Private/UObject/PropertyBaseObject.cpp
@@ -687,7 +687,7 @@ void FObjectPropertyBase::CheckValidObject(void* ValueAddress, TObjectPtr<UObjec
 					TEXT("Serialized %s for a property of %s. Reference will be reverted back to %s.\n    Property = %s\n    Item = %s"),
 					*ObjectClass->GetFullName(),
 					*PropertyClass->GetFullName(),
-					*DefaultValue->GetFullName(),
+					OldValue ? *OldValue->GetFullName() : TEXT("None"), ,
 					*GetFullName(),
 					*Object.GetFullName()
 				);
