ZenLoader: Include preload dependencies in package headers instead of relying on predetermined export bundle dependencies. This means that each package can be treated in isolation when staging. Also removes the need to keep track of export and bundle counts in the package store which saves some memory.

diff --git a/Engine/Source/Developer/IoStoreUtilities/Internal/PackageStoreOptimizer.h b/Engine/Source/Developer/IoStoreUtilities/Internal/PackageStoreOptimizer.h
--- a/Engine/Source/Developer/IoStoreUtilities/Internal/PackageStoreOptimizer.h
+++ b/Engine/Source/Developer/IoStoreUtilities/Internal/PackageStoreOptimizer.h
@@ -91,7 +91,7 @@ public:
 
 	uint64 GetHeaderSize() const
 	{
-		return HeaderSize;
+		return HeaderBuffer.DataSize();
 	}
 
 	uint64 GetNameCount() const
@@ -107,18 +107,10 @@ public:
 private:
 	struct FExportGraphNode;
 
-	struct FExternalDependency
-	{
-		int32 ImportIndex = -1;
-		FExportBundleEntry::EExportCommandType ExportBundleCommandType;
-	};
-
 	struct FExportGraphNode
 	{
 		FExportBundleEntry BundleEntry;
 		TArray<FExportGraphNode*> InternalDependencies;
-		TArray<FExternalDependency> ExternalDependencies;
-		int32 ExportBundleIndex = -1;
 		int32 IncomingEdgeCount = 0;
 		bool bIsPublic = false;
 	};
@@ -132,7 +124,6 @@ private:
 		FPackageObjectIndex SuperIndex;
 		FPackageObjectIndex TemplateIndex;
 		EObjectFlags ObjectFlags = RF_NoFlags;
-		uint64 CookedSerialOffset = uint64(-1);
 		uint64 SerialOffset = uint64(-1);
 		uint64 SerialSize = uint64(-1);
 		bool bNotForClient = false;
@@ -141,12 +132,6 @@ private:
 		FExportGraphNode* Nodes[FExportBundleEntry::ExportCommandType_Count] = { nullptr };
 	};
 
-	struct FExportBundle
-	{
-		uint64 SerialOffset = uint64(-1);
-		TArray<FExportBundleEntry> Entries;
-	};
-
 	struct FUnresolvedExport
 	{
 		FString FullName;
@@ -162,42 +147,11 @@ private:
 		bool bIsImportOptional = false;
 	};
 
-	struct FInternalArc
-	{
-		int32 FromExportBundleIndex;
-		int32 ToExportBundleIndex;
-
-		bool operator==(const FInternalArc& Other) const
-		{
-			return FromExportBundleIndex == Other.FromExportBundleIndex &&
-				ToExportBundleIndex == Other.ToExportBundleIndex;
-		}
-
-		friend FORCEINLINE uint32 GetTypeHash(const FInternalArc& Arc)
-		{
-			return HashCombine(GetTypeHash(Arc.FromExportBundleIndex), GetTypeHash(Arc.ToExportBundleIndex));
-		}
-	};
-
-	struct FExternalArc
-	{
-		int32 FromImportIndex;
-		FExportBundleEntry::EExportCommandType FromCommandType;
-		int32 ToExportBundleIndex;
-
-		bool operator==(const FExternalArc& Other) const
-		{
-			return FromImportIndex == Other.FromImportIndex &&
-				FromCommandType == Other.FromCommandType &&
-				ToExportBundleIndex == Other.ToExportBundleIndex;
-		}
-	};
-
 	struct FGraphData
 	{
-		TArray<FExportBundle> ExportBundles;
-		TArray<FInternalArc> InternalArcs;
-		TMap<FPackageId, TArray<FExternalArc>> ExternalArcs;
+		TArray<FExportBundleEntry> ExportBundleEntries;
+		TArray<FDependencyBundleHeader> DependencyBundleHeaders;
+		TArray<FDependencyBundleEntry> DependencyBundleEntries;
 	};
 
 	struct FExportBundleGraphNode
@@ -244,15 +198,7 @@ private:
 
 	uint32 PackageFlags = 0;
 	uint32 CookedHeaderSize = 0;
-	uint64 HeaderSize = 0;
-	uint64 ExportsSerialSize = 0;
 	
-	TArray<FExportBundleGraphNode> ExportBundleGraphNodes;
-	FPackageStorePackage::FExportBundle* CurrentBundle = nullptr;
-	TArray<FExportBundleGraphNode*> NodesWithNoIncomingEdges;
-	bool bTemporaryMark = false;
-	bool bPermanentMark = false;
-
 	friend class FPackageStoreOptimizer;
 };
 
@@ -269,13 +215,10 @@ public:
 
 	FPackageStorePackage* CreateMissingPackage(const FName& Name) const;
 	FPackageStorePackage* CreatePackageFromCookedHeader(const FName& Name, const FIoBuffer& CookedHeaderBuffer) const;
-	FPackageStorePackage* CreatePackageFromZenPackageHeader(const FName& Name, const FIoBuffer& Buffer, int32 ExportBundleCount, const TArrayView<const FPackageId>& ImportedPackageIds) const;
-	void FinalizePackage(FPackageStorePackage* Package) const;
-	FIoBuffer CreatePackageBuffer(const FPackageStorePackage* Package, const FIoBuffer& CookedExportsBuffer, TArray<FFileRegion>* InOutFileRegions) const;
+	FIoBuffer CreatePackageBuffer(const FPackageStorePackage* Package, const FIoBuffer& CookedExportsBuffer) const;
 	FPackageStoreEntryResource CreatePackageStoreEntry(const FPackageStorePackage* Package, const FPackageStorePackage* OptionalSegmentPackage) const;
 	IOSTOREUTILITIES_API FIoBuffer CreateScriptObjectsBuffer() const;
 	void LoadScriptObjectsBuffer(const FIoBuffer& ScriptObjectsBuffer);
-	void OptimizeExportBundles(const TMap<FPackageId, FPackageStorePackage*>& PackagesMap);
 
 private:
 	struct FScriptObjectData
@@ -297,42 +240,18 @@ private:
 		TArray<FObjectDataResource> DataResources;
 	};
 
-	struct FZenPackageHeaderData
-	{
-		FZenPackageSummary Summary;
-		TOptional<FZenPackageVersioningInfo> VersioningInfo;
-		TArray<FName> ImportedPackageNames;
-		TArray<uint64> ImportedPublicExportHashes;
-		TArray<FDisplayNameEntryId> NameMap;
-		TArray<FBulkDataMapEntry> BulkDataEntries;
-		TArray<FPackageObjectIndex> Imports; // FH: Imports might need to have more info to be able to resolve export hash with import as outer
-		TArray<FExportMapEntry> Exports;
-		TArray<FExportBundleHeader> ExportBundleHeaders;
-		TArray<FExportBundleEntry> ExportBundleEntries;
-		TArray<FPackageStorePackage::FInternalArc> InternalArcs;
-		TArray<FPackageStorePackage::FExternalArc> ExternalArcs;
-	};
-
 	using FExportGraphEdges = TMultiMap<FPackageStorePackage::FExportGraphNode*, FPackageStorePackage::FExportGraphNode*>;
-	using FExportBundleGraphEdges = TMultiMap<FPackageStorePackage::FExportBundleGraphNode*, FPackageStorePackage::FExportBundleGraphNode*>;
 
 	static uint64 GetPublicExportHash(FStringView PackageRelativeExportPath);
 	FCookedHeaderData LoadCookedHeader(const FIoBuffer& CookedHeaderBuffer) const;
-	FZenPackageHeaderData LoadZenPackageHeader(const FIoBuffer& HeaderBuffer, int32 ExportBundleCount) const;
 	void ResolveImport(FPackageStorePackage::FUnresolvedImport* Imports, const FObjectImport* ObjectImports, int32 LocalImportIndex) const;
 	void ProcessImports(const FCookedHeaderData& CookedHeaderData, FPackageStorePackage* Package, TArray<FPackageStorePackage::FUnresolvedImport>& UnresolvedImports) const;
-	void ProcessImports(const FZenPackageHeaderData& ZenHeaderData, FPackageStorePackage* Package) const;
 	void ResolveExport(FPackageStorePackage::FUnresolvedExport* Exports, const FObjectExport* ObjectExports, const int32 LocalExportIndex, const FName& PackageName, FPackageStorePackage::FUnresolvedImport* Imports, const FObjectImport* ObjectImports) const;
 	void ProcessExports(const FCookedHeaderData& CookedHeaderData, FPackageStorePackage* Package, FPackageStorePackage::FUnresolvedImport* Imports) const;
-	void ProcessExports(const FZenPackageHeaderData& ZenHeaderData, FPackageStorePackage* Package) const;
 	void ProcessPreloadDependencies(const FCookedHeaderData& CookedHeaderData, FPackageStorePackage* Package) const;
-	void ProcessPreloadDependencies(const FZenPackageHeaderData& ZenHeaderData, FPackageStorePackage* Package) const;
 	void ProcessDataResources(const FCookedHeaderData& CookedHeaderData, FPackageStorePackage* Package) const;
-	TArray<FPackageStorePackage*> SortPackagesInLoadOrder(const TMap<FPackageId, FPackageStorePackage*>& PackagesMap) const;
-	void SerializeGraphData(const TArray<FPackageStorePackage::FImportedPackageRef>& ImportedPackages, FPackageStorePackage::FGraphData& GraphData, FBufferWriter& GraphArchive) const;
 	TArray<FPackageStorePackage::FExportGraphNode*> SortExportGraphNodesInLoadOrder(FPackageStorePackage* Package, FExportGraphEdges& Edges) const;
-	TArray<FPackageStorePackage::FExportBundleGraphNode*> SortExportBundleGraphNodesInLoadOrder(const TArray<FPackageStorePackage*>& Packages, FExportBundleGraphEdges& Edges) const;
-	void CreateExportBundles(FPackageStorePackage* Package) const;
+	void CreateExportBundle(FPackageStorePackage* Package) const;
 	void FinalizePackageHeader(FPackageStorePackage* Package) const;
 	void FindScriptObjectsRecursive(FPackageObjectIndex OuterIndex, UObject* Object);
 	void FindScriptObjects();
diff --git a/Engine/Source/Developer/IoStoreUtilities/Private/IoStoreUtilities.cpp b/Engine/Source/Developer/IoStoreUtilities/Private/IoStoreUtilities.cpp
--- a/Engine/Source/Developer/IoStoreUtilities/Private/IoStoreUtilities.cpp
+++ b/Engine/Source/Developer/IoStoreUtilities/Private/IoStoreUtilities.cpp
@@ -608,8 +608,6 @@ struct FCookedPackage
 	uint64 LoadOrder = MAX_uint64; // Ordered by dependencies
 	uint64 DiskLayoutOrder = MAX_uint64; // Final order after considering input order files
 	FPackageStoreEntryResource PackageStoreEntry;
-	FPackageStorePackage* OptimizedPackage = nullptr;
-	FPackageStorePackage* OptimizedOptionalSegmentPackage = nullptr;
 	int32 PreOrderNumber = -1; // Used for sorting in load order
 	bool bPermanentMark = false; // Used for sorting in load order
 	bool bIsLocalized = false;
@@ -620,6 +618,8 @@ struct FLegacyCookedPackage
 {
 	FString FileName;
 	FString OptionalSegmentFileName;
+	FPackageStorePackage* OptimizedPackage = nullptr;
+	FPackageStorePackage* OptimizedOptionalSegmentPackage = nullptr;
 };
 
 enum class EContainerChunkType
@@ -1774,35 +1774,6 @@ static void ParsePackageAssetsFromFiles(TArray<FCookedPackage*>& Packages, const
 	FMemory::Free(OptionalSegmentUAssetMemory);
 }
 
-static void ParsePackageAssetsFromZen(FCookedPackageStore& PackageStore, TArray<FCookedPackage*>& Packages, const FPackageStoreOptimizer& PackageStoreOptimizer)
-{
-	IOSTORE_CPU_SCOPE(ParsePackageAssetsFromZen);
-	UE_LOG(LogIoStore, Display, TEXT("Parsing packages..."));
-
-	const int32 TotalPackageCount = Packages.Num();
-	
-	ParallelFor(TotalPackageCount, [
-		&PackageStore,
-		&PackageStoreOptimizer,
-		&Packages](int32 Index)
-	{
-		FCookedPackage* Package = Packages[Index];
-		TIoStatusOr<FIoBuffer> HeaderBuffer = PackageStore.ReadPackageHeaderFromZen(Package->GlobalPackageId, 0);
-		Package->UExpSize = Package->UAssetSize - HeaderBuffer.ValueOrDie().DataSize();
-		Package->UAssetSize = HeaderBuffer.ValueOrDie().DataSize();
-		const FPackageStoreEntryResource& PackageStoreEntry = Package->PackageStoreEntry;
-		Package->OptimizedPackage = PackageStoreOptimizer.CreatePackageFromZenPackageHeader(Package->PackageName, HeaderBuffer.ValueOrDie(), PackageStoreEntry.ExportInfo.ExportBundleCount, PackageStoreEntry.ImportedPackageIds);
-		check(Package->OptimizedPackage->GetId() == Package->GlobalPackageId);
-		if (Package->OptionalSegmentUAssetSize)
-		{
-			TIoStatusOr<FIoBuffer> OptionalSegmentHeaderBuffer = PackageStore.ReadPackageHeaderFromZen(Package->GlobalPackageId, 1);
-			Package->OptionalSegmentUAssetSize = OptionalSegmentHeaderBuffer.ValueOrDie().DataSize();
-			Package->OptimizedOptionalSegmentPackage = PackageStoreOptimizer.CreatePackageFromZenPackageHeader(Package->PackageName, OptionalSegmentHeaderBuffer.ValueOrDie(), PackageStoreEntry.OptionalSegmentExportInfo.ExportBundleCount, PackageStoreEntry.OptionalSegmentImportedPackageIds);
-			check(Package->OptimizedOptionalSegmentPackage->GetId() == Package->GlobalPackageId);
-		}
-	}, EParallelForFlags::Unbalanced);
-}
-
 static TUniquePtr<FIoStoreReader> CreateIoStoreReader(const TCHAR* Path, const FKeyChain& KeyChain)
 {
 	TUniquePtr<FIoStoreReader> IoStoreReader(new FIoStoreReader());
@@ -2851,7 +2822,6 @@ private:
 		uint64 SourceBufferSize;
 		FGraphEventRef CompletionEvent;
 		FIoBuffer SourceBuffer;
-		bool bHasUpdatedExportBundleRegions = false;
 		FQueueEntry* QueueEntry = nullptr;
 	};
 
@@ -2876,8 +2846,11 @@ private:
 		: public FWriteContainerTargetFileRequest
 	{
 	public:
-		FLooseFileWriteRequest(FIoStoreWriteRequestManager& InManager,const FContainerTargetFile& InTargetFile)
-			: FWriteContainerTargetFileRequest(InManager, InTargetFile) { }
+		FLooseFileWriteRequest(FIoStoreWriteRequestManager& InManager, const FContainerTargetFile& InTargetFile)
+			: FWriteContainerTargetFileRequest(InManager, InTargetFile)
+			, Package(static_cast<FLegacyCookedPackage*>(InTargetFile.Package))
+		{
+		}
 
 		virtual void LoadSourceBufferAsync() override
 		{
@@ -2894,14 +2867,12 @@ private:
 
 				if (TargetFile.ChunkType == EContainerChunkType::PackageData)
 				{
-					SourceBuffer = Manager.PackageStoreOptimizer.CreatePackageBuffer(TargetFile.Package->OptimizedPackage, SourceBuffer, bHasUpdatedExportBundleRegions ? nullptr : &FileRegions);
-					bHasUpdatedExportBundleRegions = true;
+					SourceBuffer = Manager.PackageStoreOptimizer.CreatePackageBuffer(Package->OptimizedPackage, SourceBuffer);
 				}
 				else if (TargetFile.ChunkType == EContainerChunkType::OptionalSegmentPackageData)
 				{
-					check(TargetFile.Package->OptimizedOptionalSegmentPackage);
-					SourceBuffer = Manager.PackageStoreOptimizer.CreatePackageBuffer(TargetFile.Package->OptimizedOptionalSegmentPackage, SourceBuffer, bHasUpdatedExportBundleRegions ? nullptr : &FileRegions);
-					bHasUpdatedExportBundleRegions = true;
+					check(Package->OptimizedOptionalSegmentPackage);
+					SourceBuffer = Manager.PackageStoreOptimizer.CreatePackageBuffer(Package->OptimizedOptionalSegmentPackage, SourceBuffer);
 				}
 				OnSourceBufferLoaded();
 			};
@@ -2910,6 +2881,9 @@ private:
 				QueueEntry->FileHandle->ReadRequest(0, SourceBuffer.DataSize(), AIOP_Normal, &Callback, SourceBuffer.Data()));
 			QueueEntry->ReleaseRef(Manager);
 		}
+
+	private:
+		FLegacyCookedPackage* Package;
 	};
 
 	class FZenWriteRequest
@@ -2928,22 +2902,6 @@ private:
 				{
 					SourceBuffer = Status.ConsumeValueOrDie();
 					Manager.ZenSourceBytes[(int8)TargetFile.ChunkId.GetChunkType()] += SourceBuffer.DataSize();
-					if (TargetFile.ChunkType == EContainerChunkType::PackageData)
-					{
-						check(TargetFile.Package->UAssetSize > 0);
-						const uint64 HeaderSize = TargetFile.Package->UAssetSize;
-						FIoBuffer ExportsBuffer(SourceBuffer.Data() + HeaderSize, SourceBuffer.DataSize() - HeaderSize, SourceBuffer);
-						SourceBuffer = Manager.PackageStoreOptimizer.CreatePackageBuffer(TargetFile.Package->OptimizedPackage, ExportsBuffer, nullptr);
-						if (!bHasUpdatedExportBundleRegions)
-						{
-							// The regions are relative to the start of the export data, adjust to include the the new header size
-							for (FFileRegion& Region : FileRegions)
-							{
-								Region.Offset += TargetFile.Package->OptimizedPackage->GetHeaderSize();
-							}
-							bHasUpdatedExportBundleRegions = true;
-						}
-					}
 					OnSourceBufferLoaded();
 				});
 		}
@@ -3650,10 +3608,6 @@ void CreateContainerHeader(FContainerTargetSpec& ContainerTarget, bool bIsOption
 					TargetEntriesWriter = &OptionalSegmentStoreEntriesWriter;
 				}
 
-				// OptionalStoreEntry
-				FPackageStoreExportInfo OptionalSegmentExportInfo = Entry.OptionalSegmentExportInfo;
-				TargetEntriesWriter->StoreTocArchive << OptionalSegmentExportInfo;
-
 				// OptionalImportedPackages
 				const TArray<FPackageId>& OptionalSegmentImportedPackageIds = Entry.OptionalSegmentImportedPackageIds;
 				SerializePackageEntryCArrayHeader(*TargetEntriesWriter, OptionalSegmentImportedPackageIds.Num());
@@ -3692,10 +3646,6 @@ void CreateContainerHeader(FContainerTargetSpec& ContainerTarget, bool bIsOption
 				}
 			}
 
-			// StoreEntries
-			FPackageStoreExportInfo ExportInfo = Entry.ExportInfo;
-			StoreEntriesWriter.StoreTocArchive << ExportInfo;
-
 			// ImportedPackages
 			const TArray<FPackageId>& ImportedPackageIds = Entry.ImportedPackageIds;
 			SerializePackageEntryCArrayHeader(StoreEntriesWriter, ImportedPackageIds.Num());
@@ -3835,13 +3785,36 @@ int32 CreateTarget(const FIoStoreArguments& Arguments, const FIoStoreWriterSetti
 		}
 	}
 
-	if (Arguments.PackageStore->HasZenStoreClient())
-	{
-		ParsePackageAssetsFromZen(*Arguments.PackageStore, Packages, PackageStoreOptimizer);
-	}
-	else
+	const bool bIsLegacyStage = !Arguments.PackageStore->HasZenStoreClient();
+	if (bIsLegacyStage)
 	{
 		ParsePackageAssetsFromFiles(Packages, PackageStoreOptimizer);
+		if (Arguments.bFileRegions)
+		{
+			// The file regions for packages are relative to the start of the uexp file so we need to make them relative to the start of the export bundle chunk instead
+			for (FContainerTargetSpec* ContainerTarget : ContainerTargets)
+			{
+				for (FContainerTargetFile& TargetFile : ContainerTarget->TargetFiles)
+				{
+					if (TargetFile.ChunkType == EContainerChunkType::PackageData)
+					{
+						uint64 HeaderSize = static_cast<FLegacyCookedPackage*>(TargetFile.Package)->OptimizedPackage->GetHeaderSize();
+						for (FFileRegion& Region : TargetFile.FileRegions)
+						{
+							Region.Offset += HeaderSize;
+						}
+					}
+					else if (TargetFile.ChunkType == EContainerChunkType::OptionalSegmentPackageData)
+					{
+						uint64 HeaderSize = static_cast<FLegacyCookedPackage*>(TargetFile.Package)->OptimizedOptionalSegmentPackage->GetHeaderSize();
+						for (FFileRegion& Region : TargetFile.FileRegions)
+						{
+							Region.Offset += HeaderSize;
+						}
+					}
+				}
+			}
+		}
 	}
 
 	UE_LOG(LogIoStore, Display, TEXT("Processing shader libraries, compressing with Oodle %s, level %d (%s)"), FOodleDataCompression::ECompressorToString(Arguments.ShaderOodleCompressor), (int32)Arguments.ShaderOodleLevel, FOodleDataCompression::ECompressionLevelToString(Arguments.ShaderOodleLevel));
@@ -3892,7 +3865,6 @@ int32 CreateTarget(const FIoStoreArguments& Arguments, const FIoStoreWriterSetti
 	};
 
 	{
-		// Append all chunks that don't change after this point
 		IOSTORE_CPU_SCOPE(AppendChunks);
 		for (FContainerTargetSpec* ContainerTarget : ContainerTargets)
 		{
@@ -3900,10 +3872,7 @@ int32 CreateTarget(const FIoStoreArguments& Arguments, const FIoStoreWriterSetti
 			{
 				for (FContainerTargetFile& TargetFile : ContainerTarget->TargetFiles)
 				{
-					if (TargetFile.ChunkType != EContainerChunkType::PackageData && TargetFile.ChunkType != EContainerChunkType::OptionalSegmentPackageData)
-					{
-						AppendTargetFileChunk(ContainerTarget, TargetFile);
-					}
+					AppendTargetFileChunk(ContainerTarget, TargetFile);
 				}
 			}
 		}
@@ -3912,33 +3881,6 @@ int32 CreateTarget(const FIoStoreArguments& Arguments, const FIoStoreWriterSetti
 	UE_LOG(LogIoStore, Display, TEXT("Processing redirects..."));
 	ProcessRedirects(Arguments, PackageIdMap);
 
-	UE_LOG(LogIoStore, Display, TEXT("Optimizing packages..."));
-	{
-		TMap<FPackageId, FPackageStorePackage*> OptimizedPackagesMap;
-		for (FCookedPackage* Package : Packages)
-		{
-			check(Package->OptimizedPackage);
-			OptimizedPackagesMap.Add(Package->OptimizedPackage->GetId(), Package->OptimizedPackage);
-		}
-		PackageStoreOptimizer.OptimizeExportBundles(OptimizedPackagesMap);
-	}
-
-	UE_LOG(LogIoStore, Display, TEXT("Finalizing packages..."));
-	{
-		IOSTORE_CPU_SCOPE(FinalizingPackages);
-		for (FCookedPackage* Package : Packages)
-		{
-			check(Package->OptimizedPackage);
-			PackageStoreOptimizer.FinalizePackage(Package->OptimizedPackage);
-			if (Package->OptimizedOptionalSegmentPackage)
-			{
-				PackageStoreOptimizer.FinalizePackage(Package->OptimizedOptionalSegmentPackage);
-			}
-			// Update package store entry after optimizing export bundles
-			Package->PackageStoreEntry = PackageStoreOptimizer.CreatePackageStoreEntry(Package->OptimizedPackage, Package->OptimizedOptionalSegmentPackage);
-		}
-	}
-
 	UE_LOG(LogIoStore, Display, TEXT("Creating disk layout..."));
 	FString ClusterCSVPath;
 	if (FParse::Value(FCommandLine::Get(), TEXT("ClusterCSV="), ClusterCSVPath))
@@ -3950,19 +3892,11 @@ int32 CreateTarget(const FIoStoreArguments& Arguments, const FIoStoreWriterSetti
 	CreateDiskLayout(ContainerTargets, Packages, Arguments.OrderMaps, PackageIdMap, Arguments.bClusterByOrderFilePriority);
 
 	{
-		IOSTORE_CPU_SCOPE(AppendRemaningChunks);
+		IOSTORE_CPU_SCOPE(AppendContainerHeaderChunks);
 		for (FContainerTargetSpec* ContainerTarget : ContainerTargets)
 		{
 			if (ContainerTarget->IoStoreWriter)
 			{
-				for (FContainerTargetFile& TargetFile : ContainerTarget->TargetFiles)
-				{
-					if (TargetFile.ChunkType == EContainerChunkType::PackageData || TargetFile.ChunkType == EContainerChunkType::OptionalSegmentPackageData)
-					{
-						AppendTargetFileChunk(ContainerTarget, TargetFile);
-					}
-				}
-
 				auto WriteContainerHeaderChunk = [](FIoContainerHeader& Header, IIoStoreWriter* IoStoreWriter)
 				{
 					FLargeMemoryWriter HeaderAr(0, true);
@@ -4243,8 +4177,12 @@ int32 CreateTarget(const FIoStoreArguments& Arguments, const FIoStoreWriterSetti
 		UExpSize += Package->UExpSize;
 		UAssetSize += Package->UAssetSize;
 		UBulkSize += Package->TotalBulkDataSize;
-		NameMapCount += Package->OptimizedPackage->GetNameCount();
-		HeaderSize += Package->OptimizedPackage->GetHeaderSize();
+		if (bIsLegacyStage)
+		{
+			const FLegacyCookedPackage* LegacyPackage = static_cast<const FLegacyCookedPackage*>(Package);
+			NameMapCount += LegacyPackage->OptimizedPackage->GetNameCount();
+			HeaderSize += LegacyPackage->OptimizedPackage->GetHeaderSize();
+		}
 		int32 PackageImportedPackagesCount = Package->PackageStoreEntry.ImportedPackageIds.Num();
 		ImportedPackagesCount += PackageImportedPackagesCount;
 		NoImportedPackagesCount += PackageImportedPackagesCount == 0;
@@ -4294,11 +4232,17 @@ int32 CreateTarget(const FIoStoreArguments& Arguments, const FIoStoreWriterSetti
 	UE_LOG(LogIoStore, Display, TEXT("Input:  %8.2f MB for %d Unique shaders"), (double)UniqueShaderSize / 1024.0 / 1024.0, UniqueShaderCount);
 	UE_LOG(LogIoStore, Display, TEXT("Input:  %8.2f MB for %d Inline shaders"), (double)InlineShaderSize / 1024.0 / 1024.0, InlineShaderCount);
 	UE_LOG(LogIoStore, Display, TEXT(""));
-	UE_LOG(LogIoStore, Display, TEXT("Output: %8llu Name map entries"), NameMapCount);
+	if (bIsLegacyStage)
+	{
+		UE_LOG(LogIoStore, Display, TEXT("Output: %8llu Name map entries"), NameMapCount);
+	}
 	UE_LOG(LogIoStore, Display, TEXT("Output: %8llu Imported package entries"), ImportedPackagesCount);
 	UE_LOG(LogIoStore, Display, TEXT("Output: %8llu Packages without imports"), NoImportedPackagesCount);
 	UE_LOG(LogIoStore, Display, TEXT("Output: %8d Public runtime script objects"), PackageStoreOptimizer.GetTotalScriptObjectCount());
-	UE_LOG(LogIoStore, Display, TEXT("Output: %8.2lf GB HeaderData"), (double)HeaderSize / 1024.0 / 1024.0 / 1024.0);
+	if (bIsLegacyStage)
+	{
+		UE_LOG(LogIoStore, Display, TEXT("Output: %8.2lf GB HeaderData"), (double)HeaderSize / 1024.0 / 1024.0 / 1024.0);
+	}
 	UE_LOG(LogIoStore, Display, TEXT("Output: %8.2lf MB InitialLoadData"), (double)InitialLoadSize / 1024.0 / 1024.0);
 
 	if (ChunkDatabase.IsValid())
@@ -5011,13 +4955,12 @@ namespace DescribeUtils
 		FName PackageName;
 		uint32 PackageFlags = 0;
 		int32 NameCount = -1;
-		int32 ExportBundleCount = -1;
 		TArray<FPackageLocation, TInlineAllocator<1>> Locations;
 		TArray<FPackageId> ImportedPackageIds;
 		TArray<uint64> ImportedPublicExportHashes;
 		TArray<FImportDesc> Imports;
 		TArray<FExportDesc> Exports;
-		TArray<TArray<FExportBundleEntryDesc>, TInlineAllocator<1>> ExportBundles;
+		TArray<FExportBundleEntryDesc> ExportBundleEntries;
 	};
 
 	// Info loaded about a set of containers for the purposes of dumping to text in Describe or exploring some other way for debugging 
@@ -5250,7 +5193,6 @@ namespace DescribeUtils
 					const FPackageId& PackageId = ContainerHeader.PackageIds[PackageIndex++];
 					FPackageDesc* PackageDesc = new FPackageDesc();
 					PackageDesc->PackageId = PackageId;
-					PackageDesc->ExportBundleCount = ContainerEntry.ExportBundleCount;
 					PackageDesc->ImportedPackageIds = TArrayView<FPackageId>(ContainerEntry.ImportedPackages.Data(), ContainerEntry.ImportedPackages.Num());
 					Job.Packages.Add(PackageDesc);
 					++TotalPackageCount;
@@ -5382,38 +5324,30 @@ namespace DescribeUtils
 				ExportDesc.SuperIndex = ExportMapEntry.SuperIndex;
 				ExportDesc.TemplateIndex = ExportMapEntry.TemplateIndex;
 				ExportDesc.PublicExportHash = ExportMapEntry.PublicExportHash;
+				ExportDesc.SerialOffset = PackageSummary->HeaderSize + ExportMapEntry.CookedSerialOffset;
 				ExportDesc.SerialSize = ExportMapEntry.CookedSerialSize;
 			}
 
-			const FExportBundleHeader* ExportBundleHeaders = reinterpret_cast<const FExportBundleHeader*>(PackageSummaryData + PackageSummary->GraphDataOffset);
 			const FExportBundleEntry* ExportBundleEntries = reinterpret_cast<const FExportBundleEntry*>(PackageSummaryData + PackageSummary->ExportBundleEntriesOffset);
-			uint64 CurrentExportOffset = PackageSummary->HeaderSize;
-			for (int32 ExportBundleIndex = 0; ExportBundleIndex < Job.PackageDesc->ExportBundleCount; ++ExportBundleIndex)
-			{
-				TArray<FExportBundleEntryDesc>& ExportBundleDesc = Job.PackageDesc->ExportBundles.AddDefaulted_GetRef();
-				const FExportBundleHeader* ExportBundle = ExportBundleHeaders + ExportBundleIndex;
-				const FExportBundleEntry* BundleEntry = ExportBundleEntries + ExportBundle->FirstEntryIndex;
-				const FExportBundleEntry* BundleEntryEnd = BundleEntry + ExportBundle->EntryCount;
-				check(BundleEntry <= BundleEntryEnd);
-				while (BundleEntry < BundleEntryEnd)
+			const FExportBundleEntry* BundleEntry = ExportBundleEntries;
+			int32 ExportBundleEntriesCount = Job.PackageDesc->Exports.Num() * 2;
+			const FExportBundleEntry* BundleEntryEnd = BundleEntry + ExportBundleEntriesCount;
+			Job.PackageDesc->ExportBundleEntries.Reserve(ExportBundleEntriesCount);
+			while (BundleEntry < BundleEntryEnd)
+			{
+				FExportBundleEntryDesc& EntryDesc = Job.PackageDesc->ExportBundleEntries.AddDefaulted_GetRef();
+				EntryDesc.CommandType = FExportBundleEntry::EExportCommandType(BundleEntry->CommandType);
+				EntryDesc.LocalExportIndex = BundleEntry->LocalExportIndex;
+				EntryDesc.Export = &Job.PackageDesc->Exports[BundleEntry->LocalExportIndex];
+				if (BundleEntry->CommandType == FExportBundleEntry::ExportCommandType_Serialize)
 				{
-					FExportBundleEntryDesc& EntryDesc = ExportBundleDesc.AddDefaulted_GetRef();
-					EntryDesc.CommandType = FExportBundleEntry::EExportCommandType(BundleEntry->CommandType);
-					EntryDesc.LocalExportIndex = BundleEntry->LocalExportIndex;
-					EntryDesc.Export = &Job.PackageDesc->Exports[BundleEntry->LocalExportIndex];
-					if (BundleEntry->CommandType == FExportBundleEntry::ExportCommandType_Serialize)
+					if (bIncludeExportHashes)
 					{
-						EntryDesc.Export->SerialOffset = CurrentExportOffset;
-						CurrentExportOffset += EntryDesc.Export->SerialSize;
-
-						if (bIncludeExportHashes)
-						{
-							check(EntryDesc.Export->SerialOffset + EntryDesc.Export->SerialSize <= IoBuffer.ValueOrDie().DataSize());
-							FSHA1::HashBuffer(IoBuffer.ValueOrDie().Data() + EntryDesc.Export->SerialOffset, EntryDesc.Export->SerialSize, EntryDesc.Export->ExportHash.Hash);
-						}
+						check(EntryDesc.Export->SerialOffset + EntryDesc.Export->SerialSize <= IoBuffer.ValueOrDie().DataSize());
+						FSHA1::HashBuffer(IoBuffer.ValueOrDie().Data() + EntryDesc.Export->SerialOffset, EntryDesc.Export->SerialSize, EntryDesc.Export->ExportHash.Hash);
 					}
-					++BundleEntry;
 				}
+				++BundleEntry;
 			}
 		}, EParallelForFlags::Unbalanced);
 
@@ -5754,7 +5688,6 @@ int32 Describe(
 			OutputOverride->Logf(ELogVerbosity::Display, TEXT("\t\t        NameCount: %d"), PackageDesc->NameCount);
 			OutputOverride->Logf(ELogVerbosity::Display, TEXT("\t\t      ImportCount: %d"), PackageDesc->Imports.Num());
 			OutputOverride->Logf(ELogVerbosity::Display, TEXT("\t\t      ExportCount: %d"), PackageDesc->Exports.Num());
-			OutputOverride->Logf(ELogVerbosity::Display, TEXT("\t\tExportBundleCount: %d"), PackageDesc->ExportBundleCount);
 
 			OutputOverride->Logf(ELogVerbosity::Display, TEXT("--------------------------------------------"));
 			OutputOverride->Logf(ELogVerbosity::Display, TEXT("Locations"));
@@ -5801,23 +5734,17 @@ int32 Describe(
 			}
 
 			OutputOverride->Logf(ELogVerbosity::Display, TEXT("--------------------------------------------"));
-			OutputOverride->Logf(ELogVerbosity::Display, TEXT("Export Bundles"));
+			OutputOverride->Logf(ELogVerbosity::Display, TEXT("Export Bundle"));
 			OutputOverride->Logf(ELogVerbosity::Display, TEXT("=========="));
-			Index = 0;
-			for (const TArray<FExportBundleEntryDesc>& ExportBundle : PackageDesc->ExportBundles)
+			for (const FExportBundleEntryDesc& ExportBundleEntry : PackageDesc->ExportBundleEntries)
 			{
-				OutputOverride->Logf(ELogVerbosity::Display, TEXT("\t*************************"));
-				OutputOverride->Logf(ELogVerbosity::Display, TEXT("\tExport Bundle %d"), Index++);
-				for (const FExportBundleEntryDesc& ExportBundleEntry : ExportBundle)
+				if (ExportBundleEntry.CommandType == FExportBundleEntry::ExportCommandType_Create)
 				{
-					if (ExportBundleEntry.CommandType == FExportBundleEntry::ExportCommandType_Create)
-					{
-						OutputOverride->Logf(ELogVerbosity::Display, TEXT("\t\t           Create: %d '%s'"), ExportBundleEntry.LocalExportIndex, *ExportBundleEntry.Export->Name.ToString());
-					}
-					else
-					{
-						OutputOverride->Logf(ELogVerbosity::Display, TEXT("\t\t        Serialize: %d '%s'"), ExportBundleEntry.LocalExportIndex, *ExportBundleEntry.Export->Name.ToString());
-					}
+					OutputOverride->Logf(ELogVerbosity::Display, TEXT("\t\t           Create: %d '%s'"), ExportBundleEntry.LocalExportIndex, *ExportBundleEntry.Export->Name.ToString());
+				}
+				else
+				{
+					OutputOverride->Logf(ELogVerbosity::Display, TEXT("\t\t        Serialize: %d '%s'"), ExportBundleEntry.LocalExportIndex, *ExportBundleEntry.Export->Name.ToString());
 				}
 			}
 		}
@@ -6732,8 +6659,6 @@ int32 Staged2Zen(const FString& BuildPath, const FKeyChain& KeyChain, const FStr
 				
 				FPackageStoreEntryResource& PackageStoreEntryResource = FindPackageInfo->PackageStoreEntry;
 				PackageStoreEntryResource.PackageName = *FindPackageName;
-				PackageStoreEntryResource.ExportInfo.ExportBundleCount = StoreEntry->ExportBundleCount;
-				PackageStoreEntryResource.ExportInfo.ExportCount = StoreEntry->ExportCount;
 				PackageStoreEntryResource.ImportedPackageIds.SetNum(StoreEntry->ImportedPackages.Num());
 				FMemory::Memcpy(PackageStoreEntryResource.ImportedPackageIds.GetData(), StoreEntry->ImportedPackages.Data(), sizeof(FPackageId) * StoreEntry->ImportedPackages.Num()); //-V575
 			}
diff --git a/Engine/Source/Developer/IoStoreUtilities/Private/PackageStoreOptimizer.cpp b/Engine/Source/Developer/IoStoreUtilities/Private/PackageStoreOptimizer.cpp
--- a/Engine/Source/Developer/IoStoreUtilities/Private/PackageStoreOptimizer.cpp
+++ b/Engine/Source/Developer/IoStoreUtilities/Private/PackageStoreOptimizer.cpp
@@ -117,43 +117,9 @@ FPackageStorePackage* FPackageStoreOptimizer::CreatePackageFromCookedHeader(cons
 	ProcessPreloadDependencies(CookedHeaderData, Package);
 	ProcessDataResources(CookedHeaderData, Package);
 
-	CreateExportBundles(Package);
+	CreateExportBundle(Package);
 
-	return Package;
-}
-
-FPackageStorePackage* FPackageStoreOptimizer::CreatePackageFromZenPackageHeader(const FName& Name, const FIoBuffer& Buffer, int32 ExportBundleCount, const TArrayView<const FPackageId>& ImportedPackageIds) const
-{
-	FPackageStorePackage* Package = new FPackageStorePackage();
-	Package->Id = FPackageId::FromName(Name);
-
-	// The package id should be generated from the original name
-	// Support for optional package is not implemented when loading from the package store yet however
-	FString NameStr = Name.ToString();
-	int32 Index = NameStr.Find(FPackagePath::GetOptionalSegmentExtensionModifier());
-	if (Index != INDEX_NONE)
-	{
-		unimplemented();
-	}
-
-	Package->Name = Name;
-	
-	FZenPackageHeaderData ZenHeaderData = LoadZenPackageHeader(Buffer, ExportBundleCount);
-	if (ZenHeaderData.VersioningInfo.IsSet())
-	{
-		Package->VersioningInfo.Emplace(ZenHeaderData.VersioningInfo.GetValue());
-	}
-	Package->PackageFlags = ZenHeaderData.Summary.PackageFlags;
-	Package->CookedHeaderSize = ZenHeaderData.Summary.CookedHeaderSize;
-	for (FDisplayNameEntryId DisplayId : ZenHeaderData.NameMap)
-	{
-		Package->NameMapBuilder.AddName(DisplayId);
-	}
-	Package->BulkDataEntries = MoveTemp(ZenHeaderData.BulkDataEntries);
-	ProcessImports(ZenHeaderData, Package);
-	ProcessExports(ZenHeaderData, Package);
-	ProcessPreloadDependencies(ZenHeaderData, Package);
-	CreateExportBundles(Package);
+	FinalizePackageHeader(Package);
 
 	return Package;
 }
@@ -265,105 +231,6 @@ FPackageStoreOptimizer::FCookedHeaderData FPackageStoreOptimizer::LoadCookedHead
 }
 PRAGMA_ENABLE_DEPRECATION_WARNINGS
 
-FPackageStoreOptimizer::FZenPackageHeaderData FPackageStoreOptimizer::LoadZenPackageHeader(const FIoBuffer& HeaderBuffer, int32 ExportBundleCount) const
-{
-	FZenPackageHeaderData ZenPackageHeaderData;
-
-	const uint8* HeaderData = HeaderBuffer.Data();
-
-	FZenPackageSummary& Summary = ZenPackageHeaderData.Summary;
-	Summary = *reinterpret_cast<const FZenPackageSummary*>(HeaderData);
-	check(HeaderBuffer.DataSize() == Summary.HeaderSize);
-
-	TArrayView<const uint8> HeaderDataView(HeaderData + sizeof(FZenPackageSummary), Summary.HeaderSize - sizeof(FZenPackageSummary));
-	FMemoryReaderView HeaderDataReader(HeaderDataView);
-	
-	if (Summary.bHasVersioningInfo)
-	{
-		FZenPackageVersioningInfo& VersioningInfo = ZenPackageHeaderData.VersioningInfo.Emplace();
-		HeaderDataReader << VersioningInfo;
-	}
-
-	TArray<FDisplayNameEntryId>& NameMap = ZenPackageHeaderData.NameMap;
-	NameMap = LoadNameBatch(HeaderDataReader);
-
-	const FZenPackageVersioningInfo* VersioningInfo = ZenPackageHeaderData.VersioningInfo.GetPtrOrNull();
-	if (VersioningInfo == nullptr || VersioningInfo->PackageVersion >= EUnrealEngineObjectUE5Version::DATA_RESOURCES)
-	{
-		int64 BulkDataMapSize = 0;
-		HeaderDataReader << BulkDataMapSize;
-		const uint8* BulkDataMapData = HeaderData + sizeof(FZenPackageSummary) + HeaderDataReader.Tell();
-		ZenPackageHeaderData.BulkDataEntries = MakeArrayView(reinterpret_cast<const FBulkDataMapEntry*>(BulkDataMapData), BulkDataMapSize / sizeof(FBulkDataMapEntry));
-	}
-
-	TArrayView<const uint8> ImportedPackageNamesDataView(HeaderData + Summary.ImportedPackageNamesOffset, Summary.HeaderSize - Summary.ImportedPackageNamesOffset);
-	FMemoryReaderView ImportedPackageNamesDataReader(ImportedPackageNamesDataView);
-	FZenPackageImportedPackageNamesContainer ImportedPackageNamesContainer;
-	ImportedPackageNamesDataReader << ImportedPackageNamesContainer;
-	ZenPackageHeaderData.ImportedPackageNames = MoveTemp(ImportedPackageNamesContainer.Names);
-
-	ZenPackageHeaderData.ImportedPublicExportHashes =
-		MakeArrayView<const uint64>(
-			reinterpret_cast<const uint64*>(HeaderData + Summary.ImportedPublicExportHashesOffset),
-			(Summary.ImportMapOffset - Summary.ImportedPublicExportHashesOffset) / sizeof(uint64));
-
-	ZenPackageHeaderData.Imports =
-		MakeArrayView<const FPackageObjectIndex>(
-			reinterpret_cast<const FPackageObjectIndex*>(HeaderData + Summary.ImportMapOffset),
-			(Summary.ExportMapOffset - Summary.ImportMapOffset) / sizeof(FPackageObjectIndex));
-
-	ZenPackageHeaderData.Exports =
-		MakeArrayView<const FExportMapEntry>(
-			reinterpret_cast<const FExportMapEntry*>(HeaderData + Summary.ExportMapOffset),
-			(Summary.ExportBundleEntriesOffset - Summary.ExportMapOffset) / sizeof(FExportMapEntry));
-
-	ZenPackageHeaderData.ExportBundleHeaders =
-		MakeArrayView<const FExportBundleHeader>(
-			reinterpret_cast<const FExportBundleHeader*>(HeaderData + Summary.GraphDataOffset),
-			ExportBundleCount);
-
-	ZenPackageHeaderData.ExportBundleEntries =
-		MakeArrayView<const FExportBundleEntry>(
-			reinterpret_cast<const FExportBundleEntry*>(HeaderData + Summary.ExportBundleEntriesOffset),
-			(Summary.GraphDataOffset - Summary.ExportBundleEntriesOffset) / sizeof(FExportBundleEntry));
-
-	const uint64 ExportBundleHeadersSize = sizeof(FExportBundleHeader) * ExportBundleCount;
-	const uint64 ArcsDataOffset = Summary.GraphDataOffset + ExportBundleHeadersSize;
-	const uint64 ArcsDataSize = Summary.HeaderSize - ArcsDataOffset;
-
-	FMemoryReaderView ArcsAr(MakeArrayView<const uint8>(HeaderData + ArcsDataOffset, ArcsDataSize));
-
-	int32 InternalArcsCount = 0;
-	ArcsAr << InternalArcsCount;
-
-	for (int32 Idx = 0; Idx < InternalArcsCount; ++Idx)
-	{
-		FPackageStorePackage::FInternalArc& InternalArc = ZenPackageHeaderData.InternalArcs.AddDefaulted_GetRef();
-		ArcsAr << InternalArc.FromExportBundleIndex;
-		ArcsAr << InternalArc.ToExportBundleIndex;
-	}
-
-	for (FName ImportedPackageName : ZenPackageHeaderData.ImportedPackageNames)
-	{
-		int32 ExternalArcsCount = 0;
-		ArcsAr << ExternalArcsCount;
-
-		for (int32 Idx = 0; Idx < ExternalArcsCount; ++Idx)
-		{
-			FPackageStorePackage::FExternalArc ExternalArc;
-			ArcsAr << ExternalArc.FromImportIndex;
-			uint8 FromCommandType = 0;
-			ArcsAr << FromCommandType;
-			ExternalArc.FromCommandType = static_cast<FExportBundleEntry::EExportCommandType>(FromCommandType);
-			ArcsAr << ExternalArc.ToExportBundleIndex;
-
-			ZenPackageHeaderData.ExternalArcs.Add(ExternalArc);
-		}
-	}
-
-	return ZenPackageHeaderData;
-}
-
 void FPackageStoreOptimizer::ResolveImport(FPackageStorePackage::FUnresolvedImport* Imports, const FObjectImport* ObjectImports, int32 LocalImportIndex) const
 {
 	FPackageStorePackage::FUnresolvedImport* Import = Imports + LocalImportIndex;
@@ -467,17 +334,6 @@ void FPackageStoreOptimizer::ProcessImports(const FCookedHeaderData& CookedHeade
 	}
 }
 
-void FPackageStoreOptimizer::ProcessImports(const FZenPackageHeaderData& ZenHeaderData, FPackageStorePackage* Package) const
-{
-	Package->ImportedPackages.Reserve(ZenHeaderData.ImportedPackageNames.Num());
-	for (FName ImportedPackageName : ZenHeaderData.ImportedPackageNames)
-	{
-		Package->ImportedPackages.Emplace(ImportedPackageName);
-	}
-	Package->ImportedPublicExportHashes = ZenHeaderData.ImportedPublicExportHashes;
-	Package->Imports = ZenHeaderData.Imports;
-}
-
 void FPackageStoreOptimizer::ResolveExport(
 	FPackageStorePackage::FUnresolvedExport* Exports,
 	const FObjectExport* ObjectExports,
@@ -553,14 +409,13 @@ void FPackageStoreOptimizer::ProcessExports(const FCookedHeaderData& CookedHeade
 	for (int32 ExportIndex = 0; ExportIndex < ExportCount; ++ExportIndex)
 	{
 		const FObjectExport& ObjectExport = CookedHeaderData.ObjectExports[ExportIndex];
-		Package->ExportsSerialSize += ObjectExport.SerialSize;
 
 		FPackageStorePackage::FExport& Export = Package->Exports[ExportIndex];
 		FPackageStorePackage::FUnresolvedExport& UnresolvedExport = UnresolvedExports[ExportIndex];
 		Export.ObjectName = ObjectExport.ObjectName;
 		Export.ObjectFlags = ObjectExport.ObjectFlags;
-		Export.CookedSerialOffset = ObjectExport.SerialOffset;
-		Export.SerialOffset = ObjectExport.SerialOffset - CookedHeaderData.Summary.TotalHeaderSize;
+		check(ObjectExport.SerialOffset >= Package->CookedHeaderSize);
+		Export.SerialOffset = ObjectExport.SerialOffset - Package->CookedHeaderSize;
 		Export.SerialSize = ObjectExport.SerialSize;
 		Export.bNotForClient = ObjectExport.bNotForClient;
 		Export.bNotForServer = ObjectExport.bNotForServer;
@@ -596,372 +451,103 @@ void FPackageStoreOptimizer::ProcessExports(const FCookedHeaderData& CookedHeade
 	}
 }
 
-void FPackageStoreOptimizer::ProcessExports(const FZenPackageHeaderData& ZenHeaderData, FPackageStorePackage* Package) const
-{
-	const int32 ExportCount = ZenHeaderData.Exports.Num();
-	Package->Exports.SetNum(ExportCount);
-	Package->ExportGraphNodes.Reserve(ExportCount * 2);
-
-	const TArray<FDisplayNameEntryId>& NameMap = ZenHeaderData.NameMap;
-
-	Package->ImportedPublicExportHashes = ZenHeaderData.ImportedPublicExportHashes;
-	for (int32 ExportIndex = 0; ExportIndex < ExportCount; ++ExportIndex)
-	{
-		const FExportMapEntry& ExportEntry = ZenHeaderData.Exports[ExportIndex];
-		Package->ExportsSerialSize += ExportEntry.CookedSerialSize;
-
-		FPackageStorePackage::FExport& Export = Package->Exports[ExportIndex];
-		Export.ObjectName = ExportEntry.ObjectName.ResolveName(NameMap);
-		Export.PublicExportHash = ExportEntry.PublicExportHash;
-		Export.OuterIndex = ExportEntry.OuterIndex;
-		Export.ClassIndex = ExportEntry.ClassIndex;
-		Export.SuperIndex = ExportEntry.SuperIndex;
-		Export.TemplateIndex = ExportEntry.TemplateIndex;
-		Export.ObjectFlags = ExportEntry.ObjectFlags;
-		Export.CookedSerialOffset = ExportEntry.CookedSerialOffset;
-		Export.SerialSize = ExportEntry.CookedSerialSize;
-		Export.bNotForClient = ExportEntry.FilterFlags == EExportFilterFlags::NotForClient;
-		Export.bNotForServer = ExportEntry.FilterFlags == EExportFilterFlags::NotForServer;
-		Export.bIsPublic = (ExportEntry.ObjectFlags & RF_Public) > 0;
-
-		for (uint8 CommandType = 0; CommandType < FExportBundleEntry::ExportCommandType_Count; ++CommandType)
-		{
-			FPackageStorePackage::FExportGraphNode& Node = Package->ExportGraphNodes.AddDefaulted_GetRef();
-			Node.BundleEntry.CommandType = FExportBundleEntry::EExportCommandType(CommandType);
-			Node.BundleEntry.LocalExportIndex = ExportIndex;
-			Node.bIsPublic = Export.bIsPublic;
-			Export.Nodes[CommandType] = &Node;
-		}
-
-		Package->NameMapBuilder.MarkNameAsReferenced(Export.ObjectName);
-	}
-
-	uint64 ExportSerialOffset = 0;
-	for (const FExportBundleHeader& ExportBundleHeader : ZenHeaderData.ExportBundleHeaders)
-	{
-		int32 ExportEntryIndex = ExportBundleHeader.FirstEntryIndex;
-		for (uint32 Idx = 0; Idx < ExportBundleHeader.EntryCount; ++Idx)
-		{
-			const FExportBundleEntry& BundleEntry = ZenHeaderData.ExportBundleEntries[ExportEntryIndex++];
-			if (BundleEntry.CommandType == FExportBundleEntry::ExportCommandType_Serialize)
-			{
-				FPackageStorePackage::FExport& Export = Package->Exports[BundleEntry.LocalExportIndex];
-				Export.SerialOffset = ExportSerialOffset;
-				ExportSerialOffset += Export.SerialSize;
-			}
-		}
-	}
-	check(ExportSerialOffset == Package->ExportsSerialSize);
-}
-
-TArray<FPackageStorePackage*> FPackageStoreOptimizer::SortPackagesInLoadOrder(const TMap<FPackageId, FPackageStorePackage*>& PackagesMap) const
-{
-	TRACE_CPUPROFILER_EVENT_SCOPE(SortPackagesInLoadOrder);
-	TArray<FPackageStorePackage*> Packages;
-	PackagesMap.GenerateValueArray(Packages);
-	Algo::Sort(Packages, [](const FPackageStorePackage* A, const FPackageStorePackage* B)
-	{
-		return A->Id < B->Id;
-	});
-
-	TMap<FPackageStorePackage*, TArray<FPackageStorePackage*>> SortedEdges;
-	for (FPackageStorePackage* Package : Packages)
-	{
-		for (const FPackageStorePackage::FImportedPackageRef& ImportedPackage : Package->ImportedPackages)
-		{
-			FPackageStorePackage* FindImportedPackage = PackagesMap.FindRef(ImportedPackage.Id);
-			if (FindImportedPackage)
-			{
-				TArray<FPackageStorePackage*>& SourceArray = SortedEdges.FindOrAdd(FindImportedPackage);
-				SourceArray.Add(Package);
-			}
-		}
-	}
-	for (auto& KV : SortedEdges)
-	{
-		TArray<FPackageStorePackage*>& SourceArray = KV.Value;
-		Algo::Sort(SourceArray, [](const FPackageStorePackage* A, const FPackageStorePackage* B)
-		{
-			return A->Id < B->Id;
-		});
-	}
-
-	TArray<FPackageStorePackage*> Result;
-	Result.Reserve(Packages.Num());
-	struct
-	{
-		void Visit(FPackageStorePackage* Package)
-		{
-			if (Package->bPermanentMark || Package->bTemporaryMark)
-			{
-				return;
-			}
-			Package->bTemporaryMark = true;
-			TArray<FPackageStorePackage*>* TargetNodes = Edges.Find(Package);
-			if (TargetNodes)
-			{
-				for (FPackageStorePackage* ToNode : *TargetNodes)
-				{
-					Visit(ToNode);
-				}
-			}
-			Package->bTemporaryMark = false;
-			Package->bPermanentMark = true;
-			Result.Add(Package);
-		}
-
-		TMap<FPackageStorePackage*, TArray<FPackageStorePackage*>>& Edges;
-		TArray<FPackageStorePackage*>& Result;
-	} Visitor{ SortedEdges, Result };
-
-	for (FPackageStorePackage* Package : Packages)
-	{
-		Visitor.Visit(Package);
-	}
-	check(Result.Num() == Packages.Num());
-	Algo::Reverse(Result);
-	Swap(Result, Packages);
-	return Packages;
-}
-
-TArray<FPackageStorePackage::FExportBundleGraphNode*> FPackageStoreOptimizer::SortExportBundleGraphNodesInLoadOrder(const TArray<FPackageStorePackage*>& Packages, FExportBundleGraphEdges& Edges) const
-{
-	TRACE_CPUPROFILER_EVENT_SCOPE(SortExportBundleGraphNodesInLoadOrder);
-	int32 NodeCount = 0;
-	for (FPackageStorePackage* Package : Packages)
-	{
-		NodeCount += Package->ExportGraphNodes.Num();
-	}
-	for (auto& KV : Edges)
-	{
-		FPackageStorePackage::FExportBundleGraphNode* ToNode = KV.Value;
-		++ToNode->IncomingEdgeCount;
-	}
-
-	auto NodeSorter = [](const FPackageStorePackage::FExportBundleGraphNode& A, const FPackageStorePackage::FExportBundleGraphNode& B)
-	{
-		return A.Index < B.Index;
-	};
-
-	for (FPackageStorePackage* Package : Packages)
-	{
-		for (FPackageStorePackage::FExportBundleGraphNode& Node : Package->ExportBundleGraphNodes)
-		{
-			if (Node.IncomingEdgeCount == 0)
-			{
-				Package->NodesWithNoIncomingEdges.HeapPush(&Node, NodeSorter);
-			}
-		}
-	}
-	TArray<FPackageStorePackage::FExportBundleGraphNode*> LoadOrder;
-	LoadOrder.Reserve(NodeCount);
-	while (LoadOrder.Num() < NodeCount)
-	{
-		bool bMadeProgress = false;
-		for (FPackageStorePackage* Package : Packages)
-		{
-			while (Package->NodesWithNoIncomingEdges.Num())
-			{
-				FPackageStorePackage::FExportBundleGraphNode* RemovedNode;
-				Package->NodesWithNoIncomingEdges.HeapPop(RemovedNode, NodeSorter, false);
-				LoadOrder.Add(RemovedNode);
-				bMadeProgress = true;
-				for (auto EdgeIt = Edges.CreateKeyIterator(RemovedNode); EdgeIt; ++EdgeIt)
-				{
-					FPackageStorePackage::FExportBundleGraphNode* ToNode = EdgeIt.Value();
-					check(ToNode->IncomingEdgeCount > 0);
-					if (--ToNode->IncomingEdgeCount == 0)
-					{
-						ToNode->Package->NodesWithNoIncomingEdges.HeapPush(ToNode, NodeSorter);
-					}
-					EdgeIt.RemoveCurrent();
-				}
-			}
-		}
-		check(bMadeProgress);
-	}
-	return LoadOrder;
-}
-
-void FPackageStoreOptimizer::OptimizeExportBundles(const TMap<FPackageId, FPackageStorePackage*>& PackagesMap)
-{
-	TRACE_CPUPROFILER_EVENT_SCOPE(OptimizeExportBundles);
-	TArray<FPackageStorePackage*> Packages = SortPackagesInLoadOrder(PackagesMap);
-	for (FPackageStorePackage* Package : Packages)
-	{
-		Package->ExportBundleGraphNodes.Reserve(Package->GraphData.ExportBundles.Num());
-		for (FPackageStorePackage::FExportBundle& ExportBundle : Package->GraphData.ExportBundles)
-		{
-			FPackageStorePackage::FExportBundleGraphNode& Node = Package->ExportBundleGraphNodes.AddDefaulted_GetRef();
-			Node.Package = Package;
-			Node.Index = Package->ExportBundleGraphNodes.Num() - 1;
-			Node.ExportGraphNodes.Reserve(ExportBundle.Entries.Num());
-			for (FExportBundleEntry& ExportBundleEntry : ExportBundle.Entries)
-			{
-				FPackageStorePackage::FExport& Export = Package->Exports[ExportBundleEntry.LocalExportIndex];
-				Node.ExportGraphNodes.Add(Export.Nodes[ExportBundleEntry.CommandType]);
-			}
-		}
-		Package->GraphData.ExportBundles.Empty();
-	}
-
-	FExportBundleGraphEdges Edges;
-	TSet<FPackageStorePackage::FExportBundleGraphNode*> Dependencies;
-	for (FPackageStorePackage* Package : Packages)
-	{
-		for (FPackageStorePackage::FExportBundleGraphNode& ExportBundleGraphNode : Package->ExportBundleGraphNodes)
-		{
-			for (FPackageStorePackage::FExportGraphNode* ExportGraphNode : ExportBundleGraphNode.ExportGraphNodes)
-			{
-				check(ExportGraphNode->ExportBundleIndex >= 0);
-				for (FPackageStorePackage::FExportGraphNode* InternalDependency : ExportGraphNode->InternalDependencies)
-				{
-					check(InternalDependency->ExportBundleIndex >= 0);
-					FPackageStorePackage::FExportBundleGraphNode* FromNode = &Package->ExportBundleGraphNodes[InternalDependency->ExportBundleIndex];
-					Dependencies.Add(FromNode);
-				}
-				for (FPackageStorePackage::FExternalDependency& ExternalDependency : ExportGraphNode->ExternalDependencies)
-				{
-					FPackageObjectIndex FromImport = Package->Imports[ExternalDependency.ImportIndex];
-					check(FromImport.IsPackageImport());
-					FPackageImportReference FromPackageImportRef = FromImport.ToPackageImportRef();
-					FPackageId FromPackageId = Package->ImportedPackages[FromPackageImportRef.GetImportedPackageIndex()].Id;
-					uint64 FromPublicExportHash = Package->ImportedPublicExportHashes[FromPackageImportRef.GetImportedPublicExportHashIndex()];
-					FPackageStorePackage* FindFromPackage = PackagesMap.FindRef(FromPackageId);
-					if (FindFromPackage)
-					{
-						bool bFoundExport = false;
-						for (int32 ExportIndex = 0; ExportIndex < FindFromPackage->Exports.Num(); ++ExportIndex)
-						{
-							FPackageStorePackage::FExport& FromExport = FindFromPackage->Exports[ExportIndex];
-							if (FromExport.PublicExportHash == FromPublicExportHash)
-							{
-								FPackageStorePackage::FExportGraphNode* FromExportGraphNode = FromExport.Nodes[ExternalDependency.ExportBundleCommandType];
-								check(FromExportGraphNode->ExportBundleIndex >= 0);
-								FPackageStorePackage::FExportBundleGraphNode* FromNode = &FindFromPackage->ExportBundleGraphNodes[FromExportGraphNode->ExportBundleIndex];
-								Dependencies.Add(FromNode);
-								bFoundExport = true;
-								break;
-							}
-						}
-					}
-				}
-			}
-			for (FPackageStorePackage::FExportBundleGraphNode* FromNode : Dependencies)
-			{
-				Edges.Add(FromNode, &ExportBundleGraphNode);
-			}
-			Dependencies.Reset();
-		}
-	}
-
-	TArray<FPackageStorePackage::FExportBundleGraphNode*> LoadOrder = SortExportBundleGraphNodesInLoadOrder(Packages, Edges);
-
-	FPackageStorePackage* PreviousPackage = nullptr;
-	for (const FPackageStorePackage::FExportBundleGraphNode* Node : LoadOrder)
-	{
-		check(Node);
-		FPackageStorePackage* Package = Node->Package;
-		check(Package);
-		if (Package->CurrentBundle == nullptr || Package != PreviousPackage)
-		{
-			Package->CurrentBundle = &Package->GraphData.ExportBundles.AddDefaulted_GetRef();
-		}
-		for (FPackageStorePackage::FExportGraphNode* ExportGraphNode : Node->ExportGraphNodes)
-		{
-			Package->CurrentBundle->Entries.Add(ExportGraphNode->BundleEntry);
-			ExportGraphNode->ExportBundleIndex = Package->GraphData.ExportBundles.Num() - 1;
-		}
-		PreviousPackage = Package;
-	}
-}
-
 void FPackageStoreOptimizer::ProcessPreloadDependencies(const FCookedHeaderData& CookedHeaderData, FPackageStorePackage* Package) const
 {
 	TRACE_CPUPROFILER_EVENT_SCOPE(ProcessPreloadDependencies);
 
-	auto AddInternalDependency = [](FPackageStorePackage* Package, int32 FromExportIndex, FExportBundleEntry::EExportCommandType FromExportBundleCommandType, FPackageStorePackage::FExportGraphNode* ToNode)
+	auto AddNodeDependency = [](FPackageStorePackage* Package, int32 FromExportIndex, FExportBundleEntry::EExportCommandType FromExportBundleCommandType, FPackageStorePackage::FExportGraphNode* ToNode)
 	{
 		FPackageStorePackage::FExport& FromExport = Package->Exports[FromExportIndex];
 		FPackageStorePackage::FExportGraphNode* FromNode = FromExport.Nodes[FromExportBundleCommandType];
 		ToNode->InternalDependencies.Add(FromNode);
 	};
 
-	auto AddExternalDependency = [](FPackageStorePackage* Package, int32 FromImportIndex, FExportBundleEntry::EExportCommandType FromExportBundleCommandType, FPackageStorePackage::FExportGraphNode* ToNode)
-	{
-		FPackageObjectIndex FromImport = Package->Imports[FromImportIndex];
-		if (FromImport.IsScriptImport())
-		{
-			return;
-		}
-
-		FPackageStorePackage::FExternalDependency& ExternalDependency = ToNode->ExternalDependencies.AddDefaulted_GetRef();
-		ExternalDependency.ImportIndex = FromImportIndex;
-		ExternalDependency.ExportBundleCommandType = FromExportBundleCommandType;
-	};
+	TArray<FDependencyBundleHeader>& DependencyBundleHeaders = Package->GraphData.DependencyBundleHeaders;
+	TArray<FDependencyBundleEntry>& DependencyBundleEntries = Package->GraphData.DependencyBundleEntries;
 
 	for (int32 ExportIndex = 0; ExportIndex < Package->Exports.Num(); ++ExportIndex)
 	{
 		FPackageStorePackage::FExport& Export = Package->Exports[ExportIndex];
 		const FObjectExport& ObjectExport = CookedHeaderData.ObjectExports[ExportIndex];
 
-		AddInternalDependency(Package, ExportIndex, FExportBundleEntry::ExportCommandType_Create, Export.Nodes[FExportBundleEntry::ExportCommandType_Serialize]);
+		AddNodeDependency(Package, ExportIndex, FExportBundleEntry::ExportCommandType_Create, Export.Nodes[FExportBundleEntry::ExportCommandType_Serialize]);
 
+		FDependencyBundleHeader& DependencyBundleHeader = DependencyBundleHeaders.AddDefaulted_GetRef();
+		FMemory::Memzero(&DependencyBundleHeader, sizeof(FDependencyBundleHeader));
+		
 		if (ObjectExport.FirstExportDependency >= 0)
 		{
-			int32 RunningIndex = ObjectExport.FirstExportDependency;
-			for (int32 Index = ObjectExport.SerializationBeforeSerializationDependencies; Index > 0; Index--)
+			DependencyBundleHeader.FirstEntryIndex = DependencyBundleEntries.Num();
+			int32 StartIndex = ObjectExport.FirstExportDependency +
+				ObjectExport.SerializationBeforeSerializationDependencies +
+				ObjectExport.CreateBeforeSerializationDependencies +
+				ObjectExport.SerializationBeforeCreateDependencies;
+			for (int32 Index = StartIndex; Index < StartIndex + ObjectExport.CreateBeforeCreateDependencies; ++Index)
 			{
-				FPackageIndex Dep = CookedHeaderData.PreloadDependencies[RunningIndex++];
+				FPackageIndex Dep = CookedHeaderData.PreloadDependencies[Index];
 				if (Dep.IsExport())
 				{
-					AddInternalDependency(Package, Dep.ToExport(), FExportBundleEntry::ExportCommandType_Serialize, Export.Nodes[FExportBundleEntry::ExportCommandType_Serialize]);
+					AddNodeDependency(Package, Dep.ToExport(), FExportBundleEntry::ExportCommandType_Create, Export.Nodes[FExportBundleEntry::ExportCommandType_Create]);
 				}
-				else
+				if (Dep.IsExport() || !Package->Imports[Dep.ToImport()].IsScriptImport())
 				{
-					AddExternalDependency(Package, Dep.ToImport(), FExportBundleEntry::ExportCommandType_Serialize, Export.Nodes[FExportBundleEntry::ExportCommandType_Serialize]);
+					DependencyBundleEntries.AddDefaulted_GetRef().LocalImportOrExportIndex = Dep;
+					++DependencyBundleHeader.EntryCount[FExportBundleEntry::ExportCommandType_Create][FExportBundleEntry::ExportCommandType_Create];
 				}
 			}
 
-			for (int32 Index = ObjectExport.CreateBeforeSerializationDependencies; Index > 0; Index--)
+			StartIndex = ObjectExport.FirstExportDependency +
+				ObjectExport.SerializationBeforeSerializationDependencies +
+				ObjectExport.CreateBeforeSerializationDependencies;
+			for (int32 Index = StartIndex; Index < StartIndex + ObjectExport.SerializationBeforeCreateDependencies; ++Index)
 			{
-				FPackageIndex Dep = CookedHeaderData.PreloadDependencies[RunningIndex++];
+				FPackageIndex Dep = CookedHeaderData.PreloadDependencies[Index];
 				if (Dep.IsExport())
 				{
-					AddInternalDependency(Package, Dep.ToExport(), FExportBundleEntry::ExportCommandType_Create, Export.Nodes[FExportBundleEntry::ExportCommandType_Serialize]);
+					AddNodeDependency(Package, Dep.ToExport(), FExportBundleEntry::ExportCommandType_Serialize, Export.Nodes[FExportBundleEntry::ExportCommandType_Create]);
 				}
-				else
+				if (Dep.IsExport() || !Package->Imports[Dep.ToImport()].IsScriptImport())
 				{
-					AddExternalDependency(Package, Dep.ToImport(), FExportBundleEntry::ExportCommandType_Create, Export.Nodes[FExportBundleEntry::ExportCommandType_Serialize]);
+					DependencyBundleEntries.AddDefaulted_GetRef().LocalImportOrExportIndex = Dep;
+					++DependencyBundleHeader.EntryCount[FExportBundleEntry::ExportCommandType_Create][FExportBundleEntry::ExportCommandType_Serialize];
 				}
 			}
-
-			for (int32 Index = ObjectExport.SerializationBeforeCreateDependencies; Index > 0; Index--)
+			
+			StartIndex = ObjectExport.FirstExportDependency +
+				ObjectExport.SerializationBeforeSerializationDependencies;
+			for (int32 Index = StartIndex; Index < StartIndex + ObjectExport.CreateBeforeSerializationDependencies; ++Index)
 			{
-				FPackageIndex Dep = CookedHeaderData.PreloadDependencies[RunningIndex++];
+				FPackageIndex Dep = CookedHeaderData.PreloadDependencies[Index];
 				if (Dep.IsExport())
 				{
-					AddInternalDependency(Package, Dep.ToExport(), FExportBundleEntry::ExportCommandType_Serialize, Export.Nodes[FExportBundleEntry::ExportCommandType_Create]);
+					AddNodeDependency(Package, Dep.ToExport(), FExportBundleEntry::ExportCommandType_Create, Export.Nodes[FExportBundleEntry::ExportCommandType_Serialize]);
 				}
-				else
+				if (Dep.IsExport() || !Package->Imports[Dep.ToImport()].IsScriptImport())
 				{
-					AddExternalDependency(Package, Dep.ToImport(), FExportBundleEntry::ExportCommandType_Serialize, Export.Nodes[FExportBundleEntry::ExportCommandType_Create]);
+					DependencyBundleEntries.AddDefaulted_GetRef().LocalImportOrExportIndex = Dep;
+					++DependencyBundleHeader.EntryCount[FExportBundleEntry::ExportCommandType_Serialize][FExportBundleEntry::ExportCommandType_Create];
 				}
 			}
 
-			for (int32 Index = ObjectExport.CreateBeforeCreateDependencies; Index > 0; Index--)
+			StartIndex = ObjectExport.FirstExportDependency;
+			for (int32 Index = StartIndex; Index < StartIndex + ObjectExport.SerializationBeforeSerializationDependencies; ++Index)
 			{
-				FPackageIndex Dep = CookedHeaderData.PreloadDependencies[RunningIndex++];
+				FPackageIndex Dep = CookedHeaderData.PreloadDependencies[Index];
 				if (Dep.IsExport())
 				{
-					AddInternalDependency(Package, Dep.ToExport(), FExportBundleEntry::ExportCommandType_Create, Export.Nodes[FExportBundleEntry::ExportCommandType_Create]);
+					AddNodeDependency(Package, Dep.ToExport(), FExportBundleEntry::ExportCommandType_Serialize, Export.Nodes[FExportBundleEntry::ExportCommandType_Serialize]);
 				}
-				else
+				if (Dep.IsExport() || !Package->Imports[Dep.ToImport()].IsScriptImport())
 				{
-					AddExternalDependency(Package, Dep.ToImport(), FExportBundleEntry::ExportCommandType_Create, Export.Nodes[FExportBundleEntry::ExportCommandType_Create]);
+					DependencyBundleEntries.AddDefaulted_GetRef().LocalImportOrExportIndex = Dep;
+					++DependencyBundleHeader.EntryCount[FExportBundleEntry::ExportCommandType_Serialize][FExportBundleEntry::ExportCommandType_Serialize];
 				}
 			}
 		}
+		else
+		{
+			DependencyBundleHeader.FirstEntryIndex = -1;
+		}
 	}
 }
 
@@ -979,52 +565,6 @@ void FPackageStoreOptimizer::ProcessDataResources(const FCookedHeaderData& Cooke
 	}
 }
 
-void FPackageStoreOptimizer::ProcessPreloadDependencies(const FZenPackageHeaderData& ZenHeaderData, FPackageStorePackage* Package) const
-{
-	TRACE_CPUPROFILER_EVENT_SCOPE(ProcessPreloadDependencies);
-
-	for (const FPackageStorePackage::FInternalArc& InternalArc : ZenHeaderData.InternalArcs)
-	{
-		const FExportBundleHeader& FromExportBundle = ZenHeaderData.ExportBundleHeaders[InternalArc.FromExportBundleIndex];
-		const FExportBundleHeader& ToExportBundle = ZenHeaderData.ExportBundleHeaders[InternalArc.ToExportBundleIndex];
-
-		uint32 FromBundleEntryIndex = FromExportBundle.FirstEntryIndex;
-		const uint32 LastFromBundleEntryIndex = FromBundleEntryIndex + FromExportBundle.EntryCount;
-		while (FromBundleEntryIndex < LastFromBundleEntryIndex)
-		{
-			const FExportBundleEntry& FromBundleEntry = ZenHeaderData.ExportBundleEntries[FromBundleEntryIndex++];
-			FPackageStorePackage::FExport& FromExport = Package->Exports[FromBundleEntry.LocalExportIndex];
-			FPackageStorePackage::FExportGraphNode* FromNode = FromExport.Nodes[FromBundleEntry.CommandType];
-
-			uint32 ToBundleEntryIndex = ToExportBundle.FirstEntryIndex;
-			const uint32 LastToBundleEntryIndex = ToBundleEntryIndex + ToExportBundle.EntryCount;
-			while (ToBundleEntryIndex < LastToBundleEntryIndex)
-			{
-				const FExportBundleEntry& ToBundleEntry = ZenHeaderData.ExportBundleEntries[ToBundleEntryIndex++];
-				FPackageStorePackage::FExport& ToExport = Package->Exports[ToBundleEntry.LocalExportIndex];
-				FPackageStorePackage::FExportGraphNode* ToNode = ToExport.Nodes[ToBundleEntry.CommandType];
-				ToNode->InternalDependencies.Add(FromNode);
-			}
-		}
-	}
-
-	for (const FPackageStorePackage::FExternalArc& ExternalArc : ZenHeaderData.ExternalArcs)
-	{
-		const FExportBundleHeader& ToExportBundle = ZenHeaderData.ExportBundleHeaders[ExternalArc.ToExportBundleIndex];
-		uint32 ToBundleEntryIndex = ToExportBundle.FirstEntryIndex;
-		const uint32 LastToBundleEntryIndex = ToBundleEntryIndex + ToExportBundle.EntryCount;
-		while (ToBundleEntryIndex < LastToBundleEntryIndex)
-		{
-			const FExportBundleEntry& ToBundleEntry = ZenHeaderData.ExportBundleEntries[ToBundleEntryIndex++];
-			FPackageStorePackage::FExport& ToExport = Package->Exports[ToBundleEntry.LocalExportIndex];
-			FPackageStorePackage::FExportGraphNode* ToNode = ToExport.Nodes[ToBundleEntry.CommandType];
-			FPackageStorePackage::FExternalDependency& ExternalDependency = ToNode->ExternalDependencies.AddDefaulted_GetRef();
-			ExternalDependency.ImportIndex = ExternalArc.FromImportIndex;
-			ExternalDependency.ExportBundleCommandType = ExternalArc.FromCommandType;
-		}
-	}
-}
-
 TArray<FPackageStorePackage::FExportGraphNode*> FPackageStoreOptimizer::SortExportGraphNodesInLoadOrder(FPackageStorePackage* Package, FExportGraphEdges& Edges) const
 {
 	TRACE_CPUPROFILER_EVENT_SCOPE(SortExportGraphNodesInLoadOrder);
@@ -1080,7 +620,7 @@ TArray<FPackageStorePackage::FExportGraphNode*> FPackageStoreOptimizer::SortExpo
 	return LoadOrder;
 }
 
-void FPackageStoreOptimizer::CreateExportBundles(FPackageStorePackage* Package) const
+void FPackageStoreOptimizer::CreateExportBundle(FPackageStorePackage* Package) const
 {
 	TRACE_CPUPROFILER_EVENT_SCOPE(CreateExportBundles);
 	FExportGraphEdges Edges;
@@ -1092,80 +632,9 @@ void FPackageStoreOptimizer::CreateExportBundles(FPackageStorePackage* Package)
 		}
 	}
 	TArray<FPackageStorePackage::FExportGraphNode*> LoadOrder = SortExportGraphNodesInLoadOrder(Package, Edges);
-	FPackageStorePackage::FExportBundle* CurrentBundle = nullptr;
 	for (FPackageStorePackage::FExportGraphNode* Node : LoadOrder)
 	{
-		if (!CurrentBundle)
-		{
-			Package->CurrentBundle = &Package->GraphData.ExportBundles.AddDefaulted_GetRef();
-		}
-		Package->CurrentBundle->Entries.Add(Node->BundleEntry);
-		Node->ExportBundleIndex = Package->GraphData.ExportBundles.Num() - 1;
-		if (Node->bIsPublic)
-		{
-			CurrentBundle = nullptr;
-		}
-	}
-}
-
-void FPackageStoreOptimizer::SerializeGraphData(const TArray<FPackageStorePackage::FImportedPackageRef>& ImportedPackages, FPackageStorePackage::FGraphData& GraphData, FBufferWriter& GraphArchive) const
-{
-	uint32 ExportBundleEntryIndex = 0;
-	for (const FPackageStorePackage::FExportBundle& ExportBundle : GraphData.ExportBundles)
-	{
-		const uint32 EntryCount = ExportBundle.Entries.Num();
-		FExportBundleHeader ExportBundleHeader{ ExportBundle.SerialOffset, ExportBundleEntryIndex, EntryCount };
-		GraphArchive << ExportBundleHeader;
-		ExportBundleEntryIndex += EntryCount;
-	}
-	Algo::Sort(GraphData.InternalArcs, [](const FPackageStorePackage::FInternalArc& A, const FPackageStorePackage::FInternalArc& B)
-	{
-		if (A.ToExportBundleIndex == B.ToExportBundleIndex)
-		{
-			return A.FromExportBundleIndex < B.FromExportBundleIndex;
-		}
-		return A.ToExportBundleIndex < B.ToExportBundleIndex;
-	});
-	int32 InternalArcsCount = GraphData.InternalArcs.Num();
-	GraphArchive << InternalArcsCount;
-	for (FPackageStorePackage::FInternalArc& InternalArc : GraphData.InternalArcs)
-	{
-		GraphArchive << InternalArc.FromExportBundleIndex;
-		GraphArchive << InternalArc.ToExportBundleIndex;
-	}
-
-	for (const FPackageStorePackage::FImportedPackageRef& ImportedPackage : ImportedPackages)
-	{
-		TArray<FPackageStorePackage::FExternalArc>* FindArcsFromImportedPackage = GraphData.ExternalArcs.Find(ImportedPackage.Id);
-		if (!FindArcsFromImportedPackage)
-		{
-			int32 ExternalArcCount = 0;
-			GraphArchive << ExternalArcCount;
-		}
-		else
-		{
-			Algo::Sort(*FindArcsFromImportedPackage, [](const FPackageStorePackage::FExternalArc& A, const FPackageStorePackage::FExternalArc& B)
-			{
-				if (A.ToExportBundleIndex == B.ToExportBundleIndex)
-				{
-					if (A.FromImportIndex == B.FromImportIndex)
-					{
-						return A.FromCommandType < B.FromCommandType;
-					}
-					return A.FromImportIndex < B.FromImportIndex;
-				}
-				return A.ToExportBundleIndex < B.ToExportBundleIndex;
-			});
-			int32 ExternalArcCount = FindArcsFromImportedPackage->Num();
-			GraphArchive << ExternalArcCount;
-			for (FPackageStorePackage::FExternalArc& Arc : *FindArcsFromImportedPackage)
-			{
-				GraphArchive << Arc.FromImportIndex;
-				uint8 FromCommandType = uint8(Arc.FromCommandType);
-				GraphArchive << FromCommandType;
-				GraphArchive << Arc.ToExportBundleIndex;
-			}
-		}
+		Package->GraphData.ExportBundleEntries.Add(Node->BundleEntry);
 	}
 }
 
@@ -1189,7 +658,7 @@ void FPackageStoreOptimizer::FinalizePackageHeader(FPackageStorePackage* Package
 	for (const FPackageStorePackage::FExport& Export : Package->Exports)
 	{
 		FExportMapEntry ExportMapEntry;
-		ExportMapEntry.CookedSerialOffset = Export.CookedSerialOffset;
+		ExportMapEntry.CookedSerialOffset = Export.SerialOffset;
 		ExportMapEntry.CookedSerialSize = Export.SerialSize;
 		Package->NameMapBuilder.MarkNameAsReferenced(Export.ObjectName);
 		ExportMapEntry.ObjectName = Package->NameMapBuilder.MapName(Export.ObjectName);
@@ -1214,18 +683,25 @@ void FPackageStoreOptimizer::FinalizePackageHeader(FPackageStorePackage* Package
 	uint64 ExportMapSize = ExportMapArchive.Tell();
 
 	FBufferWriter ExportBundleEntriesArchive(nullptr, 0, EBufferWriterFlags::AllowResize | EBufferWriterFlags::TakeOwnership);
-	for (const FPackageStorePackage::FExportBundle& ExportBundle : Package->GraphData.ExportBundles)
+	for (FExportBundleEntry BundleEntry : Package->GraphData.ExportBundleEntries)
 	{
-		for (FExportBundleEntry BundleEntry : ExportBundle.Entries)
-		{
-			ExportBundleEntriesArchive << BundleEntry;
-		}
+		ExportBundleEntriesArchive << BundleEntry;
 	}
 	uint64 ExportBundleEntriesSize = ExportBundleEntriesArchive.Tell();
 
-	FBufferWriter GraphArchive(nullptr, 0, EBufferWriterFlags::AllowResize | EBufferWriterFlags::TakeOwnership);
-	SerializeGraphData(Package->ImportedPackages, Package->GraphData, GraphArchive);
-	uint64 GraphDataSize = GraphArchive.Tell();
+	FBufferWriter DependencyBundleHeadersArchive(nullptr, 0, EBufferWriterFlags::AllowResize | EBufferWriterFlags::TakeOwnership);
+	for (FDependencyBundleHeader DependencyBundleHeader : Package->GraphData.DependencyBundleHeaders)
+	{
+		DependencyBundleHeadersArchive << DependencyBundleHeader;
+	}
+
+	FBufferWriter DependencyBundleEntriesArchive(nullptr, 0, EBufferWriterFlags::AllowResize | EBufferWriterFlags::TakeOwnership);
+	for (FDependencyBundleEntry DependencyBundleEntry : Package->GraphData.DependencyBundleEntries)
+	{
+		DependencyBundleEntriesArchive << DependencyBundleEntry;
+	}
+
+	uint64 GraphDataSize = DependencyBundleHeadersArchive.Tell() + DependencyBundleEntriesArchive.Tell();
 
 	Package->NameMapBuilder.MarkNameAsReferenced(Package->Name);
 	FMappedName MappedPackageName = Package->NameMapBuilder.MapName(Package->Name);
@@ -1259,7 +735,7 @@ void FPackageStoreOptimizer::FinalizePackageHeader(FPackageStorePackage* Package
 	}
 	uint64 BulkDataMapSize = BulkDataMapAr.Tell();
 
-	Package->HeaderSize =
+	uint64 HeaderSize =
 		sizeof(FZenPackageSummary)
 		+ VersioningInfoSize
 		+ NameMapSize
@@ -1271,15 +747,15 @@ void FPackageStoreOptimizer::FinalizePackageHeader(FPackageStorePackage* Package
 		+ ImportedPackagesSize
 		+ BulkDataMapSize + sizeof(int64);
 
-	Package->HeaderBuffer = FIoBuffer(Package->HeaderSize);
+	Package->HeaderBuffer = FIoBuffer(HeaderSize);
 	uint8* HeaderData = Package->HeaderBuffer.Data();
-	FMemory::Memzero(HeaderData, Package->HeaderSize);
+	FMemory::Memzero(HeaderData, HeaderSize);
 	FZenPackageSummary* PackageSummary = reinterpret_cast<FZenPackageSummary*>(HeaderData);
-	PackageSummary->HeaderSize = Package->HeaderSize;
+	PackageSummary->HeaderSize = HeaderSize;
 	PackageSummary->Name = MappedPackageName;
 	PackageSummary->PackageFlags = Package->PackageFlags;
 	PackageSummary->CookedHeaderSize = Package->CookedHeaderSize;
-	FBufferWriter HeaderArchive(HeaderData, Package->HeaderSize);
+	FBufferWriter HeaderArchive(HeaderData, HeaderSize);
 	HeaderArchive.Seek(sizeof(FZenPackageSummary));
 
 	if (Package->VersioningInfo.IsSet())
@@ -1305,116 +781,22 @@ void FPackageStoreOptimizer::FinalizePackageHeader(FPackageStorePackage* Package
 	HeaderArchive.Serialize(ExportMapArchive.GetWriterData(), ExportMapArchive.Tell());
 	PackageSummary->ExportBundleEntriesOffset = HeaderArchive.Tell();
 	HeaderArchive.Serialize(ExportBundleEntriesArchive.GetWriterData(), ExportBundleEntriesArchive.Tell());
-	PackageSummary->GraphDataOffset = HeaderArchive.Tell();
-	HeaderArchive.Serialize(GraphArchive.GetWriterData(), GraphArchive.Tell());
+	PackageSummary->DependencyBundleHeadersOffset = HeaderArchive.Tell();
+	HeaderArchive.Serialize(DependencyBundleHeadersArchive.GetWriterData(), DependencyBundleHeadersArchive.Tell());
+	PackageSummary->DependencyBundleEntriesOffset = HeaderArchive.Tell();
+	HeaderArchive.Serialize(DependencyBundleEntriesArchive.GetWriterData(), DependencyBundleEntriesArchive.Tell());
 	PackageSummary->ImportedPackageNamesOffset = HeaderArchive.Tell();
 	HeaderArchive.Serialize(ImportedPackagesArchive.GetWriterData(), ImportedPackagesArchive.Tell());
 	check(HeaderArchive.Tell() == PackageSummary->HeaderSize)
 }
 
-void FPackageStoreOptimizer::FinalizePackage(FPackageStorePackage* Package) const
-{
-	for (FPackageStorePackage::FExportGraphNode& Node : Package->ExportGraphNodes)
-	{
-		check(Node.ExportBundleIndex >= 0);
-		TSet<FPackageStorePackage::FInternalArc> InternalArcsSet;
-		for (FPackageStorePackage::FExportGraphNode* InternalDependency : Node.InternalDependencies)
-		{
-			FPackageStorePackage::FInternalArc Arc;
-			check(InternalDependency->ExportBundleIndex >= 0);
-			Arc.FromExportBundleIndex = InternalDependency->ExportBundleIndex;
-			Arc.ToExportBundleIndex = Node.ExportBundleIndex;
-			if (Arc.FromExportBundleIndex != Arc.ToExportBundleIndex)
-			{
-				bool bIsAlreadyInSet = false;
-				InternalArcsSet.Add(Arc, &bIsAlreadyInSet);
-				if (!bIsAlreadyInSet)
-				{
-					Package->GraphData.InternalArcs.Add(Arc);
-				}
-			}
-		}
-
-		for (FPackageStorePackage::FExternalDependency& ExternalDependency : Node.ExternalDependencies)
-		{
-			check(ExternalDependency.ImportIndex >= 0);
-			const FPackageObjectIndex& Import = Package->Imports[ExternalDependency.ImportIndex];
-			check(Import.IsPackageImport());
-			FPackageImportReference PackageImportRef = Import.ToPackageImportRef();
-			TArray<FPackageStorePackage::FExternalArc>& ArcsFromImportedPackage = Package->GraphData.ExternalArcs.FindOrAdd(Package->ImportedPackages[PackageImportRef.GetImportedPackageIndex()].Id);
-			FPackageStorePackage::FExternalArc Arc;
-			Arc.FromImportIndex = ExternalDependency.ImportIndex;
-			Arc.FromCommandType = ExternalDependency.ExportBundleCommandType;
-			Arc.ToExportBundleIndex = Node.ExportBundleIndex;
-			if (!ArcsFromImportedPackage.Contains(Arc))
-			{
-				ArcsFromImportedPackage.Add(Arc);
-			}
-		}
-	}
-	uint64 SerialOffset = 0;
-	for (FPackageStorePackage::FExportBundle& ExportBundle : Package->GraphData.ExportBundles)
-	{
-		ExportBundle.SerialOffset = SerialOffset;
-		for (const FExportBundleEntry& BundleEntry : ExportBundle.Entries)
-		{
-			if (BundleEntry.CommandType == FExportBundleEntry::ExportCommandType_Serialize)
-			{
-				const FPackageStorePackage::FExport& Export = Package->Exports[BundleEntry.LocalExportIndex];
-				SerialOffset += Export.SerialSize;
-			}
-		}
-	}
-
-	FinalizePackageHeader(Package);
-}
-
-FIoBuffer FPackageStoreOptimizer::CreatePackageBuffer(const FPackageStorePackage* Package, const FIoBuffer& CookedExportsDataBuffer, TArray<FFileRegion>* InOutFileRegions) const
+FIoBuffer FPackageStoreOptimizer::CreatePackageBuffer(const FPackageStorePackage* Package, const FIoBuffer& CookedExportsDataBuffer) const
 {
 	check(Package->HeaderBuffer.DataSize() > 0);
-	check(Package->HeaderBuffer.DataSize() == Package->HeaderSize);
-	const uint64 BundleBufferSize = Package->HeaderSize + Package->ExportsSerialSize;
+	const uint64 BundleBufferSize = Package->HeaderBuffer.DataSize() + CookedExportsDataBuffer.DataSize();
 	FIoBuffer BundleBuffer(BundleBufferSize);
-	FMemory::Memcpy(BundleBuffer.Data(), Package->HeaderBuffer.Data(), Package->HeaderSize);
-	uint64 BundleBufferOffset = Package->HeaderSize;
-
-	TArray<FFileRegion> OutputRegions;
-
-	for (const FPackageStorePackage::FExportBundle& ExportBundle : Package->GraphData.ExportBundles)
-	{
-		for (const FExportBundleEntry& BundleEntry : ExportBundle.Entries)
-		{
-			if (BundleEntry.CommandType == FExportBundleEntry::ExportCommandType_Serialize)
-			{
-				const FPackageStorePackage::FExport& Export = Package->Exports[BundleEntry.LocalExportIndex];
-				check(Export.SerialOffset + Export.SerialSize <= CookedExportsDataBuffer.DataSize());
-				FMemory::Memcpy(BundleBuffer.Data() + BundleBufferOffset, CookedExportsDataBuffer.Data() + Export.SerialOffset, Export.SerialSize);
-
-				if (InOutFileRegions)
-				{
-					// Find overlapping regions and adjust them to match the new offset of the export data
-					for (const FFileRegion& Region : *InOutFileRegions)
-					{
-						if (Export.SerialOffset <= Region.Offset && Region.Offset + Region.Length <= Export.SerialOffset + Export.SerialSize)
-						{
-							FFileRegion NewRegion = Region;
-							NewRegion.Offset -= Export.SerialOffset;
-							NewRegion.Offset += BundleBufferOffset;
-							OutputRegions.Add(NewRegion);
-						}
-					}
-				}
-				BundleBufferOffset += Export.SerialSize;
-			}
-		}
-	}
-	check(BundleBufferOffset == BundleBuffer.DataSize());
-
-	if (InOutFileRegions)
-	{
-		*InOutFileRegions = OutputRegions;
-	}
-
+	FMemory::Memcpy(BundleBuffer.Data(), Package->HeaderBuffer.Data(), Package->HeaderBuffer.DataSize());
+	FMemory::Memcpy(BundleBuffer.Data() + Package->HeaderBuffer.DataSize(), CookedExportsDataBuffer.Data(), CookedExportsDataBuffer.DataSize());
 	return BundleBuffer;
 }
 
@@ -1589,8 +971,6 @@ FPackageStoreEntryResource FPackageStoreOptimizer::CreatePackageStoreEntry(const
 	
 	Result.PackageName = Package->Name;
 	Result.PackageId = FPackageId::FromName(Package->Name);
-	Result.ExportInfo.ExportCount = Package->Exports.Num();
-	Result.ExportInfo.ExportBundleCount = Package->GraphData.ExportBundles.Num();
 	Result.ImportedPackageIds.Reserve(Package->ImportedPackages.Num());
 	for (const FPackageStorePackage::FImportedPackageRef& ImportedPackage : Package->ImportedPackages)
 	{
@@ -1599,8 +979,6 @@ FPackageStoreEntryResource FPackageStoreOptimizer::CreatePackageStoreEntry(const
 	
 	if (OptionalSegmentPackage)
 	{
-		Result.OptionalSegmentExportInfo.ExportCount = OptionalSegmentPackage->Exports.Num();
-		Result.OptionalSegmentExportInfo.ExportBundleCount = OptionalSegmentPackage->GraphData.ExportBundles.Num();
 		Result.OptionalSegmentImportedPackageIds.Reserve(OptionalSegmentPackage->ImportedPackages.Num());
 		for (const FPackageStorePackage::FImportedPackageRef& ImportedPackage : OptionalSegmentPackage->ImportedPackages)
 		{
diff --git a/Engine/Source/Developer/IoStoreUtilities/Private/ZenStoreWriter.cpp b/Engine/Source/Developer/IoStoreUtilities/Private/ZenStoreWriter.cpp
--- a/Engine/Source/Developer/IoStoreUtilities/Private/ZenStoreWriter.cpp
+++ b/Engine/Source/Developer/IoStoreUtilities/Private/ZenStoreWriter.cpp
@@ -228,19 +228,14 @@ void FZenStoreWriter::WritePackageData(const FPackageInfo& Info, FLargeMemoryWri
 	FIoBuffer CookedHeaderBuffer = FIoBuffer(PackageData.Data(), Info.HeaderSize, PackageData);
 	FIoBuffer CookedExportsBuffer = FIoBuffer(PackageData.Data() + Info.HeaderSize, PackageData.DataSize() - Info.HeaderSize, PackageData);
 	Entry.OptimizedPackage.Reset(PackageStoreOptimizer->CreatePackageFromCookedHeader(Info.PackageName, CookedHeaderBuffer));
-	PackageStoreOptimizer->FinalizePackage(Entry.OptimizedPackage.Get());
 	Entry.FileRegions = FileRegions;
 	for (FFileRegion& Region : Entry.FileRegions)
 	{
-		// Adjust regions so they are relative to the start of the exports buffer
+		// Adjust regions so they are relative to the start of the export bundle buffer
 		Region.Offset -= Info.HeaderSize;
+		Region.Offset += Entry.OptimizedPackage->GetHeaderSize();
 	}
-	FIoBuffer PackageBuffer = PackageStoreOptimizer->CreatePackageBuffer(Entry.OptimizedPackage.Get(), CookedExportsBuffer, &Entry.FileRegions);
-	for (FFileRegion& Region : Entry.FileRegions)
-	{
-		// Adjust regions once more so they are relative to the start of the exports data again
-		Region.Offset -= Entry.OptimizedPackage->GetHeaderSize();
-	}
+	FIoBuffer PackageBuffer = PackageStoreOptimizer->CreatePackageBuffer(Entry.OptimizedPackage.Get(), CookedExportsBuffer);
 
 	// Commit to Zen build store
 
diff --git a/Engine/Source/Runtime/Core/Private/IO/PackageStore.cpp b/Engine/Source/Runtime/Core/Private/IO/PackageStore.cpp
--- a/Engine/Source/Runtime/Core/Private/IO/PackageStore.cpp
+++ b/Engine/Source/Runtime/Core/Private/IO/PackageStore.cpp
@@ -14,43 +14,13 @@
 #include "Serialization/CompactBinaryWriter.h"
 #include "Templates/Greater.h"
 
-FArchive& operator<<(FArchive& Ar, FPackageStoreExportInfo& ExportInfo)
-{
-	Ar << ExportInfo.ExportCount;
-	Ar << ExportInfo.ExportBundleCount;
-
-	return Ar;
-}
-
-FCbWriter& operator<<(FCbWriter& Writer, const FPackageStoreExportInfo& ExportInfo)
-{
-	Writer.BeginObject();
-	Writer << "exportcount" << ExportInfo.ExportCount;
-	Writer << "exportbundlecount" << ExportInfo.ExportBundleCount;
-	Writer.EndObject();
-	
-	return Writer;
-}
-
-FPackageStoreExportInfo FPackageStoreExportInfo::FromCbObject(const FCbObject& Obj)
-{
-	FPackageStoreExportInfo ExportInfo;
-	
-	ExportInfo.ExportCount			= Obj["exportcount"].AsInt32();
-	ExportInfo.ExportBundleCount	= Obj["exportbundlecount"].AsInt32();
-
-	return ExportInfo;
-}
-
 FArchive& operator<<(FArchive& Ar, FPackageStoreEntryResource& PackageStoreEntry)
 {
 	uint32 Flags = static_cast<uint32>(PackageStoreEntry.Flags);
 
 	Ar << Flags;
 	Ar << PackageStoreEntry.PackageName;
-	Ar << PackageStoreEntry.ExportInfo;
 	Ar << PackageStoreEntry.ImportedPackageIds;
-	Ar << PackageStoreEntry.OptionalSegmentExportInfo;
 	Ar << PackageStoreEntry.OptionalSegmentImportedPackageIds;
 
 	if (Ar.IsLoading())
@@ -68,7 +38,6 @@ FCbWriter& operator<<(FCbWriter& Writer, const FPackageStoreEntryResource& Packa
 
 	Writer << "flags" << static_cast<uint32>(PackageStoreEntry.Flags);
 	Writer << "packagename" << PackageStoreEntry.PackageName.ToString();
-	Writer << "exportinfo" << PackageStoreEntry.ExportInfo;
 
 	if (PackageStoreEntry.ImportedPackageIds.Num())
 	{
@@ -90,8 +59,6 @@ FCbWriter& operator<<(FCbWriter& Writer, const FPackageStoreEntryResource& Packa
 		Writer.EndArray();
 	}
 
-	Writer << "optionalsegmentexportinfo" << PackageStoreEntry.OptionalSegmentExportInfo;
-
 	if (PackageStoreEntry.OptionalSegmentImportedPackageIds.Num())
 	{
 		Writer.BeginArray("optionalsegmentimportedpackageids");
@@ -114,7 +81,6 @@ FPackageStoreEntryResource FPackageStoreEntryResource::FromCbObject(const FCbObj
 	Entry.Flags				= static_cast<EPackageStoreEntryFlags>(Obj["flags"].AsUInt32());
 	Entry.PackageName		= FName(Obj["packagename"].AsString());
 	Entry.PackageId			= FPackageId::FromName(Entry.PackageName);
-	Entry.ExportInfo		= FPackageStoreExportInfo::FromCbObject(Obj["exportinfo"].AsObject());
 	
 	if (Obj["importedpackageids"])
 	{
@@ -133,8 +99,6 @@ FPackageStoreEntryResource FPackageStoreEntryResource::FromCbObject(const FCbObj
 		}
 	}
 
-	Entry.OptionalSegmentExportInfo = FPackageStoreExportInfo::FromCbObject(Obj["optionalsegmentexportinfo"].AsObject());
-
 	if (Obj["optionalsegmentimportedpackageids"])
 	{
 		for (FCbFieldView ArrayField : Obj["optionalsegmentimportedpackageids"])
diff --git a/Engine/Source/Runtime/Core/Public/IO/IoContainerHeader.h b/Engine/Source/Runtime/Core/Public/IO/IoContainerHeader.h
--- a/Engine/Source/Runtime/Core/Public/IO/IoContainerHeader.h
+++ b/Engine/Source/Runtime/Core/Public/IO/IoContainerHeader.h
@@ -42,8 +42,6 @@ public:
  */
 struct FFilePackageStoreEntry
 {
-	int32 ExportCount;
-	int32 ExportBundleCount;
 	TFilePackageStoreEntryCArrayView<FPackageId> ImportedPackages;
 	TFilePackageStoreEntryCArrayView<FSHAHash> ShaderMapHashes;
 };
@@ -70,6 +68,7 @@ enum class EIoContainerHeaderVersion : uint32
 	Initial = 0,
 	LocalizedPackages = 1,
 	OptionalSegmentPackages = 2,
+	NoExportInfo = 3,
 
 	LatestPlusOne,
 	Latest = LatestPlusOne - 1
diff --git a/Engine/Source/Runtime/Core/Public/IO/PackageStore.h b/Engine/Source/Runtime/Core/Public/IO/PackageStore.h
--- a/Engine/Source/Runtime/Core/Public/IO/PackageStore.h
+++ b/Engine/Source/Runtime/Core/Public/IO/PackageStore.h
@@ -19,21 +19,6 @@ class FCbObject;
 class FCbWriter;
 class FStructuredArchiveSlot;
 
-/**
- * Package export information.
- */
-struct FPackageStoreExportInfo
-{
-	int32 ExportCount = 0;
-	int32 ExportBundleCount = 0;
-	
-	CORE_API friend FArchive& operator<<(FArchive& Ar, FPackageStoreExportInfo& ExportInfo);
-	
-	CORE_API friend FCbWriter& operator<<(FCbWriter& Writer, const FPackageStoreExportInfo& ExportInfo);
-	
-	CORE_API static FPackageStoreExportInfo FromCbObject(const FCbObject& Obj);
-};
-
 /**
  * Package store entry status.
  */
@@ -50,12 +35,11 @@ enum class EPackageStoreEntryStatus
  */ 
 struct FPackageStoreEntry
 {
-	FPackageStoreExportInfo ExportInfo;
 	TArrayView<const FPackageId> ImportedPackageIds;
 	TArrayView<const FSHAHash> ShaderMapHashes;
 #if WITH_EDITOR
-	FPackageStoreExportInfo OptionalSegmentExportInfo;
 	TArrayView<const FPackageId> OptionalSegmentImportedPackageIds;
+	bool bHasOptionalSegment = false;
 #endif
 };
 
@@ -84,14 +68,10 @@ struct FPackageStoreEntryResource
 	/** The package name. */
 	FName PackageName;
 	FPackageId PackageId;
-	/** The package export information. */
-	FPackageStoreExportInfo ExportInfo;
 	/** Imported package IDs. */
 	TArray<FPackageId> ImportedPackageIds;
 	/** Referenced shader map hashes. */
 	TArray<FSHAHash> ShaderMapHashes;
-	/** The editor data package export information. */
-	FPackageStoreExportInfo OptionalSegmentExportInfo;
 	/** Editor data imported package IDs. */
 	TArray<FPackageId> OptionalSegmentImportedPackageIds;
 
diff --git a/Engine/Source/Runtime/CoreUObject/Private/Serialization/AsyncLoading2.cpp b/Engine/Source/Runtime/CoreUObject/Private/Serialization/AsyncLoading2.cpp
--- a/Engine/Source/Runtime/CoreUObject/Private/Serialization/AsyncLoading2.cpp
+++ b/Engine/Source/Runtime/CoreUObject/Private/Serialization/AsyncLoading2.cpp
@@ -133,11 +133,23 @@ FArchive& operator<<(FArchive& Ar, FExportBundleEntry& ExportBundleEntry)
 	return Ar;
 }
 
-FArchive& operator<<(FArchive& Ar, FExportBundleHeader& ExportBundleHeader)
+FArchive& operator<<(FArchive& Ar, FDependencyBundleEntry& DependencyBundleEntry)
 {
-	Ar << ExportBundleHeader.SerialOffset;
-	Ar << ExportBundleHeader.FirstEntryIndex;
-	Ar << ExportBundleHeader.EntryCount;
+	Ar << DependencyBundleEntry.LocalImportOrExportIndex;
+
+	return Ar;
+}
+
+FArchive& operator<<(FArchive& Ar, FDependencyBundleHeader& DependencyBundleHeader)
+{
+	Ar << DependencyBundleHeader.FirstEntryIndex;
+	for (int32 I = 0; I < FExportBundleEntry::ExportCommandType_Count; ++I)
+	{
+		for (int32 J = 0; J < FExportBundleEntry::ExportCommandType_Count; ++J)
+		{
+			Ar << DependencyBundleHeader.EntryCount[I][J];
+		}
+	}
 
 	return Ar;
 }
@@ -1398,7 +1410,6 @@ struct FAsyncPackageHeaderData
 {
 	uint32 CookedHeaderSize = 0;
 	uint32 ExportCount = 0; // Need to keep this count around after ExportMap is cleared
-	uint32 ExportBundleCount = 0; // Need to keep this count around after ExportBundleHeaders is cleared
 	TOptional<FZenPackageVersioningInfo> VersioningInfo;
 	FNameMap NameMap;
 	FName PackageName;
@@ -1407,17 +1418,16 @@ struct FAsyncPackageHeaderData
 	TArrayView<const uint64> ImportedPublicExportHashes;
 	TArrayView<const FPackageObjectIndex> ImportMap;
 	TArrayView<const FExportMapEntry> ExportMap;
-	TArrayView<const uint8> ArcsData;
 	TArrayView<const FBulkDataMapEntry> BulkDataMap;
-	TArrayView<const FExportBundleHeader> ExportBundleHeaders;
 	TArrayView<const FExportBundleEntry> ExportBundleEntries;
+	TArrayView<const FDependencyBundleHeader> DependencyBundleHeaders;
+	TArrayView<const FDependencyBundleEntry> DependencyBundleEntries;
 
 	// Backed by allocation in FAsyncPackageData
 	TArrayView<FPackageId> ImportedPackageIds;
 	TArrayView<FAsyncPackage2*> ImportedAsyncPackagesView;
 	TArrayView<FExportObject> ExportsView;
-	TArrayView<FExportBundleHeader> ExportBundleHeadersCopyForPostLoad; // TODO: Can we use ConstructedObjects or Exports instead for postloading?
-	TArrayView<FExportBundleEntry> ExportBundleEntriesCopyForPostLoad;
+	TArrayView<FExportBundleEntry> ExportBundleEntriesCopyForPostLoad; // TODO: Can we use ConstructedObjects or Exports instead for posloading?
 
 	TArray<FName> ImportedPackageNames;
 
@@ -1427,10 +1437,10 @@ struct FAsyncPackageHeaderData
 		ImportedPublicExportHashes = TArrayView<const uint64>();
 		ImportMap = TArrayView<const FPackageObjectIndex>();
 		ExportMap = TArrayView<const FExportMapEntry>();
-		ArcsData = TArrayView<const uint8>();
 		BulkDataMap = TArrayView<const FBulkDataMapEntry>();
-		ExportBundleHeaders = TArrayView<const FExportBundleHeader>();
 		ExportBundleEntries = TArrayView<const FExportBundleEntry>();
+		DependencyBundleHeaders = TArrayView<const FDependencyBundleHeader>();
+		DependencyBundleEntries = TArrayView<const FDependencyBundleEntry>();
 	}
 };
 
@@ -1572,34 +1582,34 @@ public:
 class FExportArchive final : public FArchive
 {
 public:
-	FExportArchive(const uint8* AllExportDataPtr, const uint8* CurrentExportPtr, uint64 AllExportDataSize)
+	FExportArchive(const FIoBuffer& IoBuffer)
 		: IoDispatcher(FIoDispatcher::Get())
 	{
 #if (!DEVIRTUALIZE_FLinkerLoad_Serialize)
 		ActiveFPLB = &InlineFPLB;
 #endif
-		ActiveFPLB->OriginalFastPathLoadBuffer = AllExportDataPtr;
-		ActiveFPLB->StartFastPathLoadBuffer = CurrentExportPtr;
-		ActiveFPLB->EndFastPathLoadBuffer = AllExportDataPtr + AllExportDataSize;
+		ActiveFPLB->OriginalFastPathLoadBuffer = IoBuffer.Data();
+		ActiveFPLB->StartFastPathLoadBuffer = ActiveFPLB->OriginalFastPathLoadBuffer;
+		ActiveFPLB->EndFastPathLoadBuffer = IoBuffer.Data() + IoBuffer.DataSize();
 	}
 
 	~FExportArchive()
 	{
 	}
-	void ExportBufferBegin(UObject* Object, uint64 InExportCookedFileSerialOffset, uint64 InExportSerialSize)
+	void ExportBufferBegin(UObject* Object, uint64 InExportSerialOffset, uint64 InExportSerialSize)
 	{
 		CurrentExport = Object;
-		CookedSerialOffset = InExportCookedFileSerialOffset;
-		BufferSerialOffset = (ActiveFPLB->StartFastPathLoadBuffer - ActiveFPLB->OriginalFastPathLoadBuffer);
-		CookedSerialSize = InExportSerialSize;
+		ExportSerialOffset = HeaderData->PackageSummary->HeaderSize + InExportSerialOffset;
+		ExportSerialSize = InExportSerialSize;
+		ActiveFPLB->StartFastPathLoadBuffer = ActiveFPLB->OriginalFastPathLoadBuffer + ExportSerialOffset;
 	}
 
 	void ExportBufferEnd()
 	{
 		CurrentExport = nullptr;
-		CookedSerialOffset = 0;
-		BufferSerialOffset = 0;
-		CookedSerialSize = 0;
+		ExportSerialOffset = 0;
+		ExportSerialSize = 0;
+		ActiveFPLB->StartFastPathLoadBuffer = ActiveFPLB->OriginalFastPathLoadBuffer;
 	}
 
 	void CheckBufferPosition(const TCHAR* Text, uint64 Offset = 0)
@@ -1607,7 +1617,7 @@ public:
 #if DO_CHECK
 		const uint64 BufferPosition = (ActiveFPLB->StartFastPathLoadBuffer - ActiveFPLB->OriginalFastPathLoadBuffer) + Offset;
 		const bool bIsInsideExportBuffer =
-			(BufferSerialOffset <= BufferPosition) && (BufferPosition <= BufferSerialOffset + CookedSerialSize);
+			(ExportSerialOffset <= BufferPosition) && (BufferPosition <= ExportSerialOffset + ExportSerialSize);
 
 		UE_ASYNC_PACKAGE_CLOG(
 			!bIsInsideExportBuffer,
@@ -1615,7 +1625,7 @@ public:
 			TEXT("%s: Position %llu is outside of the current export buffer (%lld,%lld)."),
 			Text,
 			BufferPosition,
-			BufferSerialOffset, BufferSerialOffset + CookedSerialSize);
+			ExportSerialOffset, ExportSerialOffset + ExportSerialSize);
 #endif
 	}
 
@@ -1627,20 +1637,30 @@ public:
 
 	virtual int64 TotalSize() override
 	{
-		return HeaderData->CookedHeaderSize + (ActiveFPLB->EndFastPathLoadBuffer - ActiveFPLB->OriginalFastPathLoadBuffer);
+		if (bExportsCookedToSeparateArchive)
+		{
+			return ExportSerialSize;
+		}
+		else
+		{
+			int64 CookedFileSize = (ActiveFPLB->EndFastPathLoadBuffer - ActiveFPLB->OriginalFastPathLoadBuffer);
+			CookedFileSize -= HeaderData->PackageSummary->HeaderSize;
+			CookedFileSize += HeaderData->CookedHeaderSize;
+			return CookedFileSize;
+		}
 	}
 
 	virtual int64 Tell() override
 	{
 		if (bExportsCookedToSeparateArchive)
 		{
-			return (ActiveFPLB->StartFastPathLoadBuffer - ActiveFPLB->OriginalFastPathLoadBuffer) - BufferSerialOffset;
+			return (ActiveFPLB->StartFastPathLoadBuffer - ActiveFPLB->OriginalFastPathLoadBuffer) - ExportSerialOffset;
 		}
 		else
 		{
 			int64 CookedFilePosition = (ActiveFPLB->StartFastPathLoadBuffer - ActiveFPLB->OriginalFastPathLoadBuffer);
-			CookedFilePosition -= BufferSerialOffset;
-			CookedFilePosition += CookedSerialOffset;
+			CookedFilePosition -= HeaderData->PackageSummary->HeaderSize;
+			CookedFilePosition += HeaderData->CookedHeaderSize;
 			return CookedFilePosition;
 		}
 	}
@@ -1649,13 +1669,13 @@ public:
 	{
 		if (bExportsCookedToSeparateArchive)
 		{
-			ActiveFPLB->StartFastPathLoadBuffer = ActiveFPLB->OriginalFastPathLoadBuffer + BufferSerialOffset + Position;
+			ActiveFPLB->StartFastPathLoadBuffer = ActiveFPLB->OriginalFastPathLoadBuffer + ExportSerialOffset + Position;
 		}
 		else
 		{
 			uint64 BufferPosition = (uint64)Position;
-			BufferPosition -= CookedSerialOffset;
-			BufferPosition += BufferSerialOffset;
+			BufferPosition -= HeaderData->CookedHeaderSize;
+			BufferPosition += HeaderData->PackageSummary->HeaderSize;
 			ActiveFPLB->StartFastPathLoadBuffer = ActiveFPLB->OriginalFastPathLoadBuffer + BufferPosition;
 		}
 		CheckBufferPosition(TEXT("InvalidSeek"));
@@ -1831,7 +1851,7 @@ public:
 			if (int64 PayloadSize = Meta.GetSize(); PayloadSize > 0 && Meta.HasAnyFlags(BULKDATA_Unused) == false)
 			{
 				// Set the offset from the beginning of the I/O chunk in order to make inline bulk data reloadable
-				const int64 ExportBundleChunkOffset = (ActiveFPLB->StartFastPathLoadBuffer - ExportBundleChunkData);
+				const int64 ExportBundleChunkOffset = (ActiveFPLB->StartFastPathLoadBuffer - ActiveFPLB->OriginalFastPathLoadBuffer);
 				Meta.SetOffset(ExportBundleChunkOffset);
 				BulkData.BulkChunkId = CreateIoChunkId(PackageId.Value(), ChunkIndex, EIoChunkType::ExportBundleData);
 				Serialize(BulkData.ReallocateData(PayloadSize), PayloadSize);
@@ -1925,12 +1945,10 @@ private:
 	const FAsyncPackageHeaderData* HeaderData = nullptr;
 	const FLinkerInstancingContext* InstanceContext = nullptr;
 	UObject* CurrentExport = nullptr;
-	uint64 CookedSerialOffset = 0;
-	uint64 CookedSerialSize = 0;
-	uint64 BufferSerialOffset = 0;
+	uint64 ExportSerialOffset = 0;
+	uint64 ExportSerialSize = 0;
 	bool bIsOptionalSegment = false;
 	bool bExportsCookedToSeparateArchive = false;
-	const uint8* ExportBundleChunkData = nullptr;
 
 	void FixupSoftObjectPathForInstancedPackage(FSoftObjectPath& InOutSoftObjectPath)
 	{
@@ -2263,12 +2281,6 @@ enum EEventLoadNode2 : uint8
 	ExportBundle_NumPhases,
 };
 
-struct FAsyncPackageExportToBundleMapping
-{
-	uint64 ExportHash;
-	int32 BundleIndex[FExportBundleEntry::ExportCommandType_Count];
-};
-
 struct FAsyncPackageData
 {
 	uint8* MemoryBuffer0 = nullptr;
@@ -2277,21 +2289,16 @@ struct FAsyncPackageData
 	TArrayView<FAsyncPackage2*> ImportedAsyncPackages;
 	TArrayView<FEventLoadNode2> ExportBundleNodes;
 	TArrayView<const FSHAHash> ShaderMapHashes;
-	TArrayView<FAsyncPackageExportToBundleMapping> ExportToBundleMappings;
 	int32 TotalExportBundleCount = 0;
 };
 
 struct FAsyncPackageSerializationState
 {
 	FIoRequest IoRequest;
-	const uint8* AllExportDataPtr = nullptr;
-	const uint8* CurrentExportDataPtr = nullptr;
 
 	void ReleaseIoRequest()
 	{
 		IoRequest.Release();
-		AllExportDataPtr = nullptr;
-		CurrentExportDataPtr = nullptr;
 	}
 };
 
@@ -2806,7 +2813,7 @@ public:
 	static EEventLoadNodeExecutionResult Event_DeferredPostLoadExportBundle(FAsyncLoadingThreadState2& ThreadState, FAsyncPackage2* Package, int32 ExportBundleIndex);
 
 	void EventDrivenCreateExport(const FAsyncPackageHeaderData& Header, int32 LocalExportIndex);
-	bool EventDrivenSerializeExport(const FAsyncPackageHeaderData& Header, int32 LocalExportIndex, FExportArchive& Ar);
+	bool EventDrivenSerializeExport(const FAsyncPackageHeaderData& Header, int32 LocalExportIndex, FExportArchive* Ar);
 
 	UObject* EventDrivenIndexToObject(const FAsyncPackageHeaderData& Header, FPackageObjectIndex Index, bool bCheckSerialized);
 	template<class T>
@@ -2831,7 +2838,6 @@ private:
 	void InitializeExportArchive(FExportArchive& Ar, bool bIsOptionalSegment);
 	void CreatePackageNodes(const FAsyncLoadEventSpec* EventSpecs);
 	void CreateExportBundleNodes(const FAsyncLoadEventSpec* EventSpecs);
-	void SetupSerializedArcs(const FAsyncPackageHeaderData& Header, const TArrayView<FEventLoadNode2>& ExportBundleNodes, const TArrayView<FAsyncPackage2*>& ImportedAsyncPackages);
 	void SetupScriptDependencies();
 #if ALT2_ENABLE_PACKAGE_DEPENDENCY_DEBUGGING
 	bool HasDependencyToPackageDebug(FAsyncPackage2* Package);
@@ -2869,6 +2875,14 @@ private:
 	EEventLoadNodeExecutionResult PostLoadLinkerLoadPackageExports(FAsyncLoadingThreadState2& ThreadState);
 #endif
 
+
+	void ProcessExportDependencies(const FAsyncPackageHeaderData& Header, int32 LocalExportIndex, FExportBundleEntry::EExportCommandType CommandType);
+	int32 GetPublicExportIndex(uint64 ExportHash, FAsyncPackageHeaderData*& OutHeader);
+	UObject* ConditionalCreateExport(const FAsyncPackageHeaderData& Header, int32 LocalExportIndex);
+	UObject* ConditionalSerializeExport(const FAsyncPackageHeaderData& Header, int32 LocalExportIndex);
+	UObject* ConditionalCreateImport(const FAsyncPackageHeaderData& Header, int32 LocalImportIndex);
+	UObject* ConditionalSerializeImport(const FAsyncPackageHeaderData& Header, int32 LocalImportIndex);
+
 	/**
 	 * Begin async loading process. Simulates parts of BeginLoad.
 	 *
@@ -3444,7 +3458,7 @@ private:
 		const int32 ImportedPackagesCount = PackageStoreEntry.ImportedPackageIds.Num();
 		const uint64 ImportedPackageIdsMemSize = Align(sizeof(FPackageId) * ImportedPackagesCount, 8);
 #if WITH_EDITOR
-		const bool bHasOptionalSegment = PackageStoreEntry.OptionalSegmentExportInfo.ExportBundleCount > 0;
+		const bool bHasOptionalSegment = PackageStoreEntry.bHasOptionalSegment;
 		const int32 OptionalSegmentImportedPackagesCount = PackageStoreEntry.OptionalSegmentImportedPackageIds.Num();
 		const uint64 OptionalSegmentImportedPackageIdsMemSize = Align(sizeof(FPackageId) * OptionalSegmentImportedPackagesCount, 8);
 
@@ -3487,8 +3501,7 @@ private:
 		FMemory::Memcpy((void*)HeaderData.ImportedPackageIds.GetData(), PackageStoreEntry.ImportedPackageIds.GetData(), sizeof(FPackageId) * ImportedPackagesCount);
 		DataPtr += ImportedPackageIdsMemSize;
 
-		Data.TotalExportBundleCount = PackageStoreEntry.ExportInfo.ExportBundleCount;
-		HeaderData.ExportBundleCount = PackageStoreEntry.ExportInfo.ExportBundleCount;
+		Data.TotalExportBundleCount = 1;
 		HeaderData.ImportedAsyncPackagesView = Data.ImportedAsyncPackages;
 #if WITH_EDITOR
 		if (bHasOptionalSegment)
@@ -3499,8 +3512,7 @@ private:
 			FMemory::Memcpy((void*)OptionalSegmentHeaderData.ImportedPackageIds.GetData(), PackageStoreEntry.OptionalSegmentImportedPackageIds.GetData(), sizeof(FPackageId) * OptionalSegmentImportedPackagesCount);
 			DataPtr += OptionalSegmentImportedPackageIdsMemSize;
 
-			Data.TotalExportBundleCount += PackageStoreEntry.OptionalSegmentExportInfo.ExportBundleCount;
-			OptionalSegmentHeaderData.ExportBundleCount = PackageStoreEntry.OptionalSegmentExportInfo.ExportBundleCount;
+			++Data.TotalExportBundleCount;
 			HeaderData.ImportedAsyncPackagesView = Data.ImportedAsyncPackages.Left(ImportedPackagesCount);
 			OptionalSegmentHeaderData.ImportedAsyncPackagesView = Data.ImportedAsyncPackages.Right(OptionalSegmentImportedPackagesCount);
 		}
@@ -3527,34 +3539,27 @@ private:
 		FAsyncPackageData& Data = AsyncPackage->Data;
 		FAsyncPackageHeaderData& HeaderData = AsyncPackage->HeaderData;
 		int32 TotalExportCount = HeaderData.ExportMap.Num();
-		const uint64 ExportBundleHeadersCopyMemSize = Align(HeaderData.ExportBundleHeaders.Num() * sizeof(FExportBundleHeader), 8);
 		const uint64 ExportBundleEntriesCopyMemSize = Align(HeaderData.ExportBundleEntries.Num() * sizeof(FExportBundleEntry), 8);
 #if WITH_EDITOR
 		FAsyncPackageHeaderData* OptionalSegmentHeaderData = AsyncPackage->OptionalSegmentHeaderData.GetPtrOrNull();
-		uint64 OptionalSegmentExportBundleHeadersCopyMemSize = 0;
 		uint64 OptionalSegmentExportBundleEntriesCopyMemSize = 0;
 		if (OptionalSegmentHeaderData)
 		{
 			TotalExportCount += OptionalSegmentHeaderData->ExportMap.Num();
-			OptionalSegmentExportBundleHeadersCopyMemSize = Align(OptionalSegmentHeaderData->ExportBundleHeaders.Num() * sizeof(FExportBundleHeader), 8);
 			OptionalSegmentExportBundleEntriesCopyMemSize = Align(OptionalSegmentHeaderData->ExportBundleEntries.Num() * sizeof(FExportBundleEntry), 8);
 		}
 #endif
 		const int32 ExportBundleNodeCount = Data.TotalExportBundleCount * EEventLoadNode2::ExportBundle_NumPhases;
 		const uint64 ExportBundleNodesMemSize = Align(sizeof(FEventLoadNode2) * ExportBundleNodeCount, 8);
 		const uint64 ExportsMemSize = Align(sizeof(FExportObject) * TotalExportCount, 8);
-		const uint64 ExportToBundleMappingMemSize = Align(sizeof(FAsyncPackageExportToBundleMapping) * TotalExportCount, 8);
 
 		const uint64 MemoryBufferSize =
 			ExportBundleNodesMemSize +
 			ExportsMemSize +
 #if WITH_EDITOR
-			OptionalSegmentExportBundleHeadersCopyMemSize +
 			OptionalSegmentExportBundleEntriesCopyMemSize +
 #endif
-			ExportBundleHeadersCopyMemSize +
-			ExportBundleEntriesCopyMemSize +
-			ExportToBundleMappingMemSize;
+			ExportBundleEntriesCopyMemSize;
 
 #if PLATFORM_32BITS
 		if (MemoryBufferSize > MAX_Int32)
@@ -3571,11 +3576,6 @@ private:
 		DataPtr += ExportsMemSize;
 		Data.ExportBundleNodes = MakeArrayView(reinterpret_cast<FEventLoadNode2*>(DataPtr), ExportBundleNodeCount);
 		DataPtr += ExportBundleNodesMemSize;
-		Data.ExportToBundleMappings = MakeArrayView(reinterpret_cast<FAsyncPackageExportToBundleMapping*>(DataPtr), TotalExportCount);
-		DataPtr += ExportToBundleMappingMemSize;
-		HeaderData.ExportBundleHeadersCopyForPostLoad = MakeArrayView(reinterpret_cast<FExportBundleHeader*>(DataPtr), HeaderData.ExportBundleHeaders.Num());
-		FMemory::Memcpy(DataPtr, HeaderData.ExportBundleHeaders.GetData(), HeaderData.ExportBundleHeaders.Num() * sizeof(FExportBundleHeader));
-		DataPtr += ExportBundleHeadersCopyMemSize;
 		HeaderData.ExportBundleEntriesCopyForPostLoad = MakeArrayView(reinterpret_cast<FExportBundleEntry*>(DataPtr), HeaderData.ExportBundleEntries.Num());
 		FMemory::Memcpy(DataPtr, HeaderData.ExportBundleEntries.GetData(), HeaderData.ExportBundleEntries.Num() * sizeof(FExportBundleEntry));
 		DataPtr += ExportBundleEntriesCopyMemSize;
@@ -3585,9 +3585,6 @@ private:
 #if WITH_EDITOR
 		if (OptionalSegmentHeaderData)
 		{
-			OptionalSegmentHeaderData->ExportBundleHeadersCopyForPostLoad = MakeArrayView(reinterpret_cast<FExportBundleHeader*>(DataPtr), OptionalSegmentHeaderData->ExportBundleHeaders.Num());
-			FMemory::Memcpy(DataPtr, OptionalSegmentHeaderData->ExportBundleHeaders.GetData(), OptionalSegmentHeaderData->ExportBundleHeaders.Num() * sizeof(FExportBundleHeader));
-			DataPtr += OptionalSegmentExportBundleHeadersCopyMemSize;
 			OptionalSegmentHeaderData->ExportBundleEntriesCopyForPostLoad = MakeArrayView(reinterpret_cast<FExportBundleEntry*>(DataPtr), OptionalSegmentHeaderData->ExportBundleEntries.Num());
 			FMemory::Memcpy(DataPtr, OptionalSegmentHeaderData->ExportBundleEntries.GetData(), OptionalSegmentHeaderData->ExportBundleEntries.Num() * sizeof(FExportBundleEntry));
 			DataPtr += OptionalSegmentExportBundleEntriesCopyMemSize;
@@ -3606,39 +3603,6 @@ private:
 			Export = FExportObject();
 		}
 
-		// Map publc export hashes to export bundles
-		if (!AsyncPackage->bLoadHasFailed
-#if ALT2_ENABLE_LINKERLOAD_SUPPORT
-			&& !AsyncPackage->LinkerLoadState.IsSet()
-#endif
-			)
-		{
-			for (int32 ExportBundleIndex = 0; ExportBundleIndex < AsyncPackage->Data.TotalExportBundleCount; ++ExportBundleIndex)
-			{
-				const FAsyncPackageHeaderData* TargetHeaderData = &AsyncPackage->HeaderData;
-				int32 LocalExportBundleIndex = ExportBundleIndex;
-				int32 ExportIndexOffset = 0;
-#if WITH_EDITOR
-				if (ExportBundleIndex >= AsyncPackage->HeaderData.ExportBundleHeaders.Num())
-				{
-					check(OptionalSegmentHeaderData);
-					TargetHeaderData = OptionalSegmentHeaderData;
-					LocalExportBundleIndex -= AsyncPackage->HeaderData.ExportBundleHeaders.Num();
-					ExportIndexOffset = AsyncPackage->HeaderData.ExportMap.Num();
-				}
-#endif
-				const FExportBundleHeader& ExportBundle = TargetHeaderData->ExportBundleHeaders[LocalExportBundleIndex];
-				for (int32 EntryIndex = ExportBundle.FirstEntryIndex, EntryEnd = ExportBundle.FirstEntryIndex + ExportBundle.EntryCount; EntryIndex < EntryEnd; ++EntryIndex)
-				{
-					const FExportBundleEntry& BundleEntry = TargetHeaderData->ExportBundleEntries[EntryIndex];
-					const FExportMapEntry& ExportMapEntry = TargetHeaderData->ExportMap[BundleEntry.LocalExportIndex];
-					FAsyncPackageExportToBundleMapping& ExportToBundleMapping = AsyncPackage->Data.ExportToBundleMappings[BundleEntry.LocalExportIndex + ExportIndexOffset];
-					ExportToBundleMapping.ExportHash = ExportMapEntry.PublicExportHash;
-					ExportToBundleMapping.BundleIndex[BundleEntry.CommandType] = ExportBundleIndex;
-				}
-			}
-		}
-
 		if (!AsyncPackage->bHasStartedImportingPackages)
 		{
 			FIoBatch IoBatch = IoDispatcher.NewBatch();
@@ -3997,7 +3961,6 @@ bool FAsyncLoadingThread2::CreateAsyncPackagesFromQueue(FAsyncLoadingThreadState
 		{
 			// Initialize package with a fake package store entry
 			FPackageStoreEntry FakePackageEntry;
-			FakePackageEntry.ExportInfo.ExportBundleCount = 1;
 			InitializeAsyncPackageFromPackageStore(ThreadState, &IoBatch, PendingPackage, FakePackageEntry);
 			// Simulate StartLoading() getting back a failed IoRequest and let it go through all package states
 			PendingPackage->AsyncPackageLoadingState = EAsyncPackageLoadingState2::WaitingForIo;
@@ -4507,87 +4470,6 @@ FUObjectSerializeContext* FAsyncPackage2::GetSerializeContext()
 	return FUObjectThreadContext::Get().GetSerializeContext();
 }
 
-void FAsyncPackage2::SetupSerializedArcs(const FAsyncPackageHeaderData& Header, const TArrayView<FEventLoadNode2>& ExportBundleNodes, const TArrayView<FAsyncPackage2*>& ImportedAsyncPackages)
-{
-	TRACE_CPUPROFILER_EVENT_SCOPE(SetupExternalArcs);
-
-	FSimpleArchive ArcsArchive(Header.ArcsData.GetData(), Header.ArcsData.Num());
-	int32 InternalArcsCount;
-	ArcsArchive << InternalArcsCount;
-	for (int32 InternalArcIndex = 0; InternalArcIndex < InternalArcsCount; ++InternalArcIndex)
-	{
-		int32 FromExportBundleIndex;
-		ArcsArchive << FromExportBundleIndex;
-		int32 ToExportBundleIndex;
-		ArcsArchive << ToExportBundleIndex;
-		if (FromExportBundleIndex >= Data.TotalExportBundleCount || ToExportBundleIndex >= Data.TotalExportBundleCount)
-		{
-			UE_LOG(LogStreaming, Fatal, TEXT("Corrupt package data when serializing arcs"));
-			return;
-		}
-		uint32 FromNodeIndexBase = FromExportBundleIndex * EEventLoadNode2::ExportBundle_NumPhases;
-		uint32 ToNodeIndexBase = ToExportBundleIndex * EEventLoadNode2::ExportBundle_NumPhases;
-		for (int32 Phase = 0; Phase < EEventLoadNode2::ExportBundle_NumPhases; ++Phase)
-		{
-			uint32 ToNodeIndex = ToNodeIndexBase + Phase;
-			uint32 FromNodeIndex = FromNodeIndexBase + Phase;
-			ExportBundleNodes[ToNodeIndex].DependsOn(&ExportBundleNodes[FromNodeIndex]);
-		}
-	}
-	for (const FAsyncPackage2* ImportedPackage : ImportedAsyncPackages)
-	{
-		int32 ExternalArcCount;
-		ArcsArchive << ExternalArcCount;
-
-		int32 PreviousFromExportBundleIndex = -1;
-		int32 PreviousToExportBundleIndex = -1;
-
-		for (int32 ExternalArcIndex = 0; ExternalArcIndex < ExternalArcCount; ++ExternalArcIndex)
-		{
-			int32 FromImportIndex;
-			uint8 FromCommandType;
-			int32 ToExportBundleIndex;
-			ArcsArchive << FromImportIndex;
-			ArcsArchive << FromCommandType;
-			ArcsArchive << ToExportBundleIndex;
-			if (ImportedPackage)
-			{
-				if (FromImportIndex >= Header.ImportMap.Num() ||
-					FromCommandType >= FExportBundleEntry::ExportCommandType_Count ||
-					ToExportBundleIndex >= Data.TotalExportBundleCount)
-				{
-					UE_LOG(LogStreaming, Fatal, TEXT("Corrupt package data when serializing arcs"));
-					return;
-				}
-				FPackageObjectIndex GlobalImportIndex = Header.ImportMap[FromImportIndex];
-				FPackageImportReference PackageImportRef = GlobalImportIndex.ToPackageImportRef();
-				const uint64 ImportedPublicExportHash = Header.ImportedPublicExportHashes[PackageImportRef.GetImportedPublicExportHashIndex()];
-				for (const FAsyncPackageExportToBundleMapping& ExportToBundleMapping : ImportedPackage->Data.ExportToBundleMappings)
-				{
-					if (ExportToBundleMapping.ExportHash == ImportedPublicExportHash)
-					{
-						int32 FromExportBundleIndex = ExportToBundleMapping.BundleIndex[FromCommandType];
-						if (PreviousFromExportBundleIndex != FromExportBundleIndex || PreviousToExportBundleIndex != ToExportBundleIndex)
-						{
-							PreviousFromExportBundleIndex = FromExportBundleIndex;
-							PreviousToExportBundleIndex = ToExportBundleIndex;
-							uint32 FromNodeIndexBase = FromExportBundleIndex * EEventLoadNode2::ExportBundle_NumPhases;
-							uint32 ToNodeIndexBase = ToExportBundleIndex * EEventLoadNode2::ExportBundle_NumPhases;
-							for (int32 Phase = 0; Phase < EEventLoadNode2::ExportBundle_NumPhases; ++Phase)
-							{
-								uint32 ToNodeIndex = ToNodeIndexBase + Phase;
-								uint32 FromNodeIndex = FromNodeIndexBase + Phase;
-								ExportBundleNodes[ToNodeIndex].DependsOn(&ImportedPackage->Data.ExportBundleNodes[FromNodeIndex]);
-							}
-						}
-						break;
-					}
-				}
-			}
-		}
-	}
-}
-
 void FAsyncPackage2::SetupScriptDependencies()
 {
 	TRACE_CPUPROFILER_EVENT_SCOPE(SetupScriptDependencies);
@@ -5158,8 +5040,7 @@ static void ReadAsyncPackageHeader(FAsyncPackageSerializationState& Serializatio
 		(PackageSummary->ExportBundleEntriesOffset - PackageSummary->ExportMapOffset) / sizeof(FExportMapEntry));
 	HeaderData.ExportCount = HeaderData.ExportMap.Num();
 
-	const uint64 ExportBundleHeadersOffset = PackageSummary->GraphDataOffset;
-	const uint64 ExportBundleEntriesSize = ExportBundleHeadersOffset - PackageSummary->ExportBundleEntriesOffset;
+	const uint64 ExportBundleEntriesSize = PackageSummary->DependencyBundleHeadersOffset - PackageSummary->ExportBundleEntriesOffset;
 	const int32 ExportBundleEntriesCount = ExportBundleEntriesSize / sizeof(FExportBundleEntry);
 
 	if (ExportBundleEntriesCount != HeaderData.ExportCount * FExportBundleEntry::ExportCommandType_Count)
@@ -5168,17 +5049,15 @@ static void ReadAsyncPackageHeader(FAsyncPackageSerializationState& Serializatio
 		return;
 	}
 
-	HeaderData.ExportBundleHeaders = TArrayView<const FExportBundleHeader>(
-		reinterpret_cast<const FExportBundleHeader*>(PackageHeaderDataPtr + ExportBundleHeadersOffset),
-		HeaderData.ExportBundleCount);
 	HeaderData.ExportBundleEntries = TArrayView<const FExportBundleEntry>(
 		reinterpret_cast<const FExportBundleEntry*>(PackageHeaderDataPtr + PackageSummary->ExportBundleEntriesOffset),
 		ExportBundleEntriesCount);
-
-	const uint64 ExportBundleHeadersSize = sizeof(FExportBundleHeader) * HeaderData.ExportBundleHeaders.Num();
-	const uint64 ArcsDataOffset = ExportBundleHeadersOffset + ExportBundleHeadersSize;
-	const uint64 ArcsDataSize = PackageSummary->ImportedPackageNamesOffset - ArcsDataOffset;
-	HeaderData.ArcsData = TArrayView64<const uint8>(PackageHeaderDataPtr + ArcsDataOffset, ArcsDataSize);
+	HeaderData.DependencyBundleHeaders = TArrayView<const FDependencyBundleHeader>(
+		reinterpret_cast<const FDependencyBundleHeader*>(PackageHeaderDataPtr + PackageSummary->DependencyBundleHeadersOffset),
+		(PackageSummary->DependencyBundleEntriesOffset - PackageSummary->DependencyBundleHeadersOffset) / sizeof(FDependencyBundleHeader));
+	HeaderData.DependencyBundleEntries = TArrayView<const FDependencyBundleEntry>(
+		reinterpret_cast<const FDependencyBundleEntry*>(PackageHeaderDataPtr + PackageSummary->DependencyBundleEntriesOffset),
+		(PackageSummary->ImportedPackageNamesOffset - PackageSummary->DependencyBundleEntriesOffset) / sizeof(FDependencyBundleEntry));
 
 #if WITH_EDITOR
 	// Only read the imported package names in editor builds where we need to interact with uncooked packages
@@ -5188,8 +5067,6 @@ static void ReadAsyncPackageHeader(FAsyncPackageSerializationState& Serializatio
 	ImportedPackageNamesDataReader << Container;
 	HeaderData.ImportedPackageNames = MoveTemp(Container.Names);
 #endif
-
-	SerializationState.AllExportDataPtr = PackageHeaderDataPtr + PackageSummary->HeaderSize;
 }
 
 #if ALT2_ENABLE_LINKERLOAD_SUPPORT
@@ -5328,7 +5205,6 @@ EEventLoadNodeExecutionResult FAsyncPackage2::ProcessLinkerLoadPackageSummary(FA
 	}
 
 	FPackageStoreEntry PackageStoreEntry;
-	PackageStoreEntry.ExportInfo.ExportBundleCount = 1;
 	PackageStoreEntry.ImportedPackageIds = ImportedPackageIds;
 	AsyncLoadingThread.InitializeAsyncPackageFromPackageStore(ThreadState, nullptr, this, PackageStoreEntry);
 
@@ -5738,13 +5614,8 @@ void FAsyncPackage2::InitializeExportArchive(FExportArchive& Ar, bool bIsOptiona
 	{
 		Ar.HeaderData = OptionalSegmentHeaderData.GetPtrOrNull();
 		check(Ar.HeaderData);
-		Ar.ExportBundleChunkData = OptionalSegmentSerializationState->IoRequest.GetResultOrDie().Data();
 	}
-	else
 #endif
-	{
-		Ar.ExportBundleChunkData = SerializationState.IoRequest.GetResultOrDie().Data();
-	}
 	Ar.ImportStore = &ImportStore;
 	Ar.ExternalReadDependencies = &ExternalReadDependencies;
 	Ar.InstanceContext = InstanceContext.Get();
@@ -5774,41 +5645,33 @@ EEventLoadNodeExecutionResult FAsyncPackage2::Event_ProcessExportBundle(FAsyncLo
 		bool bIsOptionalSegment = false;
 #if WITH_EDITOR
 		const FAsyncPackageHeaderData* HeaderData;
-		const FExportBundleHeader* ExportBundle;
 		FAsyncPackageSerializationState* SerializationState;
-		bIsOptionalSegment = InExportBundleIndex >= Package->HeaderData.ExportBundleHeaders.Num();
+		bIsOptionalSegment = InExportBundleIndex == 1;
 
 		if (bIsOptionalSegment)
 		{
 			HeaderData = Package->OptionalSegmentHeaderData.GetPtrOrNull();
 			check(HeaderData);
-			ExportBundle = &HeaderData->ExportBundleHeaders[InExportBundleIndex - Package->HeaderData.ExportBundleHeaders.Num()];
 			SerializationState = Package->OptionalSegmentSerializationState.GetPtrOrNull();
 			check(SerializationState);
 		}
 		else
 		{
+			check(InExportBundleIndex == 0);
 			HeaderData = &Package->HeaderData;
-			ExportBundle = &Package->HeaderData.ExportBundleHeaders[InExportBundleIndex];
 			SerializationState = &Package->SerializationState;
 		}
 #else
 		const FAsyncPackageHeaderData* HeaderData = &Package->HeaderData;
-		const FExportBundleHeader* ExportBundle = &Package->HeaderData.ExportBundleHeaders[InExportBundleIndex];
 		FAsyncPackageSerializationState* SerializationState = &Package->SerializationState;
 #endif
 		const FIoBuffer& IoBuffer = SerializationState->IoRequest.GetResultOrDie();
-		const uint64 AllExportDataSize = IoBuffer.DataSize() - (SerializationState->AllExportDataPtr - IoBuffer.Data());
-		if (Package->ExportBundleEntryIndex == 0)
-		{
-			SerializationState->CurrentExportDataPtr = SerializationState->AllExportDataPtr + ExportBundle->SerialOffset;
-		}
-		FExportArchive Ar(SerializationState->AllExportDataPtr, SerializationState->CurrentExportDataPtr, AllExportDataSize);
+		FExportArchive Ar(IoBuffer);
 		Package->InitializeExportArchive(Ar, bIsOptionalSegment);
 
-		while (Package->ExportBundleEntryIndex < int32(ExportBundle->EntryCount))
+		while (Package->ExportBundleEntryIndex < HeaderData->ExportBundleEntries.Num())
 		{
-			const FExportBundleEntry& BundleEntry = HeaderData->ExportBundleEntries[ExportBundle->FirstEntryIndex + Package->ExportBundleEntryIndex];
+			const FExportBundleEntry& BundleEntry = HeaderData->ExportBundleEntries[Package->ExportBundleEntryIndex];
 			if (ThreadState.IsTimeLimitExceeded(TEXT("Event_ProcessExportBundle")))
 			{
 				return EEventLoadNodeExecutionResult::Timeout;
@@ -5818,41 +5681,18 @@ EEventLoadNodeExecutionResult FAsyncPackage2::Event_ProcessExportBundle(FAsyncLo
 
 			if (BundleEntry.CommandType == FExportBundleEntry::ExportCommandType_Create)
 			{
-				Package->EventDrivenCreateExport(*HeaderData, BundleEntry.LocalExportIndex);
+				if (!Export.Object)
+				{
+					Package->EventDrivenCreateExport(*HeaderData, BundleEntry.LocalExportIndex);
+				}
 			}
 			else
 			{
 				check(BundleEntry.CommandType == FExportBundleEntry::ExportCommandType_Serialize);
-
-				const uint64 CookedSerialSize = ExportMapEntry.CookedSerialSize;
-				UObject* Object = Export.Object;
-
-				check(SerializationState->CurrentExportDataPtr + CookedSerialSize <= IoBuffer.Data() + IoBuffer.DataSize());
-				check(Object || Export.bFiltered || Export.bExportLoadFailed);
-
-				Ar.ExportBufferBegin(Object, ExportMapEntry.CookedSerialOffset, ExportMapEntry.CookedSerialSize);
-
-				const int64 Pos = Ar.Tell();
-				UE_ASYNC_PACKAGE_CLOG(
-					CookedSerialSize > uint64(Ar.TotalSize() - Pos), Fatal, Package->Desc, TEXT("ObjectSerializationError"),
-					TEXT("%s: Serial size mismatch: Expected read size %d, Remaining archive size: %d"),
-					Object ? *Object->GetFullName() : TEXT("null"), CookedSerialSize, uint64(Ar.TotalSize() - Pos));
-
-				const bool bSerialized = Package->EventDrivenSerializeExport(*HeaderData, BundleEntry.LocalExportIndex, Ar);
-				if (!bSerialized)
+				if (Export.Object && Export.Object->HasAllFlags(RF_NeedLoad))
 				{
-					Ar.Skip(CookedSerialSize);
+					Package->EventDrivenSerializeExport(*HeaderData, BundleEntry.LocalExportIndex, &Ar);
 				}
-				UE_ASYNC_PACKAGE_CLOG(
-					CookedSerialSize != uint64(Ar.Tell() - Pos), Fatal, Package->Desc, TEXT("ObjectSerializationError"),
-					TEXT("%s: Serial size mismatch: Expected read size %d, Actual read size %d"),
-					Object ? *Object->GetFullName() : TEXT("null"), CookedSerialSize, uint64(Ar.Tell() - Pos));
-
-				Ar.ExportBufferEnd();
-
-				check((Object && !Object->HasAnyFlags(RF_NeedLoad)) || Export.bFiltered || Export.bExportLoadFailed);
-
-				SerializationState->CurrentExportDataPtr += CookedSerialSize;
 			}
 			++Package->ExportBundleEntryIndex;
 		}
@@ -5931,6 +5771,165 @@ UObject* FAsyncPackage2::EventDrivenIndexToObject(const FAsyncPackageHeaderData&
 	return Result;
 }
 
+void FAsyncPackage2::ProcessExportDependencies(const FAsyncPackageHeaderData& Header, int32 LocalExportIndex, FExportBundleEntry::EExportCommandType CommandType)
+{
+	static_assert(FExportBundleEntry::ExportCommandType_Count == 2, "Expected the only export command types fo be Create and Serialize");
+
+	TRACE_CPUPROFILER_EVENT_SCOPE(ProcessExportDependencies);
+	const FDependencyBundleHeader& DependencyBundle = Header.DependencyBundleHeaders[LocalExportIndex];
+	if (DependencyBundle.FirstEntryIndex < 0)
+	{
+		return;
+	}
+	int32 RunningIndex = DependencyBundle.FirstEntryIndex;
+	if (CommandType == FExportBundleEntry::ExportCommandType_Serialize)
+	{
+		// Skip over the dependency entries for Create
+		for (int32 Index = 0; Index < FExportBundleEntry::ExportCommandType_Count; ++Index)
+		{
+			RunningIndex += DependencyBundle.EntryCount[FExportBundleEntry::ExportCommandType_Create][Index];
+		}
+	}
+	
+	for (int32 Index = DependencyBundle.EntryCount[CommandType][FExportBundleEntry::ExportCommandType_Create]; Index > 0; --Index)
+	{
+		const FDependencyBundleEntry& Dep = Header.DependencyBundleEntries[RunningIndex++];
+		if (Dep.LocalImportOrExportIndex.IsExport())
+		{
+			ConditionalCreateExport(Header, Dep.LocalImportOrExportIndex.ToExport());
+		}
+		else
+		{
+			ConditionalCreateImport(Header, Dep.LocalImportOrExportIndex.ToImport());
+		}
+	}
+
+	for (int32 Index = DependencyBundle.EntryCount[CommandType][FExportBundleEntry::ExportCommandType_Serialize]; Index > 0; Index--)
+	{
+		const FDependencyBundleEntry& Dep = Header.DependencyBundleEntries[RunningIndex++];
+		if (Dep.LocalImportOrExportIndex.IsExport())
+		{
+			ConditionalSerializeExport(Header, Dep.LocalImportOrExportIndex.ToExport());
+		}
+		else
+		{
+			ConditionalSerializeImport(Header, Dep.LocalImportOrExportIndex.ToImport());
+		}
+	}
+}
+
+int32 FAsyncPackage2::GetPublicExportIndex(uint64 ExportHash, FAsyncPackageHeaderData*& OutHeader)
+{
+	for (int32 ExportIndex = 0; ExportIndex < HeaderData.ExportMap.Num(); ++ExportIndex)
+	{
+		if (HeaderData.ExportMap[ExportIndex].PublicExportHash == ExportHash)
+		{
+			OutHeader = &HeaderData;
+			return ExportIndex;
+		}
+	}
+#if WITH_EDITOR
+	if (FAsyncPackageHeaderData* OptionalSegmentHeaderDataPtr = OptionalSegmentHeaderData.GetPtrOrNull())
+	{
+		for (int32 ExportIndex = 0; ExportIndex < OptionalSegmentHeaderDataPtr->ExportMap.Num(); ++ExportIndex)
+		{
+			if (OptionalSegmentHeaderDataPtr->ExportMap[ExportIndex].PublicExportHash == ExportHash)
+			{
+				OutHeader = OptionalSegmentHeaderDataPtr;
+				return HeaderData.ExportMap.Num() + ExportIndex;
+			}
+		}
+	}
+#endif
+	return -1;
+}
+UObject* FAsyncPackage2::ConditionalCreateExport(const FAsyncPackageHeaderData& Header, int32 LocalExportIndex)
+{
+	if (!Header.ExportsView[LocalExportIndex].Object)
+	{
+		FAsyncPackageScope2 Scope(this);
+		EventDrivenCreateExport(Header, LocalExportIndex);
+	}
+	return Header.ExportsView[LocalExportIndex].Object;
+}
+
+UObject* FAsyncPackage2::ConditionalSerializeExport(const FAsyncPackageHeaderData& Header, int32 LocalExportIndex)
+{
+	FExportObject& Export = Header.ExportsView[LocalExportIndex];
+	
+	if (!Export.Object && !(Export.bFiltered | Export.bExportLoadFailed))
+	{
+		ConditionalCreateExport(Header, LocalExportIndex);
+	}
+	
+	if (!Export.Object || (Export.bFiltered | Export.bExportLoadFailed))
+	{
+		return nullptr;
+	}
+
+	if (Export.Object->HasAllFlags(RF_NeedLoad))
+	{
+		FAsyncPackageScope2 Scope(this);
+		EventDrivenSerializeExport(Header, LocalExportIndex, nullptr);
+	}
+
+	return Export.Object;
+}
+
+UObject* FAsyncPackage2::ConditionalCreateImport(const FAsyncPackageHeaderData& Header, int32 LocalImportIndex)
+{
+	const FPackageObjectIndex& ObjectIndex = Header.ImportMap[LocalImportIndex];
+	check(ObjectIndex.IsPackageImport());
+	if (UObject* FromImportStore = ImportStore.FindOrGetImportObject(Header, ObjectIndex))
+	{
+		return FromImportStore;
+	}
+
+	FPackageImportReference PackageImportRef = ObjectIndex.ToPackageImportRef();
+	FAsyncPackage2* SourcePackage = Header.ImportedAsyncPackagesView[PackageImportRef.GetImportedPackageIndex()];
+	if (!SourcePackage)
+	{
+		return nullptr;
+	}
+	uint64 ExportHash = Header.ImportedPublicExportHashes[PackageImportRef.GetImportedPublicExportHashIndex()];
+	FAsyncPackageHeaderData* SourcePackageHeader = nullptr;
+	int32 ExportIndex = SourcePackage->GetPublicExportIndex(ExportHash, SourcePackageHeader);
+	if (ExportIndex < 0)
+	{
+		return nullptr;
+	}
+	return SourcePackage->ConditionalCreateExport(*SourcePackageHeader, ExportIndex);
+}
+
+UObject* FAsyncPackage2::ConditionalSerializeImport(const FAsyncPackageHeaderData& Header, int32 LocalImportIndex)
+{
+	const FPackageObjectIndex& ObjectIndex = Header.ImportMap[LocalImportIndex];
+	check(ObjectIndex.IsPackageImport());
+
+	if (UObject* FromImportStore = ImportStore.FindOrGetImportObject(Header, ObjectIndex))
+	{
+		if (!FromImportStore->HasAllFlags(RF_NeedLoad))
+		{
+			return FromImportStore;
+		}
+	}
+
+	FPackageImportReference PackageImportRef = ObjectIndex.ToPackageImportRef();
+	FAsyncPackage2* SourcePackage = Header.ImportedAsyncPackagesView[PackageImportRef.GetImportedPackageIndex()];
+	if (!SourcePackage)
+	{
+		return nullptr;
+	}
+	uint64 ExportHash = Header.ImportedPublicExportHashes[PackageImportRef.GetImportedPublicExportHashIndex()];
+	FAsyncPackageHeaderData* SourcePackageHeader = nullptr;
+	int32 ExportIndex = SourcePackage->GetPublicExportIndex(ExportHash, SourcePackageHeader);
+	if (ExportIndex < 0)
+	{
+		return nullptr;
+	}
+	return SourcePackage->ConditionalSerializeExport(*SourcePackageHeader, ExportIndex);
+}
+
 void FAsyncPackage2::EventDrivenCreateExport(const FAsyncPackageHeaderData& Header, int32 LocalExportIndex)
 {
 	TRACE_CPUPROFILER_EVENT_SCOPE(CreateExport);
@@ -5988,6 +5987,8 @@ void FAsyncPackage2::EventDrivenCreateExport(const FAsyncPackageHeaderData& Head
 		return;
 	}
 
+	ProcessExportDependencies(Header, LocalExportIndex, FExportBundleEntry::ExportCommandType_Create);
+
 	bool bIsCompleteyLoaded = false;
 	UClass* LoadClass = Export.ClassIndex.IsNull() ? UClass::StaticClass() : CastEventDrivenIndexToObject<UClass>(Header, Export.ClassIndex, true);
 	UObject* ThisParent = Export.OuterIndex.IsNull() ? LinkerRoot : EventDrivenIndexToObject(Header, Export.OuterIndex, false);
@@ -6161,7 +6162,7 @@ void FAsyncPackage2::EventDrivenCreateExport(const FAsyncPackageHeaderData& Head
 	Object->SetInternalFlags(FlagsToSet);
 }
 
-bool FAsyncPackage2::EventDrivenSerializeExport(const FAsyncPackageHeaderData& Header, int32 LocalExportIndex, FExportArchive& Ar)
+bool FAsyncPackage2::EventDrivenSerializeExport(const FAsyncPackageHeaderData& Header, int32 LocalExportIndex, FExportArchive* Ar)
 {
 	LLM_SCOPE(ELLMTag::UObject);
 	TRACE_CPUPROFILER_EVENT_SCOPE(SerializeExport);
@@ -6193,6 +6194,25 @@ bool FAsyncPackage2::EventDrivenSerializeExport(const FAsyncPackageHeaderData& H
 		return false;
 	}
 
+	TOptional<FExportArchive> LocalAr;
+	if (!Ar)
+	{
+#if WITH_EDITOR
+		if (&Header == OptionalSegmentHeaderData.GetPtrOrNull())
+		{
+			Ar = &LocalAr.Emplace(OptionalSegmentSerializationState->IoRequest.GetResultOrDie());
+			InitializeExportArchive(*Ar, true);
+		}
+		else
+#endif
+		{
+			Ar = &LocalAr.Emplace(SerializationState.IoRequest.GetResultOrDie());
+			InitializeExportArchive(*Ar, false);
+		}
+	}
+
+	ProcessExportDependencies(Header, LocalExportIndex, FExportBundleEntry::ExportCommandType_Serialize);
+
 	// If this is a struct, make sure that its parent struct is completely loaded
 	if (UStruct* Struct = dynamic_cast<UStruct*>(Object))
 	{
@@ -6222,20 +6242,32 @@ bool FAsyncPackage2::EventDrivenSerializeExport(const FAsyncPackageHeaderData& H
 	UObject* PrevSerializedObject = LoadContext->SerializedObject;
 	LoadContext->SerializedObject = Object;
 
-	Ar.TemplateForGetArchetypeFromLoader = ExportObject.TemplateObject;
+	Ar->ExportBufferBegin(Object, Export.CookedSerialOffset, Export.CookedSerialSize);
+
+	const int64 Pos = Ar->Tell();
+
+	check(!Ar->TemplateForGetArchetypeFromLoader);
+	Ar->TemplateForGetArchetypeFromLoader = ExportObject.TemplateObject;
 
 	if (Object->HasAnyFlags(RF_ClassDefaultObject))
 	{
 		TRACE_CPUPROFILER_EVENT_SCOPE(SerializeDefaultObject);
-		Object->GetClass()->SerializeDefaultObject(Object, Ar);
+		Object->GetClass()->SerializeDefaultObject(Object, *Ar);
 	}
 	else
 	{
 		TRACE_CPUPROFILER_EVENT_SCOPE(SerializeObject);
 		UE_SERIALIZE_ACCCESS_SCOPE(Object);
-		Object->Serialize(Ar);
+		Object->Serialize(*Ar);
 	}
-	Ar.TemplateForGetArchetypeFromLoader = nullptr;
+	Ar->TemplateForGetArchetypeFromLoader = nullptr;
+
+	UE_ASYNC_PACKAGE_CLOG(
+		Export.CookedSerialSize != uint64(Ar->Tell() - Pos), Fatal, Desc, TEXT("ObjectSerializationError"),
+		TEXT("%s: Serial size mismatch: Expected read size %d, Actual read size %d"),
+		Object ? *Object->GetFullName() : TEXT("null"), Export.CookedSerialSize, uint64(Ar->Tell() - Pos));
+
+	Ar->ExportBufferEnd();
 
 	Object->SetFlags(RF_LoadCompleted);
 	LoadContext->SerializedObject = PrevSerializedObject;
@@ -6493,22 +6525,6 @@ void FAsyncPackage2::ConditionalBeginProcessPackageExports(FAsyncLoadingThreadSt
 	WaitForAllDependenciesToReachState(ThreadState, &FAsyncPackage2::AllDependenciesSetupState, EAsyncPackageLoadingState2::DependenciesReady, AsyncLoadingThread.ConditionalBeginProcessExportsTick,
 		[&ThreadState](FAsyncPackage2* Package)
 		{
-			if (!Package->bLoadHasFailed
-#if ALT2_ENABLE_LINKERLOAD_SUPPORT
-			&& !Package->LinkerLoadState.IsSet()
-#endif
-				)
-			{
-				TArrayView<FEventLoadNode2> ExportBundleNodesView = Package->Data.ExportBundleNodes;
-				TArrayView<FAsyncPackage2*> ImportedAsyncPackagesView = Package->Data.ImportedAsyncPackages;
-				Package->SetupSerializedArcs(Package->HeaderData, ExportBundleNodesView.Left(Package->HeaderData.ExportBundleHeaders.Num() * EEventLoadNode2::ExportBundle_NumPhases), ImportedAsyncPackagesView.Left(Package->HeaderData.ImportedPackageIds.Num()));
-#if WITH_EDITOR
-				if (Package->OptionalSegmentHeaderData.IsSet())
-				{
-					Package->SetupSerializedArcs(*Package->OptionalSegmentHeaderData, ExportBundleNodesView.Right(Package->OptionalSegmentHeaderData->ExportBundleHeaders.Num() * EEventLoadNode2::ExportBundle_NumPhases), ImportedAsyncPackagesView.Right(Package->OptionalSegmentHeaderData->ImportedPackageIds.Num()));
-				}
-#endif
-			}
 			for (int32 ExportBundleIndex = 0; ExportBundleIndex < Package->Data.TotalExportBundleCount; ++ExportBundleIndex)
 			{
 				Package->GetExportBundleNode(EEventLoadNode2::ExportBundle_Process, ExportBundleIndex).ReleaseBarrier(&ThreadState);
@@ -6585,27 +6601,24 @@ EEventLoadNodeExecutionResult FAsyncPackage2::Event_PostLoadExportBundle(FAsyncL
 
 #if WITH_EDITOR
 		const FAsyncPackageHeaderData* HeaderData;
-		const FExportBundleHeader* ExportBundle;
-		if (InExportBundleIndex >= Package->HeaderData.ExportBundleHeadersCopyForPostLoad.Num())
+		if (InExportBundleIndex == 1)
 		{
 			HeaderData = Package->OptionalSegmentHeaderData.GetPtrOrNull();
 			check(HeaderData);
-			ExportBundle = &HeaderData->ExportBundleHeadersCopyForPostLoad[InExportBundleIndex - Package->HeaderData.ExportBundleHeadersCopyForPostLoad.Num()];
 		}
 		else
 		{
+			check(InExportBundleIndex == 0);
 			HeaderData = &Package->HeaderData;
-			ExportBundle = &Package->HeaderData.ExportBundleHeadersCopyForPostLoad[InExportBundleIndex];
 		}
 #else
 		const FAsyncPackageHeaderData* HeaderData = &Package->HeaderData;
-		const FExportBundleHeader* ExportBundle = &Package->HeaderData.ExportBundleHeadersCopyForPostLoad[InExportBundleIndex];
 #endif
 
 		TRACE_LOADTIME_POSTLOAD_SCOPE;
-		while (Package->ExportBundleEntryIndex < int32(ExportBundle->EntryCount))
+		while (Package->ExportBundleEntryIndex < HeaderData->ExportBundleEntriesCopyForPostLoad.Num())
 		{
-			const FExportBundleEntry& BundleEntry = HeaderData->ExportBundleEntriesCopyForPostLoad[ExportBundle->FirstEntryIndex + Package->ExportBundleEntryIndex];
+			const FExportBundleEntry& BundleEntry = HeaderData->ExportBundleEntriesCopyForPostLoad[Package->ExportBundleEntryIndex];
 			if (ThreadState.IsTimeLimitExceeded(TEXT("Event_PostLoadExportBundle")))
 			{
 				LoadingState = EEventLoadNodeExecutionResult::Timeout;
@@ -6707,27 +6720,24 @@ EEventLoadNodeExecutionResult FAsyncPackage2::Event_DeferredPostLoadExportBundle
 
 #if WITH_EDITOR
 		const FAsyncPackageHeaderData* HeaderData;
-		const FExportBundleHeader* ExportBundle;
-		if (InExportBundleIndex >= Package->HeaderData.ExportBundleHeadersCopyForPostLoad.Num())
+		if (InExportBundleIndex == 1)
 		{
 			HeaderData = Package->OptionalSegmentHeaderData.GetPtrOrNull();
 			check(HeaderData);
-			ExportBundle = &HeaderData->ExportBundleHeadersCopyForPostLoad[InExportBundleIndex - Package->HeaderData.ExportBundleHeadersCopyForPostLoad.Num()];
 		}
 		else
 		{
+			check(InExportBundleIndex == 0);
 			HeaderData = &Package->HeaderData;
-			ExportBundle = &Package->HeaderData.ExportBundleHeadersCopyForPostLoad[InExportBundleIndex];
 		}
 #else
 		const FAsyncPackageHeaderData* HeaderData = &Package->HeaderData;
-		const FExportBundleHeader* ExportBundle = &Package->HeaderData.ExportBundleHeadersCopyForPostLoad[InExportBundleIndex];
 #endif
 
 		TRACE_LOADTIME_POSTLOAD_SCOPE;
-		while (Package->ExportBundleEntryIndex < int32(ExportBundle->EntryCount))
+		while (Package->ExportBundleEntryIndex < HeaderData->ExportBundleEntriesCopyForPostLoad.Num())
 		{
-			const FExportBundleEntry& BundleEntry = HeaderData->ExportBundleEntriesCopyForPostLoad[ExportBundle->FirstEntryIndex + Package->ExportBundleEntryIndex];
+			const FExportBundleEntry& BundleEntry = HeaderData->ExportBundleEntriesCopyForPostLoad[Package->ExportBundleEntryIndex];
 			if (ThreadState.IsTimeLimitExceeded(TEXT("Event_DeferredPostLoadExportBundle")))
 			{
 				LoadingState = EEventLoadNodeExecutionResult::Timeout;
diff --git a/Engine/Source/Runtime/CoreUObject/Public/Serialization/AsyncLoading2.h b/Engine/Source/Runtime/CoreUObject/Public/Serialization/AsyncLoading2.h
--- a/Engine/Source/Runtime/CoreUObject/Public/Serialization/AsyncLoading2.h
+++ b/Engine/Source/Runtime/CoreUObject/Public/Serialization/AsyncLoading2.h
@@ -247,6 +247,7 @@ enum class EZenPackageVersion : uint32
 	Initial,
 	DataResourceTable,
 	ImportedPackageNames,
+	ExportDependencies,
 
 	LatestPlusOne,
 	Latest = LatestPlusOne - 1
@@ -283,7 +284,8 @@ struct FZenPackageSummary
 	int32 ImportMapOffset;
 	int32 ExportMapOffset;
 	int32 ExportBundleEntriesOffset;
-	int32 GraphDataOffset;
+	int32 DependencyBundleHeadersOffset;
+	int32 DependencyBundleEntriesOffset;
 	int32 ImportedPackageNamesOffset;
 };
 
@@ -304,16 +306,19 @@ struct FExportBundleEntry
 	COREUOBJECT_API friend FArchive& operator<<(FArchive& Ar, FExportBundleEntry& ExportBundleEntry);
 };
 
-/**
- * Export bundle header
- */
-struct FExportBundleHeader
+struct FDependencyBundleEntry
+{
+	FPackageIndex LocalImportOrExportIndex;
+
+	COREUOBJECT_API friend FArchive& operator<<(FArchive& Ar, FDependencyBundleEntry& DependencyBundleEntry);
+};
+
+struct FDependencyBundleHeader
 {
-	uint64 SerialOffset;
-	uint32 FirstEntryIndex;
-	uint32 EntryCount;
+	int32 FirstEntryIndex;
+	uint32 EntryCount[FExportBundleEntry::ExportCommandType_Count][FExportBundleEntry::ExportCommandType_Count];
 
-	COREUOBJECT_API friend FArchive& operator<<(FArchive& Ar, FExportBundleHeader& ExportBundleHeader);
+	COREUOBJECT_API friend FArchive& operator<<(FArchive& Ar, FDependencyBundleHeader& DependencyBundleHeader);
 };
 
 struct FScriptObjectEntry
@@ -339,7 +344,7 @@ static_assert(sizeof(FMappedName) >= sizeof(FMinimalName));
  */
 struct FExportMapEntry
 {
-	uint64 CookedSerialOffset = 0;
+	uint64 CookedSerialOffset = 0; // Offset from start of exports data (HeaderSize + CookedSerialOffset gives actual offset in iobuffer)
 	uint64 CookedSerialSize = 0;
 	FMappedName ObjectName;
 	FPackageObjectIndex OuterIndex;
diff --git a/Engine/Source/Runtime/PakFile/Private/FilePackageStore.cpp b/Engine/Source/Runtime/PakFile/Private/FilePackageStore.cpp
--- a/Engine/Source/Runtime/PakFile/Private/FilePackageStore.cpp
+++ b/Engine/Source/Runtime/PakFile/Private/FilePackageStore.cpp
@@ -54,17 +54,14 @@ EPackageStoreEntryStatus FFilePackageStoreBackend::GetPackageStoreEntry(FPackage
 	const FFilePackageStoreEntry* FindEntry = StoreEntriesMap.FindRef(PackageId);
 	if (FindEntry)
 	{
-		OutPackageStoreEntry.ExportInfo.ExportCount = FindEntry->ExportCount;
-		OutPackageStoreEntry.ExportInfo.ExportBundleCount = FindEntry->ExportBundleCount;
 		OutPackageStoreEntry.ImportedPackageIds = MakeArrayView(FindEntry->ImportedPackages.Data(), FindEntry->ImportedPackages.Num());
 		OutPackageStoreEntry.ShaderMapHashes = MakeArrayView(FindEntry->ShaderMapHashes.Data(), FindEntry->ShaderMapHashes.Num());
 #if WITH_EDITOR
 		const FFilePackageStoreEntry* FindOptionalSegmentEntry = OptionalSegmentStoreEntriesMap.FindRef(PackageId);
 		if (FindOptionalSegmentEntry)
 		{
-			OutPackageStoreEntry.OptionalSegmentExportInfo.ExportCount = FindOptionalSegmentEntry->ExportCount;
-			OutPackageStoreEntry.OptionalSegmentExportInfo.ExportBundleCount = FindOptionalSegmentEntry->ExportBundleCount;
 			OutPackageStoreEntry.OptionalSegmentImportedPackageIds = MakeArrayView(FindOptionalSegmentEntry->ImportedPackages.Data(), FindOptionalSegmentEntry->ImportedPackages.Num());
+			OutPackageStoreEntry.bHasOptionalSegment = true;
 		}
 #endif
 		return EPackageStoreEntryStatus::Ok;
diff --git a/Engine/Source/Runtime/StorageServerClient/Private/CookOnTheFlyPackageStore.cpp b/Engine/Source/Runtime/StorageServerClient/Private/CookOnTheFlyPackageStore.cpp
--- a/Engine/Source/Runtime/StorageServerClient/Private/CookOnTheFlyPackageStore.cpp
+++ b/Engine/Source/Runtime/StorageServerClient/Private/CookOnTheFlyPackageStore.cpp
@@ -127,7 +127,6 @@ EPackageStoreEntryStatus FCookOnTheFlyPackageStoreBackend::CreatePackageStoreEnt
 	if (EntryInfo.Status == EPackageStoreEntryStatus::Ok)
 	{
 		const FPackageStoreEntryResource& Entry = PackageEntries[EntryInfo.EntryIndex];
-		OutPackageStoreEntry.ExportInfo = Entry.ExportInfo;
 		OutPackageStoreEntry.ImportedPackageIds = Entry.ImportedPackageIds;
 		return EPackageStoreEntryStatus::Ok;
 	}
diff --git a/Engine/Source/Runtime/StorageServerClient/Private/StorageServerPackageStore.cpp b/Engine/Source/Runtime/StorageServerClient/Private/StorageServerPackageStore.cpp
--- a/Engine/Source/Runtime/StorageServerClient/Private/StorageServerPackageStore.cpp
+++ b/Engine/Source/Runtime/StorageServerClient/Private/StorageServerPackageStore.cpp
@@ -12,7 +12,6 @@ FStorageServerPackageStoreBackend::FStorageServerPackageStoreBackend(FStorageSer
 	Connection.PackageStoreRequest([this](FPackageStoreEntryResource&& PackageStoreEntryResource)
 		{
 			FStoreEntry& LocalEntry = StoreEntriesMap.FindOrAdd(PackageStoreEntryResource.GetPackageId());
-			LocalEntry.ExportInfo = MoveTemp(PackageStoreEntryResource.ExportInfo);
 			LocalEntry.ImportedPackages = MoveTemp(PackageStoreEntryResource.ImportedPackageIds);
 			LocalEntry.ShaderMapHashes = MoveTemp(PackageStoreEntryResource.ShaderMapHashes);
 		});
@@ -23,7 +22,6 @@ EPackageStoreEntryStatus FStorageServerPackageStoreBackend::GetPackageStoreEntry
 	const FStoreEntry* FindEntry = StoreEntriesMap.Find(PackageId);
 	if (FindEntry)
 	{
-		OutPackageStoreEntry.ExportInfo = FindEntry->ExportInfo;
 		OutPackageStoreEntry.ImportedPackageIds = FindEntry->ImportedPackages;
 		OutPackageStoreEntry.ShaderMapHashes = FindEntry->ShaderMapHashes;
 		return EPackageStoreEntryStatus::Ok;
diff --git a/Engine/Source/Runtime/StorageServerClient/Private/StorageServerPackageStore.h b/Engine/Source/Runtime/StorageServerClient/Private/StorageServerPackageStore.h
--- a/Engine/Source/Runtime/StorageServerClient/Private/StorageServerPackageStore.h
+++ b/Engine/Source/Runtime/StorageServerClient/Private/StorageServerPackageStore.h
@@ -38,7 +38,6 @@ public:
 private:
 	struct FStoreEntry
 	{
-		FPackageStoreExportInfo ExportInfo;
 		TArray<FPackageId> ImportedPackages;
 		TArray<FSHAHash> ShaderMapHashes;
 	};
