When collecting property references, don't exlude the property owner from the reference list and instead exclude the owner of the reference list

diff --git a/Engine/Source/Editor/KismetCompiler/Private/KismetCompiler.cpp b/Engine/Source/Editor/KismetCompiler/Private/KismetCompiler.cpp
--- a/Engine/Source/Editor/KismetCompiler/Private/KismetCompiler.cpp
+++ b/Engine/Source/Editor/KismetCompiler/Private/KismetCompiler.cpp
@@ -5006,7 +5006,7 @@ void FKismetCompilerContext::CompileFunctions(EInternalCompilerFlags InternalFla
 				if (FunctionContext.IsValid())
 				{
 					UFunction* Function = FunctionContext.Function; 
-					FArchiveScriptReferenceCollector ObjRefCollector(Function->ScriptAndPropertyObjectReferences);
+					FArchiveScriptReferenceCollector ObjRefCollector(Function->ScriptAndPropertyObjectReferences, Function);
 					for (int32 iCode = 0; iCode < Function->Script.Num();)
 					{
 						Function->SerializeExpr(iCode, ObjRefCollector);
diff --git a/Engine/Source/Runtime/CoreUObject/Private/UObject/Class.cpp b/Engine/Source/Runtime/CoreUObject/Private/UObject/Class.cpp
--- a/Engine/Source/Runtime/CoreUObject/Private/UObject/Class.cpp
+++ b/Engine/Source/Runtime/CoreUObject/Private/UObject/Class.cpp
@@ -709,7 +709,8 @@ void UStruct::GetPreloadDependencies(TArray<UObject*>& OutDeps)
 
 void UStruct::CollectBytecodeReferencedObjects(TArray<UObject*>& OutReferencedObjects)
 {
-	FArchiveScriptReferenceCollector ObjRefCollector(OutReferencedObjects);
+	// Collect references from byte code but exclude ourselves from the refs lists
+	FArchiveScriptReferenceCollector ObjRefCollector(OutReferencedObjects, this);
 
 	int32 BytecodeIndex = 0;
 	while (BytecodeIndex < Script.Num())
diff --git a/Engine/Source/Runtime/CoreUObject/Public/Serialization/ArchiveScriptReferenceCollector.h b/Engine/Source/Runtime/CoreUObject/Public/Serialization/ArchiveScriptReferenceCollector.h
--- a/Engine/Source/Runtime/CoreUObject/Public/Serialization/ArchiveScriptReferenceCollector.h
+++ b/Engine/Source/Runtime/CoreUObject/Public/Serialization/ArchiveScriptReferenceCollector.h
@@ -50,8 +50,9 @@ public:
 	*
 	* @param	InObjectArray			Array to add object references to
 	*/
-	FArchiveScriptReferenceCollector(TArray<UObject*>& InObjectArray)
+	explicit FArchiveScriptReferenceCollector(TArray<UObject*>& InObjectArray, UObject* InExcludeOwner = nullptr)
 		: ObjectArray(InObjectArray)
+		, ExcludeOwner(InExcludeOwner)
 	{
 		ArIsObjectReferenceCollector = true;
 		this->SetIsPersistent(false);
@@ -67,7 +68,7 @@ protected:
 	virtual FArchive& operator<<(UObject*& Object) override
 	{
 		// Avoid duplicate entries.
-		if (Object != nullptr && !ObjectArray.Contains(Object))
+		if (Object != nullptr && Object != ExcludeOwner && !ObjectArray.Contains(Object))
 		{
 			check(Object->IsValidLowLevel());
 			ObjectArray.Add(Object);
@@ -97,7 +98,7 @@ protected:
 		if (Field != nullptr)
 		{
 			// It's faster to collect references via AddReferencedObjects than serialization
-			FPropertyReferenceCollector Collector(Field->GetOwnerUObject(), ObjectArray);
+			FPropertyReferenceCollector Collector(ExcludeOwner, ObjectArray);
 			Field->AddReferencedObjects(Collector);
 		}
 		return *this;
@@ -105,4 +106,6 @@ protected:
 
 	/** Stored reference to array of objects we add object references to */
 	TArray<UObject*>&		ObjectArray;
+	/** Script or property owner object that should be excluded from ObjectArray list */
+	UObject*				ExcludeOwner;
 };
