Fixed triggering static_assert fo FProperty size when sizeof(FName) == 8

diff --git a/Engine/Source/Runtime/CoreUObject/Private/UObject/Property.cpp b/Engine/Source/Runtime/CoreUObject/Private/UObject/Property.cpp
--- a/Engine/Source/Runtime/CoreUObject/Private/UObject/Property.cpp
+++ b/Engine/Source/Runtime/CoreUObject/Private/UObject/Property.cpp
@@ -763,7 +763,7 @@ bool FGCStackSizeHelper::ConvertPathToProperties(UClass* ObjectClass, const FNam
 	FProperty implementation.
 -----------------------------------------------------------------------------*/
 
-#if UE_GAME
+#if UE_GAME && UE_FNAME_OUTLINE_NUMBER
 	static_assert(sizeof(FProperty) <= 104, "FProperty was optimized to reduce its size so most of the classes that inherent from it will fall withing 112 bytes bin of MallocBinned3");
 #endif
 
