Remove ObjectPtr.h's dependency on a fully-defined "class UClass"

diff --git a/Engine/Source/Runtime/CoreUObject/Private/UObject/ObjectPtr.cpp b/Engine/Source/Runtime/CoreUObject/Private/UObject/ObjectPtr.cpp
--- a/Engine/Source/Runtime/CoreUObject/Private/UObject/ObjectPtr.cpp
+++ b/Engine/Source/Runtime/CoreUObject/Private/UObject/ObjectPtr.cpp
@@ -1,6 +1,7 @@
 // Copyright Epic Games, Inc. All Rights Reserved.
 
 #include "UObject/ObjectPtr.h"
+#include "UObject/Class.h"
 
 #if UE_WITH_OBJECT_HANDLE_LATE_RESOLVE
 
@@ -46,4 +47,16 @@ FString FObjectPtr::GetFullName(EObjectFullNameFlags Flags) const
 	}
 }
 
-#endif
\ No newline at end of file
+#endif
+
+bool FObjectPtr::IsA(const UClass* SomeBase) const
+{
+	checkfSlow(SomeBase, TEXT("IsA(NULL) cannot yield meaningful results"));
+
+	if (const UClass* ThisClass = GetClass())
+	{
+		return ThisClass->IsChildOf(SomeBase);
+	}
+
+	return false;
+}
diff --git a/Engine/Source/Runtime/CoreUObject/Public/UObject/ObjectPtr.h b/Engine/Source/Runtime/CoreUObject/Public/UObject/ObjectPtr.h
--- a/Engine/Source/Runtime/CoreUObject/Public/UObject/ObjectPtr.h
+++ b/Engine/Source/Runtime/CoreUObject/Public/UObject/ObjectPtr.h
@@ -4,7 +4,6 @@
 
 #include "HAL/Platform.h"
 #include "Serialization/StructuredArchive.h"
-#include "UObject/Class.h"
 #include "UObject/Object.h"
 #include "UObject/ObjectHandle.h"
 
@@ -44,6 +43,8 @@
 	#define UE_OBJECT_PTR_NONCONFORMANCE_SUPPORT 0
 #endif
 
+class UClass;
+	
 template <typename T>
 struct TObjectPtr;
 
@@ -193,17 +194,7 @@ public:
 	FORCEINLINE FObjectHandle GetHandle() const { return Handle; }
 	FORCEINLINE FObjectHandle& GetHandleRef() const { return Handle; }
 
-	FORCEINLINE bool IsA(const UClass* SomeBase) const
-	{
-		checkfSlow(SomeBase, TEXT("IsA(NULL) cannot yield meaningful results"));
-
-		if (const UClass* ThisClass = GetClass())
-		{
-			return ThisClass->IsChildOf(SomeBase);
-		}
-
-		return false;
-	}
+	COREUOBJECT_API bool IsA(const UClass* SomeBase) const;
 
 	template <typename T>
 	FORCEINLINE bool IsA() const
