autortfm fixes:
- call PropertyInitCallback under Close rather than a nested Transact so that failure propagates out and when we return from Open, we will properly abort and longjmp

diff --git a/Engine/Source/Runtime/CoreUObject/Private/UObject/UObjectGlobals.cpp b/Engine/Source/Runtime/CoreUObject/Private/UObject/UObjectGlobals.cpp
--- a/Engine/Source/Runtime/CoreUObject/Private/UObject/UObjectGlobals.cpp
+++ b/Engine/Source/Runtime/CoreUObject/Private/UObject/UObjectGlobals.cpp
@@ -3854,10 +3854,17 @@ void FObjectInitializer::PostConstructInit()
 	{
 		// autortfm todo: if this transaction aborts and we are in a transaction's open nest,
 		// we need to have a way of propagating out that abort
-		UE_AUTORTFM_TRANSACT(
+		if(AutoRTFM::IsTransactional())
+		{
+			AutoRTFM::EContextStatus Status = AutoRTFM::Close([&]
+			{
+				PropertyInitCallback();
+			});
+		}
+		else
 		{
 			PropertyInitCallback();
-		});
+		}
 	}
 	// After the call to `PropertyInitCallback` to allow the callback to modify the instancing graph
 	if (bNeedInstancing || bNeedSubobjectInstancing)
