[Experimental] Option to not mark outer level/OFPA assets as dirty when propagating default value changes to instances that already match the archetype's value.

Notes:
- Applies only to changes propagated from BP editor class default edits.
- Disabled by default. Can be enabled locally via Blueprint editor settings.
- When enabled, instances will still be marked dirty if the default value change results in the instance value becoming realigned with the archetype. This is because subsequent changes will now propagate to the instance, and since the per-instance value may have been delta-serialized, it needs to be re-serialized in order to update the current saved value.

diff --git a/Engine/Source/Editor/BlueprintGraph/Private/BlueprintEditorSettings.cpp b/Engine/Source/Editor/BlueprintGraph/Private/BlueprintEditorSettings.cpp
--- a/Engine/Source/Editor/BlueprintGraph/Private/BlueprintEditorSettings.cpp
+++ b/Engine/Source/Editor/BlueprintGraph/Private/BlueprintEditorSettings.cpp
@@ -56,6 +56,7 @@ UBlueprintEditorSettings::UBlueprintEditorSettings(const FObjectInitializer& Obj
 	, bEnableContextMenuTimeSlicing(false)
 	, ContextMenuTimeSlicingThresholdMs(50)
 	, bFavorPureCastNodes(false)
+	, bDoNotMarkAllInstancesDirtyOnDefaultValueChange(false)
 	// Compiler Settings
 	, SaveOnCompile(SoC_Never)
 	, bJumpToNodeErrors(false)
diff --git a/Engine/Source/Editor/BlueprintGraph/Public/BlueprintEditorSettings.h b/Engine/Source/Editor/BlueprintGraph/Public/BlueprintEditorSettings.h
--- a/Engine/Source/Editor/BlueprintGraph/Public/BlueprintEditorSettings.h
+++ b/Engine/Source/Editor/BlueprintGraph/Public/BlueprintEditorSettings.h
@@ -209,6 +209,10 @@ public:
 	/** If enabled, then placed cast nodes will default to their "pure" form (meaning: without execution pins). */
 	UPROPERTY(EditAnywhere, config, AdvancedDisplay, Category = Experimental, meta = (DisplayName = "Default to Using Pure Cast Nodes"))
 	bool bFavorPureCastNodes;
+	
+	/** If enabled, assets containing Blueprint instances (e.g. maps) will not be marked dirty when default values are edited, unless it results in the instance becoming realigned with the new default value. */
+	UPROPERTY(EditAnywhere, config, Category = Experimental)
+	bool bDoNotMarkAllInstancesDirtyOnDefaultValueChange;
 
 	// Compiler Settings
 public:
diff --git a/Engine/Source/Editor/Kismet/Private/BlueprintEditor.cpp b/Engine/Source/Editor/Kismet/Private/BlueprintEditor.cpp
--- a/Engine/Source/Editor/Kismet/Private/BlueprintEditor.cpp
+++ b/Engine/Source/Editor/Kismet/Private/BlueprintEditor.cpp
@@ -2020,6 +2020,9 @@ void FBlueprintEditor::CommonInitialization(const TArray<UBlueprint*>& InitBluep
 
 	CreateDefaultTabContents(InitBlueprints);
 
+	FCoreUObjectDelegates::OnPreObjectPropertyChanged.AddSP(this, &FBlueprintEditor::OnPreObjectPropertyChanged);
+	FCoreUObjectDelegates::OnObjectPropertyChanged.AddSP(this, &FBlueprintEditor::OnPostObjectPropertyChanged);
+
 	FKismetEditorUtilities::OnBlueprintUnloaded.AddSP(this, &FBlueprintEditor::OnBlueprintUnloaded);
 }
 
@@ -2827,6 +2830,9 @@ FBlueprintEditor::~FBlueprintEditor()
 		SubobjectDataSubsystem->OnNewSubobjectAdded().RemoveAll(this);
 	}
 
+	FCoreUObjectDelegates::OnPreObjectPropertyChanged.RemoveAll(this);
+	FCoreUObjectDelegates::OnObjectPropertyChanged.RemoveAll(this);
+
 	FKismetEditorUtilities::OnBlueprintUnloaded.RemoveAll(this);
 
 	FGlobalTabmanager::Get()->OnActiveTabChanged_Unsubscribe(OnActiveTabChangedDelegateHandle);
@@ -9574,6 +9580,74 @@ void FBlueprintEditor::OnFinishedChangingProperties(const FPropertyChangedEvent&
 	}
 }
 
+void FBlueprintEditor::OnPreObjectPropertyChanged(UObject* InObject, const FEditPropertyChain& EditPropertyChain)
+{
+	if (GetDefault<UBlueprintEditorSettings>()->bDoNotMarkAllInstancesDirtyOnDefaultValueChange)
+	{
+		// Note: While we could rely on our notify hook to determine whether an incoming change belongs to this
+		// context, NotifyPreChange() may end up getting called after this event is broadcast (e.g. from inside
+		// of a details customization). So it's not safe to assume they will get called in any specific order.
+		// It is, however, safe to assume this will get called for an archetype/CDO prior to change propagation.
+		if (Inspector.IsValid())
+		{
+			check(InObject);
+
+			// If we're modifying an object that's selected into our property editor context (e.g. the Blueprint CDO, or an SCS
+			// component template), set up change propagation so that instances do not always mark their outer package as dirty.
+			bool bIsObjectSelectedForEditing = false;
+			const TArray<TWeakObjectPtr<UObject>> SelectedObjects = Inspector->GetSelectedObjects();
+			if (SelectedObjects.Contains(InObject))
+			{
+				bIsObjectSelectedForEditing = true;
+			}
+			else if (IsEditingSingleBlueprint())
+			{
+				// Resolve the property that's about to be changed.
+				FProperty* PropertyAboutToChange = nullptr;
+				if (FEditPropertyChain::TDoubleLinkedListNode* PropertyNode = EditPropertyChain.GetActiveNode())
+				{
+					PropertyAboutToChange = PropertyNode->GetValue();
+				}
+
+				// When actions from the "My Blueprint" tab are selected (e.g. class variable actions), a property "wrapper"
+				// object is selected to the inspector view instead, and the details customization will then pass the object
+				// instance that's about to be changed (e.g. the Blueprint CDO) as the input parameter to this delegate.
+				if (PropertyAboutToChange)
+				{
+					const UBlueprint* BlueprintContext = GetBlueprintObj();
+					check(BlueprintContext);
+
+					// Redirect to the skeleton class if the property's owner class was generated by the Blueprint that's
+					// associated with this editing context. For details customization purposes, the property that's wrapped
+					// belongs to the skeleton class and not the actual generated class (see SMyBlueprint::CollectAllActions).
+					const UClass* OwnerClass = PropertyAboutToChange->GetOwnerClass();
+					if (OwnerClass && OwnerClass->ClassGeneratedBy == BlueprintContext && BlueprintContext->SkeletonGeneratedClass)
+					{
+						PropertyAboutToChange = BlueprintContext->SkeletonGeneratedClass->FindPropertyByName(PropertyAboutToChange->GetFName());
+						if (PropertyAboutToChange && SelectedObjects.Contains(PropertyAboutToChange->GetUPropertyWrapper()) && InObject == OwnerClass->GetDefaultObject())
+						{
+							bIsObjectSelectedForEditing = true;
+						}
+					}
+				}
+			}
+			
+			if (bIsObjectSelectedForEditing)
+			{
+				InObject->SetEditChangePropagationFlags(EEditChangePropagationFlags::OnlyMarkRealignedInstancesAsDirty);
+			}
+		}
+	}
+}
+
+void FBlueprintEditor::OnPostObjectPropertyChanged(UObject* InObject, FPropertyChangedEvent& PropertyChangedEvent)
+{
+	check(InObject);
+
+	// Reset propagation flags so that any changes external to our editing context will use the standard (default) behavior.
+	InObject->SetEditChangePropagationFlags(EEditChangePropagationFlags::None);
+}
+
 FName FBlueprintEditor::GetToolkitFName() const
 {
 	return FName("BlueprintEditor");
diff --git a/Engine/Source/Editor/Kismet/Public/BlueprintEditor.h b/Engine/Source/Editor/Kismet/Public/BlueprintEditor.h
--- a/Engine/Source/Editor/Kismet/Public/BlueprintEditor.h
+++ b/Engine/Source/Editor/Kismet/Public/BlueprintEditor.h
@@ -903,6 +903,12 @@ protected:
 	/** Handles the unloading of Blueprints (by closing the editor, if it operating on the Blueprint being unloaded)*/
 	void OnBlueprintUnloaded(UBlueprint* InBlueprint);
 
+	/** Called when a property change is about to be propagated to instances of the Blueprint */
+	void OnPreObjectPropertyChanged(UObject* InObject, const FEditPropertyChain& EditPropertyChain);
+
+	/** Called after a property change has been propagated to instances of the Blueprint */
+	void OnPostObjectPropertyChanged(UObject* InObject, FPropertyChangedEvent& PropertyChangedEvent);
+
 	//@TODO: Should the breakpoint/watch modification operations be whole-blueprint, or current-graph?
 
 	/** Deletes all breakpoints for the blueprint being edited */
diff --git a/Engine/Source/Runtime/CoreUObject/Private/UObject/Obj.cpp b/Engine/Source/Runtime/CoreUObject/Private/UObject/Obj.cpp
--- a/Engine/Source/Runtime/CoreUObject/Private/UObject/Obj.cpp
+++ b/Engine/Source/Runtime/CoreUObject/Private/UObject/Obj.cpp
@@ -96,6 +96,37 @@ static UPackage*			GObjTransientPkg								= NULL;
 #endif
 
 #if WITH_EDITOR
+	struct FPropagatedEditChangeAnnotation
+	{
+		/**
+		 * If set, archetype edits will not mark an instance dirty unless it results
+		 * in the instance realigning with the archetype after the change (that will
+		 * then result in future archetype changes being propagated to the instance).
+		 */
+		uint8 bDeferredMarkAsDirty : 1;
+		/**
+		 * If set, this instance will be affected by an archetype change (i.e. it
+		 * matched the archetype prior to propagating the change).
+		 */
+		uint8 bIdenticalToArchetype : 1;
+		/**
+		 * If set, the package containing this instance was already marked as dirty
+		 * prior to propagating the change.
+		 */
+		uint8 bWasPackageDirtyOnEdit : 1;
+
+		FPropagatedEditChangeAnnotation()
+		: bDeferredMarkAsDirty(false)
+		, bIdenticalToArchetype(false)
+		, bWasPackageDirtyOnEdit(false)
+		{}
+
+		FORCEINLINE bool IsDefault() const
+		{
+			return !bDeferredMarkAsDirty;
+		}
+	};
+	static FUObjectAnnotationSparse<FPropagatedEditChangeAnnotation, true> PropagatedEditChangeAnnotation;
 	UObject::FAssetRegistryTag::FOnGetObjectAssetRegistryTags UObject::FAssetRegistryTag::OnGetExtraObjectTags;
 	UObject::FAssetRegistryTag::FOnGetExtendedAssetRegistryTagsForSave UObject::FAssetRegistryTag::OnGetExtendedAssetRegistryTagsForSave;
 #endif // WITH_EDITOR
@@ -365,7 +396,29 @@ void UObject::PostLoad()
 #if WITH_EDITOR
 void UObject::PreEditChange(FProperty* PropertyAboutToChange)
 {
-	Modify(!GIsTransacting && !(PropertyAboutToChange && PropertyAboutToChange->HasAnyPropertyFlags(CPF_SkipSerialization)));
+	bool bShouldMarkAsDirty = true;
+
+	if (GIsTransacting)
+	{
+		// Don't mark the outer package as dirty during an undo/redo operation.
+		bShouldMarkAsDirty = false;
+	}
+	else if (PropertyAboutToChange && PropertyAboutToChange->HasAnyPropertyFlags(CPF_SkipSerialization))
+	{
+		// Don't mark the outer package as dirty if we're about to change a non-serializable property.
+		bShouldMarkAsDirty = false;
+	}
+	else
+	{
+		FPropagatedEditChangeAnnotation Annotation = PropagatedEditChangeAnnotation.GetAnnotation(this);
+		if (Annotation.bDeferredMarkAsDirty)
+		{
+			// Don't mark the outer package as dirty if annotated to be deferred (e.g. during propagation).
+			bShouldMarkAsDirty = false;
+		}
+	}
+
+	Modify(bShouldMarkAsDirty);
 }
 
 
@@ -518,10 +571,46 @@ void UObject::PropagatePreEditChange( TArray<UObject*>& AffectedObjects, FEditPr
 		}
 	}
 
+	check(PropertyAboutToChange.GetActiveNode() != nullptr);
+	const FProperty* ChangedProperty = PropertyAboutToChange.GetActiveNode()->GetValue();
+	FPropagatedEditChangeAnnotation Annotation = PropagatedEditChangeAnnotation.GetAnnotation(this);
+
 	for ( int32 i = 0; i < Instances.Num(); i++ )
 	{
 		UObject* Obj = Instances[i];
 
+		// To defer marking instances as dirty, check to see if the instance
+		// matches the value stored in its archetype, and flag it if so. We'll
+		// use this later to determine if we need to mark the package as dirty,
+		// rather than always marking all affected archetype instances as dirty.
+		if (Annotation.bDeferredMarkAsDirty)
+		{
+			// Start with the assumption that the instance matches the archetype. In
+			// that case, we won't need to dirty the package after applying the change.
+			Annotation.bIdenticalToArchetype = true;
+
+			// Note that some elements may match and thus will propagate, but we may
+			// need to dirty the package later even if only one element differs here.
+			for (int32 ArrayIdx = 0; ArrayIdx < ChangedProperty->ArrayDim; ++ArrayIdx)
+			{
+				if (!ChangedProperty->Identical_InContainer(this, Obj, ArrayIdx, PPF_DeepComparison))
+				{
+					Annotation.bIdenticalToArchetype = false;
+					break;
+				}
+			}
+
+			// Determine if the package is already marked as dirty.
+			Annotation.bWasPackageDirtyOnEdit = Obj->GetPackage()->IsDirty();
+
+			// Temporarily annotate the instance for change propagation.
+			PropagatedEditChangeAnnotation.AddAnnotation(Obj, Annotation);
+		}
+
+		// Note: This test is not the same as the flag above - change propagation can
+		// be filtered via the event (e.g. container properties via the Property Editor).
+		// For most cases (i.e. non-container), all archetype instances will pass here,
+		// regardless of whether or not they differ from the default prior to the change.
 		if ( PropertyAboutToChange.IsArchetypeInstanceAffected(Obj) )
 		{
 			// this object must now be included in any undo/redo operations
@@ -557,11 +646,51 @@ void UObject::PropagatePostEditChange( TArray<UObject*>& AffectedObjects, FPrope
 	}
 
 	check(PropertyChangedEvent.PropertyChain.GetActiveMemberNode() != nullptr);
+	const FProperty* ChangedProperty = PropertyChangedEvent.PropertyChain.GetActiveMemberNode()->GetValue();
 
+	TSet<UPackage*> PackagesMarkedAsDirty;
 	for ( int32 i = 0; i < Instances.Num(); i++ )
 	{
 		UObject* Obj = Instances[i];
 
+		UPackage* Package = Obj->GetPackage();
+		check(Package);
+
+		// Deferred marking instances as dirty - if our previous value did not
+		// match the instance but now the current value does, we need to mark
+		// the package dirty to indicate to the user that it needs to be saved.
+		FPropagatedEditChangeAnnotation Annotation = PropagatedEditChangeAnnotation.GetAndRemoveAnnotation(Obj);
+		if (Annotation.bDeferredMarkAsDirty && !Annotation.bWasPackageDirtyOnEdit && !PackagesMarkedAsDirty.Contains(Package))
+		{
+			// Clear the dirty flag if the previous value matched the archetype and if
+			// the package was not already marked as dirty prior to change propagation.
+			const bool bIsPackageDirty = Package->IsDirty();
+			if (bIsPackageDirty && Annotation.bIdenticalToArchetype)
+			{
+				Package->SetDirtyFlag(false);
+			}
+			else if (!bIsPackageDirty && !Annotation.bIdenticalToArchetype)
+			{
+				// If any index matches, that element will no longer be delta-serialized,
+				// so we need to dirty the package. If the property has multiple entries
+				// and we arrived here, that means at least one element differed from its
+				// previous value in the source, but that may not be the one that changed.
+				for (int32 ArrayIdx = 0; ArrayIdx < ChangedProperty->ArrayDim; ++ArrayIdx)
+				{
+					if (ChangedProperty->Identical_InContainer(this, Obj, ArrayIdx, PPF_DeepComparison))
+					{
+						// Using this API so that we don't unnecessarily mark certain packages (e.g. transient).
+						Obj->MarkPackageDirty();
+						PackagesMarkedAsDirty.Add(Package);
+						break;
+					}
+				}
+			}
+		}
+
+		// Note: This is not the same as the flag above - change propagation can be
+		// filtered via the event (e.g. container properties via the Property Editor),
+		// but for most cases (i.e. non-container), all archetype instances pass here.
 		if ( PropertyChangedEvent.HasArchetypeInstanceChanged(Obj) )
 		{
 			// notify the object that all changes are complete
@@ -573,6 +702,13 @@ void UObject::PropagatePostEditChange( TArray<UObject*>& AffectedObjects, FPrope
 	}
 }
 
+void UObject::SetEditChangePropagationFlags(EEditChangePropagationFlags InFlags)
+{
+	FPropagatedEditChangeAnnotation Annotation;
+	Annotation.bDeferredMarkAsDirty = !!(InFlags & EEditChangePropagationFlags::OnlyMarkRealignedInstancesAsDirty);
+	PropagatedEditChangeAnnotation.AddAnnotation(this, MoveTemp(Annotation));
+}
+
 void UObject::PreEditUndo()
 {
 	PreEditChange(NULL);
diff --git a/Engine/Source/Runtime/CoreUObject/Public/UObject/Object.h b/Engine/Source/Runtime/CoreUObject/Public/UObject/Object.h
--- a/Engine/Source/Runtime/CoreUObject/Public/UObject/Object.h
+++ b/Engine/Source/Runtime/CoreUObject/Public/UObject/Object.h
@@ -57,6 +57,25 @@ namespace ECastCheckedType
 	};
 };
 
+#if WITH_EDITOR
+/** Flags to control how editor change events get propagated */
+enum class EEditChangePropagationFlags : uint8
+{
+	/**
+	 * No flags; use default change event propagation behavior.
+	 */
+	None = 0,
+	/**
+	 * If set, only the subset of affected instances that realign with
+	 * their archetype will be marked dirty. If not set, affected instances
+	 * will always be marked dirty, even if they already match the archetype.
+	 */
+	 OnlyMarkRealignedInstancesAsDirty = (1 << 0),
+};
+
+ENUM_CLASS_FLAGS(EEditChangePropagationFlags);
+#endif // WITH_EDITOR
+
 /** 
  * The base class of all UE objects. The type of an object is defined by its UClass.
  * This provides support functions for creating and using objects, and virtual functions that should be overridden in child classes.
@@ -468,6 +487,9 @@ public:
 	 */
 	virtual void PostEditChangeChainProperty( struct FPropertyChangedChainEvent& PropertyChangedEvent );
 
+	/** Called to set the flags that control how editor change events get propagated to archetype instances. */
+	void SetEditChangePropagationFlags(EEditChangePropagationFlags InFlags);
+
 	/** Called before applying a transaction to the object.  Default implementation simply calls PreEditChange. */
 	virtual void PreEditUndo();
 
