SavePackage support for OnlyEditorOnly: Remove exports that are private and are only referenced by editoronly properties from the cookable exports. Remove imports and softobjectpaths that are only referenced from these OnlyEditorOnly exports from the cooked package's imports. For editor saves, mark these OnlyEditorOnly-referenced imports and softobjectpaths as editoronly dependencies. Also return the list of ImportPackages and SoftPackageReferences to the cooker so it can add them to the cook.
#rnx

diff --git a/Engine/Source/Developer/TargetPlatform/Public/Interfaces/ITargetPlatform.h b/Engine/Source/Developer/TargetPlatform/Public/Interfaces/ITargetPlatform.h
--- a/Engine/Source/Developer/TargetPlatform/Public/Interfaces/ITargetPlatform.h
+++ b/Engine/Source/Developer/TargetPlatform/Public/Interfaces/ITargetPlatform.h
@@ -643,10 +643,7 @@ public:
 	 */
 	virtual bool AllowAudioVisualData() const = 0;
 
-	/**
-	 * Checks if this Target will want to load this object (generally used to mark an object to not be cooked for this target,
-	 * as called by ConditionallyExcludeObjectForTarget)
-	 */
+	/** Checks if this Target will want to load this object (generally used to mark an object to not be cooked for this target) */
 	virtual bool AllowObject(const class UObject* Object) const = 0;
 
 	/**
diff --git a/Engine/Source/Runtime/CoreUObject/Private/UObject/SavePackage/PackageHarvester.cpp b/Engine/Source/Runtime/CoreUObject/Private/UObject/SavePackage/PackageHarvester.cpp
--- a/Engine/Source/Runtime/CoreUObject/Private/UObject/SavePackage/PackageHarvester.cpp
+++ b/Engine/Source/Runtime/CoreUObject/Private/UObject/SavePackage/PackageHarvester.cpp
@@ -45,8 +45,7 @@ EObjectMark GenerateMarksForObject(const UObject* InObject, const ITargetPlatfor
 	}
 #endif
 	
-	// CDOs must be included if their class is so only inherit marks, for everything else we check the native overrides as well
-	if (SavePackageUtilities::IsStrippedEditorOnlyObject(InObject, false, false))
+	if (SavePackageUtilities::IsStrippedEditorOnlyObject(InObject, true, false))
 	{
 		Marks = (EObjectMark)(Marks | OBJECTMARK_EditorOnly);
 	}
@@ -60,71 +59,67 @@ EObjectMark GenerateMarksForObject(const UObject* InObject, const ITargetPlatfor
 	return Marks;
 }
 
-bool ConditionallyExcludeObjectForTarget(FSaveContext& SaveContext, UObject* Obj, ESaveRealm HarvestingContext)
+bool ConditionallyExcludeObjectForRealm(FSaveContext& SaveContext, UObject* Obj, ESaveRealm HarvestingContext)
 {
 	if (!Obj || Obj->GetOutermost()->GetFName() == GLongCoreUObjectPackageName)
 	{
 		// No object or in CoreUObject, don't exclude
 		return false;
 	}
-
-	bool bExcluded = false;
-	if (SaveContext.GetHarvestedRealm(HarvestingContext).IsExcluded(Obj))
+	FHarvestedRealm& RealmBeingChecked = SaveContext.GetHarvestedRealm(HarvestingContext);
+	if (RealmBeingChecked.IsExcluded(Obj))
 	{
 		return true;
 	}
-	else if (!SaveContext.GetHarvestedRealm(HarvestingContext).IsIncluded(Obj))
+	if (RealmBeingChecked.IsIncluded(Obj))
 	{
-		const EObjectMark ExcludedObjectMarks = SaveContext.GetExcludedObjectMarks(HarvestingContext);
-		const ITargetPlatform* TargetPlatform = SaveContext.GetTargetPlatform();
-		EObjectMark ObjectMarks = GenerateMarksForObject(Obj, TargetPlatform);
-		if (!(ObjectMarks & ExcludedObjectMarks))
-		{
-			UObject* ObjOuter = Obj->GetOuter();
-			UClass* ObjClass = Obj->GetClass();
+		return false;
+	}
 
-			if (ConditionallyExcludeObjectForTarget(SaveContext, ObjClass, HarvestingContext))
-			{
-				// If the object class is excluded, the object must be excluded too
-				bExcluded = true;
-			}
-			else if (ConditionallyExcludeObjectForTarget(SaveContext, ObjOuter, HarvestingContext))
-			{
-				// If the object outer is excluded, the object must be excluded too
-				bExcluded = true;
-			}
+	const EObjectMark ExcludedObjectMarks = SaveContext.GetExcludedObjectMarks(HarvestingContext);
+	const ITargetPlatform* TargetPlatform = SaveContext.GetTargetPlatform();
+	EObjectMark ObjectMarks = GenerateMarksForObject(Obj, TargetPlatform);
+	if (!!(ObjectMarks & ExcludedObjectMarks))
+	{
+		RealmBeingChecked.AddExcluded(Obj);
+		return true;
+	}
 
-			// Check parent struct if we have one
-			UStruct* ThisStruct = Cast<UStruct>(Obj);
-			if (ThisStruct && ThisStruct->GetSuperStruct())
-			{
-				UObject* SuperStruct = ThisStruct->GetSuperStruct();
-				if (ConditionallyExcludeObjectForTarget(SaveContext, SuperStruct, HarvestingContext))
-				{
-					bExcluded = true;
-				}
-			}
+	// If the object class is excluded, the object must be excluded too
+	if (ConditionallyExcludeObjectForRealm(SaveContext, Obj->GetClass(), HarvestingContext))
+	{
+		RealmBeingChecked.AddExcluded(Obj);
+		return true;
+	}
 
-			// Check archetype, this may not have been covered in the case of components
-			UObject* Archetype = Obj->GetArchetype();
-			if (Archetype)
-			{
-				if (ConditionallyExcludeObjectForTarget(SaveContext, Archetype, HarvestingContext))
-				{
-					bExcluded = true;
-				}
-			}
-		}
-		else
-		{
-			bExcluded = true;
-		}
-		if (bExcluded)
+	// If the object outer is excluded, the object must be excluded too
+	if (ConditionallyExcludeObjectForRealm(SaveContext, Obj->GetOuter(), HarvestingContext))
+	{
+		RealmBeingChecked.AddExcluded(Obj);
+		return true;
+	}
+
+	// Check parent struct if we have one
+	UStruct* ThisStruct = Cast<UStruct>(Obj);
+	if (ThisStruct)
+	{
+		UObject* SuperStruct = ThisStruct->GetSuperStruct();
+		if (SuperStruct && ConditionallyExcludeObjectForRealm(SaveContext, SuperStruct, HarvestingContext))
 		{
-			SaveContext.GetHarvestedRealm(HarvestingContext).AddExcluded(Obj);
+			RealmBeingChecked.AddExcluded(Obj);
+			return true;
 		}
 	}
-	return bExcluded;
+
+	// Check archetype, this may not have been covered in the case of components
+	UObject* Archetype = Obj->GetArchetype();
+	if (Archetype && ConditionallyExcludeObjectForRealm(SaveContext, Archetype, HarvestingContext))
+	{
+		RealmBeingChecked.AddExcluded(Obj);
+		return true;
+	}
+
+	return false;
 }
 
 bool DoesObjectNeedLoadForEditorGame(UObject* InObject)
@@ -146,34 +141,200 @@ bool DoesObjectNeedLoadForEditorGame(UObject* InObject)
 	return bNeedsLoadForEditorGame;
 }
 
-FPackageHarvester::FExportScope::FExportScope(FPackageHarvester& InHarvester, const FExportWithContext& InToProcess, bool bIsEditorOnlyObject)
+FPackageHarvester::FExportScope::FExportScope(FPackageHarvester& InHarvester, const FExportWithContext& InToProcess)
 	: Harvester(InHarvester)
-	, PreviousRealm(InHarvester.CurrentExportHarvestingRealm)
+	, PreviousExportHarvestingRealms(InHarvester.CurrentExportHarvestingRealms)
 	, bPreviousFilterEditorOnly(InHarvester.IsFilterEditorOnly())
 {
-	check(Harvester.CurrentExportDependencies.CurrentExport == nullptr);
+	check(!Harvester.HasAnyExportHarvestingRealms());
 	Harvester.CurrentExportDependencies = { InToProcess.Export };
-	Harvester.CurrentExportHarvestingRealm = InToProcess.HarvestedFromRealm;
-	Harvester.bIsEditorOnlyExportOnStack = bIsEditorOnlyObject;
+	Harvester.CurrentExportHarvestingRealms = InToProcess.HarvestedFromRealms;
 
-	// if we are auto generating optional package, then do no filter editor properties for that harvest
-	if (Harvester.SaveContext.IsSaveAutoOptional() && InToProcess.HarvestedFromRealm == ESaveRealm::Optional)
+	if (Harvester.CurrentExportHarvestingRealms.Contains(ESaveRealm::Optional))
 	{
-		Harvester.SetFilterEditorOnly(false);
+		// The optional realm has to be harvested individually because we set flags separately for it
+		check(Harvester.CurrentExportHarvestingRealms.Num() == 1);
+
+		// if we are auto generating optional package, then do not filter editor properties for that harvest
+		if (Harvester.SaveContext.IsSaveAutoOptional())
+		{
+			Harvester.SetFilterEditorOnly(false);
+		}
 	}
 }
 
 FPackageHarvester::FExportScope::~FExportScope()
 {
 	Harvester.AppendCurrentExportDependencies();
-	Harvester.CurrentExportHarvestingRealm = PreviousRealm;
+	Harvester.CurrentExportHarvestingRealms = PreviousExportHarvestingRealms;
 	Harvester.SetFilterEditorOnly(bPreviousFilterEditorOnly);
 }
 
+FPackageHarvester::FIgnoreDependenciesScope::FIgnoreDependenciesScope(FPackageHarvester& InHarvester)
+	: Harvester(InHarvester)
+	, bPreviousValue(Harvester.CurrentExportDependencies.bIgnoreDependencies)
+{
+	Harvester.CurrentExportDependencies.bIgnoreDependencies = true;
+}
+
+FPackageHarvester::FIgnoreDependenciesScope::~FIgnoreDependenciesScope()
+{
+	Harvester.CurrentExportDependencies.bIgnoreDependencies = bPreviousValue;
+}
+
+FPackageHarvester::FHarvestScope::FHarvestScope(FPackageHarvester& InHarvester)
+	: Harvester(InHarvester)
+	, PreviousExportHarvestingRealms(InHarvester.CurrentExportHarvestingRealms)
+	, bActive(true)
+{
+}
+
+FPackageHarvester::FHarvestScope::FHarvestScope(FHarvestScope&& Other)
+	: Harvester(Other.Harvester)
+	, PreviousExportHarvestingRealms(MoveTemp(Other.PreviousExportHarvestingRealms))
+	, bActive(Other.bActive)
+{
+	Other.bActive = false;
+}
+
+FPackageHarvester::FHarvestScope::~FHarvestScope()
+{
+	if (bActive)
+	{
+		Harvester.CurrentExportHarvestingRealms = PreviousExportHarvestingRealms;
+	}
+}
+
+bool FPackageHarvester::FHarvestScope::IsEmpty() const
+{
+	return Harvester.CurrentExportHarvestingRealms.IsEmpty();
+}
+
+FPackageHarvester::FHarvestScope FPackageHarvester::EnterRootReferencesScope()
+{
+	check(!HasAnyExportHarvestingRealms());
+	FHarvestScope Scope(*this);
+
+	CurrentExportHarvestingRealms.Add(ESaveRealm::Game);
+	if (!SaveContext.IsCooking())
+	{
+		CurrentExportHarvestingRealms.Add(ESaveRealm::Editor);
+	}
+	return Scope;
+}
+
+FPackageHarvester::FHarvestScope FPackageHarvester::EnterConditionalEditorOnlyScope(bool bIsEditorOnly)
+{
+	FHarvestScope Scope(*this);
+	if (bIsEditorOnly)
+	{
+		CurrentExportHarvestingRealms.RemoveSwap(ESaveRealm::Game, false /* bAllowShrinking */);
+	}
+	return Scope;
+}
+
+FPackageHarvester::FHarvestScope FPackageHarvester::EnterConditionalOptionalObjectScope(UObject* Object)
+{
+	FHarvestScope Scope(*this);
+	if (!HasAnyExportHarvestingRealms())
+	{
+		return Scope;
+	}
+	bool bIsOptionalObject = ShouldObjectBeHarvestedInOptionalRealm(Object);
+	if (!bIsOptionalObject)
+	{
+		return Scope;
+	}
+
+	// No matter the realm referring to it, optional objects are added into the optional realm
+	// if we are harvesting an optional object while in a different context, record an illegal reference to display on validation
+	ForEachExportHarvestingRealm([this, Object](ESaveRealm HarvestingRealm)
+		{
+			if (HarvestingRealm != ESaveRealm::Optional)
+			{
+				SaveContext.RecordIllegalReference(CurrentExportDependencies.CurrentExport,
+					Object, EIllegalRefReason::ReferenceToOptional);
+			}
+		});
+
+	CurrentExportHarvestingRealms.Reset();
+	CurrentExportHarvestingRealms.Add(ESaveRealm::Optional);
+	return Scope;
+}
+
+FPackageHarvester::FHarvestScope FPackageHarvester::EnterNewExportOnlyScope(UObject* Export)
+{
+	FHarvestScope Scope(*this);
+	CurrentExportHarvestingRealms.RemoveAllSwap([this, Export](ESaveRealm HarvestingRealm)
+		{
+			return SaveContext.GetHarvestedRealm(HarvestingRealm).IsExport(Export);
+		}, false /* bAllowShrinking */);
+	return Scope;
+}
+
+FPackageHarvester::FHarvestScope FPackageHarvester::EnterNotExcludedScope(UObject* Object)
+{
+	FHarvestScope Scope(*this);
+	CurrentExportHarvestingRealms.RemoveAllSwap([this, Object](ESaveRealm HarvestingRealm)
+		{
+			return ConditionallyExcludeObjectForRealm(SaveContext, Object, HarvestingRealm);
+		}, false /* bAllowShrinking */);
+	return Scope;
+}
+
+FPackageHarvester::FHarvestScope FPackageHarvester::EnterNotPreviouslyExcludedScope(UObject* Object)
+{
+	FHarvestScope Scope(*this);
+	CurrentExportHarvestingRealms.RemoveAllSwap([this, Object](ESaveRealm HarvestingRealm)
+		{
+			return SaveContext.GetHarvestedRealm(HarvestingRealm).IsExcluded(Object); 
+		}, false /* bAllowShrinking */);
+	return Scope;
+}
+
+FPackageHarvester::FHarvestScope FPackageHarvester::EnterIncludedScope(UObject* Object)
+{
+	FHarvestScope Scope(*this);
+	CurrentExportHarvestingRealms.RemoveAllSwap([this, Object](ESaveRealm HarvestingRealm)
+		{
+			return !SaveContext.GetHarvestedRealm(HarvestingRealm).IsIncluded(Object);
+		}, false /* bAllowShrinking */);
+	return Scope;
+}
+
+bool FPackageHarvester::IsObjNative(UObject* InObj)
+{
+	bool bIsNative = InObj->IsNative();
+	UObject* Outer = InObj->GetOuter();
+	while (!bIsNative && Outer)
+	{
+		bIsNative |= Cast<UClass>(Outer) != nullptr && Outer->IsNative();
+		Outer = Outer->GetOuter();
+	}
+	return bIsNative;
+};
+
+bool FPackageHarvester::ShouldObjectBeHarvestedInOptionalRealm(UObject* InObj) const
+{
+	return InObj->GetClass()->HasAnyClassFlags(CLASS_Optional) && SaveContext.IsCooking();
+}
+
+template <typename CallbackType>
+void FPackageHarvester::ForEachExportHarvestingRealm(CallbackType&& Callback)
+{
+	for (ESaveRealm HarvestingRealm : CurrentExportHarvestingRealms)
+	{
+		Callback(HarvestingRealm);
+	}
+}
+
+bool FPackageHarvester::HasAnyExportHarvestingRealms()
+{
+	return !CurrentExportHarvestingRealms.IsEmpty();
+}
+
 FPackageHarvester::FPackageHarvester(FSaveContext& InContext)
 	: SaveContext(InContext)
-	, CurrentExportHarvestingRealm(ESaveRealm::None)
-	, bIsEditorOnlyExportOnStack(false)
 {
 	this->SetIsSaving(true);
 	this->SetIsPersistent(true);
@@ -198,61 +359,73 @@ FPackageHarvester::FExportWithContext FPackageHarvester::PopExportToProcess()
 }
 
 void FPackageHarvester::ProcessExport(const FExportWithContext& InProcessContext)
-{	
+{
 	UObject* Export = InProcessContext.Export;
+	FExportScope HarvesterScope(*this, InProcessContext);
 
-	// No need to check marks since we do not set them on objects anymore
-	bool bReferencerIsEditorOnly = SavePackageUtilities::IsStrippedEditorOnlyObject(Export, true /* bCheckRecursive */, false /* bCheckMarks */) && !Export->HasNonEditorOnlyReferences();
-	FExportScope HarvesterScope(*this, InProcessContext, bReferencerIsEditorOnly);
-	
-	// The export scope set the current harvesting context
-	FHarvestedRealm& HarvestedRealm = SaveContext.GetHarvestedRealm(CurrentExportHarvestingRealm);
-	check(HarvestedRealm.IsExport(Export));
-
-	// Harvest its class 
+	// Harvest its class; warn if class is not included
 	UClass* Class = Export->GetClass();
 	*this << Class;
-	if (!HarvestedRealm.IsIncluded(Class))
-	{
-		SaveContext.RecordIllegalReference(Export, Class, EIllegalRefReason::UnsaveableClass, GetUnsaveableReason(Class));
-	}
+	ForEachExportHarvestingRealm([this, Class, Export](ESaveRealm HarvestingRealm)
+		{
+			FHarvestedRealm& RealmData = SaveContext.GetHarvestedRealm(HarvestingRealm);
+			if (!RealmData.IsIncluded(Class))
+			{
+				SaveContext.RecordIllegalReference(Export, Class, EIllegalRefReason::UnsaveableClass,
+					GetUnsaveableReason(Class, HarvestingRealm));
+			}
+		});
 
 	// Harvest the export outer
 	if (UObject* Outer = Export->GetOuter())
 	{
-		auto ShouldHarvestOuterAsDependencies = [this](UObject* InObject, UObject* InOuter)
+		// Harvest the outer as dependencies if the outer is not in the package or if the outer is a ref from optional to non optional object in an optional context
+		auto HarvestOuter = [this, &Outer](bool bShouldHarvestOuterAsDependencies)
 		{
-			// Harvest the outer as dependencies if the outer is not in the package or if the outer is a ref from optional to non optional object in an optional context
-			return !InOuter->IsInPackage(SaveContext.GetPackage()) ||
-				(CurrentExportHarvestingRealm == ESaveRealm::Optional && 
-					InObject->GetClass()->HasAnyClassFlags(CLASS_Optional) && !InOuter->GetClass()->HasAnyClassFlags(CLASS_Optional));
+			if (bShouldHarvestOuterAsDependencies)
+			{
+				*this << Outer;
+			}
+			else
+			{
+				// Legacy behavior does not add an export outer as a preload dependency if that outer is also an export since those are handled already by the EDL
+				FIgnoreDependenciesScope IgnoreDependencies(*this);
+				*this << Outer;
+			}
 		};
 
-		if (ShouldHarvestOuterAsDependencies(Export, Outer))
+		bool bShouldHarvestOuterAsDependencies = !Outer->IsInPackage(SaveContext.GetPackage());
+		if (!CurrentExportHarvestingRealms.Contains(ESaveRealm::Optional))
 		{
-			*this << Outer;
+			HarvestOuter(bShouldHarvestOuterAsDependencies);
 		}
 		else
 		{
-			// Legacy behavior does not add an export outer as a preload dependency if that outer is also an export since those are handled already by the EDL
-			FIgnoreDependenciesScope IgnoreDependencies(*this);
-			*this << Outer;
+			// The optional realm has to be harvested individually because we set flags separately for it
+			check(CurrentExportHarvestingRealms.Num() == 1);
+			bool bShouldHarvestOuterAsDependenciesInOptional = bShouldHarvestOuterAsDependencies ||
+				(Export->GetClass()->HasAnyClassFlags(CLASS_Optional) && !Outer->GetClass()->HasAnyClassFlags(CLASS_Optional));
+			HarvestOuter(bShouldHarvestOuterAsDependenciesInOptional);
 		}
-		if (!HarvestedRealm.IsIncluded(Outer))
-		{
-			// Only packages or object having the currently saved package as outer are allowed to have no outer
-			if (!Export->IsA<UPackage>() && Outer != SaveContext.GetPackage())
+
+		// Only packages or object having the currently saved package as outer are allowed to have no outer
+		ForEachExportHarvestingRealm([this, Outer, Export](ESaveRealm HarvestingRealm)
 			{
-				SaveContext.RecordIllegalReference(Export, Outer, EIllegalRefReason::UnsaveableOuter, GetUnsaveableReason(Outer));
-			}
-		}
+				FHarvestedRealm& RealmData = SaveContext.GetHarvestedRealm(HarvestingRealm);
+				if (!RealmData.IsIncluded(Outer))
+				{
+					if (!Export->IsA<UPackage>() && Outer != SaveContext.GetPackage())
+					{
+						SaveContext.RecordIllegalReference(Export, Outer, EIllegalRefReason::UnsaveableOuter,
+							GetUnsaveableReason(Outer, HarvestingRealm));
+					}
+				}
+			});
 	}
 
 	// Harvest its template, if any
 	UObject* Template = Export->GetArchetype();
-	if (Template
-		 && (Template != Class->GetDefaultObject() || SaveContext.IsCooking())
-		)
+	if (Template && (Template != Class->GetDefaultObject() || SaveContext.IsCooking()))
 	{
 		*this << Template;
 	}
@@ -312,133 +485,170 @@ void FPackageHarvester::ProcessExport(const FExportWithContext& InProcessContext
 
 void FPackageHarvester::TryHarvestExport(UObject* InObject)
 {
-	// Those should have been already validated
-	check(InObject && InObject->IsInPackage(SaveContext.GetPackage()));
+	check(!HasAnyExportHarvestingRealms());
 
-	// Get the realm in which we should harvest this export
-	EIllegalRefReason Reason = EIllegalRefReason::None;
-	ESaveRealm HarvestContext = GetObjectHarvestingRealm(InObject, Reason);
-	if (!SaveContext.GetHarvestedRealm(HarvestContext).IsExport(InObject))
+	FHarvestScope ScopeSetHarvestingRealms(*this);
+	FExportingRealmsArray HarvestingRealms;
+	// if the object is optional and we are cooking, harvest in the Optional context
+	if (ShouldObjectBeHarvestedInOptionalRealm(InObject))
 	{
-		SaveContext.MarkUnsaveable(InObject);
-		bool bExcluded = false;
-		if (!InObject->HasAnyFlags(RF_Transient))
-		{
-			bExcluded = ConditionallyExcludeObjectForTarget(SaveContext, InObject, HarvestContext);
-		}
-		if (!InObject->HasAnyFlags(RF_Transient) && !bExcluded)
+		HarvestingRealms.Add(ESaveRealm::Optional);
+	}
+	else
+	{
+		// For regular objects, exports are added by default to the game realm, and if making an editor
+		// save, also to the editor realm
+		HarvestingRealms.Add(ESaveRealm::Game);
+		if (!SaveContext.IsCooking())
 		{
-			// It passed filtering so mark as export
-			HarvestExport(InObject, HarvestContext);
+			HarvestingRealms.Add(ESaveRealm::Editor);
 		}
+	}
+	CurrentExportHarvestingRealms = HarvestingRealms;
 
-		// If we have a illegal ref reason, record it
-		if (Reason != EIllegalRefReason::None)
-		{
-			SaveContext.RecordIllegalReference(CurrentExportDependencies.CurrentExport, InObject, Reason);
-		}
+	TryHarvestExportInternal(InObject);
+}
+
+void FPackageHarvester::TryHarvestExportInternal(UObject* InObject)
+{
+	// Those should have been already validated
+	check(InObject && InObject->IsInPackage(SaveContext.GetPackage()));
+
+	// Switch to the Optional Realm if the Object is optional
+	FHarvestScope OptionalObjectScope = EnterConditionalOptionalObjectScope(InObject);
+
+#if WITH_EDITORONLY_DATA
+	// Remove the Game realm if the object is editoronly
+	bool bIsEditorOnlyObject = SavePackageUtilities::IsStrippedEditorOnlyObject(InObject, true /* bCheckRecursive */, false /* bCheckMarks */)
+		&& !InObject->HasNonEditorOnlyReferences();
+	FHarvestScope EditorOnlyScope = EnterConditionalEditorOnlyScope(bIsEditorOnlyObject);
+#endif
+
+	// Remove realms for which we have already harvested the export
+	FHarvestScope NewExportOnlyScope = EnterNewExportOnlyScope(InObject);
+	if (NewExportOnlyScope.IsEmpty())
+	{
+		return;
 	}
+
+	// Check whether the object is unsaveable and skip adding it as an export to any realm if so
+	SaveContext.MarkUnsaveable(InObject);
+	if (InObject->HasAnyFlags(RF_Transient))
+	{
+		return;
+	}
+
+	// Filter out any realms in which the export is excluded
+	FHarvestScope NotExcludedScope = EnterNotExcludedScope(InObject);
+	if (NotExcludedScope.IsEmpty())
+	{
+		return;
+	}
+
+	HarvestExport(InObject);
 }
 
 void FPackageHarvester::TryHarvestImport(UObject* InObject)
 {
 	// Those should have been already validated
-	check(InObject);
 	check(InObject && !InObject->IsInPackage(SaveContext.GetPackage()));
+	if (InObject == nullptr)
+	{
+		return;
+	}
 
-	if ( InObject==nullptr )
+	// Do not add unsaveable imports to any realm
+	if (SaveContext.IsUnsaveable(InObject))
 	{
 		return;
 	}
 
-	auto IsObjNative = [](UObject* InObj)
+	// Do not add imports in packages excluded from cooking to any realm
+	if (InObject && FCoreUObjectDelegates::ShouldCookPackageForPlatform.IsBound())
 	{
-		bool bIsNative = InObj->IsNative();
-		UObject* Outer = InObj->GetOuter();
-		while (!bIsNative && Outer)
+		if (!FCoreUObjectDelegates::ShouldCookPackageForPlatform.Execute(InObject->GetOutermost(), CookingTarget()))
 		{
-			bIsNative |= Cast<UClass>(Outer) != nullptr && Outer->IsNative();
-			Outer = Outer->GetOuter();
+			return;
 		}
-		return bIsNative;
-	};
+	}
 
-	bool bExcluded = ConditionallyExcludeObjectForTarget(SaveContext, InObject, CurrentExportHarvestingRealm);
-	bool bExcludePackageFromCook = InObject && FCoreUObjectDelegates::ShouldCookPackageForPlatform.IsBound() ? !FCoreUObjectDelegates::ShouldCookPackageForPlatform.Execute(InObject->GetOutermost(), CookingTarget()) : false;
-	if (!bExcludePackageFromCook && !bExcluded && !SaveContext.IsUnsaveable(InObject))
+	// Filter out any realms in which the export is excluded
+	FHarvestScope NotExcludedScope = EnterNotExcludedScope(InObject);
+	if (NotExcludedScope.IsEmpty())
 	{
-		bool bIsNative = IsObjNative(InObject);
-		HarvestImport(InObject);
+		return;
+	}
 
-		UObject* ObjOuter = InObject->GetOuter();
-		UClass* ObjClass = InObject->GetClass();
-		FName ObjName = InObject->GetFName();
-		if (SaveContext.IsCooking())
+	HarvestImport(InObject);
+	ProcessImport(InObject);
+}
+
+void FPackageHarvester::ProcessImport(UObject* InObject)
+{
+	bool bIsNative = IsObjNative(InObject);
+	UObject* ObjOuter = InObject->GetOuter();
+	UClass* ObjClass = InObject->GetClass();
+	FName ObjName = InObject->GetFName();
+	if (SaveContext.IsCooking())
+	{
+		// The ignore dependencies check is is necessary not to have infinite recursive calls
+		if (!bIsNative && !CurrentExportDependencies.bIgnoreDependencies)
 		{
-			// The ignore dependencies check is is necessary not to have infinite recursive calls
-			if (!bIsNative && !CurrentExportDependencies.bIgnoreDependencies)
+			UClass* ClassObj = Cast<UClass>(InObject);
+			UObject* CDO = ClassObj ? ClassObj->GetDefaultObject() : nullptr;
+			if (CDO)
 			{
-				UClass* ClassObj = Cast<UClass>(InObject);
-				UObject* CDO = ClassObj ? ClassObj->GetDefaultObject() : nullptr;
-				if (CDO)
-				{
-					FIgnoreDependenciesScope IgnoreDependencies(*this);
+				FIgnoreDependenciesScope IgnoreDependencies(*this);
 
-					// Gets all subobjects defined in a class, including the CDO, CDO components and blueprint-created components
-					TArray<UObject*> ObjectTemplates;
-					ObjectTemplates.Add(CDO);
-					SavePackageUtilities::GetCDOSubobjects(CDO, ObjectTemplates);
-					for (UObject* ObjTemplate : ObjectTemplates)
-					{
-						// Recurse into templates
-						*this << ObjTemplate;
-					}
+				// Gets all subobjects defined in a class, including the CDO, CDO components and blueprint-created components
+				TArray<UObject*> ObjectTemplates;
+				ObjectTemplates.Add(CDO);
+				SavePackageUtilities::GetCDOSubobjects(CDO, ObjectTemplates);
+				for (UObject* ObjTemplate : ObjectTemplates)
+				{
+					// Recurse into templates
+					*this << ObjTemplate;
 				}
 			}
 		}
+	}
 
-		// Harvest the import name
-		HarvestPackageHeaderName(ObjName);
+	// Harvest the import name
+	HarvestPackageHeaderName(ObjName);
 
-		// Recurse into outer, package override and non native class
-		if (ObjOuter)
-		{
-			*this << ObjOuter;
-		}
-		UPackage* Package = InObject->GetExternalPackage();
-		if (Package && Package != InObject)
-		{
-			*this << Package;
-		}
-		else
-		{
-			if (!IsFilterEditorOnly())
-			{
-				// operator<<(FStructuredArchive::FSlot Slot, FObjectImport& I) will need to write NAME_None for this empty ExternalPackage pointer
-				HarvestPackageHeaderName(NAME_None);
-			}
-		}
-
-		// For things with a BP-created class we need to recurse into that class so the import ClassPackage will load properly
-		// We don't do this for native classes to avoid bloating the import table, but we need to harvest their name and outer (package) name
-		if (!ObjClass->IsNative())
-		{
-			*this << ObjClass; 
-		}	
-		else
+	// Recurse into outer, package override and non native class
+	if (ObjOuter)
+	{
+		*this << ObjOuter;
+	}
+	UPackage* Package = InObject->GetExternalPackage();
+	if (Package && Package != InObject)
+	{
+		*this << Package;
+	}
+	else
+	{
+		if (!IsFilterEditorOnly())
 		{
-			HarvestPackageHeaderName(ObjClass->GetFName());
-			HarvestPackageHeaderName(ObjClass->GetOuter()->GetFName());
+			// operator<<(FStructuredArchive::FSlot Slot, FObjectImport& I) will need to write NAME_None for this empty ExternalPackage pointer
+			HarvestPackageHeaderName(NAME_None);
 		}
 	}
 
-	// Check for illegal reference
-	EIllegalRefReason Reason = EIllegalRefReason::None;
-	GetObjectHarvestingRealm(InObject, Reason);
-	if (Reason != EIllegalRefReason::None)
+	// For things with a BP-created class we need to recurse into that class so the import ClassPackage will load properly
+	// We don't do this for native classes to avoid bloating the import table, but we need to harvest their name and outer (package) name
+	if (!ObjClass->IsNative())
 	{
-		SaveContext.RecordIllegalReference(CurrentExportDependencies.CurrentExport, InObject, Reason);
+		*this << ObjClass; 
+	}	
+	else
+	{
+		HarvestPackageHeaderName(ObjClass->GetFName());
+		HarvestPackageHeaderName(ObjClass->GetOuter()->GetFName());
 	}
+
+	// If we have an illegal reference to an optional object, record it
+	EnterConditionalOptionalObjectScope(InObject);
 }
 
 FString FPackageHarvester::GetArchiveName() const
@@ -461,9 +671,20 @@ void FPackageHarvester::MarkSearchableName(const TObjectPtr<const UObject>& Type
 }
 
 FArchive& FPackageHarvester::operator<<(UObject*& Obj)
-{	
-	// if the object is null or already marked excluded, we can skip the harvest
-	if (!Obj || SaveContext.GetHarvestedRealm(CurrentExportHarvestingRealm).IsExcluded(Obj))
+{
+	// if the object is null, skip the harvest
+	if (!Obj)
+	{
+		return *this;
+	}
+
+#if WITH_EDITORONLY_DATA
+	FHarvestScope EditorOnlyScope = EnterConditionalEditorOnlyScope(IsEditorOnlyPropertyOnTheStack());
+#endif
+
+	// if the object is already marked excluded, skip the harvest
+	FHarvestScope NotPreviouslyExcludedScope = EnterNotPreviouslyExcludedScope(Obj);
+	if (NotPreviouslyExcludedScope.IsEmpty())
 	{
 		return *this;
 	}
@@ -478,7 +699,7 @@ FArchive& FPackageHarvester::operator<<(UObject*& Obj)
 	// if the object is in the save context package, try to tag it as export
 	if (Obj->IsInPackage(SaveContext.GetPackage()))
 	{
-		TryHarvestExport(Obj);
+		TryHarvestExportInternal(Obj);
 	}
 	// Otherwise visit the import
 	else
@@ -486,19 +707,9 @@ FArchive& FPackageHarvester::operator<<(UObject*& Obj)
 		TryHarvestImport(Obj);
 	}
 
-	auto IsObjNative = [](UObject* InObj)
-	{
-		bool bIsNative = InObj->IsNative();
-		UObject* Outer = InObj->GetOuter();
-		while (!bIsNative && Outer)
-		{
-			bIsNative |= Cast<UClass>(Outer) != nullptr && Outer->IsNative();
-			Outer = Outer->GetOuter();
-		}
-		return bIsNative;
-	};
-
-	if (SaveContext.GetHarvestedRealm(CurrentExportHarvestingRealm).IsIncluded(Obj))
+	// Add a dependency from the current export to Obj if Obj was added as an import or export in any realm
+	FHarvestScope ObjIncludedScope = EnterIncludedScope(Obj);
+	if (!ObjIncludedScope.IsEmpty())
 	{
 		HarvestDependency(Obj, IsObjNative(Obj));
 	}
@@ -540,8 +751,12 @@ FArchive& FPackageHarvester::operator<<(FSoftObjectPath& Value)
 {
 	// We need to harvest NAME_None even if the path isn't valid
 	Value.SerializePath(*this);
-	// Add the soft object path to the list, we need to map invalid soft object path too
-	SaveContext.GetHarvestedRealm(CurrentExportHarvestingRealm).GetSoftObjectPathList().Add(Value);
+	// Add the soft object path to the realm's list. We need to map even invalid and uncollected soft object paths
+	// so that they can be mapped to an index when serializing the exports
+	ForEachExportHarvestingRealm([this, &Value](ESaveRealm HarvestingRealm)
+		{
+			SaveContext.GetHarvestedRealm(HarvestingRealm).GetSoftObjectPathList().Add(Value);
+		});
 	if (Value.IsValid())
 	{
 		FSoftObjectPathThreadContext& ThreadContext = FSoftObjectPathThreadContext::Get();
@@ -549,21 +764,23 @@ FArchive& FPackageHarvester::operator<<(FSoftObjectPath& Value)
 		ESoftObjectPathCollectType CollectType = ESoftObjectPathCollectType::AlwaysCollect;
 		ESoftObjectPathSerializeType SerializeType = ESoftObjectPathSerializeType::AlwaysSerialize;
 
+		// Don't track as a soft reference to a package if this is a NeverCollect SoftObjectPath
 		ThreadContext.GetSerializationOptions(ReferencingPackageName, ReferencingPropertyName, CollectType, SerializeType, this);
 		if (CollectType != ESoftObjectPathCollectType::NeverCollect && CollectType != ESoftObjectPathCollectType::NonPackage)
 		{
-			// Don't track if this is a never collect path
+#if WITH_EDITORONLY_DATA
+			// Don't track as a used-in-game soft reference to a package if this is an EditorOnly SoftObjectPath
+			// CollectType takes into account FSoftObjectPathSerializationScopes and
+			// this->FArchiveState::IsEditorOnlyPropertyOnTheStack
+			FHarvestScope EditorOnlyScope = EnterConditionalEditorOnlyScope(CollectType == ESoftObjectPathCollectType::EditorOnlyCollect);
+#endif
 			FString Path = Value.ToString();
 			FName PackageName = FName(*FPackageName::ObjectPathToPackageName(Path));
 			HarvestPackageHeaderName(PackageName);
-			SaveContext.GetHarvestedRealm(CurrentExportHarvestingRealm).GetSoftPackageReferenceList().Add(PackageName);
-#if WITH_EDITORONLY_DATA
-			if (CollectType != ESoftObjectPathCollectType::EditorOnlyCollect && !bIsEditorOnlyExportOnStack 
-				&& CurrentExportHarvestingRealm != ESaveRealm::Optional)
-#endif
-			{
-				SaveContext.GetHarvestedRealm(ESaveRealm::Game).GetSoftPackageReferenceList().Add(PackageName);
-			}
+			ForEachExportHarvestingRealm([this, PackageName](ESaveRealm HarvestingRealm)
+				{
+					SaveContext.GetHarvestedRealm(HarvestingRealm).GetSoftPackageReferenceList().Add(PackageName);
+				});
 		}
 	}
 	return *this;
@@ -653,147 +870,221 @@ void FPackageHarvester::HarvestDependency(UObject* InObj, bool bIsNative)
 	}
 }
 
-bool FPackageHarvester::CurrentExportHasDependency(UObject* InObj) const
+bool FPackageHarvester::CurrentExportHasDependency(UObject* InObj, ESaveRealm HarvestingRealm) const
 {
-	return SaveContext.GetHarvestedRealm(CurrentExportHarvestingRealm).GetObjectDependencies().Contains(InObj) || SaveContext.GetHarvestedRealm(CurrentExportHarvestingRealm).GetNativeObjectDependencies().Contains(InObj);
+	FHarvestedRealm& RealmData = SaveContext.GetHarvestedRealm(HarvestingRealm);
+	return RealmData.GetObjectDependencies().Contains(InObj) || RealmData.GetNativeObjectDependencies().Contains(InObj);
 }
 
 void FPackageHarvester::HarvestExportDataName(FName Name)
 {
-	SaveContext.GetHarvestedRealm(CurrentExportHarvestingRealm).GetNamesReferencedFromExportData().Add(Name.GetDisplayIndex());
+	ForEachExportHarvestingRealm([this, Name](ESaveRealm HarvestingRealm)
+		{
+			SaveContext.GetHarvestedRealm(HarvestingRealm).GetNamesReferencedFromExportData().Add(Name.GetDisplayIndex());
+		});
 }
 
 void FPackageHarvester::HarvestPackageHeaderName(FName Name)
 {
-	SaveContext.GetHarvestedRealm(CurrentExportHarvestingRealm).GetNamesReferencedFromPackageHeader().Add(Name.GetDisplayIndex());
+	ForEachExportHarvestingRealm([this, Name](ESaveRealm HarvestingRealm)
+		{
+			SaveContext.GetHarvestedRealm(HarvestingRealm).GetNamesReferencedFromPackageHeader().Add(Name.GetDisplayIndex());
+		});
 }
 
 void FPackageHarvester::HarvestSearchableName(UObject* TypeObject, FName Name)
 {
 	// Make sure the object is tracked as a dependency
-	if (!CurrentExportHasDependency(TypeObject))
+	bool bAlreadyTrackedInAllHarvestingRealms = true;
+	ForEachExportHarvestingRealm([this, &bAlreadyTrackedInAllHarvestingRealms, TypeObject](ESaveRealm HarvestingRealm)
+		{
+			bAlreadyTrackedInAllHarvestingRealms &= CurrentExportHasDependency(TypeObject, HarvestingRealm);
+		});
+	if (!bAlreadyTrackedInAllHarvestingRealms)
 	{
 		(*this) << TypeObject;
 	}
 
 	HarvestPackageHeaderName(Name);
-	SaveContext.GetHarvestedRealm(CurrentExportHarvestingRealm).GetSearchableNamesObjectMap().FindOrAdd(TypeObject).AddUnique(Name);
+	ForEachExportHarvestingRealm([this, TypeObject, Name](ESaveRealm HarvestingRealm)
+		{
+			FHarvestedRealm& RealmData = SaveContext.GetHarvestedRealm(HarvestingRealm);
+			RealmData.GetSearchableNamesObjectMap().FindOrAdd(TypeObject).AddUnique(Name);
+		});
 }
 
-ESaveRealm FPackageHarvester::GetObjectHarvestingRealm(UObject* InObject, EIllegalRefReason& OutReason) const
+void FPackageHarvester::HarvestExport(UObject* InObject)
 {
-	OutReason = EIllegalRefReason::None;
-	switch (CurrentExportHarvestingRealm)
-	{
-	// We are harvesting InObject from the root (i.e. asset or top level flag)
-	case ESaveRealm::None:
-		// if the object is optional and we are cooking, harvest in the Optional context
-		if (InObject->GetClass()->HasAnyClassFlags(CLASS_Optional) && SaveContext.IsCooking())
-		{
-			return ESaveRealm::Optional;
-		}
-
-		// Otherwise, just return the current default harvesting context (i.e. Game while cooking, Editor otherwise)
-		return SaveContext.CurrentHarvestingRealm;
-
-	// We are harvesting InObject from an optional object
-	case ESaveRealm::Optional:
-		// whatever the type of object we are harvesting, harvest it in the Optional context
-		return ESaveRealm::Optional;
-
-	// We are harvesting InObject from a game or editor only object
-	default:
-		//@todo FH: check CanSkipEditorReferencedPackagesWhenCooking to propagate the editor context when trimming editor reference and potentially skip processing
+	bool bFromOptionalRef = CurrentExportDependencies.CurrentExport &&
+		CurrentExportDependencies.CurrentExport->GetClass()->HasAnyClassFlags(CLASS_Optional);
+	FTaggedExport TaggedExport(InObject, !DoesObjectNeedLoadForEditorGame(InObject), bFromOptionalRef);
 
-		// if we are harvesting an optional object while in a different context, record an illegal reference to display on validation
-		if (InObject->GetClass()->HasAnyClassFlags(CLASS_Optional) && SaveContext.IsCooking())
+	ForEachExportHarvestingRealm([this, &TaggedExport, InObject](ESaveRealm HarvestingRealm)
 		{
-			OutReason = EIllegalRefReason::ReferenceToOptional;
-			return ESaveRealm::Optional;
-		}
-		// Otherwise propagate the current context
-		return CurrentExportHarvestingRealm;
-	}
-}
-
-void FPackageHarvester::HarvestExport(UObject* InObject, ESaveRealm InContext)
-{
-	bool bFromOptionalRef = CurrentExportDependencies.CurrentExport && CurrentExportDependencies.CurrentExport->GetClass()->HasAnyClassFlags(CLASS_Optional);
-	SaveContext.GetHarvestedRealm(InContext)
-		.AddExport(FTaggedExport(InObject, !DoesObjectNeedLoadForEditorGame(InObject), bFromOptionalRef));
-	SaveContext.GetHarvestedRealm(InContext).GetNamesReferencedFromPackageHeader().Add(InObject->GetFName().GetDisplayIndex());
-	ExportsToProcess.Enqueue({ InObject, InContext });
+			FHarvestedRealm& RealmData = SaveContext.GetHarvestedRealm(HarvestingRealm);
+			RealmData.AddExport(TaggedExport);
+			RealmData.GetNamesReferencedFromPackageHeader().Add(InObject->GetFName().GetDisplayIndex());
+		});
+	ExportsToProcess.Enqueue({ InObject, CurrentExportHarvestingRealms });
 }
 
 void FPackageHarvester::HarvestImport(UObject* InObject)
 {
-	bool bIsEditorOnly = false;
-#if WITH_EDITORONLY_DATA
-	bIsEditorOnly = bIsEditorOnlyExportOnStack || IsEditorOnlyPropertyOnTheStack() || CurrentExportHarvestingRealm == ESaveRealm::Optional;
-#endif
-	SaveContext.GetHarvestedRealm(CurrentExportHarvestingRealm).AddImport(InObject);
-	// No matter the current context, if the import is not editor only also add it to the game context, this is later used in asset registry saving
-	if (!bIsEditorOnly)
-	{
-		SaveContext.GetHarvestedRealm(ESaveRealm::Game).AddImport(InObject);
-	}
+	ForEachExportHarvestingRealm([this, InObject](ESaveRealm HarvestingRealm)
+		{
+			SaveContext.GetHarvestedRealm(HarvestingRealm).AddImport(InObject);
+		});
 }
 
 void FPackageHarvester::AppendCurrentExportDependencies()
 {
 	check(CurrentExportDependencies.CurrentExport);
-	SaveContext.GetHarvestedRealm(CurrentExportHarvestingRealm).GetObjectDependencies().Add(CurrentExportDependencies.CurrentExport, MoveTemp(CurrentExportDependencies.ObjectReferences));
-	SaveContext.GetHarvestedRealm(CurrentExportHarvestingRealm).GetNativeObjectDependencies().Add(CurrentExportDependencies.CurrentExport, MoveTemp(CurrentExportDependencies.NativeObjectReferences));
+	ForEachExportHarvestingRealm([this](ESaveRealm HarvestingRealm)
+		{
+			FHarvestedRealm& RealmData = SaveContext.GetHarvestedRealm(HarvestingRealm);
+			RealmData.GetObjectDependencies().Add(CurrentExportDependencies.CurrentExport,
+				MoveTemp(CurrentExportDependencies.ObjectReferences));
+			RealmData.GetNativeObjectDependencies().Add(CurrentExportDependencies.CurrentExport,
+				MoveTemp(CurrentExportDependencies.NativeObjectReferences));
+		});
 	CurrentExportDependencies.CurrentExport = nullptr;
 }
 
-FString FPackageHarvester::GetUnsaveableReason(UObject* Required)
+FString FPackageHarvester::GetUnsaveableReason(UObject* Required, ESaveRealm RealmInWhichItIsUnsaveable)
 {
-	// Copy some of the code from operator<<(UObject*), TryHarvestExport, TryHarvestImport to find out why the Required object was not included
-	FString ReasonText(TEXTVIEW("It should be included but was excluded for an unknown reason."));
-	EIllegalRefReason UnusedRealmReason = EIllegalRefReason::None;
-	ESaveRealm HarvestContext = GetObjectHarvestingRealm(Required, UnusedRealmReason);
+	UObject* Culprit;
+	FString Reason;
+	ESaveableStatus Status = GetSaveableStatusForRealm(Required, RealmInWhichItIsUnsaveable, Culprit, Reason);
+	if (Status != ESaveableStatus::Success)
+	{
+		return FString::Printf(TEXT("It %s."), *Reason);
+	}
+
 	bool bShouldBeExport = Required->IsInPackage(SaveContext.GetPackage());
+	if (bShouldBeExport)
+	{
+		return FString(TEXTVIEW("It should be an export but was excluded for an unknown reason."));
+	}
+	else
+	{
+		return FString(TEXTVIEW("It should be an import but was excluded for an unknown reason."));
 
-	FSaveContext::ESaveableStatus CulpritStatus;
-	UObject* Culprit;
-	FSaveContext::ESaveableStatus Status = SaveContext.GetSaveableStatus(Required, &Culprit, &CulpritStatus);
-	if (Status != FSaveContext::ESaveableStatus::Success)
+	}
+}
+
+ESaveableStatus FPackageHarvester::GetSaveableStatusForRealm(UObject* Obj, ESaveRealm RealmInWhichItIsUnsaveable,
+	UObject*& OutCulprit, FString& OutReason)
+{
+	// Copy some of the code from operator<<(UObject*), TryHarvestExport, TryHarvestImport to find out why the Required object was not included
+	if (!Obj || Obj->GetOutermost()->GetFName() == GLongCoreUObjectPackageName)
 	{
-		if (Status == FSaveContext::ESaveableStatus::OuterUnsaveable)
+		// No object or in CoreUObject, it is saveable
+		return ESaveableStatus::Success;
+	}
+
+	ESaveableStatus CulpritStatus;
+	ESaveableStatus Status = SaveContext.GetSaveableStatus(Obj, &OutCulprit, &CulpritStatus);
+	if (Status != ESaveableStatus::Success)
+	{
+		if (Status == ESaveableStatus::OuterUnsaveable)
 		{
-			check(Culprit);
-			ReasonText = FString::Printf(TEXT("It has outer %s which %s."), *Culprit->GetPathName(), LexToString(CulpritStatus));
+			check(OutCulprit);
+			OutReason = FString::Printf(TEXT("has outer %s which %s"), *OutCulprit->GetPathName(),
+				LexToString(CulpritStatus));
 		}
 		else
 		{
-			ReasonText =  FString::Printf(TEXT("It %s"), LexToString(Status));
+			OutCulprit = Obj;
+			OutReason = LexToString(Status);
 		}
+		return Status;
 	}
-	else if (ConditionallyExcludeObjectForTarget(SaveContext, Required, HarvestContext))
+
+	const EObjectMark ExcludedObjectMarks = SaveContext.GetExcludedObjectMarks(RealmInWhichItIsUnsaveable);
+	const ITargetPlatform* TargetPlatform = SaveContext.GetTargetPlatform();
+	EObjectMark ObjectMarks = GenerateMarksForObject(Obj, TargetPlatform);
+	if (!!(ObjectMarks & ExcludedObjectMarks))
 	{
-		ReasonText = TEXTVIEW("It is excluded for the current cooking target.");
+		OutCulprit = Obj;
+		OutReason = FString::Printf(TEXT("has ObjectMarks 0x%x that are excluded for the current cooking target"),
+			(int32)(ObjectMarks & ExcludedObjectMarks));
+		return ESaveableStatus::ExcludedByPlatform;
 	}
-	else if (bShouldBeExport)
+
+	UObject* ObjOuter = Obj->GetOuter();
+	FString RecursiveReason;
+	ESaveableStatus RecursiveStatus;
+	if (ObjOuter)
 	{
-		// The class is in the package and so should be an export; we don't know of any other reasons why it would be excluded
-		ReasonText = TEXTVIEW("It should be an export but was excluded for an unknown reason.");
+		RecursiveStatus = GetSaveableStatusForRealm(ObjOuter, RealmInWhichItIsUnsaveable, OutCulprit, RecursiveReason);
+		if (RecursiveStatus != ESaveableStatus::Success)
+		{
+			check(OutCulprit);
+			if (RecursiveStatus == ESaveableStatus::OuterUnsaveable)
+			{
+				OutReason = MoveTemp(RecursiveReason);
+			}
+			else
+			{
+				OutReason = FString::Printf(TEXT("has outer %s which %s"), *OutCulprit->GetPathName(), *RecursiveReason);
+			}
+			return ESaveableStatus::OuterUnsaveable;
+		}
 	}
-	else
+
+	UClass* ObjClass = Obj->GetClass();
+	if (ObjClass)
 	{
-		// The class is not in the package and so should be an import
-		bool bExcludePackageFromCook = FCoreUObjectDelegates::ShouldCookPackageForPlatform.IsBound() ?
-			!FCoreUObjectDelegates::ShouldCookPackageForPlatform.Execute(Required->GetOutermost(), CookingTarget()) : false;
-		if (bExcludePackageFromCook)
+		RecursiveStatus = GetSaveableStatusForRealm(ObjClass, RealmInWhichItIsUnsaveable, OutCulprit, RecursiveReason);
+		if (RecursiveStatus != ESaveableStatus::Success)
 		{
-			ReasonText = FString::Printf(TEXT("It is in package %s which is excluded from the cook by FCoreUObjectDelegates::ShouldCookPackageForPlatform."),
-				*Required->GetOutermost()->GetName());
+			OutCulprit = ObjClass;
+			OutReason = FString::Printf(TEXT("has class %s which %s"), *ObjClass->GetPathName(), *RecursiveReason);
+			return ESaveableStatus::ClassUnsaveable;
 		}
-		else
+	}
+
+	// Check parent struct if we have one
+	UStruct* ThisStruct = Cast<UStruct>(Obj);
+	UObject* SuperStruct = ThisStruct ? ThisStruct->GetSuperStruct() : nullptr;
+	if (SuperStruct)
+	{
+		RecursiveStatus = GetSaveableStatusForRealm(SuperStruct, RealmInWhichItIsUnsaveable, OutCulprit, RecursiveReason);
+		if (RecursiveStatus != ESaveableStatus::Success)
+		{
+			OutCulprit = SuperStruct;
+			OutReason = FString::Printf(TEXT("has superclass %s which %s"), *SuperStruct->GetPathName(), *RecursiveReason);
+			return ESaveableStatus::ClassUnsaveable;
+		}
+	}
+
+	// Check archetype, this may not have been covered in the case of components
+	UObject* Archetype = Obj->GetArchetype();
+	if (Archetype)
+	{
+		RecursiveStatus = GetSaveableStatusForRealm(Archetype, RealmInWhichItIsUnsaveable, OutCulprit, RecursiveReason);
+		if (RecursiveStatus != ESaveableStatus::Success)
 		{
-			// We don't know of any other reasons why it would be excluded
-			ReasonText = TEXTVIEW("It should be an import but was excluded for an unknown reason.");
+			OutCulprit = Archetype;
+			OutReason = FString::Printf(TEXT("has archetype %s which %s"), *Archetype->GetPathName(), *RecursiveReason);
+			return ESaveableStatus::ClassUnsaveable;
 		}
 	}
-	return ReasonText;
+
+	bool bShouldBeExport = Obj->IsInPackage(SaveContext.GetPackage());
+	if (!bShouldBeExport)
+	{
+		// Imports additionally can be excluded if their package is excluded for the platform
+		bool bExcludePackageFromCook = FCoreUObjectDelegates::ShouldCookPackageForPlatform.IsBound() ?
+			!FCoreUObjectDelegates::ShouldCookPackageForPlatform.Execute(Obj->GetOutermost(), CookingTarget()) : false;
+		if (bExcludePackageFromCook)
+		{
+			OutCulprit = Obj;
+			OutReason = FString::Printf(TEXT("is in package %s which is excluded from the cook by FCoreUObjectDelegates::ShouldCookPackageForPlatform."),
+				*Obj->GetOutermost()->GetName());
+			return ESaveableStatus::ExcludedByPlatform;
+		}
+	}
+
+	return ESaveableStatus::Success;
 }
\ No newline at end of file
diff --git a/Engine/Source/Runtime/CoreUObject/Private/UObject/SavePackage/PackageHarvester.h b/Engine/Source/Runtime/CoreUObject/Private/UObject/SavePackage/PackageHarvester.h
--- a/Engine/Source/Runtime/CoreUObject/Private/UObject/SavePackage/PackageHarvester.h
+++ b/Engine/Source/Runtime/CoreUObject/Private/UObject/SavePackage/PackageHarvester.h
@@ -8,6 +8,7 @@
 #include "Serialization/ArchiveUObject.h"
 #include "UObject/NameTypes.h"
 
+enum class ESaveableStatus;
 enum class ESaveRealm : uint32;
 enum class EIllegalRefReason : uint8;
 class FSaveContext;
@@ -23,12 +24,13 @@ struct FSoftObjectPath;
 class FPackageHarvester : public FArchiveUObject
 {
 public:
+	using FExportingRealmsArray = TArray<ESaveRealm, TInlineAllocator<2>>;
 	struct FExportWithContext
 	{
 		/* Export to process */
 		UObject* Export = nullptr;
-		/* The harvesting context from which this export was harvested. */
-		ESaveRealm HarvestedFromRealm;
+		/* The realms from which this export was harvested. */
+		FExportingRealmsArray HarvestedFromRealms;
 
 		operator bool() const
 		{
@@ -36,34 +38,21 @@ public:
 		}
 	};
 
-	class FExportScope
-	{
-	public:
-		FExportScope(FPackageHarvester& InHarvester, const FExportWithContext& InToProcess, bool bIsEditorOnlyObject);
-		~FExportScope();
-	private:
-		FPackageHarvester& Harvester;
-		ESaveRealm PreviousRealm;
-		bool bPreviousFilterEditorOnly;
-	};
-
-	class FIgnoreDependenciesScope
+	/** Scope object constructed when changing the CurrentExportHarvestingRealms that restores the previous realms when done */
+	class FHarvestScope
 	{
 	public:
-		FIgnoreDependenciesScope(FPackageHarvester& InHarvester)
-			: Harvester(InHarvester)
-			, bPreviousValue(Harvester.CurrentExportDependencies.bIgnoreDependencies)
-		{
-			Harvester.CurrentExportDependencies.bIgnoreDependencies = true;
-		}
-
-		~FIgnoreDependenciesScope()
-		{
-			Harvester.CurrentExportDependencies.bIgnoreDependencies = bPreviousValue;
-		}
+		FHarvestScope(FPackageHarvester& InHarvester);
+		FHarvestScope(FHarvestScope&& Other);
+		FHarvestScope(const FHarvestScope& Other) = delete;
+		FHarvestScope& operator=(FHarvestScope&& Other) = delete;
+		FHarvestScope& operator=(const FHarvestScope& Other) = delete;
+		~FHarvestScope();
+		bool IsEmpty() const;
 	private:
 		FPackageHarvester& Harvester;
-		bool bPreviousValue;
+		FExportingRealmsArray PreviousExportHarvestingRealms;
+		bool bActive;
 	};
 
 public:
@@ -71,16 +60,17 @@ public:
 
 	FExportWithContext PopExportToProcess();
 
-	void ProcessExport(const FExportWithContext& InProcessContext);
 	void TryHarvestExport(UObject* InObject);
-	void TryHarvestImport(UObject* InObject);
+	/** Add objects/names/others that are referenced when an object is being saved */
+	void ProcessExport(const FExportWithContext& InProcessContext);
+	FHarvestScope EnterRootReferencesScope();
 
 	void HarvestExportDataName(FName Name);
 	void HarvestPackageHeaderName(FName Name);
 	void HarvestSearchableName(UObject* TypeObject, FName Name);
 	void HarvestDependency(UObject* InObj, bool bIsNative);
 
-	bool CurrentExportHasDependency(UObject* InObj) const;
+	bool CurrentExportHasDependency(UObject* InObj, ESaveRealm HarvestingRealm) const;
 	TMap<UObject*, TSet<FProperty*>> ReleaseTransientPropertyOverrides();
 
 	// FArchiveUObject implementation
@@ -95,17 +85,7 @@ public:
 	virtual bool ShouldSkipProperty(const FProperty* InProperty) const;
 
 private:
-	void ResolveOverrides();
-
-	ESaveRealm GetObjectHarvestingRealm(UObject* InObject, EIllegalRefReason& OutReason) const;
-
-	void HarvestExport(UObject* InObject, ESaveRealm InContext);
-	void HarvestImport(UObject* InObject);
-
-	void AppendCurrentExportDependencies();
-	FString GetUnsaveableReason(UObject* Required);
-
-
+	/** Dependencies collected for the current export, and scope information about how to collect them. */
 	struct FExportDependencies
 	{
 		UObject* CurrentExport = nullptr;
@@ -114,12 +94,73 @@ private:
 		bool bIgnoreDependencies = false;
 	};
 
-	FSaveContext& SaveContext;
+	/**
+	 * Set dependencies and other storage to write to the given export, and set
+	 * CurrentExportHarvestingRealms and other data based on where the export was harvested.
+	 */
+	class FExportScope
+	{
+	public:
+		FExportScope(FPackageHarvester& InHarvester, const FExportWithContext& InToProcess);
+		~FExportScope();
+	private:
+		FPackageHarvester& Harvester;
+		FExportingRealmsArray PreviousExportHarvestingRealms;
+		bool bPreviousFilterEditorOnly;
+	};
 
+	/** Mark that dependencies should be ignored when harvesting further imports/exports in this scope. */
+	class FIgnoreDependenciesScope
+	{
+	public:
+		FIgnoreDependenciesScope(FPackageHarvester& InHarvester);
+		~FIgnoreDependenciesScope();
+	private:
+		FPackageHarvester& Harvester;
+		bool bPreviousValue;
+	};
+
+private:
+	/** Return whether the Object and all its Outers are native. */
+	static bool IsObjNative(UObject* InObj);
+	/** Report whether the object should be put in the optional realm even if harvesting realm is non-optional. */
+	bool ShouldObjectBeHarvestedInOptionalRealm(UObject* InObj) const;
+
+	/** Call void (ESaveRealm HarvestingRealm) on each Realm in CurrentExportHarvestingRealms */
+	template <typename CallbackType>
+	void ForEachExportHarvestingRealm(CallbackType&& Callback);
+	bool HasAnyExportHarvestingRealms();
+
+	void ResolveOverrides();
+
+	void TryHarvestExportInternal(UObject* InObject);
+	void HarvestExport(UObject* InObject);
+	void TryHarvestImport(UObject* InObject);
+	void HarvestImport(UObject* InObject);
+	/** Add objects/names/others that are referenced even when an object is just being imported */
+	void ProcessImport(UObject* InObject);
+
+	void AppendCurrentExportDependencies();
+	FString GetUnsaveableReason(UObject* Required, ESaveRealm RealmInWhichItIsUnsaveable);
+	ESaveableStatus GetSaveableStatusForRealm(UObject* Obj, ESaveRealm RealmInWhichItIsUnsaveable,
+		UObject*& OutCulprit, FString& OutReason);
+
+	/** If bIsEditorOnly is true, remove CurrentExportHarvestingRealms that do not follow editoronly references. */
+	FHarvestScope EnterConditionalEditorOnlyScope(bool bIsEditorOnly);
+	/** If the Object is optional, clear CurrentExportHarvestingRealms and populate it with the Optional Realm. */
+	FHarvestScope EnterConditionalOptionalObjectScope(UObject* Object);
+	/** Remove CurrentExportHarvestingRealms that already include Export. */
+	FHarvestScope EnterNewExportOnlyScope(UObject* Export);
+	/** Remove CurrentExportHarvestingRealms that return false for ConditionallyExcludeObjectForRealm(Object). */
+	FHarvestScope EnterNotExcludedScope(UObject* Object);
+	/** Remove CurrentExportHarvestingRealms that return true for IsExcluded(Object). */
+	FHarvestScope EnterNotPreviouslyExcludedScope(UObject* Object);
+	/** Remove CurrentExportHarvestingRealms that return false for IsIncluded(Object). */
+	FHarvestScope EnterIncludedScope(UObject* Object);
+
+	FSaveContext& SaveContext;
 	TQueue<FExportWithContext> ExportsToProcess;
 	FExportDependencies CurrentExportDependencies;
 	TMap<UObject*, TSet<FProperty*>> TransientPropertyOverrides;
-	ESaveRealm CurrentExportHarvestingRealm;
-	//@todo: bIsEditorOnlyExportOnStack can be probably be folded in CurrentExportHarvestingContext
-	bool bIsEditorOnlyExportOnStack;
+	FExportingRealmsArray CurrentExportHarvestingRealms;
 };
\ No newline at end of file
diff --git a/Engine/Source/Runtime/CoreUObject/Private/UObject/SavePackage/SaveContext.cpp b/Engine/Source/Runtime/CoreUObject/Private/UObject/SavePackage/SaveContext.cpp
--- a/Engine/Source/Runtime/CoreUObject/Private/UObject/SavePackage/SaveContext.cpp
+++ b/Engine/Source/Runtime/CoreUObject/Private/UObject/SavePackage/SaveContext.cpp
@@ -66,7 +66,7 @@ bool FSaveContext::IsUnsaveable(UObject* InObject, bool bEmitWarning) const
 	return true;
 }
 
-FSaveContext::ESaveableStatus FSaveContext::GetSaveableStatus(UObject* InObject, UObject** OutCulprit, ESaveableStatus* OutCulpritStatus) const
+ESaveableStatus FSaveContext::GetSaveableStatus(UObject* InObject, UObject** OutCulprit, ESaveableStatus* OutCulpritStatus) const
 {
 	UObject* Obj = InObject;
 	while (Obj)
@@ -97,7 +97,7 @@ FSaveContext::ESaveableStatus FSaveContext::GetSaveableStatus(UObject* InObject,
 	return ESaveableStatus::Success;
 }
 
-FSaveContext::ESaveableStatus FSaveContext::GetSaveableStatusNoOuter(UObject* Obj) const
+ESaveableStatus FSaveContext::GetSaveableStatusNoOuter(UObject* Obj) const
 {
 	// pending kill object are unsaveable
 	if (!IsValidChecked(Obj))
@@ -127,6 +127,32 @@ FSaveContext::ESaveableStatus FSaveContext::GetSaveableStatusNoOuter(UObject* Ob
 	return ESaveableStatus::Success;
 }
 
+FSavePackageResultStruct FSaveContext::GetFinalResult()
+{
+	if (Result != ESavePackageResult::Success)
+	{
+		return Result;
+	}
+
+	ESavePackageResult FinalResult = IsStubRequested() ? ESavePackageResult::GenerateStub : ESavePackageResult::Success;
+	FSavePackageResultStruct ResultData(FinalResult, TotalPackageSizeUncompressed,
+		SerializedPackageFlags, IsCompareLinker() ? MoveTemp(GetHarvestedRealm().Linker) : nullptr);
+
+	ResultData.SavedAssets = MoveTemp(SavedAssets);
+	UClass* PackageClass = UPackage::StaticClass();
+	for (UObject* Import : GetImports())
+	{
+		if (Import->IsA(PackageClass))
+		{
+			ResultData.ImportPackages.Add(Import->GetFName());
+		}
+	}
+	TSet<FName>& SoftPackageReferenceList = GetSoftPackageReferenceList();
+	ResultData.SoftPackageReferences = SoftPackageReferenceList.Array();
+
+	return ResultData;
+}
+
 namespace
 {
 	TArray<UClass*> AutomaticOptionalInclusionAssetTypeList;
@@ -165,18 +191,32 @@ void FSaveContext::SetupHarvestingRealms()
 	}
 }
 
-const TCHAR* LexToString(FSaveContext::ESaveableStatus Status)
+EObjectMark FSaveContext::GetExcludedObjectMarksForGameRealm(const ITargetPlatform* TargetPlatform)
+{
+	if (TargetPlatform)
+	{
+		return SavePackageUtilities::GetExcludedObjectMarksForTargetPlatform(TargetPlatform);
+	}
+	else
+	{
+		return static_cast<EObjectMark>(OBJECTMARK_NotForTargetPlatform | OBJECTMARK_EditorOnly);
+	}
+}
+
+const TCHAR* LexToString(ESaveableStatus Status)
 {
-	static_assert(static_cast<int32>(FSaveContext::ESaveableStatus::__Count) == 7);
+	static_assert(static_cast<int32>(ESaveableStatus::__Count) == 9);
 	switch (Status)
 	{
-	case FSaveContext::ESaveableStatus::Success: return TEXT("is saveable");
-	case FSaveContext::ESaveableStatus::PendingKill: return TEXT("is pendingkill");
-	case FSaveContext::ESaveableStatus::Transient: return TEXT("is transient");
-	case FSaveContext::ESaveableStatus::AbstractClass: return TEXT("has a Class with CLASS_Abstract");
-	case FSaveContext::ESaveableStatus::DeprecatedClass: return TEXT("has a Class with CLASS_Deprecated");
-	case FSaveContext::ESaveableStatus::NewerVersionExistsClass: return TEXT("has a Class with CLASS_NewerVersionExists");
-	case FSaveContext::ESaveableStatus::OuterUnsaveable: return TEXT("has an unsaveable Outer");
+	case ESaveableStatus::Success: return TEXT("is saveable");
+	case ESaveableStatus::PendingKill: return TEXT("is pendingkill");
+	case ESaveableStatus::Transient: return TEXT("is transient");
+	case ESaveableStatus::AbstractClass: return TEXT("has a Class with CLASS_Abstract");
+	case ESaveableStatus::DeprecatedClass: return TEXT("has a Class with CLASS_Deprecated");
+	case ESaveableStatus::NewerVersionExistsClass: return TEXT("has a Class with CLASS_NewerVersionExists");
+	case ESaveableStatus::OuterUnsaveable: return TEXT("has an unsaveable Outer");
+	case ESaveableStatus::ClassUnsaveable: return TEXT("has an unsaveable Class");
+	case ESaveableStatus::ExcludedByPlatform: return TEXT("is excluded by TargetPlatform");
 	default: return TEXT("Unknown");
 	}
 }
\ No newline at end of file
diff --git a/Engine/Source/Runtime/CoreUObject/Private/UObject/SavePackage/SaveContext.h b/Engine/Source/Runtime/CoreUObject/Private/UObject/SavePackage/SaveContext.h
--- a/Engine/Source/Runtime/CoreUObject/Private/UObject/SavePackage/SaveContext.h
+++ b/Engine/Source/Runtime/CoreUObject/Private/UObject/SavePackage/SaveContext.h
@@ -101,6 +101,20 @@ struct FIllegalReference
 	FString FormatStringArg;
 };
 
+enum class ESaveableStatus
+{
+	Success,
+	PendingKill,
+	Transient,
+	AbstractClass,
+	DeprecatedClass,
+	NewerVersionExistsClass,
+	OuterUnsaveable,
+	ClassUnsaveable,
+	ExcludedByPlatform,
+	__Count,
+};
+
 /** Hold the harvested exports and imports for a realm */
 struct FHarvestedRealm
 {
@@ -149,6 +163,10 @@ struct FHarvestedRealm
 		return IsImport(InObject) || IsExport(InObject);
 	}
 
+	/**
+	 * Used during harvesting to early exit from objects we have found referenced earlier but excluded because of
+	 * editoronly or unsaveable or otherwise.
+	 */
 	bool IsExcluded(UObject* InObject) const
 	{
 		return Excluded.Contains(InObject);
@@ -402,7 +420,7 @@ public:
 		, SaveArgs(InSaveArgs)
 		, PackageWriter(InSaveArgs.SavePackageContext ? InSaveArgs.SavePackageContext->PackageWriter : nullptr)
 		, SerializeContext(InSerializeContext)
-		, ExcludedObjectMarks(SavePackageUtilities::GetExcludedObjectMarksForTargetPlatform(SaveArgs.GetTargetPlatform()))
+		, GameRealmExcludedObjectMarks(GetExcludedObjectMarksForGameRealm(SaveArgs.GetTargetPlatform()))
 	{
 		// Assumptions & checks
 		check(InPackage);
@@ -493,8 +511,19 @@ public:
 
 	EObjectMark GetExcludedObjectMarks(ESaveRealm HarvestingRealm) const
 	{
-		// When considering excluded objects for a platform, do not consider editor only object and not for platform objects in the optional context as excluded
-		return (HarvestingRealm == ESaveRealm::Optional) ? (EObjectMark)(ExcludedObjectMarks & ~(EObjectMark::OBJECTMARK_EditorOnly|EObjectMark::OBJECTMARK_NotForTargetPlatform)) : ExcludedObjectMarks;
+		switch (HarvestingRealm)
+		{
+		case ESaveRealm::Optional:
+			// When considering excluded objects for a platform, do not consider editor only object and not for platform objects in the optional context as excluded
+			return (EObjectMark)(GameRealmExcludedObjectMarks & ~(EObjectMark::OBJECTMARK_EditorOnly | EObjectMark::OBJECTMARK_NotForTargetPlatform));
+		case ESaveRealm::Game:
+			return GameRealmExcludedObjectMarks;
+		case ESaveRealm::Editor:
+			return EObjectMark::OBJECTMARK_NOMARKS;
+		default:
+			checkNoEntry();
+			return EObjectMark::OBJECTMARK_NOMARKS;
+		}
 	}
 
 	EObjectFlags GetTopLevelFlags() const
@@ -708,17 +737,6 @@ public:
 	void MarkUnsaveable(UObject* InObject);
 
 	bool IsUnsaveable(UObject* InObject, bool bEmitWarning = true) const;
-	enum class ESaveableStatus
-	{
-		Success,
-		PendingKill,
-		Transient,
-		AbstractClass,
-		DeprecatedClass,
-		NewerVersionExistsClass,
-		OuterUnsaveable,
-		__Count,
-	};
 	ESaveableStatus GetSaveableStatus(UObject* InObject, UObject** OutCulprit = nullptr, ESaveableStatus* OutCulpritStatus = nullptr) const;
 	ESaveableStatus GetSaveableStatusNoOuter(UObject* InObject) const;
 
@@ -762,11 +780,6 @@ public:
 		return GetHarvestedRealm().IsIncluded(InObject);
 	}
 
-	bool IsExcluded(UObject* InObject) const
-	{
-		return GetHarvestedRealm().IsExcluded(InObject);
-	}
-
 	TSet<FTaggedExport>& GetExports()
 	{
 		return GetHarvestedRealm().GetExports();
@@ -964,19 +977,7 @@ public:
 		GetHarvestedRealm().SetTextFormatTempFilename(MoveTemp(InTemp));
 	}
 
-	FSavePackageResultStruct GetFinalResult()
-	{
-		if (Result != ESavePackageResult::Success)
-		{
-			return Result;
-		}
-
-		ESavePackageResult FinalResult = IsStubRequested() ? ESavePackageResult::GenerateStub : ESavePackageResult::Success;
-		FSavePackageResultStruct FinalResultStruct(FinalResult, TotalPackageSizeUncompressed,
-			SerializedPackageFlags, IsCompareLinker() ? MoveTemp(GetHarvestedRealm().Linker) : nullptr);
-		FinalResultStruct.SavedAssets = MoveTemp(SavedAssets);
-		return FinalResultStruct;
-	}
+	FSavePackageResultStruct GetFinalResult();
 
 	FObjectSaveContextData& GetObjectSaveContext()
 	{
@@ -1060,6 +1061,7 @@ private:
 
 	// Create the harvesting contexts and automatic optional context gathering options
 	void SetupHarvestingRealms();
+	static EObjectMark GetExcludedObjectMarksForGameRealm(const ITargetPlatform* TargetPlatform);
 		
 	friend class FPackageHarvester;
 
@@ -1090,8 +1092,8 @@ private:
 	// Pointer to the EDLCookChecker associated with this context
 	FEDLCookCheckerThreadState* EDLCookChecker = nullptr;
 
-	// Matching any mark in ExcludedObjectMarks indicates that an object should be excluded from being either an import or an export for this save
-	const EObjectMark ExcludedObjectMarks;
+	// An object matching any GameRealmExcludedObjectMarks should be excluded from imports or exports in the game realm
+	const EObjectMark GameRealmExcludedObjectMarks;
 
 	// Harvested custom versions
 	FCustomVersionContainer CustomVersions;
@@ -1115,4 +1117,4 @@ private:
 	TMap<UObject*, TSet<FProperty*>> TransientPropertyOverrides;
 };
 
-const TCHAR* LexToString(FSaveContext::ESaveableStatus Status);
+const TCHAR* LexToString(ESaveableStatus Status);
diff --git a/Engine/Source/Runtime/CoreUObject/Private/UObject/SavePackage/SavePackageUtilities.cpp b/Engine/Source/Runtime/CoreUObject/Private/UObject/SavePackage/SavePackageUtilities.cpp
--- a/Engine/Source/Runtime/CoreUObject/Private/UObject/SavePackage/SavePackageUtilities.cpp
+++ b/Engine/Source/Runtime/CoreUObject/Private/UObject/SavePackage/SavePackageUtilities.cpp
@@ -769,6 +769,12 @@ bool IsEditorOnlyObject(const UObject* InObject, bool bCheckRecursive, bool bChe
 	DECLARE_SCOPE_CYCLE_COUNTER(TEXT("IsEditorOnlyObject"), STAT_IsEditorOnlyObject, STATGROUP_LoadTime);
 	check(InObject);
 
+	// CDOs must be included if their class are, so if the object is a CDO, return IsEditorOnlyObject for its class
+	if (InObject->HasAnyFlags(RF_ClassDefaultObject))
+	{
+		return IsEditorOnlyObject(InObject->GetClass(), bCheckRecursive, bCheckMarks);
+	}
+
 	if ((bCheckMarks && InObject->HasAnyMarks(OBJECTMARK_EditorOnly)) || InObject->IsEditorOnly())
 	{
 		return true;
diff --git a/Engine/Source/Runtime/CoreUObject/Private/UObject/SavePackage2.cpp b/Engine/Source/Runtime/CoreUObject/Private/UObject/SavePackage2.cpp
--- a/Engine/Source/Runtime/CoreUObject/Private/UObject/SavePackage2.cpp
+++ b/Engine/Source/Runtime/CoreUObject/Private/UObject/SavePackage2.cpp
@@ -157,7 +157,15 @@ ESavePackageResult ValidatePackage(FSaveContext& SaveContext)
 					ErrorText = FText::Format(NSLOCTEXT("SavePackage2", "AssetSaveMissingTopLevelFlags", "The Asset {Name} being saved does not have any of the provided object flags (0x{Flags}); saving the package would cause data loss."),	
 						Arguments);
 				}
-				SaveContext.GetError()->Logf(ELogVerbosity::Warning, TEXT("%s"), *ErrorText.ToString());
+				if (SaveContext.IsGenerateSaveError() && SaveContext.GetError())
+				{
+					SaveContext.GetError()->Logf(ELogVerbosity::Warning, TEXT("%s"), *ErrorText.ToString());
+				}
+				else
+				{
+					UE_LOG(LogSavePackage, Warning, TEXT("%s"), *ErrorText.ToString());
+				}
+
 				return ESavePackageResult::Error;
 			}
 		}
@@ -404,30 +412,33 @@ ESavePackageResult HarvestPackage(FSaveContext& SaveContext)
 		}
 	}
 
-	// Trim PrestreamPackage list 
-	TSet<UPackage*>& PrestreamPackages = SaveContext.GetPrestreamPackages();
-	TSet<UPackage*> KeptPrestreamPackages;
-	for (UPackage* Pkg : PrestreamPackages)
 	{
-		// If the prestream package hasn't been otherwise already marked as an import, keep it as such and mark it as an import
-		if (!SaveContext.IsImport(Pkg))
+		FPackageHarvester::FHarvestScope RootReferenceScope = Harvester.EnterRootReferencesScope();
+		// Trim PrestreamPackage list 
+		TSet<UPackage*>& PrestreamPackages = SaveContext.GetPrestreamPackages();
+		TSet<UPackage*> KeptPrestreamPackages;
+		for (UPackage* Pkg : PrestreamPackages)
 		{
-			KeptPrestreamPackages.Add(Pkg);
-			SaveContext.AddImport(Pkg);
+			// If the prestream package hasn't been otherwise already marked as an import, keep it as such and mark it as an import
+			if (!SaveContext.IsImport(Pkg))
+			{
+				KeptPrestreamPackages.Add(Pkg);
+				Harvester << Pkg;
+			}
 		}
-	}
-	Exchange(PrestreamPackages, KeptPrestreamPackages);
+		Exchange(PrestreamPackages, KeptPrestreamPackages);
 
-	// Harvest the PrestreamPackage class name if needed
-	if (PrestreamPackages.Num() > 0)
-	{
-		Harvester.HarvestPackageHeaderName(SavePackageUtilities::NAME_PrestreamPackage);		
-	}
+		// Harvest the PrestreamPackage class name if needed
+		if (PrestreamPackages.Num() > 0)
+		{
+			Harvester.HarvestPackageHeaderName(SavePackageUtilities::NAME_PrestreamPackage);
+		}
 
-	// if we have a WorldTileInfo, we need to harvest its dependencies as well, i.e. Custom Version
-	if (FWorldTileInfo* WorldTileInfo = SaveContext.GetPackage()->GetWorldTileInfo())
-	{
-		Harvester << *WorldTileInfo;
+		// if we have a WorldTileInfo, we need to harvest its dependencies as well, i.e. Custom Version
+		if (FWorldTileInfo* WorldTileInfo = SaveContext.GetPackage()->GetWorldTileInfo())
+		{
+			Harvester << *WorldTileInfo;
+		}
 	}
 
 	// The Editor version is used as part of the check to see if a package is too old to use the gather cache, so we always have to add it if we have gathered loc for this asset
@@ -1409,13 +1420,13 @@ void SavePreloadDependencies(FStructuredArchive::FRecord& StructuredArchiveRoot,
 			{
 				UE_LOG(LogSavePackage, Warning, TEXT("A dependency '%s' of '%s' is in the linker table, but is pending kill or garbage. We will keep the dependency anyway (%d)."), *ToTest->GetFullName(), *ForObj->GetFullName(), CallSite);
 			}
-			bool bNotFiltered = !SaveContext.IsExcluded(ToTest);
-			if (bMandatory && !bNotFiltered)
+			bool bIncludedInHarvest = SaveContext.IsIncluded(ToTest);
+			if (bMandatory && !bIncludedInHarvest)
 			{
 				UE_LOG(LogSavePackage, Warning, TEXT("A dependency '%s' of '%s' was filtered, but is mandatory. This indicates a problem with editor only stripping. We will keep the dependency anyway (%d)."), *ToTest->GetFullName(), *ForObj->GetFullName(), CallSite);
-				bNotFiltered = true;
+				bIncludedInHarvest = true;
 			}
-			if (bNotFiltered)
+			if (bIncludedInHarvest)
 			{
 				if (!Index.IsNull())
 				{
@@ -1524,8 +1535,8 @@ void SavePreloadDependencies(FStructuredArchive::FRecord& StructuredArchiveRoot,
 						// Only include subobject archetypes
 						if (SubObj->HasAnyFlags(RF_DefaultSubObject | RF_ArchetypeObject))
 						{
-							// Don't include the archetype of SubObjects that we excluded in the harvesting phase; we didn't add their archetypes
-							if (SaveContext.IsExcluded(SubObj))
+							// Don't include the archetype of SubObjects that were not included in the harvesting phase; we didn't add their archetypes
+							if (!SaveContext.IsIncluded(SubObj))
 							{
 								continue;
 							}
diff --git a/Engine/Source/Runtime/CoreUObject/Public/UObject/Package.h b/Engine/Source/Runtime/CoreUObject/Public/UObject/Package.h
--- a/Engine/Source/Runtime/CoreUObject/Public/UObject/Package.h
+++ b/Engine/Source/Runtime/CoreUObject/Public/UObject/Package.h
@@ -102,6 +102,9 @@ struct FSavePackageResultStruct
 
 	TArray<FAssetData> SavedAssets;
 
+	TArray<FName> ImportPackages;
+	TArray<FName> SoftPackageReferences;
+
 	/** Constructors, it will implicitly construct from the result enum */
 	COREUOBJECT_API FSavePackageResultStruct();
 	COREUOBJECT_API FSavePackageResultStruct(ESavePackageResult InResult);
