ZenLoader in editor: Propagate load flags to the linker and set PKG_ForDiffing also for async loads.
On behalf of CarlMagnus.Nordin

diff --git a/Engine/Source/Runtime/CoreUObject/Private/Serialization/AsyncLoading2.cpp b/Engine/Source/Runtime/CoreUObject/Private/Serialization/AsyncLoading2.cpp
--- a/Engine/Source/Runtime/CoreUObject/Private/Serialization/AsyncLoading2.cpp
+++ b/Engine/Source/Runtime/CoreUObject/Private/Serialization/AsyncLoading2.cpp
@@ -4841,14 +4841,16 @@ void FAsyncPackage2::InitializeLinkerLoadState(const FLinkerInstancingContext* I
 
 void FAsyncPackage2::CreateLinker(const FLinkerInstancingContext* InstancingContext)
 {
+	uint32 LinkerFlags = (LOAD_Async | LOAD_NoVerify | LOAD_SkipLoadImportedPackages | Desc.LoadFlags);
+#if WITH_EDITOR
+	if ((!FApp::IsGame() || GIsEditor) && (Desc.PackageFlags & PKG_PlayInEditor) != 0)
+	{
+		LinkerFlags |= LOAD_PackageForPIE;
+	}
+#endif
 	FLinkerLoad* Linker = FLinkerLoad::FindExistingLinkerForPackage(LinkerRoot);
 	if (!Linker)
 	{
-		uint32 LinkerFlags = (LOAD_Async | LOAD_NoVerify | LOAD_SkipLoadImportedPackages);
-		if ((!FApp::IsGame() || GIsEditor) && (Desc.PackageFlags & PKG_PlayInEditor) != 0)
-		{
-			LinkerFlags |= LOAD_PackageForPIE;
-		}
 		FUObjectSerializeContext* LoadContext = GetSerializeContext();
 #if ALT2_ENABLE_NEW_ARCHIVE_FOR_LINKERLOAD
 		Linker = new FLinkerLoad(LinkerRoot, Desc.PackagePathToLoad, LinkerFlags, InstancingContext ? *InstancingContext : FLinkerInstancingContext());
@@ -4862,7 +4864,7 @@ void FAsyncPackage2::CreateLinker(const FLinkerInstancingContext* InstancingCont
 	}
 	else
 	{
-		Linker->LoadFlags |= LOAD_Async | LOAD_NoVerify | LOAD_SkipLoadImportedPackages;
+		Linker->LoadFlags |= LinkerFlags;
 	}
 	check(Linker);
 	check(Linker->LinkerRoot == LinkerRoot);
diff --git a/Engine/Source/Runtime/CoreUObject/Private/UObject/UObjectGlobals.cpp b/Engine/Source/Runtime/CoreUObject/Private/UObject/UObjectGlobals.cpp
--- a/Engine/Source/Runtime/CoreUObject/Private/UObject/UObjectGlobals.cpp
+++ b/Engine/Source/Runtime/CoreUObject/Private/UObject/UObjectGlobals.cpp
@@ -1568,10 +1568,17 @@ UPackage* LoadPackageInternal(UPackage* InOuter, const FPackagePath& PackagePath
 
 		ThreadContext.SyncLoadUsingAsyncLoaderCount++;
 		EPackageFlags PackageFlags = PKG_None;
+#if WITH_EDITOR
+		// If we are loading a package for diffing, set the package flag
+		if (LoadFlags & LOAD_ForDiff)
+		{
+			PackageFlags |= PKG_ForDiffing;
+		}
 		if ((!FApp::IsGame() || GIsEditor) && (LoadFlags & LOAD_PackageForPIE) != 0)
 		{
 			PackageFlags |= PKG_PlayInEditor;
 		}
+#endif
 		constexpr int32 PIEInstanceID = INDEX_NONE;
 		constexpr int32 Priority = INT32_MAX;
 		int32 RequestID = LoadPackageAsync(PackagePath, PackageName, FLoadPackageAsyncDelegate(), PackageFlags, PIEInstanceID, Priority, InstancingContext, LoadFlags);
