Remove TObjectPtr's dependency on a fully-defined "class UClass".

diff --git a/Engine/Plugins/Media/MediaPlate/Source/MediaPlateEditor/Public/MediaPlateEditorModule.h b/Engine/Plugins/Media/MediaPlate/Source/MediaPlateEditor/Public/MediaPlateEditorModule.h
--- a/Engine/Plugins/Media/MediaPlate/Source/MediaPlateEditor/Public/MediaPlateEditorModule.h
+++ b/Engine/Plugins/Media/MediaPlate/Source/MediaPlateEditor/Public/MediaPlateEditorModule.h
@@ -8,6 +8,7 @@
 #include "Modules/ModuleManager.h"
 #include "TickableEditorObject.h"
 #include "UObject/ObjectPtr.h"
+#include "UObject/WeakObjectPtr.h"
 
 class FDragDropEvent;
 class FPlacementModeID;
diff --git a/Engine/Source/Runtime/Core/Public/UObject/WeakObjectPtrTemplates.h b/Engine/Source/Runtime/Core/Public/UObject/WeakObjectPtrTemplates.h
--- a/Engine/Source/Runtime/Core/Public/UObject/WeakObjectPtrTemplates.h
+++ b/Engine/Source/Runtime/Core/Public/UObject/WeakObjectPtrTemplates.h
@@ -12,6 +12,9 @@
 
 #include <type_traits>
 
+
+struct FWeakObjectPtr;
+
 /**
  * TWeakObjectPtr is the templated version of the generic FWeakObjectPtr
  */
@@ -22,7 +25,7 @@ struct TWeakObjectPtr : private TWeakObjectPtrBase
 
 	// Although templated, these parameters are not intended to be anything other than the default,
 	// and are only templates for module organization reasons.
-	static_assert(std::is_same_v<TWeakObjectPtrBase, FWeakObjectPtr>, "TWeakObjectPtrBase should not be overridden");
+	static_assert(std::is_same_v<TWeakObjectPtrBase*, FWeakObjectPtr*>, "TWeakObjectPtrBase should not be overridden");
 
 public:
 	using ElementType = T;
diff --git a/Engine/Source/Runtime/CoreUObject/Private/UObject/ObjectPtr.cpp b/Engine/Source/Runtime/CoreUObject/Private/UObject/ObjectPtr.cpp
--- a/Engine/Source/Runtime/CoreUObject/Private/UObject/ObjectPtr.cpp
+++ b/Engine/Source/Runtime/CoreUObject/Private/UObject/ObjectPtr.cpp
@@ -1,6 +1,7 @@
 // Copyright Epic Games, Inc. All Rights Reserved.
 
 #include "UObject/ObjectPtr.h"
+#include "UObject/Class.h"
 
 #if UE_WITH_OBJECT_HANDLE_LATE_RESOLVE
 
@@ -46,4 +47,16 @@ FString FObjectPtr::GetFullName(EObjectFullNameFlags Flags) const
 	}
 }
 
-#endif
\ No newline at end of file
+#endif
+
+bool FObjectPtr::IsA(const UClass* SomeBase) const
+{
+	checkfSlow(SomeBase, TEXT("IsA(NULL) cannot yield meaningful results"));
+
+	if (const UClass* ThisClass = GetClass())
+	{
+		return ThisClass->IsChildOf(SomeBase);
+	}
+
+	return false;
+}
diff --git a/Engine/Source/Runtime/CoreUObject/Public/UObject/Class.h b/Engine/Source/Runtime/CoreUObject/Public/UObject/Class.h
--- a/Engine/Source/Runtime/CoreUObject/Public/UObject/Class.h
+++ b/Engine/Source/Runtime/CoreUObject/Public/UObject/Class.h
@@ -68,6 +68,7 @@
 #include "UObject/TopLevelAssetPath.h"
 #include "UObject/UObjectGlobals.h"
 #include "UObject/UnrealNames.h"
+#include "UObject/ObjectPtr.h"
 
 #if UE_ENABLE_INCLUDE_ORDER_DEPRECATED_IN_5_2
 #include "Misc/PackageAccessTracking.h"
@@ -4171,4 +4172,4 @@ UE_DECLARE_CORE_VARIANT_TYPE(Sphere3, Sphere);
 
 #if UE_ENABLE_INCLUDE_ORDER_DEPRECATED_IN_5_2
 #include "CoreMinimal.h"
-#endif
\ No newline at end of file
+#endif
diff --git a/Engine/Source/Runtime/CoreUObject/Public/UObject/ObjectPtr.h b/Engine/Source/Runtime/CoreUObject/Public/UObject/ObjectPtr.h
--- a/Engine/Source/Runtime/CoreUObject/Public/UObject/ObjectPtr.h
+++ b/Engine/Source/Runtime/CoreUObject/Public/UObject/ObjectPtr.h
@@ -4,7 +4,6 @@
 
 #include "HAL/Platform.h"
 #include "Serialization/StructuredArchive.h"
-#include "UObject/Class.h"
 #include "UObject/Object.h"
 #include "UObject/ObjectHandle.h"
 
@@ -44,6 +43,8 @@
 	#define UE_OBJECT_PTR_NONCONFORMANCE_SUPPORT 0
 #endif
 
+class UClass;
+	
 template <typename T>
 struct TObjectPtr;
 
@@ -193,17 +194,7 @@ public:
 	FORCEINLINE FObjectHandle GetHandle() const { return Handle; }
 	FORCEINLINE FObjectHandle& GetHandleRef() const { return Handle; }
 
-	FORCEINLINE bool IsA(const UClass* SomeBase) const
-	{
-		checkfSlow(SomeBase, TEXT("IsA(NULL) cannot yield meaningful results"));
-
-		if (const UClass* ThisClass = GetClass())
-		{
-			return ThisClass->IsChildOf(SomeBase);
-		}
-
-		return false;
-	}
+	COREUOBJECT_API bool IsA(const UClass* SomeBase) const;
 
 	template <typename T>
 	FORCEINLINE bool IsA() const
@@ -995,3 +986,7 @@ template <typename T>
 using TPointedToType = typename TPointedToTypeImpl<T>::Type;
 
 //------------------------------------------------------------------------------
+
+#if UE_ENABLE_INCLUDE_ORDER_DEPRECATED_IN_5_3
+#include "UObject/Class.h"
+#endif
