Fix IsEditorOnlyObject for CDOs: CDOs need to be marked as editoronly, even if their class is not editoronly, if their archetype is editoronly. Blueprint compilation relies on this for CanBeExposed to return true correctly for some blueprints that are defined in Script packages that are accidentally not marked as editoronly packages.
#rnx

diff --git a/Engine/Source/Runtime/CoreUObject/Private/UObject/SavePackage/SavePackageUtilities.cpp b/Engine/Source/Runtime/CoreUObject/Private/UObject/SavePackage/SavePackageUtilities.cpp
--- a/Engine/Source/Runtime/CoreUObject/Private/UObject/SavePackage/SavePackageUtilities.cpp
+++ b/Engine/Source/Runtime/CoreUObject/Private/UObject/SavePackage/SavePackageUtilities.cpp
@@ -769,15 +769,14 @@ bool IsEditorOnlyObject(const UObject* InObject, bool bCheckRecursive, bool bChe
 	DECLARE_SCOPE_CYCLE_COUNTER(TEXT("IsEditorOnlyObject"), STAT_IsEditorOnlyObject, STATGROUP_LoadTime);
 	check(InObject);
 
-	// CDOs must be included if their class are, so if the object is a CDO, return IsEditorOnlyObject for its class
-	if (InObject->HasAnyFlags(RF_ClassDefaultObject))
+	// CDOs must be included if their class and archetype and outer are included.
+	// Ignore their value of IsEditorOnly
+	if (!InObject->HasAnyFlags(RF_ClassDefaultObject))
 	{
-		return IsEditorOnlyObject(InObject->GetClass(), bCheckRecursive, bCheckMarks);
-	}
-
-	if ((bCheckMarks && InObject->HasAnyMarks(OBJECTMARK_EditorOnly)) || InObject->IsEditorOnly())
-	{
-		return true;
+		if ((bCheckMarks && InObject->HasAnyMarks(OBJECTMARK_EditorOnly)) || InObject->IsEditorOnly())
+		{
+			return true;
+		}
 	}
 
 	// If this is a package that is editor only or the object is in editor-only package,
