cosmetic delegate-related changes, just to unclutter a bigger CL

diff --git a/Engine/Source/Runtime/Core/Public/Delegates/DelegateBase.h b/Engine/Source/Runtime/Core/Public/Delegates/DelegateBase.h
--- a/Engine/Source/Runtime/Core/Public/Delegates/DelegateBase.h
+++ b/Engine/Source/Runtime/Core/Public/Delegates/DelegateBase.h
@@ -50,24 +50,13 @@ class FDelegateBase //-V1062
 	friend class TDelegateBase;
 
 protected:
-	/**
-	 * Creates and initializes a new instance.
-	 *
-	 * @param InDelegateInstance The delegate instance to assign.
-	 */
-	explicit FDelegateBase()
-		: DelegateSize(0)
-	{
-	}
+	explicit FDelegateBase() = default;
 
 	~FDelegateBase()
 	{
 		Unbind();
 	}
 
-	/**
-	 * Move constructor.
-	 */
 	FDelegateBase(FDelegateBase&& Other)
 	{
 		DelegateAllocator.MoveToEmpty(Other.DelegateAllocator);
@@ -75,9 +64,6 @@ protected:
 		Other.DelegateSize = 0;
 	}
 
-	/**
-	 * Move assignment.
-	 */
 	FDelegateBase& operator=(FDelegateBase&& Other)
 	{
 		UE_DELEGATES_MT_SCOPED_WRITE_ACCESS(AccessDetector);
@@ -149,7 +135,7 @@ private:
 
 private:
 	FDelegateAllocatorType::ForElementType<FAlignedInlineDelegateType> DelegateAllocator;
-	int32 DelegateSize;
+	int32 DelegateSize = 0;
 
 	UE_DELEGATES_MT_ACCESS_DETECTOR(AccessDetector);
 };
diff --git a/Engine/Source/Runtime/Core/Public/Delegates/DelegateSignatureImpl.inl b/Engine/Source/Runtime/Core/Public/Delegates/DelegateSignatureImpl.inl
--- a/Engine/Source/Runtime/Core/Public/Delegates/DelegateSignatureImpl.inl
+++ b/Engine/Source/Runtime/Core/Public/Delegates/DelegateSignatureImpl.inl
@@ -528,7 +528,7 @@ public:
 		DelegateInstanceInterfaceType* LocalDelegateInstance = GetDelegateInstanceProtected();
 
 		// If this assert goes off, Execute() was called before a function was bound to the delegate.
-		// Consider using ExecuteIfSafe() instead.
+		// Consider using ExecuteIfBound() instead.
 		checkSlow(LocalDelegateInstance != nullptr);
 
 		return LocalDelegateInstance->Execute(Params...);
@@ -630,6 +630,9 @@ public:
 		return *this;
 	}
 
+	TMulticastDelegate(TMulticastDelegate&&) = default;
+	TMulticastDelegate& operator=(TMulticastDelegate&&) = default;
+
 public:
 
 	/**
@@ -868,9 +871,6 @@ public:
 		return bResult;
 	}
 
-	TMulticastDelegate(TMulticastDelegate&&) = default;
-	TMulticastDelegate& operator=(TMulticastDelegate&&) = default;
-
 	/**
 	 * Broadcasts this delegate to all bound objects, except to those that may have expired.
 	 *
diff --git a/Engine/Source/Runtime/Core/Public/Delegates/MulticastDelegateBase.h b/Engine/Source/Runtime/Core/Public/Delegates/MulticastDelegateBase.h
--- a/Engine/Source/Runtime/Core/Public/Delegates/MulticastDelegateBase.h
+++ b/Engine/Source/Runtime/Core/Public/Delegates/MulticastDelegateBase.h
@@ -225,7 +225,7 @@ protected:
 		UE_DELEGATES_MT_SCOPED_WRITE_ACCESS(AccessDetector);
 
 		FDelegateHandle Result;
-		if (GetDelegateInstanceProtectedHelper(NewDelegateBaseRef))
+		if (NewDelegateBaseRef.IsBound())
 		{
 			// compact but obey threshold of when this will trigger
 			CompactInvocationList(true);
diff --git a/Engine/Source/Runtime/Core/Public/UObject/ScriptDelegates.h b/Engine/Source/Runtime/Core/Public/UObject/ScriptDelegates.h
--- a/Engine/Source/Runtime/Core/Public/UObject/ScriptDelegates.h
+++ b/Engine/Source/Runtime/Core/Public/UObject/ScriptDelegates.h
@@ -422,6 +422,8 @@ class TMulticastScriptDelegate
 
 public:
 
+	using InvocationListType = TArray<UnicastDelegateType>;
+
 	TMulticastScriptDelegate() = default;
 
 public:
@@ -618,7 +620,7 @@ public:
 		{
 			FString AllDelegatesString = TEXT( "[" );
 			bool bAddComma = false;
-			for( typename FInvocationList::TConstIterator CurDelegate( InvocationList ); CurDelegate; ++CurDelegate )
+			for( typename InvocationListType::TConstIterator CurDelegate( InvocationList ); CurDelegate; ++CurDelegate )
 			{
 				if (bAddComma)
 				{
@@ -721,7 +723,7 @@ public:
 		UE_DELEGATES_MT_SCOPED_READ_ACCESS(AccessDetector);
 
 		TArray< UObject* > OutputList;
-		for( typename FInvocationList::TIterator CurDelegate( InvocationList ); CurDelegate; ++CurDelegate )
+		for( typename InvocationListType::TIterator CurDelegate( InvocationList ); CurDelegate; ++CurDelegate )
 		{
 			UObject* CurObject = CurDelegate->GetUObject();
 			if( CurObject != nullptr )
@@ -742,7 +744,7 @@ public:
 		UE_DELEGATES_MT_SCOPED_READ_ACCESS(AccessDetector);
 
 		TArray< UObject* > OutputList;
-		for( typename FInvocationList::TIterator CurDelegate( InvocationList ); CurDelegate; ++CurDelegate )
+		for( typename InvocationListType::TIterator CurDelegate( InvocationList ); CurDelegate; ++CurDelegate )
 		{
 			UObject* CurObject = CurDelegate->GetUObjectEvenIfUnreachable();
 			if( CurObject != nullptr )
@@ -770,14 +772,14 @@ protected:
 	*/
 	void AddInternal( const UnicastDelegateType& InDelegate )
 	{
-#if !(UE_BUILD_SHIPPING || UE_BUILD_TEST)
+#if DO_ENSURE
 		// Verify same function isn't already bound
 		const int32 NumFunctions = InvocationList.Num();
 		for( int32 CurFunctionIndex = 0; CurFunctionIndex < NumFunctions; ++CurFunctionIndex )
 		{
 			(void)ensure( InvocationList[ CurFunctionIndex ] != InDelegate );
 		}
-#endif // !(UE_BUILD_SHIPPING || UE_BUILD_TEST)
+#endif // DO_CHECK
 		InvocationList.Add( InDelegate );
 	}
 
@@ -831,13 +833,10 @@ protected:
 		});
 	}
 
-public:
-	typedef TArray<UnicastDelegateType> FInvocationList;
-
 protected:
 
 	/** Ordered list functions to invoke when the Broadcast function is called */
-	mutable FInvocationList InvocationList;		// Mutable so that we can housekeep list even with 'const' broadcasts
+	mutable InvocationListType InvocationList;		// Mutable so that we can housekeep list even with 'const' broadcasts
 
 	UE_DELEGATES_MT_ACCESS_DETECTOR(AccessDetector);
 
@@ -856,5 +855,5 @@ protected:
 template <typename TWeakPtr>
 struct TIsZeroConstructType<TMulticastScriptDelegate<TWeakPtr>>
 {
-	static constexpr bool Value = TIsZeroConstructType<typename TMulticastScriptDelegate<TWeakPtr>::FInvocationList>::Value;
+	static constexpr bool Value = TIsZeroConstructType<typename TMulticastScriptDelegate<TWeakPtr>::InvocationListType>::Value;
 };
diff --git a/Engine/Source/Runtime/CoreUObject/Private/UObject/PropertyMulticastDelegate.cpp b/Engine/Source/Runtime/CoreUObject/Private/UObject/PropertyMulticastDelegate.cpp
--- a/Engine/Source/Runtime/CoreUObject/Private/UObject/PropertyMulticastDelegate.cpp
+++ b/Engine/Source/Runtime/CoreUObject/Private/UObject/PropertyMulticastDelegate.cpp
@@ -9,7 +9,7 @@
 #include "UObject/LinkerPlaceholderFunction.h"
 #include "Serialization/ArchiveUObjectFromStructuredArchive.h"
 
-FMulticastScriptDelegate::FInvocationList FMulticastDelegateProperty::EmptyList;
+FMulticastScriptDelegate::InvocationListType FMulticastDelegateProperty::EmptyList;
 
 FMulticastDelegateProperty::FMulticastDelegateProperty(FFieldVariant InOwner, const UECodeGen_Private::FMulticastDelegatePropertyParams& Prop, EPropertyFlags AdditionalPropertyFlags /*= CPF_None*/)
 	: FProperty(InOwner, (const UECodeGen_Private::FPropertyParamsBaseWithOffset&)Prop, AdditionalPropertyFlags)
@@ -43,8 +43,8 @@ void FMulticastDelegateProperty::InstanceSubobjects(void* Data, void const* Defa
 		for( int32 i=0; i<ArrayDim; i++ )
 		{
 			// Fix up references to the class default object (if necessary)
-			FMulticastScriptDelegate::FInvocationList::TIterator CurInvocation(GetInvocationList((uint8*)Data + i));
-			FMulticastScriptDelegate::FInvocationList::TIterator DefaultInvocation(GetInvocationList((uint8*)DefaultData + i));
+			FMulticastScriptDelegate::InvocationListType::TIterator CurInvocation(GetInvocationList((uint8*)Data + i));
+			FMulticastScriptDelegate::InvocationListType::TIterator DefaultInvocation(GetInvocationList((uint8*)DefaultData + i));
 			for(; CurInvocation && DefaultInvocation; ++CurInvocation, ++DefaultInvocation )
 			{
 				FScriptDelegate& DestDelegateInvocation = *CurInvocation;
@@ -76,7 +76,7 @@ void FMulticastDelegateProperty::InstanceSubobjects(void* Data, void const* Defa
 	{
 		for( int32 i=0; i<ArrayDim; i++ )
 		{
-			for( FMulticastScriptDelegate::FInvocationList::TIterator CurInvocation(GetInvocationList((uint8*)Data + i)); CurInvocation; ++CurInvocation )
+			for( FMulticastScriptDelegate::InvocationListType::TIterator CurInvocation(GetInvocationList((uint8*)Data + i)); CurInvocation; ++CurInvocation )
 			{
 				FScriptDelegate& DestDelegateInvocation = *CurInvocation;
 				UObject* CurrentUObject = DestDelegateInvocation.GetUObject();
@@ -93,8 +93,8 @@ void FMulticastDelegateProperty::InstanceSubobjects(void* Data, void const* Defa
 
 bool FMulticastDelegateProperty::Identical( const void* A, const void* B, uint32 PortFlags ) const
 {
-	const FMulticastScriptDelegate::FInvocationList& ListA = GetInvocationList(A);
-	const FMulticastScriptDelegate::FInvocationList& ListB = GetInvocationList(B);
+	const FMulticastScriptDelegate::InvocationListType& ListA = GetInvocationList(A);
+	const FMulticastScriptDelegate::InvocationListType& ListB = GetInvocationList(B);
 
 	const int32 ListASize = ListA.Num();
 	if (ListASize != ListB.Num())
@@ -181,7 +181,7 @@ FString FMulticastDelegateProperty::GetCPPTypeForwardDeclaration() const
 
 void FMulticastDelegateProperty::ExportText_Internal( FString& ValueStr, const void* PropertyValueOrContainer, EPropertyPointerType PropertyPointerType, const void* DefaultValue, UObject* Parent, int32 PortFlags, UObject* ExportRootScope ) const
 {
-	const FMulticastScriptDelegate::FInvocationList* InvocationList = nullptr;
+	const FMulticastScriptDelegate::InvocationListType* InvocationList = nullptr;
 	
 	if (PropertyPointerType == EPropertyPointerType::Container && HasGetter())
 	{
@@ -198,7 +198,7 @@ void FMulticastDelegateProperty::ExportText_Internal( FString& ValueStr, const v
 	ValueStr += TEXT( "(" );
 
 	bool bIsFirstFunction = true;
-	for (FMulticastScriptDelegate::FInvocationList::TConstIterator CurInvocation(*InvocationList); CurInvocation; ++CurInvocation)
+	for (FMulticastScriptDelegate::InvocationListType::TConstIterator CurInvocation(*InvocationList); CurInvocation; ++CurInvocation)
 	{
 		if (CurInvocation->IsBound())
 		{
@@ -394,7 +394,7 @@ void FMulticastInlineDelegateProperty::SetMulticastDelegate(void* PropertyValue,
 	*(FMulticastScriptDelegate*)PropertyValue = MoveTemp(ScriptDelegate);
 }
 
-FMulticastScriptDelegate::FInvocationList& FMulticastInlineDelegateProperty::GetInvocationList(const void* PropertyValue) const
+FMulticastScriptDelegate::InvocationListType& FMulticastInlineDelegateProperty::GetInvocationList(const void* PropertyValue) const
 {
 	return (PropertyValue ? ((FMulticastScriptDelegate*)PropertyValue)->InvocationList : EmptyList);
 }
@@ -505,7 +505,7 @@ void FMulticastSparseDelegateProperty::SetMulticastDelegate(void* PropertyValue,
 }
 
 
-FMulticastScriptDelegate::FInvocationList& FMulticastSparseDelegateProperty::GetInvocationList(const void* PropertyValue) const
+FMulticastScriptDelegate::InvocationListType& FMulticastSparseDelegateProperty::GetInvocationList(const void* PropertyValue) const
 {
 	if (FSparseDelegate* SparseDelegate = (FSparseDelegate*)PropertyValue)
 	{
diff --git a/Engine/Source/Runtime/CoreUObject/Public/UObject/UnrealType.h b/Engine/Source/Runtime/CoreUObject/Public/UObject/UnrealType.h
--- a/Engine/Source/Runtime/CoreUObject/Public/UObject/UnrealType.h
+++ b/Engine/Source/Runtime/CoreUObject/Public/UObject/UnrealType.h
@@ -5965,8 +5965,8 @@ public:
 protected:
 	friend class FProperty;
 
-	static FMulticastScriptDelegate::FInvocationList EmptyList;
-	virtual FMulticastScriptDelegate::FInvocationList& GetInvocationList(const void* PropertyValue) const PURE_VIRTUAL(FMulticastDelegateProperty::GetInvocationList, return EmptyList;);
+	static FMulticastScriptDelegate::InvocationListType EmptyList;
+	virtual FMulticastScriptDelegate::InvocationListType& GetInvocationList(const void* PropertyValue) const PURE_VIRTUAL(FMulticastDelegateProperty::GetInvocationList, return EmptyList;);
 
 
 	const TCHAR* ImportText_Add( const TCHAR* Buffer, void* PropertyValue, int32 PortFlags, UObject* Parent, FOutputDevice* ErrorText ) const;
@@ -6081,7 +6081,7 @@ PRAGMA_ENABLE_DEPRECATION_WARNINGS
 	virtual void ClearDelegate(UObject* Parent = nullptr, void* PropertyValue = nullptr) const override;
 
 protected:
-	virtual FMulticastScriptDelegate::FInvocationList& GetInvocationList(const void* PropertyValue) const;
+	virtual FMulticastScriptDelegate::InvocationListType& GetInvocationList(const void* PropertyValue) const;
 	// End of FMulticastDelegateProperty interface
 };
 
@@ -6132,7 +6132,7 @@ PRAGMA_ENABLE_DEPRECATION_WARNINGS
 	virtual void ClearDelegate(UObject* Parent = nullptr, void* PropertyValue = nullptr) const override;
 
 protected:
-	virtual FMulticastScriptDelegate::FInvocationList& GetInvocationList(const void* PropertyValue) const;
+	virtual FMulticastScriptDelegate::InvocationListType& GetInvocationList(const void* PropertyValue) const;
 	// End of FMulticastDelegateProperty interface
 
 private:
