Fix potential race on ObjectRestoreAfterInitProps during object construction

#rnx

diff --git a/Engine/Source/Runtime/CoreUObject/Private/UObject/UObjectGlobals.cpp b/Engine/Source/Runtime/CoreUObject/Private/UObject/UObjectGlobals.cpp
--- a/Engine/Source/Runtime/CoreUObject/Private/UObject/UObjectGlobals.cpp
+++ b/Engine/Source/Runtime/CoreUObject/Private/UObject/UObjectGlobals.cpp
@@ -3218,7 +3218,7 @@ bool StaticAllocateObjectErrorTests( const UClass* InClass, UObject* InOuter, FN
 * this is only used for classes in the script compiler
 **/
 //@todo UE this is clunky
-static FRestoreForUObjectOverwrite* ObjectRestoreAfterInitProps = NULL;  
+static thread_local FRestoreForUObjectOverwrite* ObjectRestoreAfterInitProps = nullptr;
 
 COREUOBJECT_API bool GOutputCookingWarnings = false;
 
@@ -3816,11 +3816,11 @@ void FObjectInitializer::PostConstructInit()
 	bool bNeedSubobjectInstancing = InitSubobjectProperties(bAllowInstancing);
 
 	// Restore class information if replacing native class.
-	if (ObjectRestoreAfterInitProps != NULL)
+	if (ObjectRestoreAfterInitProps != nullptr)
 	{
 		ObjectRestoreAfterInitProps->Restore();
 		delete ObjectRestoreAfterInitProps;
-		ObjectRestoreAfterInitProps = NULL;
+		ObjectRestoreAfterInitProps = nullptr;
 	}
 
 	bool bNeedInstancing = false;
