Fixed FLinkerLoad::FindNewPathNameForClass treating unmounted paths as if they were short paths

This could happen if you had BPGC in the CoreRedirects list, and the plugin that owned the BPGC was unmounted.

It now uses FPackageName::IsShortPackageName instead of !FPackageName::IsValidObjectPath, as the latter will fail for a fully qualified but unmounted path. The intent of that check was to detect legacy short class names so that they could be resolved into full paths, so the former is more appropriate.

#rnx

diff --git a/Engine/Source/Runtime/CoreUObject/Private/UObject/LinkerLoad.cpp b/Engine/Source/Runtime/CoreUObject/Private/UObject/LinkerLoad.cpp
--- a/Engine/Source/Runtime/CoreUObject/Private/UObject/LinkerLoad.cpp
+++ b/Engine/Source/Runtime/CoreUObject/Private/UObject/LinkerLoad.cpp
@@ -6211,9 +6211,8 @@ FString FLinkerLoad::FindNewPathNameForClass(const FString& OldClassNameOrPathNa
 			NewClassPathName = NewName.ToString();
 		}
 	}
-	if (!NewClassPathName.IsEmpty() && !FPackageName::IsValidObjectPath(NewClassPathName))
+	if (!NewClassPathName.IsEmpty() && FPackageName::IsShortPackageName(NewClassPathName))
 	{
-		TArray<UObject*> FoundClasses;
 		UClass* ExistingClass = FindFirstObject<UClass>(*NewClassPathName, EFindFirstObjectOptions::None, ELogVerbosity::Fatal, TEXT("FindNewPathNameForClass"));
 		if (ExistingClass)
 		{
