Redo 24852317:

Making sure non-nullable object properties have a default object value assigned to them in case of serialization errors
+ Adding backwards compatibility support to optional properties

diff --git a/Engine/Source/Runtime/CoreUObject/Private/UObject/PropertyBaseObject.cpp b/Engine/Source/Runtime/CoreUObject/Private/UObject/PropertyBaseObject.cpp
--- a/Engine/Source/Runtime/CoreUObject/Private/UObject/PropertyBaseObject.cpp
+++ b/Engine/Source/Runtime/CoreUObject/Private/UObject/PropertyBaseObject.cpp
@@ -626,6 +626,54 @@ bool FObjectPropertyBase::AllowObjectTypeReinterpretationTo(const FObjectPropert
 	return false;
 }
 
+UObject* FObjectPropertyBase::ConstructDefaultObjectValueIfNecessary(UObject* ExistingValue) const
+{
+	UObject* NewDefaultObjectValue = nullptr;
+
+	FUObjectSerializeContext* SerializeContext = FUObjectThreadContext::Get().GetSerializeContext();
+	UObject* Outer = SerializeContext ? SerializeContext->SerializedObject : nullptr;
+	if (!Outer)
+	{
+		Outer = GetTransientPackage();
+	}
+
+	if (ExistingValue)
+	{
+		UClass* ExistingValueClass = ExistingValue->GetClass();
+		// Sanity check to make sure the existing value class matches the property class
+		if (ExistingValueClass && (ExistingValueClass->IsChildOf(PropertyClass) || ExistingValueClass->GetAuthoritativeClass()->IsChildOf(PropertyClass)))
+		{
+			if (ExistingValue->IsTemplate())
+			{
+				// Existing value is a template so we can construct a new value with it as the archetype
+				// We probably got here because an object value failed to load (missing import class) and the property is left with a template of default subobject
+				NewDefaultObjectValue = NewObject<UObject>(Outer, ExistingValue->GetClass(), ExistingValue->GetFName(), RF_NoFlags, ExistingValue);
+			}
+			else
+			{
+				// Existing value is not a template so we can use it directly
+				// Similar to the above condition but the property was not referencing an instanced value in which case it's ok to leave the CDO default here
+				NewDefaultObjectValue = ExistingValue;
+			}
+		}
+	}
+
+	if (!NewDefaultObjectValue && !PropertyClass->HasAnyClassFlags(CLASS_Abstract))
+	{
+		// Existing value did not exist or it could not be used as a template
+		// Existing value may be null in case we were serializing an array of UObjects that failed to load (missing import class). Since the array is first pre-allocated with null values
+		// it will not have any existing objects to instantiate
+		NewDefaultObjectValue = NewObject<UObject>(Outer, PropertyClass);
+	}
+
+	// Final sanity check. We still may end up with a null object if the property class is abstract and the previous object value was missing or was not compatible
+	UE_CLOG(!NewDefaultObjectValue, LogProperty, Fatal, TEXT("Failed to create default object value for property %s. Previous value: %s"), 
+		*GetFullName(), 
+		ExistingValue ? *ExistingValue->GetFullName() : TEXT("None"));
+
+	return NewDefaultObjectValue;
+}
+
 void FObjectPropertyBase::CheckValidObject(void* ValueAddress, TObjectPtr<UObject> OldValue) const
 {
 	const TObjectPtr<UObject> Object = GetObjectPtrPropertyValue(ValueAddress);
@@ -682,16 +730,18 @@ void FObjectPropertyBase::CheckValidObject(void* ValueAddress, TObjectPtr<UObjec
 			}
 			else
 			{
-				checkf(OldValue, TEXT("CheckValidObject(\"%s\") trying to assign null object value to non-nullable property \"%s\""), *Object.GetFullName(), *GetFullName());
+				UObject* DefaultValue = ConstructDefaultObjectValueIfNecessary(OldValue);
+
 				UE_LOG(LogProperty, Warning,
-					TEXT("Serialized %s for a property of %s. Reference will be reverted back to %s.\n    Property = %s\n    Item = %s"),
+					TEXT("Serialized %s for a non-nullable property of %s. Reference will be defaulted to %s.\n    Property = %s\n    Item = %s"),
 					*ObjectClass->GetFullName(),
 					*PropertyClass->GetFullName(),
-					OldValue ? *OldValue->GetFullName() : TEXT("None"),
+					DefaultValue ? *DefaultValue->GetFullName() : TEXT("None"),
 					*GetFullName(),
 					*Object.GetFullName()
 				);
-				SetObjectPropertyValue(ValueAddress, OldValue);
+
+				SetObjectPropertyValue(ValueAddress, DefaultValue);
 			}
 		}
 	}
diff --git a/Engine/Source/Runtime/CoreUObject/Private/UObject/PropertyObject.cpp b/Engine/Source/Runtime/CoreUObject/Private/UObject/PropertyObject.cpp
--- a/Engine/Source/Runtime/CoreUObject/Private/UObject/PropertyObject.cpp
+++ b/Engine/Source/Runtime/CoreUObject/Private/UObject/PropertyObject.cpp
@@ -91,7 +91,7 @@ EConvertFromTypeResult FObjectProperty::ConvertFromType(const FPropertyTag& Tag,
 	return EConvertFromTypeResult::UseSerializeItem;
 }
 
-void FObjectProperty::SerializeItem( FStructuredArchive::FSlot Slot, void* Value, void const* Defaults ) const
+void FObjectProperty::SerializeItem(FStructuredArchive::FSlot Slot, void* Value, void const* Defaults) const
 {
 	FArchive& UnderlyingArchive = Slot.GetUnderlyingArchive();
 
@@ -102,7 +102,7 @@ void FObjectProperty::SerializeItem( FStructuredArchive::FSlot Slot, void* Value
 		Slot << (*ObjectPtr);
 
 #if !(UE_BUILD_TEST || UE_BUILD_SHIPPING) 
-		if(!UnderlyingArchive.IsSaving())
+		if (!UnderlyingArchive.IsSaving())
 		{
 			CheckValidObject(ObjectPtr, *ObjectPtr);
 		}
@@ -120,34 +120,58 @@ void FObjectProperty::SerializeItem( FStructuredArchive::FSlot Slot, void* Value
 		check(CurrentValuePtr.IsResolved());
 		UObject* CurrentValue = UE::CoreUObject::Private::ReadObjectHandlePointerNoCheck(CurrentValuePtr.GetHandle());
 
-		if (ObjectValue != CurrentValue)
-		{
-			SetObjectPropertyValue(Value, ObjectValue);
+		PostSerializeObjectItem(UnderlyingArchive, Value, CurrentValue, ObjectValue, EObjectPropertyOptions::None);
+	}
+}
 
-	#if USE_CIRCULAR_DEPENDENCY_LOAD_DEFERRING
-			if (ULinkerPlaceholderExportObject* PlaceholderVal = Cast<ULinkerPlaceholderExportObject>(ObjectValue))
-			{
-				PlaceholderVal->AddReferencingPropertyValue(this, Value);
-			}
-			else if (ULinkerPlaceholderClass* PlaceholderClass = Cast<ULinkerPlaceholderClass>(ObjectValue))
-			{
-				PlaceholderClass->AddReferencingPropertyValue(this, Value);
-			}
-			// NOTE: we don't remove this from CurrentValue if it is a 
-			//       ULinkerPlaceholderExportObject; this is because this property 
-			//       could be an array inner, and another member of that array (also 
-			//       referenced through this property)... if this becomes a problem,
-			//       then we could inc/decrement a ref count per referencing property 
-			//
-			// @TODO: if this becomes problematic (because ObjectValue doesn't match 
-			//        this property's PropertyClass), then we could spawn another
-			//        placeholder object (of PropertyClass's type), or use null; but
-			//        we'd have to modify ULinkerPlaceholderExportObject::ReplaceReferencingObjectValues()
-			//        to accommodate this (as it depends on finding itself as the set value)
-	#endif // USE_CIRCULAR_DEPENDENCY_LOAD_DEFERRING
-
-			CheckValidObject(Value, CurrentValue);
+void FObjectProperty::PostSerializeObjectItem(FArchive& SerializingArchive, void* Value, UObject* CurrentValue, UObject* ObjectValue, EObjectPropertyOptions Options /*= EObjectPropertyOptions::None*/) const
+{
+	// Make sure non-nullable properties don't end up with null values
+	if (!(Options & EObjectPropertyOptions::AllowNullValuesOnNonNullableProperty) &&
+		!ObjectValue && HasAnyPropertyFlags(CPF_NonNullable) &&
+		!SerializingArchive.IsSerializingDefaults() && // null values when Serializing CDOs are allowed, they will be fixed up later
+		!SerializingArchive.IsSaving() && // Constructing new objects when saving may confuse package saving code (new import/export created after import/export collection pass)
+		!SerializingArchive.IsTransacting()) // Don't create new objects when loading from the transaction buffer
+	{
+		UObject* DefaultValue = ConstructDefaultObjectValueIfNecessary(CurrentValue);
+
+		UE_LOG(LogProperty, Warning,
+			TEXT("Failed to serialize value for non-nullable property %s. Reference will be defaulted to %s."),
+			*GetFullName(),
+			*DefaultValue->GetFullName()
+		);
+
+		SetObjectPropertyValue(Value, DefaultValue);
+		ObjectValue = DefaultValue;
+	}
+
+	if (ObjectValue != CurrentValue)
+	{
+		SetObjectPropertyValue(Value, ObjectValue);
+
+#if USE_CIRCULAR_DEPENDENCY_LOAD_DEFERRING
+		if (ULinkerPlaceholderExportObject* PlaceholderVal = Cast<ULinkerPlaceholderExportObject>(ObjectValue))
+		{
+			PlaceholderVal->AddReferencingPropertyValue(this, Value);
+		}
+		else if (ULinkerPlaceholderClass* PlaceholderClass = Cast<ULinkerPlaceholderClass>(ObjectValue))
+		{
+			PlaceholderClass->AddReferencingPropertyValue(this, Value);
 		}
+		// NOTE: we don't remove this from CurrentValue if it is a 
+		//       ULinkerPlaceholderExportObject; this is because this property 
+		//       could be an array inner, and another member of that array (also 
+		//       referenced through this property)... if this becomes a problem,
+		//       then we could inc/decrement a ref count per referencing property 
+		//
+		// @TODO: if this becomes problematic (because ObjectValue doesn't match 
+		//        this property's PropertyClass), then we could spawn another
+		//        placeholder object (of PropertyClass's type), or use null; but
+		//        we'd have to modify ULinkerPlaceholderExportObject::ReplaceReferencingObjectValues()
+		//        to accommodate this (as it depends on finding itself as the set value)
+#endif // USE_CIRCULAR_DEPENDENCY_LOAD_DEFERRING
+
+		CheckValidObject(Value, CurrentValue);
 	}
 }
 
diff --git a/Engine/Source/Runtime/CoreUObject/Public/UObject/UnrealType.h b/Engine/Source/Runtime/CoreUObject/Public/UObject/UnrealType.h
--- a/Engine/Source/Runtime/CoreUObject/Public/UObject/UnrealType.h
+++ b/Engine/Source/Runtime/CoreUObject/Public/UObject/UnrealType.h
@@ -2690,6 +2690,14 @@ protected:
 	virtual bool AllowObjectTypeReinterpretationTo(const FObjectPropertyBase* Other) const;
 	// End of FObjectPropertyBase interface
 
+	/**
+	 * Constructs a new object if the existing one is missing or is not compatible with the property class
+	 * Used for making sure non-nullable properties have valid values.
+	 * @param ExistingValue Previous object value (can be null)
+	 * @return Non-null object that was assigned to this property value address
+	 */
+	UObject* ConstructDefaultObjectValueIfNecessary(UObject* ExistingValue) const;
+
 	/* Helper functions for UObject property types that wrap the object pointer in a smart pointer */
 	template <typename T>
 	void GetWrappedUObjectPtrValues(UObject** OutObjects, const void* SrcAddress, EPropertyMemoryAccess SrcAccess, int32 ArrayIndex, int32 ArrayCount) const
@@ -2874,6 +2882,12 @@ PRAGMA_ENABLE_DEPRECATION_WARNINGS
 	}
 };
 
+enum class EObjectPropertyOptions
+{
+	None = 0,
+	AllowNullValuesOnNonNullableProperty = 1
+};
+ENUM_CLASS_FLAGS(EObjectPropertyOptions);
 
 //
 // Describes a reference variable to another object which may be nil.
@@ -2941,6 +2955,15 @@ public:
 	virtual FString GetCPPTypeCustom(FString* ExtendedTypeText, uint32 CPPExportFlags, const FString& InnerNativeTypeName)  const override;
 	// End of FObjectPropertyBase interface
 	
+	/**
+	 * Performs post serialization steps after loading a property value
+	 * @param SerializingArchive Archive used for serialization
+	 * @param Value Property address
+	 * @param CurrentValue Current Object value
+	 * @param ObjectValue Deserialized Object value
+	 */
+	void PostSerializeObjectItem(FArchive& SerializingArchive, void* Value, UObject* CurrentValue, UObject* ObjectValue, EObjectPropertyOptions Options = EObjectPropertyOptions::None) const;
+
 	inline TObjectPtr<UObject>* GetObjectPtrPropertyValuePtr(const void* PropertyValueAddress) const
 	{
 		return reinterpret_cast<TObjectPtr<UObject>*>(const_cast<void*>(PropertyValueAddress));
diff --git a/Engine/Source/Runtime/CoreUObject/Tests/ObjectPropertyTest.cpp b/Engine/Source/Runtime/CoreUObject/Tests/ObjectPropertyTest.cpp
--- a/Engine/Source/Runtime/CoreUObject/Tests/ObjectPropertyTest.cpp
+++ b/Engine/Source/Runtime/CoreUObject/Tests/ObjectPropertyTest.cpp
@@ -115,7 +115,7 @@ TEST_CASE("UE::CoreUObject::FObjectProperty::CheckValidAddressNonNullable")
 
 	UE::Testing::FWarnFilterScope _([](const TCHAR* Message, ELogVerbosity::Type Verbosity, const FName& Category)
 		{
-			if (Category == TEXT("LogProperty") && FCString::Strstr(Message, TEXT("Reference will be reverted back to")) && Verbosity == ELogVerbosity::Type::Warning)
+			if (Category == TEXT("LogProperty") && FCString::Strstr(Message, TEXT("Reference will be defaulted to")) && Verbosity == ELogVerbosity::Type::Warning)
 			{
 				return true;
 			}
@@ -128,8 +128,10 @@ TEST_CASE("UE::CoreUObject::FObjectProperty::CheckValidAddressNonNullable")
 	Obj->ObjectPtrNonNullable = reinterpret_cast<UObjectPtrTestClass*>(Obj);
 	CHECK(Obj->ObjectPtrNonNullable != nullptr);
 
-	//old value is required for non nullable properties
-	REQUIRE_CHECK(Property->CheckValidObject(&Obj->ObjectPtrNonNullable, nullptr));
+	//new value is required for non nullable properties
+	Property->CheckValidObject(&Obj->ObjectPtrNonNullable, nullptr);
+	CHECK(Obj->ObjectPtrNonNullable != nullptr);
+	CHECK(Obj->ObjectPtrNonNullable->IsA(UObjectPtrTestClass::StaticClass()));
 }
 
 class FMockArchive : public FArchive
