AsyncLoading2: Fix race crash in FAsyncLoadEventQueue2::UpdatePackagePriority on GT reading ExportBundleNodes data while it was being created by ALT in FAsyncLoadingThread2::FinishInitializeAsyncPackage.
Fallout from 24832537 that introduced FinishInitializeAsyncPackage and two-phased package initialization.
In FAsyncLoadingThread2::UpdatePackagePriority only try to reprioritize a package in the MainThreadEventQueue queue if the ALT processing is done and the package is supposed to be in the MainThreadEventQueue.

#rnx

diff --git a/Engine/Source/Runtime/CoreUObject/Private/Serialization/AsyncLoading2.cpp b/Engine/Source/Runtime/CoreUObject/Private/Serialization/AsyncLoading2.cpp
--- a/Engine/Source/Runtime/CoreUObject/Private/Serialization/AsyncLoading2.cpp
+++ b/Engine/Source/Runtime/CoreUObject/Private/Serialization/AsyncLoading2.cpp
@@ -3736,7 +3736,7 @@ void FAsyncLoadingThread2::UpdatePackagePriority(FAsyncLoadingThreadState2& Thre
 	{
 		EventQueue.UpdatePackagePriority(Package);
 	}
-	if (LoadingState <= EAsyncPackageLoadingState2::DeferredPostLoad)
+	if (LoadingState == EAsyncPackageLoadingState2::DeferredPostLoad)
 	{
 		if (ThreadState.bIsAsyncLoadingThread)
 		{
