SavePackage and Cooking: Remove requirement that packages are cooked with an ICookedPackageWriter; SavePackage only uses IPackageWriter functions and we will be adding one-off cook saves that do not implement the full ICookedPackageWriter.
#rnx

diff --git a/Engine/Source/Runtime/CoreUObject/Private/UObject/SavePackage/SaveContext.h b/Engine/Source/Runtime/CoreUObject/Private/UObject/SavePackage/SaveContext.h
--- a/Engine/Source/Runtime/CoreUObject/Private/UObject/SavePackage/SaveContext.h
+++ b/Engine/Source/Runtime/CoreUObject/Private/UObject/SavePackage/SaveContext.h
@@ -425,9 +425,9 @@ public:
 		// Assumptions & checks
 		check(InPackage);
 		check(InFilename);
-		// if we are cooking we should be doing it in the editor and with a CookedPackageWriter
+		// if we are cooking we should be doing it in the editor and with a PackageWriter
 		check(!IsCooking() || WITH_EDITOR);
-		checkf(!IsCooking() || (PackageWriter && PackageWriter->AsCookedPackageWriter()), TEXT("Cook saves require an ICookedPackageWriter"));
+		checkf(!IsCooking() || PackageWriter, TEXT("Cook saves require an IPackageWriter"));
 
 		SaveArgs.TopLevelFlags = UE::SavePackageUtilities::NormalizeTopLevelFlags(SaveArgs.TopLevelFlags, IsCooking());
 		if (PackageWriter)
