ZenLoader: Only finish loading packages that are part of the current sync load context and when flushing a specific request only call completion callbacks for the requested package
On behalf of CarlMagnus.Nordin

diff --git a/Engine/Source/Runtime/CoreUObject/Private/Serialization/AsyncLoading2.cpp b/Engine/Source/Runtime/CoreUObject/Private/Serialization/AsyncLoading2.cpp
--- a/Engine/Source/Runtime/CoreUObject/Private/Serialization/AsyncLoading2.cpp
+++ b/Engine/Source/Runtime/CoreUObject/Private/Serialization/AsyncLoading2.cpp
@@ -2973,32 +2973,34 @@ private:
 	
 	struct FCompletedPackageRequest
 	{
-		int32 RequestID = INDEX_NONE;
 		FName PackageName;
 		EAsyncLoadingResult::Type Result = EAsyncLoadingResult::Succeeded;
-		UPackage* Package = nullptr;
+		UPackage* UPackage = nullptr;
+		FAsyncPackage2* AsyncPackage = nullptr;
 		TArray<TUniquePtr<FLoadPackageAsyncDelegate>, TInlineAllocator<2>> CompletionCallbacks;
+		TArray<int32, TInlineAllocator<2>> RequestIDs;
 
 		static FCompletedPackageRequest FromMissingPackage(
 			const FAsyncPackageDesc2& Desc,
 			TUniquePtr<FLoadPackageAsyncDelegate>&& Callback)
 		{
-			FCompletedPackageRequest Res = {Desc.RequestID, Desc.UPackageName, EAsyncLoadingResult::Failed};
+			FCompletedPackageRequest Res = {Desc.UPackageName, EAsyncLoadingResult::Failed};
 			Res.CompletionCallbacks.Add(MoveTemp(Callback));
+			Res.RequestIDs.Add(Desc.RequestID);
 			return Res;
 		}
 
 		static FCompletedPackageRequest FromLoadedPackage(
-			const FAsyncPackage2& Package,
-			TArray<TUniquePtr<FLoadPackageAsyncDelegate>, TInlineAllocator<2>>&& Callbacks)
+			FAsyncPackage2* Package)
 		{
 			return FCompletedPackageRequest
 			{
-				INDEX_NONE,
-				Package.Desc.UPackageName,
-				Package.HasLoadFailed() ? EAsyncLoadingResult::Failed : EAsyncLoadingResult::Succeeded,
-				Package.LinkerRoot,
-				MoveTemp(Callbacks)
+				Package->Desc.UPackageName,
+				Package->HasLoadFailed() ? EAsyncLoadingResult::Failed : EAsyncLoadingResult::Succeeded,
+				Package->LinkerRoot,
+				Package,
+				MoveTemp(Package->CompletionCallbacks),
+				Package->RequestIDs
 			};
 		}
 
@@ -3014,11 +3016,12 @@ private:
 			TRACE_CPUPROFILER_EVENT_SCOPE(PackageCompletionCallbacks);
 			for (TUniquePtr<FLoadPackageAsyncDelegate>& CompletionCallback : CompletionCallbacks)
 			{
-				CompletionCallback->ExecuteIfBound(PackageName, Package, Result);
+				CompletionCallback->ExecuteIfBound(PackageName, UPackage, Result);
 			}
 			CompletionCallbacks.Empty();
 		}
 	};
+	TArray<FCompletedPackageRequest> CompletedPackageRequests;
 	TArray<FCompletedPackageRequest> FailedPackageRequests;
 	FCriticalSection FailedPackageRequestsCritical;
 
@@ -6979,7 +6982,8 @@ EAsyncPackageState::Type FAsyncLoadingThread2::ProcessLoadedPackagesFromGameThre
 			MainThreadEventQueue.UpdatePackagePriority(PackageToRepriortize);
 			PackageToRepriortize->ReleaseRef();
 		}
-		if (!ThreadState.SyncLoadContextStack.IsEmpty() && ThreadState.SyncLoadContextStack.Top()->ContextId)
+		uint64 SyncLoadContextId = !ThreadState.SyncLoadContextStack.IsEmpty() ? ThreadState.SyncLoadContextStack.Top()->ContextId : 0;
+		if (SyncLoadContextId)
 		{
 			bLocalDidSomething |= MainThreadEventQueue.ExecuteSyncLoadEvents(ThreadState);
 		}
@@ -6988,11 +6992,15 @@ EAsyncPackageState::Type FAsyncLoadingThread2::ProcessLoadedPackagesFromGameThre
 			bLocalDidSomething |= MainThreadEventQueue.PopAndExecute(ThreadState);
 		}
 
-		bLocalDidSomething |= LoadedPackagesToProcess.Num() > 0;
 		for (int32 PackageIndex = 0; PackageIndex < LoadedPackagesToProcess.Num(); ++PackageIndex)
 		{
 			SCOPED_LOADTIMER(ProcessLoadedPackagesTime);
 			FAsyncPackage2* Package = LoadedPackagesToProcess[PackageIndex];
+			if (Package->SyncLoadContextId < SyncLoadContextId)
+			{
+				continue;
+			}
+			bLocalDidSomething = true;
 			UE_ASYNC_PACKAGE_DEBUG(Package->Desc);
 			check(Package->AsyncPackageLoadingState >= EAsyncPackageLoadingState2::Finalize &&
 				  Package->AsyncPackageLoadingState <= EAsyncPackageLoadingState2::CreateClusters);
@@ -7122,42 +7130,60 @@ EAsyncPackageState::Type FAsyncLoadingThread2::ProcessLoadedPackagesFromGameThre
 			LocalCompletedAsyncPackages.Add(Package);
 		}
 
-		TArray<FCompletedPackageRequest> LocalCompletedPackageRequests;
 		{
 			FScopeLock _(&FailedPackageRequestsCritical);
-			Swap(LocalCompletedPackageRequests, FailedPackageRequests);
+			CompletedPackageRequests.Append(MoveTemp(FailedPackageRequests));
+			FailedPackageRequests.Reset();
 		}
-		for (FCompletedPackageRequest& FailedPackageRequest : LocalCompletedPackageRequests)
-		{
-			RemovePendingRequests(TArrayView<int32>(&FailedPackageRequest.RequestID, 1));
-			--PackagesWithRemainingWorkCounter;
-			TRACE_COUNTER_SET(AsyncLoadingPackagesWithRemainingWork, PackagesWithRemainingWorkCounter);
-		}
-
-		LocalCompletedPackageRequests.Reserve(LocalCompletedPackageRequests.Num() + LocalCompletedAsyncPackages.Num());
 		for (FAsyncPackage2* Package : LocalCompletedAsyncPackages)
 		{
 			UE_ASYNC_PACKAGE_DEBUG(Package->Desc);
 			UE_ASYNC_PACKAGE_LOG(Verbose, Package->Desc, TEXT("GameThread: LoadCompleted"),
 				TEXT("All loading of package is done, and the async package and load request will be deleted."));
 
-			LocalCompletedPackageRequests.Add(FCompletedPackageRequest::FromLoadedPackage(*Package, MoveTemp(Package->CompletionCallbacks)));
-			RemovePendingRequests(Package->RequestIDs);
+
 			check(Package->AsyncPackageLoadingState == EAsyncPackageLoadingState2::Complete);
 			Package->AsyncPackageLoadingState = EAsyncPackageLoadingState2::DeferredDelete;
 			Package->ClearImportedPackages();
-			Package->ReleaseRef();
+
+			if (Package->CompletionCallbacks.IsEmpty())
+			{
+				RemovePendingRequests(Package->RequestIDs);
+				Package->ReleaseRef();
+			}
+			else
+			{
+				// Note: We need to keep the package alive until the callback has been executed
+				CompletedPackageRequests.Add(FCompletedPackageRequest::FromLoadedPackage(Package));
+			}
 		}
 		LocalCompletedAsyncPackages.Reset();
 		
 		// Call callbacks in a batch in a stack-local array after all other work has been done to handle
 		// callbacks that may call FlushAsyncLoading
-		for (FCompletedPackageRequest& CompletedPackageRequest : LocalCompletedPackageRequests)
+		// If we're flushing a specific request only call callbacks for that request
+		for (int32 CompletedPackageRequestIndex = CompletedPackageRequests.Num() - 1; CompletedPackageRequestIndex >= 0; --CompletedPackageRequestIndex)
 		{
-			CompletedPackageRequest.CallCompletionCallbacks();
+			FCompletedPackageRequest& CompletedPackageRequest = CompletedPackageRequests[CompletedPackageRequestIndex];
+			if (FlushRequestID == INDEX_NONE || CompletedPackageRequest.RequestIDs.Contains(FlushRequestID))
+			{
+				RemovePendingRequests(CompletedPackageRequest.RequestIDs);
+				TRACE_COUNTER_SET(AsyncLoadingPackagesWithRemainingWork, PackagesWithRemainingWorkCounter);
+				CompletedPackageRequest.CallCompletionCallbacks();
+				if (CompletedPackageRequest.AsyncPackage)
+				{
+					CompletedPackageRequest.AsyncPackage->ReleaseRef();
+				}
+				else
+				{
+					// Requests for missing packages have no AsyncPackage but they count as packages with remaining work
+					--PackagesWithRemainingWorkCounter;
+				}
+				CompletedPackageRequests.RemoveAt(CompletedPackageRequestIndex);
+				bLocalDidSomething = true;
+			}
 		}
 
-		bLocalDidSomething |= LocalCompletedPackageRequests.Num() > 0;
 		if (!bLocalDidSomething)
 		{
 			break;
