Fix potential truncation issue is FScriptInstrumentationSignal

#rnx

- The truncation comes from performing pointer arithmetic to find an offset into a TArray which is then stored as an int32. Since the offset will always fit in int32 we can just cast here safely.

diff --git a/Engine/Source/Runtime/CoreUObject/Private/UObject/ScriptCore.cpp b/Engine/Source/Runtime/CoreUObject/Private/UObject/ScriptCore.cpp
--- a/Engine/Source/Runtime/CoreUObject/Private/UObject/ScriptCore.cpp
+++ b/Engine/Source/Runtime/CoreUObject/Private/UObject/ScriptCore.cpp
@@ -799,7 +799,10 @@ int32 FScriptInstrumentationSignal::GetScriptCodeOffset() const
 	}
 	else if (StackFramePtr != nullptr)
 	{
-		CodeOffset = StackFramePtr->Code - StackFramePtr->Node->Script.GetData() - 1;
+		// StackFramePtr->Code should never be outside the TArray StackFramePtr->Node->Script
+		// and so the resulting pointer calculation should always fit in a int32 (the max size 
+		// TArray stores) so it is safe to just cast.
+		CodeOffset = static_cast<int32>(StackFramePtr->Code - StackFramePtr->Node->Script.GetData() - 1);
 	}
 	return CodeOffset;
 }
