Fix potential truncation issue when logging game thread hitching during resource clean up.

#rnx

diff --git a/Engine/Source/Runtime/CoreUObject/Private/UObject/UObjectGlobals.cpp b/Engine/Source/Runtime/CoreUObject/Private/UObject/UObjectGlobals.cpp
--- a/Engine/Source/Runtime/CoreUObject/Private/UObject/UObjectGlobals.cpp
+++ b/Engine/Source/Runtime/CoreUObject/Private/UObject/UObjectGlobals.cpp
@@ -3417,8 +3417,8 @@ UObject* StaticAllocateObject
 				}
 				if (bPrinted)
 				{
-					float ThisTime = FPlatformTime::Seconds() - StallStart;
-					UE_LOG(LogUObjectGlobals, Warning, TEXT("Gamethread hitch waiting for resource cleanup on a UObject (%s) overwrite took %6.2fms. Fix the higher level code so that this does not happen."), *OldName, ThisTime * 1000.0f);
+					const double ThisTime = FPlatformTime::Seconds() - StallStart;
+					UE_LOG(LogUObjectGlobals, Warning, TEXT("Gamethread hitch waiting for resource cleanup on a UObject (%s) overwrite took %6.2fms. Fix the higher level code so that this does not happen."), *OldName, ThisTime * 1000.0);
 				}
 				// Finish destroying the object.
 				Obj->ConditionalFinishDestroy();
