Fix potential truncation issue in FWorkBlockifier

#rnx

- In this case we are performing pointer arithemetic to calculate how many objects are in the block. The max number of objects is set by FWorkBlock::ObjectCapacity (currently 495) meaning we can safely cast to int32.

diff --git a/Engine/Source/Runtime/CoreUObject/Public/UObject/FastReferenceCollector.h b/Engine/Source/Runtime/CoreUObject/Public/UObject/FastReferenceCollector.h
--- a/Engine/Source/Runtime/CoreUObject/Public/UObject/FastReferenceCollector.h
+++ b/Engine/Source/Runtime/CoreUObject/Public/UObject/FastReferenceCollector.h
@@ -196,7 +196,11 @@ private:
 	COREUOBJECT_API FWorkBlock* PopFullBlockAsync();
 	COREUOBJECT_API FWorkBlock* PopWipBlock();
 	COREUOBJECT_API FWorkBlock* StealAsyncBlock() const;
-	FORCEINLINE int32 PartialNum() const { return WipIt - Wip->Objects; }
+
+	FORCEINLINE int32 PartialNum() const 
+	{
+		return static_cast<int32>(WipIt - Wip->Objects);
+	}
 };
 
 //////////////////////////////////////////////////////////////////////////
