Update FEditorBulkData code documentation to make it more obvious how the internal data can be updated by referencing the examples in FEditorBulkDataTests.cpp

diff --git a/Engine/Source/Runtime/CoreUObject/Public/Serialization/EditorBulkData.h b/Engine/Source/Runtime/CoreUObject/Public/Serialization/EditorBulkData.h
--- a/Engine/Source/Runtime/CoreUObject/Public/Serialization/EditorBulkData.h
+++ b/Engine/Source/Runtime/CoreUObject/Public/Serialization/EditorBulkData.h
@@ -104,6 +104,12 @@ enum class ECompressionOptions : uint8
  * would be able to share local copies of the payload between branches to reduce the cost of having multiple branches 
  * with similar data on the same machine.
  * 
+ * Updating the bulkdata:
+ * Unlike the older bulkdata system, EditorBulkData does not support any lock/unlock mechanism, instead the internal 
+ * data is updated via a single call to ::UpdatePayload and the data is retrieved via a single call to ::GetPayload.
+ * 
+ * For examples see the "System.CoreUObject.Serialization.EditorBulkData.UpdatePayload" test found in
+ * Engine\Source\Runtime\CoreUObject\Private\Tests\Serialization\EditorBulkDataTests.cpp.
  */
 
 /** The base class with no type */
