Remove flagging GeneratePublicHash when referenced by optional objects as it isn't needed anymore.
This also prevents false positive diffing when cooking with optional data but diffing without

diff --git a/Engine/Source/Runtime/CoreUObject/Private/UObject/SavePackage2.cpp b/Engine/Source/Runtime/CoreUObject/Private/UObject/SavePackage2.cpp
--- a/Engine/Source/Runtime/CoreUObject/Private/UObject/SavePackage2.cpp
+++ b/Engine/Source/Runtime/CoreUObject/Private/UObject/SavePackage2.cpp
@@ -390,8 +390,9 @@ ESavePackageResult HarvestPackage(FSaveContext& SaveContext)
 				// Make sure the export is found in the game context as well
 				if (FTaggedExport* GameExport = SaveContext.GetHarvestedRealm(ESaveRealm::Game).GetExports().Find(It->Obj))
 				{
-					// Flag the export in the game context to generate it's public hash
-					GameExport->bGeneratePublicHash = true;
+					// Flagging the export in the game context to generate it's public hash isn't necessary anymore
+					//GameExport->bGeneratePublicHash = true; 
+
 					// Transform the export as an import
 					OptionalContext.AddImport(It->Obj);
 					// Flag the package itself to be an import
