SavePackage support for OnlyEditorOnly: Fixes for CookedEditor platform and Optional realm.
Fix harvesting to respect TargetPlatform->AllowsEditorObjects; CookedEditor platforms set this to true and need to allow editoronly objects.
Fix redirectors to optional objects to be placed into the Optional realm to avoid a spurious illegal reference from nonoptional to optional error.
Fix accidental ignoring of the duplication of public exports into the optional realm when handling autooptional.
#rnx

diff --git a/Engine/Source/Runtime/CoreUObject/Private/UObject/SavePackage/PackageHarvester.cpp b/Engine/Source/Runtime/CoreUObject/Private/UObject/SavePackage/PackageHarvester.cpp
--- a/Engine/Source/Runtime/CoreUObject/Private/UObject/SavePackage/PackageHarvester.cpp
+++ b/Engine/Source/Runtime/CoreUObject/Private/UObject/SavePackage/PackageHarvester.cpp
@@ -228,7 +228,16 @@ FPackageHarvester::FHarvestScope FPackageHarvester::EnterConditionalEditorOnlySc
 	FHarvestScope Scope(*this);
 	if (bIsEditorOnly)
 	{
-		CurrentExportHarvestingRealms.RemoveSwap(ESaveRealm::Game, false /* bAllowShrinking */);
+		// When saving editor (no targetplatform), editoronly objects are stored in editor realm but not
+		// game so that we can mark them as editoronly.
+		// When saving cooked, the game realm is all we are saving, and most platforms do not export those
+		// editor-only objects so we do not add them to the game realm for those platforms. But some cooked
+		// platforms include editor-only objects (CookedEditor platforms) so for those platforms we do allow
+		// editor-only objects to be added to the game realm.
+		if (!SaveContext.GetTargetPlatform() || !SaveContext.GetTargetPlatform()->AllowsEditorObjects())
+		{
+			CurrentExportHarvestingRealms.RemoveSwap(ESaveRealm::Game, false /* bAllowShrinking */);
+		}
 	}
 	return Scope;
 }
@@ -240,20 +249,28 @@ FPackageHarvester::FHarvestScope FPackageHarvester::EnterConditionalOptionalObje
 	{
 		return Scope;
 	}
-	bool bIsOptionalObject = ShouldObjectBeHarvestedInOptionalRealm(Object);
+	bool bIsOptionalObject = ShouldObjectBeHarvestedInOptionalRealm(Object, SaveContext);
 	if (!bIsOptionalObject)
 	{
 		return Scope;
 	}
 
 	// No matter the realm referring to it, optional objects are added into the optional realm
-	// if we are harvesting an optional object while in a different context, record an illegal reference to display on validation
 	ForEachExportHarvestingRealm([this, Object](ESaveRealm HarvestingRealm)
 		{
-			if (HarvestingRealm != ESaveRealm::Optional)
+			// It is illegal for a non-optional used-in-game object to reference an optional object
+			// Give an error if the referer is not optional and is not editoronly.
+			// An example of an editor-only object refererring to an optional object is a UObjectRedirector.
+			UObject* Ref = CurrentExportDependencies.CurrentExport;
+			if (HarvestingRealm != ESaveRealm::Optional && Ref)
 			{
-				SaveContext.RecordIllegalReference(CurrentExportDependencies.CurrentExport,
-					Object, EIllegalRefReason::ReferenceToOptional);
+				// TODO: Change IsEditorOnlyObject to also test NeedsLoadForClient and NeedsLoadForServer, 
+				// and then change this location to call IsEditorOnlyObject
+				bool bEditorOnly = Ref->IsEditorOnly() || (!Ref->NeedsLoadForClient() && !Ref->NeedsLoadForServer());
+				if (!bEditorOnly)
+				{
+					SaveContext.RecordIllegalReference(Ref, Object, EIllegalRefReason::ReferenceToOptional);
+				}
 			}
 		});
 
@@ -314,9 +331,13 @@ bool FPackageHarvester::IsObjNative(UObject* InObj)
 	return bIsNative;
 };
 
-bool FPackageHarvester::ShouldObjectBeHarvestedInOptionalRealm(UObject* InObj) const
+bool FPackageHarvester::ShouldObjectBeHarvestedInOptionalRealm(UObject* InObj, FSaveContext& InContext)
 {
-	return InObj->GetClass()->HasAnyClassFlags(CLASS_Optional) && SaveContext.IsCooking();
+	if (!InContext.IsCooking())
+	{
+		return false;
+	}
+	return InObj->GetClass()->HasAnyClassFlags(CLASS_Optional);
 }
 
 template <typename CallbackType>
@@ -489,15 +510,21 @@ void FPackageHarvester::TryHarvestExport(UObject* InObject)
 
 	FHarvestScope ScopeSetHarvestingRealms(*this);
 	FExportingRealmsArray HarvestingRealms;
-	// if the object is optional and we are cooking, harvest in the Optional context
-	if (ShouldObjectBeHarvestedInOptionalRealm(InObject))
+	if (ShouldObjectBeHarvestedInOptionalRealm(InObject, SaveContext))
+	{
+		// if the object is optional and we are cooking, harvest in the Optional context
+		HarvestingRealms.Add(ESaveRealm::Optional);
+	}
+	else if (SaveContext.CurrentHarvestingRealm == ESaveRealm::Optional)
 	{
+		// If we are automatically generating an optional package then the callsite will harvest all public objects into the optional
+		// realm in addition to harvesting them normally. It indicates harvesting into the optional realm by setting CurrentHarvestingRealm.
 		HarvestingRealms.Add(ESaveRealm::Optional);
 	}
 	else
 	{
-		// For regular objects, exports are added by default to the game realm, and if making an editor
-		// save, also to the editor realm
+		// For regular objects not in the autoptional special harvest context, exports are added by default to the game realm, and if making
+		// an editor save, also to the editor realm
 		HarvestingRealms.Add(ESaveRealm::Game);
 		if (!SaveContext.IsCooking())
 		{
diff --git a/Engine/Source/Runtime/CoreUObject/Private/UObject/SavePackage/PackageHarvester.h b/Engine/Source/Runtime/CoreUObject/Private/UObject/SavePackage/PackageHarvester.h
--- a/Engine/Source/Runtime/CoreUObject/Private/UObject/SavePackage/PackageHarvester.h
+++ b/Engine/Source/Runtime/CoreUObject/Private/UObject/SavePackage/PackageHarvester.h
@@ -84,6 +84,11 @@ public:
 
 	virtual bool ShouldSkipProperty(const FProperty* InProperty) const;
 
+	/** Return whether the Object and all its Outers are native. */
+	static bool IsObjNative(UObject* InObj);
+	/** Report whether the object should be put in the optional realm even if harvesting realm is non-optional. */
+	static bool ShouldObjectBeHarvestedInOptionalRealm(UObject* InObj, FSaveContext& InSaveContext);
+
 private:
 	/** Dependencies collected for the current export, and scope information about how to collect them. */
 	struct FExportDependencies
@@ -121,10 +126,6 @@ private:
 	};
 
 private:
-	/** Return whether the Object and all its Outers are native. */
-	static bool IsObjNative(UObject* InObj);
-	/** Report whether the object should be put in the optional realm even if harvesting realm is non-optional. */
-	bool ShouldObjectBeHarvestedInOptionalRealm(UObject* InObj) const;
 
 	/** Call void (ESaveRealm HarvestingRealm) on each Realm in CurrentExportHarvestingRealms */
 	template <typename CallbackType>
diff --git a/Engine/Source/Runtime/CoreUObject/Private/UObject/SavePackage2.cpp b/Engine/Source/Runtime/CoreUObject/Private/UObject/SavePackage2.cpp
--- a/Engine/Source/Runtime/CoreUObject/Private/UObject/SavePackage2.cpp
+++ b/Engine/Source/Runtime/CoreUObject/Private/UObject/SavePackage2.cpp
@@ -385,7 +385,7 @@ ESavePackageResult HarvestPackage(FSaveContext& SaveContext)
 		FHarvestedRealm& OptionalContext = SaveContext.GetHarvestedRealm(ESaveRealm::Optional);
 		for (auto It = OptionalContext.GetExports().CreateIterator(); It; ++It)
 		{
-			if (!It->Obj->GetClass()->HasAnyClassFlags(CLASS_Optional))
+			if (!FPackageHarvester::ShouldObjectBeHarvestedInOptionalRealm(It->Obj, SaveContext))
 			{
 				// Make sure the export is found in the game context as well
 				if (FTaggedExport* GameExport = SaveContext.GetHarvestedRealm(ESaveRealm::Game).GetExports().Find(It->Obj))
