Fix a number of potential truncation issues in FEDLCookChecker::ReadFromCompactBinary

#rnx

- A number of containers were potentially being set to sizes that they do not support. Now we check the sizes before reserving space and if they will be too larger we just return false.
- It is not really expected that any of these conditions will trigger.

diff --git a/Engine/Source/Runtime/CoreUObject/Private/UObject/SavePackage/SavePackageUtilities.cpp b/Engine/Source/Runtime/CoreUObject/Private/UObject/SavePackage/SavePackageUtilities.cpp
--- a/Engine/Source/Runtime/CoreUObject/Private/UObject/SavePackage/SavePackageUtilities.cpp
+++ b/Engine/Source/Runtime/CoreUObject/Private/UObject/SavePackage/SavePackageUtilities.cpp
@@ -1540,7 +1540,12 @@ bool FEDLCookChecker::ReadFromCompactBinary(FCbFieldView Field)
 	};
 
 	FCbFieldView NodesField = Field["Nodes"];
-	Nodes.Reserve(NodesField.AsArrayView().Num() / 5);
+	const uint64 NumNodes = NodesField.AsArrayView().Num() / 5;
+	if (NumNodes > MAX_int32)
+	{
+		return false;
+	}
+	Nodes.Reserve(static_cast<int32>(NumNodes));
 	if (NodesField.HasError())
 	{
 		return false;
@@ -1567,7 +1572,12 @@ bool FEDLCookChecker::ReadFromCompactBinary(FCbFieldView Field)
 	}
 
 	FCbFieldView PrereqsField = Field["NodePrereqs"];
-	NodePrereqs.Reserve(PrereqsField.AsArrayView().Num() / 2);
+	const int64 NumNodePrereqs = PrereqsField.AsArrayView().Num() / 2;
+	if (NumNodePrereqs > MAX_int32)
+	{
+		return false;
+	}
+	NodePrereqs.Reserve(static_cast<int32>(NumNodePrereqs));
 	if (PrereqsField.HasError())
 	{
 		return false;
@@ -1585,7 +1595,12 @@ bool FEDLCookChecker::ReadFromCompactBinary(FCbFieldView Field)
 	}
 
 	FCbFieldView PackagesWithUnknownExportsField = Field["PackagesWithUnknownExports"];
-	PackagesWithUnknownExports.Reserve(PackagesWithUnknownExportsField.AsArrayView().Num());
+	const int64 NumPackagesWithUnknownExports = PackagesWithUnknownExportsField.AsArrayView().Num();
+	if (NumPackagesWithUnknownExports > MAX_int32)
+	{
+		return false;
+	}
+	PackagesWithUnknownExports.Reserve(static_cast<int32>(NumPackagesWithUnknownExports));
 	if (PackagesWithUnknownExportsField.HasError())
 	{
 		return false;
