Properly handle query to GPlayInEditorID while packages are being loaded

#rnx

diff --git a/Engine/Source/Editor/UnrealEd/Private/Kismet2/KismetDebugUtilities.cpp b/Engine/Source/Editor/UnrealEd/Private/Kismet2/KismetDebugUtilities.cpp
--- a/Engine/Source/Editor/UnrealEd/Private/Kismet2/KismetDebugUtilities.cpp
+++ b/Engine/Source/Editor/UnrealEd/Private/Kismet2/KismetDebugUtilities.cpp
@@ -794,7 +794,7 @@ void FKismetDebugUtilities::AttemptToBreakExecution(UBlueprint* BlueprintObj, co
 	// Now enter within-the-frame debugging mode
 	if (bShouldInStackDebug)
 	{
-		TGuardValue<int32> GuardDisablePIE(GPlayInEditorID, INDEX_NONE);
+		FTemporaryPlayInEditorIDOverride GuardDisablePIE(INDEX_NONE);
 		TArrayView<const FFrame* const> ScriptStack = FBlueprintContextTracker::Get().GetCurrentScriptStack();
 		Data.LastExceptionMessage = Info.GetDescription();
 		FKismetEditorUtilities::BringKismetToFocusAttentionOnObject(NodeStoppedAt);
diff --git a/Engine/Source/Runtime/Core/Internal/Misc/PlayInEditorLoadingScope.h b/Engine/Source/Runtime/Core/Internal/Misc/PlayInEditorLoadingScope.h
new file mode 100644
--- /dev/null
+++ b/Engine/Source/Runtime/Core/Internal/Misc/PlayInEditorLoadingScope.h
@@ -0,0 +1,19 @@
+// Copyright Epic Games, Inc. All Rights Reserved.
+
+#pragma once
+
+#include "CoreTypes.h"
+
+namespace UE::Core::Private
+{
+
+class FPlayInEditorLoadingScope
+{
+public:
+	CORE_API FPlayInEditorLoadingScope(int32 PlayInEditorID);
+	CORE_API ~FPlayInEditorLoadingScope();
+private:
+	int32 OldValue;
+};
+
+}
\ No newline at end of file
diff --git a/Engine/Source/Runtime/Core/Private/Misc/CoreGlobals.cpp b/Engine/Source/Runtime/Core/Private/Misc/CoreGlobals.cpp
--- a/Engine/Source/Runtime/Core/Private/Misc/CoreGlobals.cpp
+++ b/Engine/Source/Runtime/Core/Private/Misc/CoreGlobals.cpp
@@ -10,6 +10,7 @@
 #include "Misc/TrackedActivity.h"
 #include "Misc/Compression.h"
 #include "Misc/LazySingleton.h"
+#include "Misc/PlayInEditorLoadingScope.h"
 #include "Misc/CommandLine.h"
 #include "ProfilingDebugging/MiscTrace.h"
 #include "GenericPlatform/GenericPlatformCrashContext.h"
@@ -336,7 +337,7 @@ bool					GIsCookerLoadingPackage = false;
 /** Whether GWorld points to the play in editor world														*/
 bool					GIsPlayInEditorWorld			= false;
 /** Unique ID for multiple PIE instances running in one process */
-int32					GPlayInEditorID					= -1;
+FPlayInEditorID			GPlayInEditorID;
 /** Whether or not PIE was attempting to play from PlayerStart							*/
 bool					GIsPIEUsingPlayerStart			= false;
 /** true if the runtime needs textures to be powers of two													*/
@@ -739,6 +740,88 @@ FIsDuplicatingClassForReinstancing::operator bool() const
 	return PRIVATE_GIsDuplicatingClassForReinstancing;
 }
 
+namespace PlayInEditorIDImpl
+{
+	int32 PRIVATE_GPlayInEditorID_GameThread = -1;
+	// We need to differentiate between the game-thread being identified as the loading thread during postload 
+	// and the actual loading thread to avoid scopes to overlap and race between both threads.
+	int32 PRIVATE_GPlayInEditorID_GameThreadAsLoadingThread = -2;
+	int32 PRIVATE_GPlayInEditorID_ActualLoadingThread = -2;
+
+	static int32* GetPointer()
+	{
+		if (IsInAsyncLoadingThread())
+		{
+			if (IsInGameThread())
+			{
+				return &PRIVATE_GPlayInEditorID_GameThreadAsLoadingThread;
+			}
+			else
+			{
+				return &PRIVATE_GPlayInEditorID_ActualLoadingThread;
+			}
+		}
+		else if (IsInGameThread())
+		{
+			return &PRIVATE_GPlayInEditorID_GameThread;
+		}
+
+		return nullptr;
+	}
+};
+
+int32 PRIVATE_GetGPlayInEditorID()
+{
+	if (int32* Pointer = PlayInEditorIDImpl::GetPointer())
+	{
+		return *Pointer;
+	}
+	else
+	{
+		//ensureMsgf(false, TEXT("GPlayInEditorID can only be read by the game-thread or loading-thread, returning -1."));
+		return -1;
+	}
+}
+
+void PRIVATE_SetGPlayInEditorID(int32 InValue)
+{
+	if (int32* Pointer = PlayInEditorIDImpl::GetPointer())
+	{
+		*Pointer = InValue;
+	}
+	else
+	{
+		checkf(false, TEXT("GPlayInEditorID can only be set by the game-thread or loading-thread."));
+	}
+}
+
+namespace UE::Core::Private
+{
+	FPlayInEditorLoadingScope::FPlayInEditorLoadingScope(int32 PlayInEditorID)
+		: OldValue(PRIVATE_GetGPlayInEditorID())
+	{
+		PRIVATE_SetGPlayInEditorID(PlayInEditorID);
+	}
+
+	FPlayInEditorLoadingScope::~FPlayInEditorLoadingScope()
+	{
+		PRIVATE_SetGPlayInEditorID(OldValue);
+	}
+}
+
+FPlayInEditorID& FPlayInEditorID::operator= (int32 InOther)
+{
+	PRIVATE_SetGPlayInEditorID(InOther);
+	return *this;
+}
+
+FPlayInEditorID::operator int32() const
+{
+	int32 Value = PRIVATE_GetGPlayInEditorID();
+	checkf(Value != -2, TEXT("GPlayInEditorID has not been properly forwarded by the loading-thread."));
+	return Value;
+}
+
 #undef LOCTEXT_NAMESPACE
 
 bool IsRunningCookOnTheFly()
diff --git a/Engine/Source/Runtime/Core/Public/CoreGlobals.h b/Engine/Source/Runtime/Core/Public/CoreGlobals.h
--- a/Engine/Source/Runtime/Core/Public/CoreGlobals.h
+++ b/Engine/Source/Runtime/Core/Public/CoreGlobals.h
@@ -450,7 +450,14 @@ extern CORE_API bool GIsCookerLoadingPackage;
 /** Whether GWorld points to the play in editor world */
 extern CORE_API bool GIsPlayInEditorWorld;
 
-extern CORE_API int32 GPlayInEditorID;
+class CORE_API FPlayInEditorID
+{
+public:
+	FPlayInEditorID& operator= (int32 InOther);
+	operator int32() const;
+};
+/** In the editor, this is set to the specific world context PIEInstance that is being currently processed */
+extern CORE_API FPlayInEditorID GPlayInEditorID;
 
 /** Whether or not PIE was attempting to play from PlayerStart */
 UE_DEPRECATED(4.25, "This variable is no longer set. Use !GEditor->GetPlayInEditorSessionInfo()->OriginalRequestParams.HasPlayWorldPlacement() instead.")
diff --git a/Engine/Source/Runtime/CoreUObject/Private/Serialization/AsyncLoading.cpp b/Engine/Source/Runtime/CoreUObject/Private/Serialization/AsyncLoading.cpp
--- a/Engine/Source/Runtime/CoreUObject/Private/Serialization/AsyncLoading.cpp
+++ b/Engine/Source/Runtime/CoreUObject/Private/Serialization/AsyncLoading.cpp
@@ -34,6 +34,7 @@
 #include "UObject/ObjectSerializeAccessScope.h"
 #include "UObject/LinkerManager.h"
 #include "Misc/Paths.h"
+#include "Misc/PlayInEditorLoadingScope.h"
 #include "Serialization/AsyncLoadingThread.h"
 #include "Misc/ExclusiveLoadPackageTimeTracker.h"
 #include "ProfilingDebugging/CookStats.h"
@@ -5964,102 +5965,106 @@ EAsyncPackageState::Type FAsyncPackage::TickAsyncPackage(bool InbUseTimeLimit, b
 
 		// Begin async loading, simulates BeginLoad
 		BeginAsyncLoad();
-
-		// We have begun loading a package that we know the name of. Let the package time tracker know.
-		FExclusiveLoadPackageTimeTracker::PushLoadPackage(Desc.PackagePath.GetPackageFName());
-
-		if (!GEventDrivenLoaderEnabled)
 		{
-			// Create raw linker. Needs to be async created via ticking before it can be used.
-			if (LoadingState == EAsyncPackageState::Complete) //-V547
-			{
-				SCOPED_LOADTIMER(Package_CreateLinker);
-				LoadingState = CreateLinker();
-			}
+#if WITH_EDITOR
+			UE::Core::Private::FPlayInEditorLoadingScope PlayInEditorIDScope(Desc.PIEInstanceID);
+#endif
+			// We have begun loading a package that we know the name of. Let the package time tracker know.
+			FExclusiveLoadPackageTimeTracker::PushLoadPackage(Desc.PackagePath.GetPackageFName());
 
-			// Async create linker.
-			if (LoadingState == EAsyncPackageState::Complete)
+			if (!GEventDrivenLoaderEnabled)
 			{
-				SCOPED_LOADTIMER(Package_FinishLinker);
-				LoadingState = FinishLinker();
-			}
+				// Create raw linker. Needs to be async created via ticking before it can be used.
+				if (LoadingState == EAsyncPackageState::Complete) //-V547
+				{
+					SCOPED_LOADTIMER(Package_CreateLinker);
+					LoadingState = CreateLinker();
+				}
 
-			// Load imports from linker import table asynchronously.
-			if (LoadingState == EAsyncPackageState::Complete && !bLoadHasFinished)
-			{
-				SCOPED_LOADTIMER(Package_LoadImports);
-				LoadingState = LoadImports();
-			}
+				// Async create linker.
+				if (LoadingState == EAsyncPackageState::Complete)
+				{
+					SCOPED_LOADTIMER(Package_FinishLinker);
+					LoadingState = FinishLinker();
+				}
 
-			// Create imports from linker import table.
-			if (LoadingState == EAsyncPackageState::Complete && !bLoadHasFinished)
-			{
-				SCOPED_LOADTIMER(Package_CreateImports);
-				LoadingState = CreateImports();
-			}
+				// Load imports from linker import table asynchronously.
+				if (LoadingState == EAsyncPackageState::Complete && !bLoadHasFinished)
+				{
+					SCOPED_LOADTIMER(Package_LoadImports);
+					LoadingState = LoadImports();
+				}
+
+				// Create imports from linker import table.
+				if (LoadingState == EAsyncPackageState::Complete && !bLoadHasFinished)
+				{
+					SCOPED_LOADTIMER(Package_CreateImports);
+					LoadingState = CreateImports();
+				}
 
 #if WITH_EDITORONLY_DATA
-			// Create and preload the package meta-data
-			if (LoadingState == EAsyncPackageState::Complete && !bLoadHasFinished)
-			{
-				SCOPED_LOADTIMER(Package_CreateMetaData);
-				LoadingState = CreateMetaData();
-			}
+				// Create and preload the package meta-data
+				if (LoadingState == EAsyncPackageState::Complete && !bLoadHasFinished)
+				{
+					SCOPED_LOADTIMER(Package_CreateMetaData);
+					LoadingState = CreateMetaData();
+				}
 #endif // WITH_EDITORONLY_DATA
 
-			// Create exports from linker export table and also preload them.
-			if (LoadingState == EAsyncPackageState::Complete && !bLoadHasFinished)
-			{
-				SCOPED_LOADTIMER(Package_CreateExports);
-				LoadingState = CreateExports();
-			}
+				// Create exports from linker export table and also preload them.
+				if (LoadingState == EAsyncPackageState::Complete && !bLoadHasFinished)
+				{
+					SCOPED_LOADTIMER(Package_CreateExports);
+					LoadingState = CreateExports();
+				}
 
-			// Call Preload on the linker for all loaded objects which causes actual serialization.
-			if (LoadingState == EAsyncPackageState::Complete && !bLoadHasFinished)
-			{
-				SCOPED_LOADTIMER(Package_PreLoadObjects);
-				LoadingState = PreLoadObjects();
-			}
+				// Call Preload on the linker for all loaded objects which causes actual serialization.
+				if (LoadingState == EAsyncPackageState::Complete && !bLoadHasFinished)
+				{
+					SCOPED_LOADTIMER(Package_PreLoadObjects);
+					LoadingState = PreLoadObjects();
+				}
 
-			if ((LoadingState == EAsyncPackageState::Complete || bLoadHasFailed) && !bLoadHasFinished)
-			{
-				const bool bInternalCallbacks = true;
-				CallCompletionCallbacks(bInternalCallbacks, bLoadHasFailed ? EAsyncLoadingResult::Failed : EAsyncLoadingResult::Succeeded);
-			}
+				if ((LoadingState == EAsyncPackageState::Complete || bLoadHasFailed) && !bLoadHasFinished)
+				{
+					const bool bInternalCallbacks = true;
+					CallCompletionCallbacks(bInternalCallbacks, bLoadHasFailed ? EAsyncLoadingResult::Failed : EAsyncLoadingResult::Succeeded);
+				}
 
-			if (LoadingState == EAsyncPackageState::Complete)
-			{
-				// We can only continue to PostLoad if all imported packages finished serializing their exports
-				for (UPackage* ImportedPackage : ImportedPackages)
+				if (LoadingState == EAsyncPackageState::Complete)
 				{
-					if (ImportedPackage && ImportedPackage->GetLinker() && ImportedPackage->GetLinker()->AsyncRoot && !static_cast<FAsyncPackage*>(ImportedPackage->GetLinker()->AsyncRoot)->bAllExportsSerialized)
+					// We can only continue to PostLoad if all imported packages finished serializing their exports
+					for (UPackage* ImportedPackage : ImportedPackages)
 					{
-						LoadingState = EAsyncPackageState::PendingImports;
-						break;
+						if (ImportedPackage && ImportedPackage->GetLinker() && ImportedPackage->GetLinker()->AsyncRoot && !static_cast<FAsyncPackage*>(ImportedPackage->GetLinker()->AsyncRoot)->bAllExportsSerialized)
+						{
+							LoadingState = EAsyncPackageState::PendingImports;
+							break;
+						}
 					}
 				}
-			}
-		} // !GEventDrivenLoaderEnabled
+			} // !GEventDrivenLoaderEnabled
 
-		if (LoadingState == EAsyncPackageState::Complete && !bLoadHasFailed && !bLoadHasFinished)
-		{
-			SCOPED_LOADTIMER(Package_ExternalReadDependencies);
-			LoadingState = FinishExternalReadDependencies();
-		}
+			if (LoadingState == EAsyncPackageState::Complete && !bLoadHasFailed && !bLoadHasFinished)
+			{
+				SCOPED_LOADTIMER(Package_ExternalReadDependencies);
+				LoadingState = FinishExternalReadDependencies();
+			}
 
-		// Call PostLoad on objects, this could cause new objects to be loaded that require
-		// another iteration of the PreLoad loop.
-		if (LoadingState == EAsyncPackageState::Complete && !bLoadHasFailed && !bLoadHasFinished)
-		{
-			SCOPED_LOADTIMER(Package_PostLoadObjects);
-			LoadingState = PostLoadObjects();
-		}
+			// Call PostLoad on objects, this could cause new objects to be loaded that require
+			// another iteration of the PreLoad loop.
+			if (LoadingState == EAsyncPackageState::Complete && !bLoadHasFailed && !bLoadHasFinished)
+			{
+				SCOPED_LOADTIMER(Package_PostLoadObjects);
+				LoadingState = PostLoadObjects();
+			}
 
-		// We are done loading the package for now. Whether it is done or not, let the package time tracker know.
-		FExclusiveLoadPackageTimeTracker::PopLoadPackage(Linker ? Linker->LinkerRoot : nullptr);
+			// We are done loading the package for now. Whether it is done or not, let the package time tracker know.
+			FExclusiveLoadPackageTimeTracker::PopLoadPackage(Linker ? Linker->LinkerRoot : nullptr);
 
-		// End async loading, simulates EndLoad
-		EndAsyncLoad();
+			// End async loading, simulates EndLoad
+			EndAsyncLoad();
+		}
 
 		// Finish objects (removing EInternalObjectFlags::AsyncLoading, dissociate imports and forced exports, 
 		// call completion callback, ...
@@ -6967,6 +6972,10 @@ EAsyncPackageState::Type FAsyncPackage::PostLoadDeferredObjects(double InTickSta
 	EAsyncPackageState::Type Result = EAsyncPackageState::Complete;
 	FAsyncLoadingTickScope InAsyncLoadingTick(AsyncLoadingThread);
 
+#if WITH_EDITOR
+	UE::Core::Private::FPlayInEditorLoadingScope PlayInEditorIDScope(Desc.PIEInstanceID);
+#endif
+
 	FUObjectSerializeContext* LoadContext = GetSerializeContext();
 	TArray<UObject*>& ObjLoadedInPostLoad = LoadContext->PRIVATE_GetObjectsLoadedInternalUseOnly();
 	TArray<UObject*> ObjLoadedInPostLoadLocal;
diff --git a/Engine/Source/Runtime/CoreUObject/Private/Serialization/AsyncLoading2.cpp b/Engine/Source/Runtime/CoreUObject/Private/Serialization/AsyncLoading2.cpp
--- a/Engine/Source/Runtime/CoreUObject/Private/Serialization/AsyncLoading2.cpp
+++ b/Engine/Source/Runtime/CoreUObject/Private/Serialization/AsyncLoading2.cpp
@@ -30,6 +30,7 @@
 #include "Serialization/DeferredMessageLog.h"
 #include "UObject/UObjectThreadContext.h"
 #include "Misc/Paths.h"
+#include "Misc/PlayInEditorLoadingScope.h"
 #include "Misc/ExclusiveLoadPackageTimeTracker.h"
 #include "ProfilingDebugging/LoadTimeTracker.h"
 #include "HAL/ThreadHeartBeat.h"
@@ -5449,6 +5450,10 @@ EEventLoadNodeExecutionResult FAsyncPackage2::PostLoadLinkerLoadPackageExports(F
 
 	FAsyncLoadingTickScope2 InAsyncLoadingTick(AsyncLoadingThread);
 
+#if WITH_EDITOR
+	UE::Core::Private::FPlayInEditorLoadingScope PlayInEditorIDScope(Desc.PIEInstanceID);
+#endif
+
 	const int32 ObjectCount = ConstructedObjects.Num();
 	while (LinkerLoadState->PostLoadExportIndex < ObjectCount)
 	{
@@ -5462,7 +5467,7 @@ EEventLoadNodeExecutionResult FAsyncPackage2::PostLoadLinkerLoadPackageExports(F
 
 	AsyncPackageLoadingState = EAsyncPackageLoadingState2::DeferredPostLoadDone;
 	ConditionalFinishLoading(ThreadState);
-	return EEventLoadNodeExecutionResult::Complete;;
+	return EEventLoadNodeExecutionResult::Complete;
 }
 #endif
 
@@ -5635,6 +5640,9 @@ EEventLoadNodeExecutionResult FAsyncPackage2::Event_ProcessExportBundle(FAsyncLo
 	check(Package->AsyncPackageLoadingState == EAsyncPackageLoadingState2::ProcessExportBundles);
 
 	FAsyncPackageScope2 Scope(Package);
+#if WITH_EDITOR
+	UE::Core::Private::FPlayInEditorLoadingScope PlayInEditorIDScope(Package->Desc.PIEInstanceID);
+#endif
 
 #if ALT2_ENABLE_LINKERLOAD_SUPPORT
 	if (Package->LinkerLoadState.IsSet())
diff --git a/Engine/Source/Runtime/LevelSequence/Private/LevelSequenceBindingReference.cpp b/Engine/Source/Runtime/LevelSequence/Private/LevelSequenceBindingReference.cpp
--- a/Engine/Source/Runtime/LevelSequence/Private/LevelSequenceBindingReference.cpp
+++ b/Engine/Source/Runtime/LevelSequence/Private/LevelSequenceBindingReference.cpp
@@ -89,7 +89,7 @@ UObject* FLevelSequenceBindingReference::Resolve(UObject* InContext, const FTopL
 		// scope. Since ResolveObject will always call FixupForPIE in editor based on GPlayInEditorID, we always override the current
 		// GPlayInEditorID to be the current PIE instance of the provided context.
 		const int32 ContextPlayInEditorID = InContext ? InContext->GetOutermost()->GetPIEInstanceID() : INDEX_NONE;
-		TGuardValue<int32> PIEGuard(GPlayInEditorID, ContextPlayInEditorID);
+		FTemporaryPlayInEditorIDOverride PIEGuard(ContextPlayInEditorID);
 	#endif
 
 		return TempPath.ResolveObject();
