Fix race in LoadConfig

#rnx

diff --git a/Engine/Source/Runtime/CoreUObject/Private/UObject/Obj.cpp b/Engine/Source/Runtime/CoreUObject/Private/UObject/Obj.cpp
--- a/Engine/Source/Runtime/CoreUObject/Private/UObject/Obj.cpp
+++ b/Engine/Source/Runtime/CoreUObject/Private/UObject/Obj.cpp
@@ -2481,7 +2481,7 @@ void UObject::LoadConfig( UClass* ConfigClass/*=NULL*/, const TCHAR* InFilename/
 	SCOPE_CYCLE_COUNTER(STAT_LoadConfig);
 
 	// OriginalClass is the class that LoadConfig() was originally called on
-	static UClass* OriginalClass = NULL;
+	static thread_local UClass* OriginalClass = nullptr;
 
 	if( !ConfigClass )
 	{
