Fixed mismatched copy constructor/assignment operator in FCoreRedirectObjectName and FCoreRedirect.

diff --git a/Engine/Source/Runtime/CoreUObject/Public/UObject/CoreRedirects.h b/Engine/Source/Runtime/CoreUObject/Public/UObject/CoreRedirects.h
--- a/Engine/Source/Runtime/CoreUObject/Public/UObject/CoreRedirects.h
+++ b/Engine/Source/Runtime/CoreUObject/Public/UObject/CoreRedirects.h
@@ -84,13 +84,7 @@ struct COREUOBJECT_API FCoreRedirectObjectName
 	FName PackageName;
 
 	/** Default to invalid names */
-	FCoreRedirectObjectName() {}
-
-	/** Copy constructor */
-	FCoreRedirectObjectName(const FCoreRedirectObjectName& Other)
-		: ObjectName(Other.ObjectName), OuterName(Other.OuterName), PackageName(Other.PackageName)
-	{
-	}
+	FCoreRedirectObjectName() = default;
 
 	/** Construct from FNames that are already expanded */
 	FCoreRedirectObjectName(FName InObjectName, FName InOuterName, FName InPackageName)
@@ -207,12 +201,6 @@ struct COREUOBJECT_API FCoreRedirect
 	/** Map of value changes, from old value to new value */
 	TMap<FString, FString> ValueChanges;
 
-	/** Construct from name strings, which may get parsed out */
-	FCoreRedirect(const FCoreRedirect& Other)
-		: RedirectFlags(Other.RedirectFlags), OldName(Other.OldName), NewName(Other.NewName), OverrideClassName(Other.OverrideClassName), ValueChanges(Other.ValueChanges)
-	{
-	}
-
 	/** Construct from name strings, which may get parsed out */
 	FCoreRedirect(ECoreRedirectFlags InRedirectFlags, FString InOldName, FString InNewName)
 		: RedirectFlags(InRedirectFlags), OldName(InOldName), NewName(InNewName)
