Add missing PlayInEditorID scope in zenloader
Fix PIEInstanceID being overwriten with -1 by the zenloader

#rnx

diff --git a/Engine/Source/Runtime/CoreUObject/Private/Serialization/AsyncLoading2.cpp b/Engine/Source/Runtime/CoreUObject/Private/Serialization/AsyncLoading2.cpp
--- a/Engine/Source/Runtime/CoreUObject/Private/Serialization/AsyncLoading2.cpp
+++ b/Engine/Source/Runtime/CoreUObject/Private/Serialization/AsyncLoading2.cpp
@@ -6612,7 +6612,10 @@ EEventLoadNodeExecutionResult FAsyncPackage2::Event_PostLoadExportBundle(FAsyncL
 		const bool bAsyncPostLoadEnabled = FAsyncLoadingThreadSettings::Get().bAsyncPostLoadEnabled;
 		const bool bIsMultithreaded = Package->AsyncLoadingThread.IsMultithreaded();
 
+		{
 #if WITH_EDITOR
+		UE::Core::Private::FPlayInEditorLoadingScope PlayInEditorIDScope(Package->Desc.PIEInstanceID);
+
 		const FAsyncPackageHeaderData* HeaderData;
 		if (InExportBundleIndex == 1)
 		{
@@ -6668,6 +6671,8 @@ EEventLoadNodeExecutionResult FAsyncPackage2::Event_PostLoadExportBundle(FAsyncL
 			++Package->ExportBundleEntryIndex;
 		}
 
+		}
+
 		// End async loading, simulates EndLoad
 		Package->EndAsyncLoad();
 	}
@@ -8124,14 +8129,20 @@ void FAsyncPackage2::CreateUPackage()
 		bCreatedLinkerRoot = true;
 	}
 
+#if WITH_EDITOR
+	// Do not overwrite PIEInstanceID for package that might already have one from a previous load
+	// or an in-memory only package created specifically for PIE.
+	if (!LinkerRoot->bHasBeenFullyLoaded && LinkerRoot->GetLoadedPath().IsEmpty())
+	{
+		LinkerRoot->SetPIEInstanceID(Desc.PIEInstanceID);
+	}
+#endif
+
 	LinkerRoot->SetFlags(RF_Public | RF_WasLoaded);
 	LinkerRoot->SetLoadedPath(Desc.PackagePathToLoad);
 	LinkerRoot->SetCanBeImportedFlag(Desc.bCanBeImported);
 	LinkerRoot->SetPackageId(Desc.UPackageId);
 	LinkerRoot->SetPackageFlags(Desc.PackageFlags);
-#if WITH_EDITOR
-	LinkerRoot->SetPIEInstanceID(Desc.PIEInstanceID);
-#endif
 
 	EInternalObjectFlags FlagsToSet = EInternalObjectFlags::Async;
 	if (Desc.bCanBeImported)
