Removed the deprecated MemoryProfiler runtime:
* Removed code switched by USE_MALLOC_PROFILER.
* Removed the USE_MALLOC_PROFILER define and the MALLOC_PROFILER macro.
* Removed PLATFORM_RUNTIME_MALLOCPROFILER_SYMBOLICATION define.
* Removed FMallocProfilerEx and FMallocProfilerEx wrappers and the GMallocProfiler global variable.
* Removed console commands: "MPROF", "DUMPALLOCSTOFILE", "SNAPSHOTMEMORY" and "SNAPSHOTMEMORYFRAME".
Please use Unreal Insights (Memory Insights; -trace=memory) instead.

diff --git a/Engine/Source/Developer/LowLevelTestsRunner/Private/TestCommon/CoreUObjectUtilities.cpp b/Engine/Source/Developer/LowLevelTestsRunner/Private/TestCommon/CoreUObjectUtilities.cpp
--- a/Engine/Source/Developer/LowLevelTestsRunner/Private/TestCommon/CoreUObjectUtilities.cpp
+++ b/Engine/Source/Developer/LowLevelTestsRunner/Private/TestCommon/CoreUObjectUtilities.cpp
@@ -100,7 +100,6 @@ void CleanupCoreUObject()
 	IPackageResourceManager::Shutdown();
 
 	FCoreDelegates::OnPreExit.Broadcast();
-	MALLOC_PROFILER(GMalloc->Exec(nullptr, TEXT("MPROF STOP"), *GLog); );
 	FCoreDelegates::OnExit.Broadcast();
 }
 
diff --git a/Engine/Source/Runtime/Core/Private/HAL/MallocLeakDetection.cpp b/Engine/Source/Runtime/Core/Private/HAL/MallocLeakDetection.cpp
--- a/Engine/Source/Runtime/Core/Private/HAL/MallocLeakDetection.cpp
+++ b/Engine/Source/Runtime/Core/Private/HAL/MallocLeakDetection.cpp
@@ -21,10 +21,6 @@
 
 #if MALLOC_LEAKDETECTION
 
-#if USE_MALLOC_PROFILER
-#error There is a deadlock that happens when using both the malloc profiler and leakdetection at the same time. Remove this error when it is fixed.
-#endif
-
 /**
  *	Need forced-initialization of this data as it can be used during global
 *	ctors.
diff --git a/Engine/Source/Runtime/Core/Private/HAL/UnrealMemory.cpp b/Engine/Source/Runtime/Core/Private/HAL/UnrealMemory.cpp
--- a/Engine/Source/Runtime/Core/Private/HAL/UnrealMemory.cpp
+++ b/Engine/Source/Runtime/Core/Private/HAL/UnrealMemory.cpp
@@ -10,15 +10,10 @@
 #include "HAL/IConsoleManager.h"
 #include "Misc/CoreDelegates.h"
 
-#if USE_MALLOC_PROFILER && WITH_ENGINE && IS_MONOLITHIC
-	#include "MallocProfilerEx.h"
-#endif
-
 /*-----------------------------------------------------------------------------
 	Memory functions.
 -----------------------------------------------------------------------------*/
 
-#include "ProfilingDebugging/MallocProfiler.h"
 #include "ProfilingDebugging/MemoryTrace.h"
 #include "HAL/MallocThreadSafeProxy.h"
 #include "HAL/MallocVerify.h"
@@ -350,7 +345,7 @@ static int FMemory_GCreateMalloc_ThreadUnsafe()
 	FPlatformMallocCrash::Get( GMalloc );
 
 #if PLATFORM_USES_FIXED_GMalloc_CLASS
-#if USE_MALLOC_PROFILER || MALLOC_VERIFY || MALLOC_LEAKDETECTION || UE_USE_MALLOC_FILL_BYTES
+#if MALLOC_VERIFY || MALLOC_LEAKDETECTION || UE_USE_MALLOC_FILL_BYTES
 #error "Turn off PLATFORM_USES_FIXED_GMalloc_CLASS in order to use special allocator proxies"
 #endif
 	if (!GMalloc->IsInternallyThreadSafe())
@@ -366,20 +361,7 @@ static int FMemory_GCreateMalloc_ThreadUnsafe()
 		GMalloc = TraceMalloc;
 		MemoryTrace_Initialize();
 	}
-	else
 #endif // UE_MEMORY_TRACE_ENABLED
-	{
-	// so now check to see if we are using a Mem Profiler which wraps the GMalloc
-#if USE_MALLOC_PROFILER
-		#if WITH_ENGINE && IS_MONOLITHIC
-			GMallocProfiler = new FMallocProfilerEx( GMalloc );
-		#else
-			GMallocProfiler = new FMallocProfiler( GMalloc );
-		#endif
-		GMallocProfiler->BeginProfiling();
-		GMalloc = GMallocProfiler;
-#endif
-	}
 
 #if WITH_MALLOC_STOMP2
 	GMalloc = FMallocStomp2::OverrideIfEnabled(GMalloc);
diff --git a/Engine/Source/Runtime/Core/Private/ProfilingDebugging/MallocProfiler.cpp b/Engine/Source/Runtime/Core/Private/ProfilingDebugging/MallocProfiler.cpp
deleted file mode 100644
diff --git a/Engine/Source/Runtime/Core/Public/HAL/MemoryBase.h b/Engine/Source/Runtime/Core/Public/HAL/MemoryBase.h
--- a/Engine/Source/Runtime/Core/Public/HAL/MemoryBase.h
+++ b/Engine/Source/Runtime/Core/Public/HAL/MemoryBase.h
@@ -32,9 +32,6 @@ enum
 CORE_API extern class FMalloc* GMalloc;
 CORE_API extern class FMalloc** GFixedMallocLocationPtr;
 
-/** Global FMallocProfiler variable to allow multiple malloc profilers to communicate. */
-MALLOC_PROFILER( CORE_API extern class FMallocProfiler* GMallocProfiler; )
-
 /** Holds generic memory stats, internally implemented as a map. */
 struct FGenericMemoryStats;
 
diff --git a/Engine/Source/Runtime/Core/Public/HAL/Platform.h b/Engine/Source/Runtime/Core/Public/HAL/Platform.h
--- a/Engine/Source/Runtime/Core/Public/HAL/Platform.h
+++ b/Engine/Source/Runtime/Core/Public/HAL/Platform.h
@@ -427,10 +427,6 @@
 	#define PLATFORM_SUPPORTS_STACK_SYMBOLS 0
 #endif
 
-#ifndef PLATFORM_RUNTIME_MALLOCPROFILER_SYMBOLICATION
-	#define PLATFORM_RUNTIME_MALLOCPROFILER_SYMBOLICATION PLATFORM_SUPPORTS_STACK_SYMBOLS
-#endif
-
 #ifndef PLATFORM_HAS_128BIT_ATOMICS
 	#define PLATFORM_HAS_128BIT_ATOMICS 0
 #endif
diff --git a/Engine/Source/Runtime/Core/Public/IOS/IOSPlatform.h b/Engine/Source/Runtime/Core/Public/IOS/IOSPlatform.h
--- a/Engine/Source/Runtime/Core/Public/IOS/IOSPlatform.h
+++ b/Engine/Source/Runtime/Core/Public/IOS/IOSPlatform.h
@@ -57,8 +57,6 @@ typedef FIOSPlatformTypes FPlatformTypes;
 #define PLATFORM_USE_FULL_TASK_GRAPH					0 // @todo platplug: not platplug, but should investigate soon anyway
 #define PLATFORM_IS_ANSI_MALLOC_THREADSAFE				1
 
-// on iOS we now perform offline symbolication as it's significantly faster. Requires bGenerateCrashReportSymbols=true in the ini file.
-#define	PLATFORM_RUNTIME_MALLOCPROFILER_SYMBOLICATION	0	
 #define PLATFORM_NUM_AUDIODECOMPRESSION_PRECACHE_BUFFERS		0
 #if PLATFORM_TVOS
 #define PLATFORM_USES_GLES								0
diff --git a/Engine/Source/Runtime/Core/Public/Misc/Build.h b/Engine/Source/Runtime/Core/Public/Misc/Build.h
--- a/Engine/Source/Runtime/Core/Public/Misc/Build.h
+++ b/Engine/Source/Runtime/Core/Public/Misc/Build.h
@@ -247,7 +247,7 @@
 		#define DO_ENSURE										1
 	#endif
 	#ifndef STATS
-		#define STATS											((WITH_UNREAL_DEVELOPER_TOOLS || !WITH_EDITORONLY_DATA || USE_STATS_WITHOUT_ENGINE || USE_MALLOC_PROFILER || FORCE_USE_STATS) && !ENABLE_STATNAMEDEVENTS)
+		#define STATS											((WITH_UNREAL_DEVELOPER_TOOLS || !WITH_EDITORONLY_DATA || USE_STATS_WITHOUT_ENGINE || FORCE_USE_STATS) && !ENABLE_STATNAMEDEVENTS)
 	#endif
 	#ifndef ALLOW_DEBUG_FILES
 		#define ALLOW_DEBUG_FILES								1
@@ -269,7 +269,7 @@
 		#define DO_ENSURE										1
 	#endif
 	#ifndef STATS
-		#define STATS											((WITH_UNREAL_DEVELOPER_TOOLS || !WITH_EDITORONLY_DATA || USE_STATS_WITHOUT_ENGINE || USE_MALLOC_PROFILER || FORCE_USE_STATS) && !ENABLE_STATNAMEDEVENTS)
+		#define STATS											((WITH_UNREAL_DEVELOPER_TOOLS || !WITH_EDITORONLY_DATA || USE_STATS_WITHOUT_ENGINE || FORCE_USE_STATS) && !ENABLE_STATNAMEDEVENTS)
 	#endif
 	#ifndef ALLOW_DEBUG_FILES
 		#define ALLOW_DEBUG_FILES								1
@@ -291,7 +291,7 @@
 		#define DO_ENSURE										USE_ENSURES_IN_SHIPPING
 	#endif
 	#ifndef STATS
-		#define STATS											((USE_MALLOC_PROFILER || FORCE_USE_STATS) && !ENABLE_STATNAMEDEVENTS)
+		#define STATS											(FORCE_USE_STATS && !ENABLE_STATNAMEDEVENTS)
 	#endif
 	#ifndef ALLOW_DEBUG_FILES
 		#define ALLOW_DEBUG_FILES								1
diff --git a/Engine/Source/Runtime/Core/Public/ProfilingDebugging/MallocProfiler.h b/Engine/Source/Runtime/Core/Public/ProfilingDebugging/MallocProfiler.h
deleted file mode 100644
diff --git a/Engine/Source/Runtime/Core/Public/ProfilingDebugging/UMemoryDefines.h b/Engine/Source/Runtime/Core/Public/ProfilingDebugging/UMemoryDefines.h
--- a/Engine/Source/Runtime/Core/Public/ProfilingDebugging/UMemoryDefines.h
+++ b/Engine/Source/Runtime/Core/Public/ProfilingDebugging/UMemoryDefines.h
@@ -15,23 +15,3 @@
 #ifndef MIMALLOC_ALLOCATOR_ALLOWED
 	#define MIMALLOC_ALLOCATOR_ALLOWED 1
 #endif
-
-/**
- *	IMPORTANT:	The malloc proxy flags are mutually exclusive.
- *				You can have either none of them set or only one of them.
- */
-/** USE_MALLOC_PROFILER				- Define this to use the FMallocProfiler allocator.			*/
-/**									  Make sure to enable Stack Frame pointers:					*/
-/**									  bOmitFramePointers = false, or /Oy-						*/
-
-#ifndef USE_MALLOC_PROFILER
-#define USE_MALLOC_PROFILER				0
-#endif
-
-#if USE_MALLOC_PROFILER
-#define MALLOC_PROFILER(x)	x	
-#else
-#define MALLOC_PROFILER(...)
-#endif
-
-
diff --git a/Engine/Source/Runtime/CoreUObject/Private/UObject/UObjectBaseUtility.cpp b/Engine/Source/Runtime/CoreUObject/Private/UObject/UObjectBaseUtility.cpp
--- a/Engine/Source/Runtime/CoreUObject/Private/UObject/UObjectBaseUtility.cpp
+++ b/Engine/Source/Runtime/CoreUObject/Private/UObject/UObjectBaseUtility.cpp
@@ -12,7 +12,6 @@
 #include "UObject/Interface.h"
 #include "Misc/StringBuilder.h"
 #include "Modules/ModuleManager.h"
-#include "ProfilingDebugging/MallocProfiler.h"
 #include "HAL/IConsoleManager.h"
 #include "Misc/ConfigCacheIni.h"
 #include "Containers/VersePath.h"
@@ -607,120 +606,6 @@ UClass* GetParentNativeClass(UClass* Class)
 	return Class;
 }
 
-#if STATS && USE_MALLOC_PROFILER
-void FScopeCycleCounterUObject::TrackObjectForMallocProfiling(const UObjectBaseUtility *InObject)
-{
-	// Get the package name from the outermost item (if available - can't use GetOutermost here)
-	FName PackageName;
-	if (InObject->GetOuter())
-	{
-		UObjectBaseUtility* Top = InObject->GetOuter();
-		for (;;)
-		{
-			UObjectBaseUtility* CurrentOuter = Top->GetOuter();
-			if (!CurrentOuter)
-			{
-				PackageName = Top->GetFName();
-				break;
-			}
-			Top = CurrentOuter;
-		}
-	}
-
-	// Get the class name (if available)
-	FName ClassName;
-	if (InObject->GetClass())
-	{
-		ClassName = InObject->GetClass()->GetFName();
-	}
-
-	TrackObjectForMallocProfiling(PackageName, ClassName, InObject->GetFName());
-}
-void FScopeCycleCounterUObject::TrackObjectForMallocProfiling(const FName InPackageName, const FName InClassName, const FName InObjectName)
-{
-	static const TCHAR PackageTagCategory[] = TEXT("Package:");
-	static const TCHAR ObjectTagCategory[] = TEXT("Object:");
-	static const TCHAR ClassTagCategory[] = TEXT("Class:");
-
-	// We use an array rather than an FString to try and minimize heap allocations
-	TArray<TCHAR, TInlineAllocator<256>> ScratchSpaceBuffer;
-
-	auto AppendNameToBuffer = [&](const FName InName)
-	{
-		const FNameEntry* NameEntry = InName.GetDisplayNameEntry();
-		if (NameEntry->IsWide())
-		{
-			WIDECHAR WideName[NAME_SIZE];
-			NameEntry->GetWideName(WideName);
-			const WIDECHAR* NameCharPtr = WideName;
-			while (*NameCharPtr != 0)
-			{
-				ScratchSpaceBuffer.Add((TCHAR)*NameCharPtr++);
-			}
-		}
-		else
-		{
-			ANSICHAR AnsiName[NAME_SIZE];
-			NameEntry->GetAnsiName(AnsiName);
-			const ANSICHAR* NameCharPtr = AnsiName;
-			while (*NameCharPtr != 0)
-			{
-				ScratchSpaceBuffer.Add((TCHAR)*NameCharPtr++);
-			}
-		}
-	};
-
-	if (!InPackageName.IsNone())
-	{
-		// "Package:/Path/To/Package"
-		ScratchSpaceBuffer.Reset();
-		ScratchSpaceBuffer.Append(PackageTagCategory, UE_ARRAY_COUNT(PackageTagCategory) - 1);
-		AppendNameToBuffer(InPackageName);
-		ScratchSpaceBuffer.Add(0);
-		PackageTag = FName(ScratchSpaceBuffer.GetData());
-		GMallocProfiler->AddTag(PackageTag);
-
-		// "Object:/Path/To/Package/ObjectName"
-		ScratchSpaceBuffer.Reset();
-		ScratchSpaceBuffer.Append(ObjectTagCategory, UE_ARRAY_COUNT(ObjectTagCategory) - 1);
-		AppendNameToBuffer(InPackageName);
-		ScratchSpaceBuffer.Add(TEXT('/'));
-		AppendNameToBuffer(InObjectName);
-		ScratchSpaceBuffer.Add(0);
-		ObjectTag = FName(ScratchSpaceBuffer.GetData());
-		GMallocProfiler->AddTag(ObjectTag);
-	}
-
-	if (!InClassName.IsNone())
-	{
-		// "Class:ClassName"
-		ScratchSpaceBuffer.Reset();
-		ScratchSpaceBuffer.Append(ClassTagCategory, UE_ARRAY_COUNT(ClassTagCategory) - 1);
-		AppendNameToBuffer(InClassName);
-		ScratchSpaceBuffer.Add(0);
-		ClassTag = FName(ScratchSpaceBuffer.GetData());
-		GMallocProfiler->AddTag(ClassTag);
-	}
-}
-void FScopeCycleCounterUObject::UntrackObjectForMallocProfiling()
-{
-	if (!PackageTag.IsNone())
-	{
-		GMallocProfiler->RemoveTag(PackageTag);
-	}
-
-	if (!ClassTag.IsNone())
-	{
-		GMallocProfiler->RemoveTag(ClassTag);
-	}
-
-	if (!ObjectTag.IsNone())
-	{
-		GMallocProfiler->RemoveTag(ObjectTag);
-	}
-}
-#endif
-
 #if !STATS && !ENABLE_STATNAMEDEVENTS && USE_LIGHTWEIGHT_STATS_FOR_HITCH_DETECTION && USE_HITCH_DETECTION && USE_LIGHTWEIGHT_UOBJECT_STATS_FOR_HITCH_DETECTION
 #include "HAL/ThreadHeartBeat.h"
 #include "HAL/ThreadManager.h"
diff --git a/Engine/Source/Runtime/CoreUObject/Private/UObject/UObjectGlobals.cpp b/Engine/Source/Runtime/CoreUObject/Private/UObject/UObjectGlobals.cpp
--- a/Engine/Source/Runtime/CoreUObject/Private/UObject/UObjectGlobals.cpp
+++ b/Engine/Source/Runtime/CoreUObject/Private/UObject/UObjectGlobals.cpp
@@ -82,10 +82,6 @@
 
 DEFINE_LOG_CATEGORY(LogUObjectGlobals);
 
-#if USE_MALLOC_PROFILER
-#include "ProfilingDebugging/MallocProfiler.h"
-#endif
-
 int32 GAllowUnversionedContentInEditor = 0;
 
 static FAutoConsoleVariableRef CVarAllowUnversionedContentInEditor(
diff --git a/Engine/Source/Runtime/CoreUObject/Public/UObject/UObjectBaseUtility.h b/Engine/Source/Runtime/CoreUObject/Public/UObject/UObjectBaseUtility.h
--- a/Engine/Source/Runtime/CoreUObject/Public/UObject/UObjectBaseUtility.h
+++ b/Engine/Source/Runtime/CoreUObject/Public/UObject/UObjectBaseUtility.h
@@ -1009,15 +1009,6 @@ COREUOBJECT_API FName GetClassTraceScope(const UObjectBaseUtility* Object);
 class FScopeCycleCounterUObject : public FCycleCounter
 {
 public:
-#if USE_MALLOC_PROFILER
-	/** Package path being tracked */
-	FName PackageTag;
-	/** Class path being tracked */
-	FName ClassTag;
-	/** Object path being tracked */
-	FName ObjectTag;
-#endif
-
 	/**
 	 * Constructor, starts timing
 	 */
@@ -1040,13 +1031,6 @@ public:
 			}
 #endif
 		}
-
-#if USE_MALLOC_PROFILER
-		if (Object)
-		{
-			TrackObjectForMallocProfiling(Object);
-		}
-#endif
 	}
 	/**
 	 * Constructor, starts timing with an alternate enable stat to use high performance disable for only SOME UObject stats
@@ -1073,13 +1057,6 @@ public:
 			}
 #endif
 		}
-
-#if USE_MALLOC_PROFILER
-		if (Object)
-		{
-			TrackObjectForMallocProfiling(Object);
-		}
-#endif
 	}
 
 #if CPUPROFILERTRACE_ENABLED
@@ -1092,17 +1069,7 @@ public:
 	FORCEINLINE_STATS ~FScopeCycleCounterUObject()
 	{
 		Stop();
-
-#if USE_MALLOC_PROFILER
-		UntrackObjectForMallocProfiling();
-#endif
 	}
-
-#if USE_MALLOC_PROFILER
-	COREUOBJECT_API void TrackObjectForMallocProfiling(const UObjectBaseUtility *InObject);
-	COREUOBJECT_API void TrackObjectForMallocProfiling(const FName InPackageName, const FName InClassName, const FName InObjectName);
-	COREUOBJECT_API void UntrackObjectForMallocProfiling();
-#endif
 };
 
 /** Declares a scope cycle counter for a specific object with a Name context */
diff --git a/Engine/Source/Runtime/Engine/Private/MallocProfilerEx.cpp b/Engine/Source/Runtime/Engine/Private/MallocProfilerEx.cpp
deleted file mode 100644
diff --git a/Engine/Source/Runtime/Engine/Private/UnrealEngine.cpp b/Engine/Source/Runtime/Engine/Private/UnrealEngine.cpp
--- a/Engine/Source/Runtime/Engine/Private/UnrealEngine.cpp
+++ b/Engine/Source/Runtime/Engine/Private/UnrealEngine.cpp
@@ -14631,7 +14631,7 @@ bool UEngine::LoadMap( FWorldContext& WorldContext, FURL URL, class UPendingNetG
 	FDisableHitchDetectorScope SuspendHitchDetector;
 
 	NETWORK_PROFILER(GNetworkProfiler.TrackSessionChange(true,URL));
-	MALLOC_PROFILER( FMallocProfiler::SnapshotMemoryLoadMapStart( URL.Map ) );
+
 	FMoviePlayerProxyBlock MoviePlayerBlock;
 	Error = TEXT("");
 
@@ -14868,8 +14868,6 @@ bool UEngine::LoadMap( FWorldContext& WorldContext, FURL URL, class UPendingNetG
 
 	FMoviePlayerProxy::BlockingTick();
 
-	MALLOC_PROFILER( FMallocProfiler::SnapshotMemoryLoadMapMid( URL.Map ); )
-
 	WorldContext.OwningGameInstance->PreloadContentForURL(URL);
 
 	UPackage* WorldPackage = NULL;
@@ -15217,7 +15215,6 @@ bool UEngine::LoadMap( FWorldContext& WorldContext, FURL URL, class UPendingNetG
 
 	STAT_ADD_CUSTOMMESSAGE_NAME( STAT_NamedMarker, *(FString( TEXT( "LoadMapComplete - " ) + URL.Map )) );
 	TRACE_BOOKMARK(TEXT("LoadMapComplete - %s"), *URL.Map);
-	MALLOC_PROFILER( FMallocProfiler::SnapshotMemoryLoadMapEnd( URL.Map ); )
 
 	double StopTime = FPlatformTime::Seconds();
 
diff --git a/Engine/Source/Runtime/Engine/Public/MallocProfilerEx.h b/Engine/Source/Runtime/Engine/Public/MallocProfilerEx.h
deleted file mode 100644
diff --git a/Engine/Source/Runtime/Launch/Private/LaunchEngineLoop.cpp b/Engine/Source/Runtime/Launch/Private/LaunchEngineLoop.cpp
--- a/Engine/Source/Runtime/Launch/Private/LaunchEngineLoop.cpp
+++ b/Engine/Source/Runtime/Launch/Private/LaunchEngineLoop.cpp
@@ -5673,9 +5673,6 @@ void FEngineLoop::Tick()
 		// Calculates average FPS/MS (outside STATS on purpose)
 		CalculateFPSTimings();
 
-		// Note the start of a new frame
-		MALLOC_PROFILER(GMalloc->Exec(nullptr, *FString::Printf(TEXT("SNAPSHOTMEMORYFRAME")),*GLog));
-
 		// handle some per-frame tasks on the rendering thread
 		ENQUEUE_RENDER_COMMAND(ResetDeferredUpdates)(
 			[](FRHICommandList& RHICmdList)
@@ -6679,8 +6676,6 @@ void FEngineLoop::AppPreExit( )
 
 	FCoreDelegates::OnPreExit.Broadcast();
 
-	MALLOC_PROFILER( GMalloc->Exec(nullptr, TEXT("MPROF STOP"), *GLog);	);
-
 #if WITH_ENGINE
 	if (FString(FCommandLine::Get()).Contains(TEXT("CreatePak")) && GetDerivedDataCache())
 	{
