Remove all deprecated code from BulkData related classes up to 5.0

#rnx

diff --git a/Engine/Source/Runtime/CoreUObject/Private/Serialization/BulkData.cpp b/Engine/Source/Runtime/CoreUObject/Private/Serialization/BulkData.cpp
--- a/Engine/Source/Runtime/CoreUObject/Private/Serialization/BulkData.cpp
+++ b/Engine/Source/Runtime/CoreUObject/Private/Serialization/BulkData.cpp
@@ -1470,11 +1470,6 @@ void FBulkData::DetachFromArchive( FArchive* Ar, bool bEnsureBulkDataIsLoaded )
 }
 #endif // WITH_EDITOR
 
-void FBulkData::StoreCompressedOnDisk(ECompressionFlags CompressionFlags)
-{
-	StoreCompressedOnDisk(FCompression::GetCompressionFormatFromDeprecatedFlags(CompressionFlags));
-}
-
 void FBulkData::StoreCompressedOnDisk( FName CompressionFormat )
 {
 	if( CompressionFormat != GetDecompressionFormat() )
diff --git a/Engine/Source/Runtime/CoreUObject/Public/Serialization/BulkData.h b/Engine/Source/Runtime/CoreUObject/Public/Serialization/BulkData.h
--- a/Engine/Source/Runtime/CoreUObject/Public/Serialization/BulkData.h
+++ b/Engine/Source/Runtime/CoreUObject/Public/Serialization/BulkData.h
@@ -57,10 +57,6 @@ enum class EFileRegionType : uint8;
 	#define USE_RUNTIME_BULKDATA 0
 #endif
 
-// Enable the following to use the more compact FBulkDataStreamingToken in places where it is implemented
-#define USE_BULKDATA_STREAMING_TOKEN UE_DEPRECATED_MACRO(5.0, "USE_BULKDATA_STREAMING_TOKEN now always evaluates to 0 and will be removed") 0
-#define STREAMINGTOKEN_PARAM(param) UE_DEPRECATED_MACRO(5.0, "STREAMINGTOKEN_PARAM now always evaluates to a NOP")
-
 // When set to 1 attempting to reload inline data will fail with the old loader in the same way that it fails in
 // the new loader to keep the results consistent.
 #define UE_KEEP_INLINE_RELOADING_CONSISTENT 0
@@ -668,9 +664,6 @@ public:
 		return   (GetBulkDataFlags() & BULKDATA_PayloadAtEndOfFile) == 0;
 	}
 
-	UE_DEPRECATED(4.25, "Use ::IsInSeparateFile() instead")
-	inline bool InSeperateFile() const { return IsInSeparateFile(); }
-
 	/**
 	* Returns whether this bulk data is currently stored in it's own file or not
 	* @return true if BULKDATA_PayloadInSeperateFile is not set
@@ -763,16 +756,8 @@ public:
 	 */
 	void ClearBulkDataFlags( uint32 BulkDataFlagsToClear );
 
-	UE_DEPRECATED(5.0, "Use GetPackagePath instead")
-	FString GetFilename() const
-	{ 
-		PRAGMA_DISABLE_DEPRECATION_WARNINGS
-		return GetPackagePath().GetLocalFullPath(GetPackageSegment());
-		PRAGMA_ENABLE_DEPRECATION_WARNINGS
-	}
-
 	/** Returns the PackagePath this bulkdata resides in */
-	UE_DEPRECATED(5.1, "Deprecated")
+	UE_DEPRECATED(5.1, "Deprecated, no replacement")
 	const FPackagePath& GetPackagePath() const
 	{ 
 		static FPackagePath Empty;
@@ -780,7 +765,7 @@ public:
 	}
 
 	/** Returns which segment of its PackagePath this bulkdata resides in */
-	UE_DEPRECATED(5.1, "Deprecated")
+	UE_DEPRECATED(5.1, "Deprecated, no replacement")
 	EPackageSegment GetPackageSegment() const
 	{ 
 		return EPackageSegment::Header;
@@ -890,8 +875,6 @@ public:
 	 *
 	 * @param CompressionFlags	Flags to use for compressing the data. Use COMPRESS_NONE for no compression, or something like COMPRESS_ZLIB to compress the data
 	 */
-	UE_DEPRECATED(4.21, "Use the FName version of StoreCompressedOnDisk")
-	void StoreCompressedOnDisk( ECompressionFlags CompressionFlags );
 	void StoreCompressedOnDisk( FName CompressionFormat );
 
 	/**
@@ -931,12 +914,6 @@ public:
 		SerializeBulkData(Ar, Data, GetBulkDataSize(), InBulkDataFlags);
 	}
 
-	UE_DEPRECATED(5.0, "Use the version that takes InBulkDataFlags")
-	void SerializeBulkData(FArchive& Ar, void* Data)
-	{
-		SerializeBulkData(Ar, Data, static_cast<EBulkDataFlags>(GetBulkDataFlags()));
-	}
-
 	FOwnedBulkDataPtr* StealFileMapping()
 	{
 		// @todo if non-mapped bulk data, do we need to detach this, or mimic GetCopy more than we do?
diff --git a/Engine/Source/Runtime/SlateCore/Private/Fonts/FontBulkData.cpp b/Engine/Source/Runtime/SlateCore/Private/Fonts/FontBulkData.cpp
--- a/Engine/Source/Runtime/SlateCore/Private/Fonts/FontBulkData.cpp
+++ b/Engine/Source/Runtime/SlateCore/Private/Fonts/FontBulkData.cpp
@@ -75,15 +75,6 @@ const void* UFontBulkData::Lock(int64& OutFontDataSizeBytes) const
 	return Data;
 }
 
-const void* UFontBulkData::Lock(int32& OutFontDataSizeBytes) const
-{
-	int64 Out64;
-	const void* Result = Lock(Out64);
-	ensureAlways(Out64 <= std::numeric_limits<int32>::max());
-	OutFontDataSizeBytes = (int32)Out64;
-	return Result;
-}
-
 void UFontBulkData::Unlock()
 {
 	bool bWasLoaded = BulkData.IsBulkDataLoaded();
diff --git a/Engine/Source/Runtime/SlateCore/Public/Fonts/FontBulkData.h b/Engine/Source/Runtime/SlateCore/Public/Fonts/FontBulkData.h
--- a/Engine/Source/Runtime/SlateCore/Public/Fonts/FontBulkData.h
+++ b/Engine/Source/Runtime/SlateCore/Public/Fonts/FontBulkData.h
@@ -26,10 +26,6 @@ public:
 	/** Locks the bulk font data and returns a read-only pointer to it */
 	const void* Lock(int64& OutFontDataSizeBytes) const;
 
-	/** Locks the bulk font data and returns a read-only pointer to it */
-	UE_DEPRECATED(5.0, "Lock with 32 bits is deprecated, use the 64 bits version.")
-	const void* Lock(int32& OutFontDataSizeBytes) const;
-
 	/** Unlock the bulk font data, after which point the pointer returned by Lock no longer is valid */
 	void Unlock();
 
