Adding the option to instance subobject templates only when processing the subobject instancing graph

#tests Performed multiple hot-reload attempts with multiple start/end game cycles witch each attempt. Changes made during hot-reload: adding new print  messages / changing print message text / adding new object memeber variables to existing classes.

#localization none

diff --git a/Engine/Source/Runtime/CoreUObject/Private/UObject/CoreNative.cpp b/Engine/Source/Runtime/CoreUObject/Private/UObject/CoreNative.cpp
--- a/Engine/Source/Runtime/CoreUObject/Private/UObject/CoreNative.cpp
+++ b/Engine/Source/Runtime/CoreUObject/Private/UObject/CoreNative.cpp
@@ -88,19 +88,24 @@ IMPLEMENT_CLASS(UObject, 0);
 -----------------------------------------------------------------------------*/
 
 FObjectInstancingGraph::FObjectInstancingGraph(bool bDisableInstancing)
+	: FObjectInstancingGraph(bDisableInstancing ? EObjectInstancingGraphOptions::DisableInstancing : EObjectInstancingGraphOptions::None)	
+{
+}
+
+FObjectInstancingGraph::FObjectInstancingGraph(EObjectInstancingGraphOptions InOptions)
 	: SourceRoot(nullptr)
 	, DestinationRoot(nullptr)
+	, InstancingOptions(InOptions)
 	, bCreatingArchetype(false)
-	, bEnableSubobjectInstancing(!bDisableInstancing)
 	, bLoadingObject(false)
 {
 }
 
-FObjectInstancingGraph::FObjectInstancingGraph( UObject* DestinationSubobjectRoot )
+FObjectInstancingGraph::FObjectInstancingGraph( UObject* DestinationSubobjectRoot, EObjectInstancingGraphOptions InOptions)
 	: SourceRoot(nullptr)
 	, DestinationRoot(nullptr)
+	, InstancingOptions(InOptions)
 	, bCreatingArchetype(false)
-	, bEnableSubobjectInstancing(true)
 	, bLoadingObject(false)
 {
 	SetDestinationRoot(DestinationSubobjectRoot);
@@ -277,11 +282,10 @@ UObject* FObjectInstancingGraph::InstancePropertyValue(UObject* SubObjectTemplat
 	bool bCausesInstancing   = !!(Flags & EInstancePropertyValueFlags::CausesInstancing);
 	bool bAllowSelfReference = !!(Flags & EInstancePropertyValueFlags::AllowSelfReference);
 
-	UObject* NewValue = CurrentValue;
-
-	check(CurrentValue);
+	checkf(CurrentValue, TEXT("Null object value encountered while trying to instance subobject of %s. Instanced object properties are never expected to be null (template: %s)."), 
+		*GetFullNameSafe(Owner), *GetNameSafe(SubObjectTemplate));
 
-	if (CurrentValue != nullptr && CurrentValue->GetClass()->HasAnyClassFlags(CLASS_DefaultToInstanced))
+	if (CurrentValue->GetClass()->HasAnyClassFlags(CLASS_DefaultToInstanced))
 	{
 		bCausesInstancing = true; // these are always instanced no matter what
 	}
@@ -290,13 +294,21 @@ UObject* FObjectInstancingGraph::InstancePropertyValue(UObject* SubObjectTemplat
 		(!bCausesInstancing && // or if this class isn't forced to be instanced and this var did not have the instance keyword
 		!bAllowSelfReference)) // and this isn't a delegate
 	{
-		return NewValue; // not instancing
+		return CurrentValue; // not instancing
 	}
 
+	if (!!(InstancingOptions & EObjectInstancingGraphOptions::InstanceTemplatesOnly) && // if we're allowed to instance templates only
+		!CurrentValue->IsTemplate()) // and the current value is not a template
+	{
+		return CurrentValue; // not instancing
+	}
+
+	UObject* NewValue = CurrentValue;
+
 	// if the object we're instancing the subobjects for (Owner) has the current subobject's outer in its archetype chain, and its archetype has a nullptr value
 	// for this subobject property it means that the archetype didn't instance its subobject, so we shouldn't either.
 
-	if (SubObjectTemplate == nullptr && (Owner != nullptr && CurrentValue != nullptr && Owner->IsBasedOnArchetype(CurrentValue->GetOuter())))
+	if (SubObjectTemplate == nullptr && (Owner != nullptr && Owner->IsBasedOnArchetype(CurrentValue->GetOuter())))
 	{
 		NewValue = nullptr;
 	}
diff --git a/Engine/Source/Runtime/CoreUObject/Public/UObject/Class.h b/Engine/Source/Runtime/CoreUObject/Public/UObject/Class.h
--- a/Engine/Source/Runtime/CoreUObject/Public/UObject/Class.h
+++ b/Engine/Source/Runtime/CoreUObject/Public/UObject/Class.h
@@ -3594,6 +3594,19 @@ enum class EInstancePropertyValueFlags
 
 ENUM_CLASS_FLAGS(EInstancePropertyValueFlags)
 
+enum class EObjectInstancingGraphOptions
+{
+	None = 0x00,
+
+	// if set, start with component instancing disabled
+	DisableInstancing = 0x01,
+
+	// if set, instance only subobject template values
+	InstanceTemplatesOnly = 0x02,
+};
+
+ENUM_CLASS_FLAGS(EObjectInstancingGraphOptions)
+
 struct COREUOBJECT_API FObjectInstancingGraph
 {
 public:
@@ -3602,14 +3615,21 @@ public:
 	 * Default Constructor 
 	 * @param bDisableInstancing - if true, start with component instancing disabled
 	**/
-	FObjectInstancingGraph(bool bDisableInstancing = false);
+	explicit FObjectInstancingGraph(bool bDisableInstancing = false);
+
+	/**
+	 * Constructor with options
+	 * @param InOptions Additional options to modify the behavior of this graph
+	**/
+	explicit FObjectInstancingGraph(EObjectInstancingGraphOptions InOptions);
 
 	/**
 	 * Standard constructor
 	 *
 	 * @param	DestinationSubobjectRoot	the top-level object that is being created
+	 * @param	InOptions					Additional options to modify the behavior of this graph
 	 */
-	FObjectInstancingGraph( class UObject* DestinationSubobjectRoot );
+	explicit FObjectInstancingGraph( class UObject* DestinationSubobjectRoot, EObjectInstancingGraphOptions InOptions = EObjectInstancingGraphOptions::None);
 
 	/**
 	 * Returns whether this instancing graph has a valid destination root.
@@ -3731,7 +3751,14 @@ public:
 	 */
 	void EnableSubobjectInstancing( bool bEnabled )
 	{
-		bEnableSubobjectInstancing = bEnabled;
+		if (bEnabled)
+		{
+			InstancingOptions &= ~EObjectInstancingGraphOptions::DisableInstancing;
+		}
+		else
+		{
+			InstancingOptions |= EObjectInstancingGraphOptions::DisableInstancing;
+		}
 	}
 
 	/**
@@ -3739,7 +3766,7 @@ public:
 	 */
 	bool IsSubobjectInstancingEnabled() const
 	{
-		return bEnableSubobjectInstancing;
+		return !(InstancingOptions & EObjectInstancingGraphOptions::DisableInstancing);
 	}
 
 	/**
@@ -3825,16 +3852,14 @@ private:
 	 */
 	class		UObject*						DestinationRoot;
 
+	/** Subobject instancing options */
+	EObjectInstancingGraphOptions				InstancingOptions;
+
 	/**
 	 * Indicates whether we are currently instancing components for an archetype.  true if we are creating or updating an archetype.
 	 */
 	bool										bCreatingArchetype;
 
-	/**
-	 * If false, components will not be instanced.
-	 */
-	bool										bEnableSubobjectInstancing;
-
 	/**
 	 * true when loading object data from disk.
 	 */
