Fix potential truncation warning in ::ResolveName

#rnx

- The pointer arithmetic is being done within the contents of an FString, which we know must fit in a int32 for length, so we know the cast is safe.

diff --git a/Engine/Source/Runtime/CoreUObject/Private/UObject/UObjectGlobals.cpp b/Engine/Source/Runtime/CoreUObject/Private/UObject/UObjectGlobals.cpp
--- a/Engine/Source/Runtime/CoreUObject/Private/UObject/UObjectGlobals.cpp
+++ b/Engine/Source/Runtime/CoreUObject/Private/UObject/UObjectGlobals.cpp
@@ -1125,7 +1125,7 @@ bool ResolveName(UObject*& InPackage, FString& InOutName, bool Create, bool Thro
 			Create         = false;
 		}
 		
-		int32 DotIndex = DelimiterOrEnd - *InOutName;
+		const int32 DotIndex = static_cast<int32>(DelimiterOrEnd - *InOutName);
 
 		TStringBuilder<FName::StringBufferSize> PartialName;
 		PartialName.Append(*InOutName, DotIndex);
