LinkerInstancingContext: Support Instancing Prefix
- Add support for Instancing Prefix
- Add Automation Test
- Bring back CL 24926885 now that Prefix are supported

#rnx

diff --git a/Engine/Source/Runtime/CoreUObject/Private/UObject/LinkerInstancingContext.cpp b/Engine/Source/Runtime/CoreUObject/Private/UObject/LinkerInstancingContext.cpp
--- a/Engine/Source/Runtime/CoreUObject/Private/UObject/LinkerInstancingContext.cpp
+++ b/Engine/Source/Runtime/CoreUObject/Private/UObject/LinkerInstancingContext.cpp
@@ -7,6 +7,7 @@
 #include "Containers/StringView.h"
 #include "Containers/UnrealString.h"
 #include "Templates/Tuple.h"
+#include "Misc/AutomationTest.h"
 
 void FLinkerInstancingContext::FixupSoftObjectPath(FSoftObjectPath& InOutSoftObjectPath) const
 {
@@ -39,26 +40,33 @@ void FLinkerInstancedPackageMap::AddPackageMapping(FName Original, FName Instanc
 
 void FLinkerInstancedPackageMap::BuildPackageMapping(FName Original, FName Instanced, const bool bBuildWorldPartitionCellMapping)
 {
-	check(GeneratedPackagesFolder.IsEmpty() && InstancedPackageSuffix.IsEmpty());
+	check(GeneratedPackagesFolder.IsEmpty() && InstancedPackageSuffix.IsEmpty() && InstancedPackagePrefix.IsEmpty());
 
 	AddPackageMapping(Original, Instanced);
 
-#if !WITH_EDITOR
-	if (bBuildWorldPartitionCellMapping)
+	if (bBuildWorldPartitionCellMapping && bEnableNonEditorPath)
 	{
+#if WITH_EDITOR
+		// This code path should only be enabled while testing
+		check(GIsAutomationTesting);
+#endif
+
 		FNameBuilder TmpOriginal(Original);
 		FNameBuilder TmpInstanced(Instanced);
 		FStringView OriginalView = TmpOriginal.ToView();
 		FStringView InstancedView = TmpInstanced.ToView();
 
+		const int32 Index = InstancedView.Find(OriginalView);
+
 		// Stash the suffix used for this instance so we can also apply it to generated packages
-		if (InstancedView.StartsWith(OriginalView))
+		if (Index != INDEX_NONE)
 		{
-			InstancedPackageSuffix = InstancedView.Mid(OriginalView.Len());
+			InstancedPackagePrefix = InstancedView.Mid(0, Index);
+			InstancedPackageSuffix = InstancedView.Mid(Index + OriginalView.Len());
 		}
 
 		// Is this a generated partitioned map package? If so, we'll also need to handle re-mapping paths to our persistent map package
-		if (!InstancedPackageSuffix.IsEmpty())
+		if (!InstancedPackagePrefix.IsEmpty() || !InstancedPackageSuffix.IsEmpty())
 		{
 			const FStringView GeneratedFolderName = TEXTVIEW("/_Generated_/");
 			
@@ -70,14 +78,16 @@ void FLinkerInstancedPackageMap::BuildPackageMapping(FName Original, FName Insta
 				if (const int32 ExtraSlashIndex = OriginalView.Find(TEXTVIEW("/"), GeneratedFolderEndIndex); ExtraSlashIndex == INDEX_NONE)
 				{
 					GeneratedPackagesFolder = OriginalView.Left(GeneratedFolderEndIndex);
+					const FString PersistentSourcePackage(OriginalView.Left(GeneratedFolderStartIndex));
+									
+					FNameBuilder PersistentPackageInstanceNameBuilder;
+					PersistentPackageInstanceNameBuilder.Append(InstancedPackagePrefix);
+					PersistentPackageInstanceNameBuilder.Append(PersistentSourcePackage);
+					PersistentPackageInstanceNameBuilder.Append(InstancedPackageSuffix);
 
-					FNameBuilder PackageNameBuilder;
-					PackageNameBuilder.Append(OriginalView.Left(GeneratedFolderStartIndex));
-					FName PersistentPackageName(PackageNameBuilder);
-					PackageNameBuilder.Append(InstancedPackageSuffix);
-					FName PersistentPackageInstanceName(PackageNameBuilder);
-										
-					AddPackageMapping(PersistentPackageName, PersistentPackageInstanceName);
+					const FName PersistentPackageInstanceName(PersistentPackageInstanceNameBuilder);
+
+					AddPackageMapping(*PersistentSourcePackage, PersistentPackageInstanceName);
 				}
 			}
 
@@ -88,7 +98,6 @@ void FLinkerInstancedPackageMap::BuildPackageMapping(FName Original, FName Insta
 			}
 		}
 	}
-#endif
 }
 
 bool FLinkerInstancedPackageMap::FixupSoftObjectPath(FSoftObjectPath& InOutSoftObjectPath) const
@@ -100,23 +109,35 @@ bool FLinkerInstancedPackageMap::FixupSoftObjectPath(FSoftObjectPath& InOutSoftO
 			InOutSoftObjectPath = FSoftObjectPath(RemappedPackage, InOutSoftObjectPath.GetAssetFName(), InOutSoftObjectPath.GetSubPathString());
 			return true;
 		}
-#if !WITH_EDITOR
 		else if (!GeneratedPackagesFolder.IsEmpty())
 		{
-			check(!InstancedPackageSuffix.IsEmpty());
+#if WITH_EDITOR
+			// This code path should only be enabled while testing
+			check(GIsAutomationTesting);
+#endif
+			check(!InstancedPackagePrefix.IsEmpty() || !InstancedPackageSuffix.IsEmpty());
 
 			FNameBuilder TmpSoftObjectPathBuilder;
 			InOutSoftObjectPath.ToString(TmpSoftObjectPathBuilder);
 
 			// Does this package path start with the generated folder path?
 			FStringView TmpSoftObjectPathView = TmpSoftObjectPathBuilder.ToView();
-			if (TmpSoftObjectPathView.StartsWith(GeneratedPackagesFolder))
+			if (const int32 GeneratedFolderIndex = TmpSoftObjectPathView.Find(GeneratedPackagesFolder); GeneratedFolderIndex != INDEX_NONE)
 			{
+				check(GeneratedFolderIndex == 0 || (TmpSoftObjectPathView.StartsWith(InstancedPackagePrefix) && GeneratedFolderIndex == InstancedPackagePrefix.Len()));
+
 				// ... and is that generated folder path immediately preceding the package name?
-				if (const int32 ExtraSlashIndex = TmpSoftObjectPathView.Find(TEXTVIEW("/"), GeneratedPackagesFolder.Len()); ExtraSlashIndex == INDEX_NONE)
+				if (const int32 ExtraSlashIndex = TmpSoftObjectPathView.Find(TEXTVIEW("/"), InstancedPackagePrefix.Len() + GeneratedPackagesFolder.Len()); ExtraSlashIndex == INDEX_NONE)
 				{
 					FNameBuilder PackageNameBuilder;
-					PackageNameBuilder.Append(InOutSoftObjectPath.GetLongPackageName());
+
+					if (InstanceMappingDirection == EInstanceMappingDirection::OriginalToInstanced)
+					{
+						PackageNameBuilder.Append(InstancedPackagePrefix);
+					}
+					
+					PackageNameBuilder.Append(InOutSoftObjectPath.GetLongPackageName().Mid(GeneratedFolderIndex));
+
 					if (InstanceMappingDirection == EInstanceMappingDirection::OriginalToInstanced)
 					{
 						if (!PackageNameBuilder.ToView().EndsWith(InstancedPackageSuffix))
@@ -138,7 +159,141 @@ bool FLinkerInstancedPackageMap::FixupSoftObjectPath(FSoftObjectPath& InOutSoftO
 				}
 			}
 		}
-#endif
 	}
 	return false;
 }
+
+IMPLEMENT_SIMPLE_AUTOMATION_TEST(FLinkerInstancingContextTests, "System.CoreUObject.LinkerInstancingContext", EAutomationTestFlags::ApplicationContextMask | EAutomationTestFlags::EngineFilter);
+bool FLinkerInstancingContextTests::RunTest(const FString& Parameters)
+{
+	// Disabled SoftObjectPath Remapping
+	{
+		FName MappedSourcePackage(TEXT("/Game/PathA/PathB/PackageName"));
+		FName UnmappedSourcePackage(TEXT("/Game/PathA/PathB/PackageNameOther"));
+		FName InstancedPackage(TEXT("/Game/PathA/PathB/PackageName_Instance"));
+
+		const FSoftObjectPath MappedSourcePackagePath(FTopLevelAssetPath(MappedSourcePackage, FName(TEXT("PackageName"))));
+		const FSoftObjectPath MappedSourceObjectPath(FTopLevelAssetPath(MappedSourcePackage, FName(TEXT("PackageName"))), "PersistentLevel.Actor.Component");
+
+		const bool bSoftObjectPathRemappingEnabled = false;
+		FLinkerInstancingContext LinkerInstancingContext(bSoftObjectPathRemappingEnabled);
+		LinkerInstancingContext.EnableAutomationTest();
+		LinkerInstancingContext.BuildPackageMapping(MappedSourcePackage, InstancedPackage);
+
+		TestEqual("RemapPackage with existing mapping - bSoftObjectPathRemappingEnabled(false)", LinkerInstancingContext.RemapPackage(MappedSourcePackage), InstancedPackage);
+		TestEqual("RemapPackage without mapping - bSoftObjectPathRemappingEnabled(false)", LinkerInstancingContext.RemapPackage(UnmappedSourcePackage), UnmappedSourcePackage);
+
+		FSoftObjectPath CopyMappedSourceObjectPackage = MappedSourcePackagePath;
+		LinkerInstancingContext.FixupSoftObjectPath(CopyMappedSourceObjectPackage);
+		TestEqual("FixupSoftObjectPath with no SubPathString - bSoftObjectPathRemappingEnabled(false)", CopyMappedSourceObjectPackage, MappedSourcePackagePath);
+
+		FSoftObjectPath CopyMappedSourceObjectPath = MappedSourceObjectPath;
+		LinkerInstancingContext.FixupSoftObjectPath(CopyMappedSourceObjectPath);
+		TestEqual("FixupSoftObjectPath with SubPathString - bSoftObjectPathRemappingEnabled(false)", CopyMappedSourceObjectPath, MappedSourceObjectPath);
+	}
+
+	// Enabled SoftObjectPath Remapping
+	{
+		FName MappedSourcePackage(TEXT("/Game/PathA/PathB/PackageName"));
+		FName UnmappedSourcePackage(TEXT("/Game/PathA/PathB/PackageNameOther"));
+		FName InstancedPackage(TEXT("/Game/PathA/PathB/PackageName_Instance"));
+
+		const FSoftObjectPath MappedSourcePackagePath(FTopLevelAssetPath(MappedSourcePackage, FName(TEXT("PackageName"))));
+		const FSoftObjectPath RemappedPackagePath(FTopLevelAssetPath(InstancedPackage, MappedSourcePackagePath.GetAssetFName()));
+		const FSoftObjectPath MappedSourceObjectPath(FTopLevelAssetPath(MappedSourcePackage, FName(TEXT("PackageName"))), "PersistentLevel.Actor.Component");
+		const FSoftObjectPath RemappedObjectPath(FTopLevelAssetPath(InstancedPackage, MappedSourceObjectPath.GetAssetFName()), MappedSourceObjectPath.GetSubPathString());
+
+		const bool bSoftObjectPathRemappingEnabled = true;
+		FLinkerInstancingContext LinkerInstancingContext(bSoftObjectPathRemappingEnabled);
+		LinkerInstancingContext.EnableAutomationTest();
+		LinkerInstancingContext.BuildPackageMapping(MappedSourcePackage, InstancedPackage);
+
+		TestEqual("RemapPackage with existing mapping - bSoftObjectPathRemappingEnabled(true)", LinkerInstancingContext.RemapPackage(MappedSourcePackage), InstancedPackage);
+		TestEqual("RemapPackage without mapping - bSoftObjectPathRemappingEnabled(true)", LinkerInstancingContext.RemapPackage(UnmappedSourcePackage), UnmappedSourcePackage);
+
+		FSoftObjectPath CopyMappedSourceObjectPackage = MappedSourcePackagePath;
+		LinkerInstancingContext.FixupSoftObjectPath(CopyMappedSourceObjectPackage);
+		TestEqual("FixupSoftObjectPath with no SubPathString - bSoftObjectPathRemappingEnabled(true)", CopyMappedSourceObjectPackage, RemappedPackagePath);
+
+		FSoftObjectPath CopyMappedSourceObjectPath = MappedSourceObjectPath;
+		LinkerInstancingContext.FixupSoftObjectPath(CopyMappedSourceObjectPath);
+		TestEqual("FixupSoftObjectPath with SubPathString - bSoftObjectPathRemappingEnabled(true)", CopyMappedSourceObjectPath, RemappedObjectPath);
+	}
+	
+	auto TestGeneratedPathFixup = [this](FName InSourceGeneratorPackage, FName InInstancedGeneratorPackage, FName InSourceGeneratedPackage, FName InInstancedGeneratedPackage, const FString& Description)
+	{
+		const bool bSoftObjectPathRemappingEnabled = true;
+		FLinkerInstancingContext LinkerInstancingContext(bSoftObjectPathRemappingEnabled);
+		LinkerInstancingContext.EnableAutomationTest();
+		LinkerInstancingContext.BuildPackageMapping(InSourceGeneratedPackage, InInstancedGeneratedPackage);
+
+		// Reverse mapping
+		FLinkerInstancedPackageMap InstancingMap(FLinkerInstancedPackageMap::EInstanceMappingDirection::InstancedToOriginal);
+		InstancingMap.EnableAutomationTest();
+		InstancingMap.BuildPackageMapping(InSourceGeneratedPackage, InInstancedGeneratedPackage);
+
+		TestEqual(FString::Format(TEXT("{0} - RemapPackage - _Generated_ Package with Suffix"), { Description }), LinkerInstancingContext.RemapPackage(InSourceGeneratedPackage), InInstancedGeneratedPackage);
+
+		auto TestSoftObjectPathFixup = [this, &LinkerInstancingContext, &InstancingMap](const FSoftObjectPath& InSource, const FSoftObjectPath& InExpectedResult, const FString& Description)
+		{
+			FSoftObjectPath CopySource = InSource;
+			LinkerInstancingContext.FixupSoftObjectPath(CopySource);
+			TestEqual(FString::Format(TEXT("{0} - FixupSoftObjectPath"), { Description }), CopySource, InExpectedResult);
+
+			LinkerInstancingContext.FixupSoftObjectPath(CopySource);
+			TestEqual(FString::Format(TEXT("{0} - FixupSoftObjectPath already fixed up"), { Description }), CopySource, InExpectedResult);
+
+			InstancingMap.FixupSoftObjectPath(CopySource);
+			TestEqual(FString::Format(TEXT("{0} - Reverse FixupSoftObjectPath"), { Description }), CopySource, InSource);
+		};
+
+		const FSoftObjectPath SourcePackagePath(FTopLevelAssetPath(InSourceGeneratorPackage, FName(TEXT("PackageName"))));
+		const FSoftObjectPath RemappedPackagePath(FTopLevelAssetPath(InInstancedGeneratorPackage, SourcePackagePath.GetAssetFName()));
+		TestSoftObjectPathFixup(SourcePackagePath, RemappedPackagePath, Description);
+
+		const FSoftObjectPath SourceObjectPath(FTopLevelAssetPath(InSourceGeneratorPackage, FName(TEXT("PackageName"))), "PersistentLevel.Actor.Component");
+		const FSoftObjectPath RemappedSourceObjectPath(FTopLevelAssetPath(InInstancedGeneratorPackage, SourceObjectPath.GetAssetFName()), SourceObjectPath.GetSubPathString());
+		TestSoftObjectPathFixup(SourceObjectPath, RemappedSourceObjectPath, Description);
+
+		const FSoftObjectPath SourceGeneratedPackagePath(FTopLevelAssetPath(InSourceGeneratedPackage, FName(TEXT("PackageName"))));
+		const FSoftObjectPath RemappedGeneratedPackagePath(FTopLevelAssetPath(InInstancedGeneratedPackage, SourceGeneratedPackagePath.GetAssetFName()));
+		TestSoftObjectPathFixup(SourceGeneratedPackagePath, RemappedGeneratedPackagePath, Description);
+
+		const FSoftObjectPath SourceGeneratedObjectPath(FTopLevelAssetPath(InSourceGeneratedPackage, FName(TEXT("PackageName"))), "PersistentLevel.Actor.Component");
+		const FSoftObjectPath RemappedSourceGeneratedObjectPath(FTopLevelAssetPath(InInstancedGeneratedPackage, SourceGeneratedObjectPath.GetAssetFName()), SourceGeneratedObjectPath.GetSubPathString());
+		TestSoftObjectPathFixup(SourceGeneratedObjectPath, RemappedSourceGeneratedObjectPath, Description);
+	};
+
+	{
+		// Generated Package Mapping (Suffix)
+		FName SourcePackage(TEXT("/Game/PathA/PathB/PackageName"));
+		FName InstancedPackage(TEXT("/Game/PathA/PathB/PackageName_LevelInstance1"));
+		FName SourceGeneratedPackage(TEXT("/Game/PathA/PathB/PackageName/_Generated_/1AEHCD0PRR98XVM12PU4N1D8X"));
+		FName InstancedGeneratedPackage(TEXT("/Game/PathA/PathB/PackageName/_Generated_/1AEHCD0PRR98XVM12PU4N1D8X_LevelInstance1"));
+
+		TestGeneratedPathFixup(SourcePackage, InstancedPackage, SourceGeneratedPackage, InstancedGeneratedPackage, TEXT("Suffix"));
+	}
+
+	{
+		// Generated Package Mapping (Prefix)
+		FName SourcePackage(TEXT("/Game/PathA/PathB/PackageName"));
+		FName InstancedPackage(TEXT("/Temp/Game/PathA/PathB/PackageName"));
+		FName SourceGeneratedPackage(TEXT("/Game/PathA/PathB/PackageName/_Generated_/1AEHCD0PRR98XVM12PU4N1D8X"));
+		FName InstancedGeneratedPackage(TEXT("/Temp/Game/PathA/PathB/PackageName/_Generated_/1AEHCD0PRR98XVM12PU4N1D8X"));
+
+		TestGeneratedPathFixup(SourcePackage, InstancedPackage, SourceGeneratedPackage, InstancedGeneratedPackage, TEXT("Prefix"));
+	}
+
+	{
+		// Generated Package Mapping (Prefix + Suffix)
+		FName SourcePackage(TEXT("/Game/PathA/PathB/PackageName"));
+		FName InstancedPackage(TEXT("/Temp/Game/PathA/PathB/PackageName_LevelInstance1"));
+		FName SourceGeneratedPackage(TEXT("/Game/PathA/PathB/PackageName/_Generated_/1AEHCD0PRR98XVM12PU4N1D8X"));
+		FName InstancedGeneratedPackage(TEXT("/Temp/Game/PathA/PathB/PackageName/_Generated_/1AEHCD0PRR98XVM12PU4N1D8X_LevelInstance1"));
+
+		TestGeneratedPathFixup(SourcePackage, InstancedPackage, SourceGeneratedPackage, InstancedGeneratedPackage, TEXT("Suffix+Prefix"));
+	}
+
+	return true;
+}
+
diff --git a/Engine/Source/Runtime/CoreUObject/Public/UObject/LinkerInstancingContext.h b/Engine/Source/Runtime/CoreUObject/Public/UObject/LinkerInstancingContext.h
--- a/Engine/Source/Runtime/CoreUObject/Public/UObject/LinkerInstancingContext.h
+++ b/Engine/Source/Runtime/CoreUObject/Public/UObject/LinkerInstancingContext.h
@@ -16,10 +16,16 @@ public:
 		InstancedToOriginal,
 	};
 
-	FLinkerInstancedPackageMap() = default;
+	FLinkerInstancedPackageMap() :
+		FLinkerInstancedPackageMap(EInstanceMappingDirection::OriginalToInstanced) {}
 
 	explicit FLinkerInstancedPackageMap(EInstanceMappingDirection MappingDirection)
 		: InstanceMappingDirection(MappingDirection)
+#if WITH_EDITOR
+		, bEnableNonEditorPath(false)
+#else
+		, bEnableNonEditorPath(true)
+#endif
 	{
 	}
 
@@ -47,6 +53,8 @@ public:
 
 private:
 	friend class FLinkerInstancingContext;
+	friend class FLinkerInstancingContextTests;
+	void EnableAutomationTest() { bEnableNonEditorPath = true; }
 
 	/**
 	 * Map between the original package name and its instance counterpart.
@@ -58,11 +66,15 @@ private:
 	/**
 	 * In which direction has this mapping been built?
 	 */
-	EInstanceMappingDirection InstanceMappingDirection = EInstanceMappingDirection::OriginalToInstanced;
+	EInstanceMappingDirection InstanceMappingDirection;
 
 	/** Data needed to re-map world partition cells */
 	FString GeneratedPackagesFolder;
+	FString InstancedPackagePrefix;
 	FString InstancedPackageSuffix;
+
+	/** Allows tests to run non editor path from editor build. */
+	bool bEnableNonEditorPath;
 };
 
 /**
@@ -165,6 +177,8 @@ public:
 	void FixupSoftObjectPath(FSoftObjectPath& InOutSoftObjectPath) const;
 
 private:
+	void EnableAutomationTest() { InstancedPackageMap.EnableAutomationTest(); }
+
 	void BuildPackageMapping(FName Original, FName Instanced)
 	{
 		InstancedPackageMap.BuildPackageMapping(Original, Instanced, GetSoftObjectPathRemappingEnabled());
@@ -177,6 +191,7 @@ private:
 
 	friend class FLinkerLoad;
 	friend struct FAsyncPackage2;
+	friend class FLinkerInstancingContextTests;
 
 	/** Map of original package name to their instance counterpart. */
 	FLinkerInstancedPackageMap InstancedPackageMap;
diff --git a/Engine/Source/Runtime/Engine/Private/WorldPartition/WorldPartitionLevelStreamingPolicy.cpp b/Engine/Source/Runtime/Engine/Private/WorldPartition/WorldPartitionLevelStreamingPolicy.cpp
--- a/Engine/Source/Runtime/Engine/Private/WorldPartition/WorldPartitionLevelStreamingPolicy.cpp
+++ b/Engine/Source/Runtime/Engine/Private/WorldPartition/WorldPartitionLevelStreamingPolicy.cpp
@@ -167,12 +167,18 @@ bool UWorldPartitionLevelStreamingPolicy::ConvertEditorPathToRuntimePath(const F
 	}
 
 	// In the editor, the _LevelInstance_ID is appended to the persistent level, while at runtime it is appended to each cell package, so we need to remap it there if present.
+	// Also handle prefixes like "/Temp"
 	FString LevelInstanceSuffix;
+	FString LevelInstancePrefix;
 	const FString WorldAssetPackageName = WorldAssetPath.GetPackageName().ToString();
 	const FString SourceWorldAssetPackageName = SourceWorldAssetPath.GetPackageName().ToString();
-	if (WorldAssetPackageName.StartsWith(SourceWorldAssetPackageName))
+	if (WorldAssetPackageName.Len() > SourceWorldAssetPackageName.Len())
 	{
-		LevelInstanceSuffix = WorldAssetPackageName.Mid(SourceWorldAssetPackageName.Len());
+		if (const int32 Index = WorldAssetPackageName.Find(SourceWorldAssetPackageName); Index != INDEX_NONE)
+		{
+			LevelInstancePrefix = WorldAssetPackageName.Mid(0, Index);
+			LevelInstanceSuffix = WorldAssetPackageName.Mid(Index + SourceWorldAssetPackageName.Len());
+		}
 	}
 
 	FString SubAssetName;
@@ -200,7 +206,7 @@ bool UWorldPartitionLevelStreamingPolicy::ConvertEditorPathToRuntimePath(const F
 #endif
 			else
 			{
-				OutPath = FString::Printf(TEXT("%s/_Generated_/%s%s.%s:%s"), *SourceWorldAssetPackageName, *(CellName->ToString()), *LevelInstanceSuffix, *WorldAssetPath.GetAssetName().ToString(), *InPath.GetSubPathString());
+				OutPath = FString::Printf(TEXT("%s%s/_Generated_/%s%s.%s:%s"), *LevelInstancePrefix, *SourceWorldAssetPackageName, *(CellName->ToString()), *LevelInstanceSuffix, *WorldAssetPath.GetAssetName().ToString(), *InPath.GetSubPathString());
 			}
 
 #if WITH_EDITOR
diff --git a/Engine/Source/Runtime/Engine/Private/WorldPartition/WorldPartitionRuntimeLevelStreamingCell.cpp b/Engine/Source/Runtime/Engine/Private/WorldPartition/WorldPartitionRuntimeLevelStreamingCell.cpp
--- a/Engine/Source/Runtime/Engine/Private/WorldPartition/WorldPartitionRuntimeLevelStreamingCell.cpp
+++ b/Engine/Source/Runtime/Engine/Private/WorldPartition/WorldPartitionRuntimeLevelStreamingCell.cpp
@@ -361,39 +361,14 @@ UWorldPartitionLevelStreamingDynamic* UWorldPartitionRuntimeLevelStreamingCell::
 		// Transfer WorldPartition's transform to LevelStreaming
 		LevelStreaming->LevelTransform = WorldPartition->GetInstanceTransform();
 
-		// When Partition outer level is an instance, make sure to also generate unique cell level instance name
-		ULevel* PartitionLevel = WorldPartition->GetTypedOuter<ULevel>();
-		if (PartitionLevel->IsInstancedLevel())
-		{
-			// Try and extract the instance suffix that was applied to the instanced level itself
-			FString InstancedLevelSuffix;
-			{
-				UPackage* PartitionLevelPackage = PartitionLevel->GetPackage();
-
-				FNameBuilder SourcePackageName(PartitionLevelPackage->GetLoadedPath().GetPackageFName());
-				FNameBuilder InstancedPackageName(PartitionLevelPackage->GetFName());
-
-				FStringView SourcePackageNameView = SourcePackageName.ToView();
-				FStringView InstancedPackageNameView = InstancedPackageName.ToView();
-
-				if (InstancedPackageNameView.StartsWith(SourcePackageNameView))
-				{
-					InstancedLevelSuffix = InstancedPackageNameView.Mid(SourcePackageNameView.Len());
-				}
-				else
-				{
-					InstancedLevelSuffix = TEXT("_");
-					InstancedLevelSuffix += FPackageName::GetShortName(PartitionLevelPackage);
-				}
-			}
-			check(!InstancedLevelSuffix.IsEmpty());
-
-			FNameBuilder InstancedLevelPackageName;
-			LevelStreaming->PackageNameToLoad.AppendString(InstancedLevelPackageName);
-			InstancedLevelPackageName += InstancedLevelSuffix;
-
-			LevelStreaming->SetWorldAssetByPackageName(FName(InstancedLevelPackageName));
-		}
+		// LevelStreaming WorldAsset is a TSoftObjectPtr<UWorld>. If the WorldPartition's world is instanced then the TSoftObjectPtr<UWorld> will be remapped by FLinkerInstancingContext SoftObject remapping
+		// example if MainWorld is instanced it will have a package suffix and/or prefix like: /Temp+PATH+_LevelInstance1
+		// MainWorld: /Game/SomePath/MainWorld_LevelInstance1
+		// 
+		// and the LevelStreaming's WorldAsset TSoftObjectPtr<UWorld> should be remapped as: 
+		// /Temp/Game/SomePath/MainWorld/_Generated_/1AEHCD0PRR98XVM12PU4N1D8X_LevelInstance_1
+
+		check(!WorldPartition->GetTypedOuter<ULevel>()->IsInstancedLevel() || (LevelStreaming->PackageNameToLoad != LevelStreaming->GetWorldAssetPackageName()));
 	}
 #endif
 
