Fix newly created flags not being properly stripped from packages on save

diff --git a/Engine/Source/Runtime/CoreUObject/Private/UObject/SavePackage/SaveContext.h b/Engine/Source/Runtime/CoreUObject/Private/UObject/SavePackage/SaveContext.h
--- a/Engine/Source/Runtime/CoreUObject/Private/UObject/SavePackage/SaveContext.h
+++ b/Engine/Source/Runtime/CoreUObject/Private/UObject/SavePackage/SaveContext.h
@@ -430,7 +430,7 @@ public:
 		checkf(!IsCooking() || PackageWriter, TEXT("Cook saves require an IPackageWriter"));
 
 		// Store initial state
-		PreviousPackageFlags = Package->GetPackageFlags();
+		InitialPackageFlags = Package->GetPackageFlags();
 
 		SaveArgs.TopLevelFlags = UE::SavePackageUtilities::NormalizeTopLevelFlags(SaveArgs.TopLevelFlags, IsCooking());
 		if (PackageWriter)
@@ -475,8 +475,11 @@ public:
 		{
 			UE::SavePackageUtilities::CallPostSaveRoot(Asset, ObjectSaveContext, bNeedPreSaveCleanup);
 		}
-		// restore initial state
-		Package->SetPackageFlagsTo(PreviousPackageFlags);
+	}
+
+	uint32 GetInitialPackageFlags() const
+	{
+		return InitialPackageFlags;
 	}
 
 	const FSavePackageArgs& GetSaveArgs() const
@@ -1092,7 +1095,7 @@ private:
 	bool bIsSaveAutoOptional = false;
 
 	// Mutated package state
-	uint32 PreviousPackageFlags;
+	uint32 InitialPackageFlags;
 
 	// Config classes shared with the old Save
 	FCanSkipEditorReferencedPackagesWhenCooking SkipEditorRefCookingSetting;
diff --git a/Engine/Source/Runtime/CoreUObject/Private/UObject/SavePackage2.cpp b/Engine/Source/Runtime/CoreUObject/Private/UObject/SavePackage2.cpp
--- a/Engine/Source/Runtime/CoreUObject/Private/UObject/SavePackage2.cpp
+++ b/Engine/Source/Runtime/CoreUObject/Private/UObject/SavePackage2.cpp
@@ -2605,7 +2605,16 @@ void ClearCachedPlatformCookedData(FSaveContext& SaveContext)
 void PostSavePackage(FSaveContext& SaveContext)
 {
 	UPackage* Package = SaveContext.GetPackage();
+	// restore initial packages flags since the save can currently mutate the package flags in an undesirable fashion, (i.e. clearing/adding editor only filtering)
+	Package->SetPackageFlagsTo(SaveContext.GetInitialPackageFlags());
 
+	// Do not run the rest of post save and send out events if the save wasn't succesful
+	if (SaveContext.Result != ESavePackageResult::Success)
+	{
+		return;
+	}
+
+	// then adjust flags that should be modified at the outcome of a save
 	if (!SaveContext.IsFromAutoSave() && !SaveContext.IsProceduralSave())
 	{
 		// Package has been saved, so unmark the NewlyCreated flag.
@@ -2647,7 +2656,6 @@ void PostSavePackage(FSaveContext& SaveContext)
 	}
 #endif
 
-
 	// Send a message that the package was saved
 	PRAGMA_DISABLE_DEPRECATION_WARNINGS;
 	UPackage::PackageSavedEvent.Broadcast(SaveContext.GetFilename(), Package);
@@ -3117,10 +3125,7 @@ FSavePackageResultStruct UPackage::Save2(UPackage* InPackage, UObject* InAsset,
 
 	// PostSave Package - edit in memory package and send events if save was successful
 	SlowTask.EnterProgressFrame();
-	if (SaveContext.Result == ESavePackageResult::Success)
-	{
-		PostSavePackage(SaveContext);
-	}
+	PostSavePackage(SaveContext);
 	return SaveContext.GetFinalResult();
 }
 
