FPathViews::Less: Fix inconsistent sort order that was breaking DirectoryTree.
"Foo" < "Foo-Bar "
but
"Foo/" > "Foo-Bar/".
In the new order, '/' and '\\' are less than every character except '\0'.
Note this will change existing sort order ('A' used to be  less than '\\', now it is more than '\\) and might cause some cooked packages with strings sorted by path to change.
#rnx

diff --git a/Engine/Source/Runtime/Core/Private/Misc/PathViews.cpp b/Engine/Source/Runtime/Core/Private/Misc/PathViews.cpp
--- a/Engine/Source/Runtime/Core/Private/Misc/PathViews.cpp
+++ b/Engine/Source/Runtime/Core/Private/Misc/PathViews.cpp
@@ -120,9 +120,13 @@ namespace UE4PathViews_Private
 			{
 				if (!IsSlashOrBackslash(BChar))
 				{
-					return AChar < BChar;
+					return true;
 				}
 			}
+			else if (IsSlashOrBackslash(BChar))
+			{
+				return false;
+			}
 			else
 			{
 				if (AChar != BChar)
diff --git a/Engine/Source/Runtime/Core/Tests/Misc/PathViewsTest.cpp b/Engine/Source/Runtime/Core/Tests/Misc/PathViewsTest.cpp
--- a/Engine/Source/Runtime/Core/Tests/Misc/PathViewsTest.cpp
+++ b/Engine/Source/Runtime/Core/Tests/Misc/PathViewsTest.cpp
@@ -264,13 +264,44 @@ public:
 		//RunEqualsLessTest(TEXT("C:/A/B"), TEXT("C:/A/../A/./B"), 0); // Collapsing .. and . is not yet implemented
 		RunEqualsLessTest(TEXT("/"), TEXT("/"), 0);
 		RunEqualsLessTest(TEXT("/"), TEXT("//"), -1);
-		RunEqualsLessTest(TEXT("/"), TEXT("C:/"), -1); // int('/') == 47 < int('C') == 67
+		RunEqualsLessTest(TEXT("/"), TEXT("C:/"), -1); // '/' sorts less than 'C'
+		RunEqualsLessTest(TEXT("/"), TEXT("C:\\"), -1); // '\\' sorts less than 'C'
+		RunEqualsLessTest(TEXT("/"), TEXT("c:/"), -1); // '/' sorts less than 'c'
+		RunEqualsLessTest(TEXT("/"), TEXT("c:\\"), -1); // '\\' sorts less than 'c'
 		RunEqualsLessTest(TEXT("/"), TEXT("A"), -1);
 		RunEqualsLessTest(TEXT("//"), TEXT("//"), 0);
 		RunEqualsLessTest(TEXT("//"), TEXT("C:/"), -1);
 		RunEqualsLessTest(TEXT("//"), TEXT("A"), -1);
 		RunEqualsLessTest(TEXT("C:/"), TEXT("C:/"), 0);
 		RunEqualsLessTest(TEXT("C:/"), TEXT("C"), 1);
+		// The directory separator is less than all characters except for '\0', so that shorter directories come before longer directories.
+		// This is true for all wchar; to save time we test it just for all ascii characters.
+		const TCHAR* ForwardSlashTestString = TEXT("0/1");
+		const TCHAR* BackSlashTestString = TEXT("0\\1");
+		TCHAR CharTestString[] = TEXT("0_1");
+		for (int IntC = 1; IntC < 128; IntC++)
+		{
+			char c = (char)IntC;
+			if (c == '/' || c == '\\')
+			{
+				continue;
+			}
+			CharTestString[1] = c;
+			RunEqualsLessTest(ForwardSlashTestString, CharTestString, -1);
+			RunEqualsLessTest(CharTestString, ForwardSlashTestString, 1);
+			RunEqualsLessTest(BackSlashTestString, CharTestString, -1);
+			RunEqualsLessTest(CharTestString, BackSlashTestString, 1);
+		}
+		// These tests use hyphen, which has int('-') == 45 < int("/") == 47 < int("\\") == 92
+		// Foo is a shorter string than Foo-Bar so it comes first, aka / sorts earlier than all other characters
+		RunEqualsLessTest(TEXT("Foo/Leaf"), TEXT("Foo-Bar/Leaf"), -1);
+		RunEqualsLessTest(TEXT("Foo\\Leaf"), TEXT("Foo-Bar\\Leaf"), -1);
+		// When the / is terminating, sort order needs to match what it is when it is in the middle
+		RunEqualsLessTest(TEXT("Foo/"), TEXT("Foo-Bar/"), -1); 
+		// When the terminating / is omitted, sort order needs to match when it is present
+		RunEqualsLessTest(TEXT("Foo"), TEXT("Foo-Bar"), -1); 
+		RunEqualsLessTest(TEXT("Foo"), TEXT("Foo-Bar/"), -1); 
+		RunEqualsLessTest(TEXT("Foo/"), TEXT("Foo-Bar/"), -1); 
 
 		return true;
 	}
diff --git a/Engine/Source/Runtime/CoreUObject/Private/Misc/DirectoryTree.cpp b/Engine/Source/Runtime/CoreUObject/Private/Misc/DirectoryTree.cpp
--- a/Engine/Source/Runtime/CoreUObject/Private/Misc/DirectoryTree.cpp
+++ b/Engine/Source/Runtime/CoreUObject/Private/Misc/DirectoryTree.cpp
@@ -43,7 +43,7 @@ IMPLEMENT_SIMPLE_AUTOMATION_TEST(FDirectoryTreeTest, "System.Core.DirectoryTree"
 bool FDirectoryTreeTest::RunTest(const FString& Parameters)
 {
 	constexpr int32 NumPathTypes = 5;
-	constexpr int32 NumPaths = 7;
+	constexpr int32 NumPaths = 9;
 	FStringView PathsByTypeAndIndex[NumPathTypes][NumPaths] =
 	{
 		{
@@ -54,6 +54,9 @@ bool FDirectoryTreeTest::RunTest(const FString& Parameters)
 			TEXTVIEW("/Plugin1/Dir2/Path2"),
 			TEXTVIEW("/Engine/Path1"),
 			TEXTVIEW("/Plugin2/Path1"),
+			// Make sure we handle suffixes of an existing string with a leading value that sorts after /
+			TEXTVIEW("/Game/Foo/Leaf"),
+			TEXTVIEW("/Game/Foo-Bar/Leaf"),
 		},
 		{
 			TEXTVIEW("d:\\root\\Project\\Content\\Dir2"),
@@ -63,6 +66,8 @@ bool FDirectoryTreeTest::RunTest(const FString& Parameters)
 			TEXTVIEW("d:\\root\\Project\\Plugins\\Plugin1\\Content\\Dir2\\Path2.uasset"),
 			TEXTVIEW("d:\\root\\Engine\\Content\\Path1.uasset"),
 			TEXTVIEW("e:\\root\\Project\\Plugins\\Plugin2\\Content\\Path1.uasset"),
+			TEXTVIEW("d:\\root\\Project\\Content\\Foo\\Leaf"),
+			TEXTVIEW("d:\\root\\Project\\Content\\Foo-Bar\\Leaf"),
 		},
 		{
 			TEXTVIEW("d:/root/Project/Content/Dir2"),
@@ -72,6 +77,8 @@ bool FDirectoryTreeTest::RunTest(const FString& Parameters)
 			TEXTVIEW("d:/root/Project/Plugins/Plugin1/Content/Dir2/Path2.uasset"),
 			TEXTVIEW("d:/root/Engine/Content/Path1.uasset"),
 			TEXTVIEW("e:/root/Project/Plugins/Plugin2/Content/Path1.uasset"),
+			TEXTVIEW("d:/root/Project/Content/Foo/Leaf"),
+			TEXTVIEW("d:/root/Project/Content/Foo-Bar/Leaf"),
 		},
 		{
 			TEXTVIEW("..\\..\\..\\Project\\Content\\Dir2"),
@@ -81,6 +88,8 @@ bool FDirectoryTreeTest::RunTest(const FString& Parameters)
 			TEXTVIEW("..\\..\\..\\Project\\Plugins\\Plugin1\\Content\\Dir2\\Path2.uasset"),
 			TEXTVIEW("..\\..\\..\\Engine\\Content\\Path1.uasset"),
 			TEXTVIEW("e:\\root\\Project\\Plugins\\Plugin2\\Content\\Path1.uasset"),
+			TEXTVIEW("..\\..\\..\\Project\\Content\\Foo\\Leaf"),
+			TEXTVIEW("..\\..\\..\\Project\\Content\\Foo-Bar\\Leaf"),
 		},
 		{
 			TEXTVIEW("../../../Project/Content/Dir2"),
@@ -90,8 +99,11 @@ bool FDirectoryTreeTest::RunTest(const FString& Parameters)
 			TEXTVIEW("../../../Project/Plugins/Plugin1/Content/Dir2/Path2.uasset"),
 			TEXTVIEW("../../../Engine/Content/Path1.uasset"),
 			TEXTVIEW("e:/root/Project/Plugins/Plugin2/Content/Path1.uasset"),
+			TEXTVIEW("../../../Project/Content/Foo/Leaf"),
+			TEXTVIEW("../../../Project/Content/Foo-Bar/Leaf"),
 		},
 	};
+	// PathSub0SubPath[i] provides a sub path of PathsByTypeAndIndex[i][0]
 	FStringView Path0SubPath[NumPathTypes] =
 	{
 		TEXTVIEW("/Game/Dir2/Sub"),
@@ -100,7 +112,7 @@ bool FDirectoryTreeTest::RunTest(const FString& Parameters)
 		TEXTVIEW("..\\..\\..\\Project\\Content\\Dir2"),
 		TEXTVIEW("../../../Project/Content/Dir2/Sub"),
 	};
-	int32 ValueByIndex[NumPaths] = { 1,2,3,4,5,6,7 };
+	int32 ValueByIndex[NumPaths] = { 1,2,3,4,5,6,7,8,9 };
 	FStringView NonPathsByTypeAndIndex0[] =
 		{
 			TEXTVIEW(""),
@@ -190,8 +202,8 @@ bool FDirectoryTreeTest::RunTest(const FString& Parameters)
 	constexpr int32 NumPermutations = 2;
 	int32 Permutations[NumPermutations][NumPaths] =
 	{
-		{ 0, 1, 2, 3, 4, 5, 6 },
-		{ 6, 5, 4, 3, 2, 1, 0 },
+		{ 0, 1, 2, 3, 4, 5, 6, 7, 8 },
+		{ 8, 7, 6, 5, 4, 3, 2, 1, 0 },
 	};
 
 	auto GetDirTreeTestName = [](FStringView InTestName, int32 PathType, int32 Permutation, int32 EditPermutationIndex,
