- Stop writing out saved ini files after reloading ini after writing a default property (the timing of the write was getting the default and saved out of sync)

diff --git a/Engine/Source/Runtime/CoreUObject/Private/UObject/Obj.cpp b/Engine/Source/Runtime/CoreUObject/Private/UObject/Obj.cpp
--- a/Engine/Source/Runtime/CoreUObject/Private/UObject/Obj.cpp
+++ b/Engine/Source/Runtime/CoreUObject/Private/UObject/Obj.cpp
@@ -3218,7 +3218,10 @@ void UObject::UpdateSinglePropertyInConfigFile(const FProperty* InProperty, cons
 		// then we don't want to touch GConfig
 		if (OverridePlatform.Len() == 0)
 		{
-			FConfigContext::ForceReloadIntoGConfig().Load(*GetClass()->ClassConfigName.ToString());
+			FConfigContext Context = FConfigContext::ForceReloadIntoGConfig();
+			// don't write the Saved out, as we just finished writing what we needed
+			Context.bWriteDestIni = false;
+			Context.Load(*GetClass()->ClassConfigName.ToString());
 		}
 	}
 	else
