Updating callsites of ToRawPtr{*}Unsafe in preparation for upcoming mutable TObjectPtr container restrictions.

diff --git a/Engine/Plugins/Animation/ACLPlugin/Source/ACLPlugin/Private/AnimationCompressionLibraryDatabase.cpp b/Engine/Plugins/Animation/ACLPlugin/Source/ACLPlugin/Private/AnimationCompressionLibraryDatabase.cpp
--- a/Engine/Plugins/Animation/ACLPlugin/Source/ACLPlugin/Private/AnimationCompressionLibraryDatabase.cpp
+++ b/Engine/Plugins/Animation/ACLPlugin/Source/ACLPlugin/Private/AnimationCompressionLibraryDatabase.cpp
@@ -442,7 +442,7 @@ void UAnimationCompressionLibraryDatabase::UpdatePreviewState(bool bBuildDatabas
 bool UAnimationCompressionLibraryDatabase::UpdateReferencingAnimSequenceList()
 {
 	// Grab every anim sequence that references us
-	TArray<UAnimSequence*> ReferencingAnimSequences;
+	TArray<TObjectPtr<UAnimSequence>> ReferencingAnimSequences;
 	for (TObjectIterator<UAnimSequence> It; It; ++It)
 	{
 		UAnimSequence* AnimSeq = *It;
diff --git a/Engine/Plugins/EnhancedInput/Source/EnhancedInput/Private/EnhancedInputSubsystemInterface.cpp b/Engine/Plugins/EnhancedInput/Source/EnhancedInput/Private/EnhancedInputSubsystemInterface.cpp
--- a/Engine/Plugins/EnhancedInput/Source/EnhancedInput/Private/EnhancedInputSubsystemInterface.cpp
+++ b/Engine/Plugins/EnhancedInput/Source/EnhancedInput/Private/EnhancedInputSubsystemInterface.cpp
@@ -359,7 +359,7 @@ void IEnhancedInputSubsystemInterface::ApplyAxisPropertyModifiers(UEnhancedPlaye
 	FInputAxisProperties AxisProperties;
 	if (PlayerInput->GetAxisProperties(Mapping.Key, AxisProperties))
 	{
-		TArray<UInputModifier*> Modifiers;
+		TArray<TObjectPtr<UInputModifier>> Modifiers;
 
 		// If a modifier already exists it should override axis properties.
 		auto HasExistingModifier = [&Mapping](UClass* OfType)
diff --git a/Engine/Plugins/Runtime/RigVM/Source/RigVMDeveloper/Private/RigVMModel/RigVMController.cpp b/Engine/Plugins/Runtime/RigVM/Source/RigVMDeveloper/Private/RigVMModel/RigVMController.cpp
--- a/Engine/Plugins/Runtime/RigVM/Source/RigVMDeveloper/Private/RigVMModel/RigVMController.cpp
+++ b/Engine/Plugins/Runtime/RigVM/Source/RigVMDeveloper/Private/RigVMModel/RigVMController.cpp
@@ -18171,7 +18171,7 @@ void URigVMController::ProcessDetachedLinks(const FRestoreLinkedPathSettings& In
 	}
 
 	// remove the links from the graph first
-	TArray<URigVMLink*> DetachedLinks;
+	TArray<TObjectPtr<URigVMLink>> DetachedLinks;
 	Swap(GetGraph()->DetachedLinks, DetachedLinks);
 
 	// destroy the links finally
diff --git a/Engine/Plugins/Runtime/StateTree/Source/StateTreeEditorModule/Private/StateTreeViewModel.cpp b/Engine/Plugins/Runtime/StateTree/Source/StateTreeEditorModule/Private/StateTreeViewModel.cpp
--- a/Engine/Plugins/Runtime/StateTree/Source/StateTreeEditorModule/Private/StateTreeViewModel.cpp
+++ b/Engine/Plugins/Runtime/StateTree/Source/StateTreeEditorModule/Private/StateTreeViewModel.cpp
@@ -285,10 +285,10 @@ void FStateTreeViewModel::NotifyStatesChangedExternally(const TSet<UStateTreeSta
 	OnStatesChanged.Broadcast(ChangedStates, PropertyChangedEvent);
 }
 
-TArray<UStateTreeState*>* FStateTreeViewModel::GetSubTrees() const
+TArray<TObjectPtr<UStateTreeState>>* FStateTreeViewModel::GetSubTrees() const
 {
 	UStateTreeEditorData* TreeData = TreeDataWeak.Get();
-	return TreeData != nullptr ? &ToRawPtrTArrayUnsafe(TreeData->SubTrees) : nullptr;
+	return TreeData != nullptr ? &TreeData->SubTrees : nullptr;
 }
 
 int32 FStateTreeViewModel::GetSubTreeCount() const
diff --git a/Engine/Plugins/Runtime/StateTree/Source/StateTreeEditorModule/Private/StateTreeViewModel.h b/Engine/Plugins/Runtime/StateTree/Source/StateTreeEditorModule/Private/StateTreeViewModel.h
--- a/Engine/Plugins/Runtime/StateTree/Source/StateTreeEditorModule/Private/StateTreeViewModel.h
+++ b/Engine/Plugins/Runtime/StateTree/Source/StateTreeEditorModule/Private/StateTreeViewModel.h
@@ -59,7 +59,7 @@ public:
 	const UStateTree* GetStateTree() const;
 	
 	// Returns array of subtrees to edit.
-	TArray<UStateTreeState*>* GetSubTrees() const;
+	TArray<TObjectPtr<UStateTreeState>>* GetSubTrees() const;
 	int32 GetSubTreeCount() const;
 	void GetSubTrees(TArray<TWeakObjectPtr<UStateTreeState>>& OutSubtrees) const;
 
diff --git a/Engine/Source/Runtime/CoreUObject/Public/UObject/ObjectPtr.h b/Engine/Source/Runtime/CoreUObject/Public/UObject/ObjectPtr.h
--- a/Engine/Source/Runtime/CoreUObject/Public/UObject/ObjectPtr.h
+++ b/Engine/Source/Runtime/CoreUObject/Public/UObject/ObjectPtr.h
@@ -692,7 +692,13 @@ FORCEINLINE T* ToRawPtr(T* Ptr)
 }
 
 template <typename T, SIZE_T Size>
-FORCEINLINE T** ToRawPtrArrayUnsafe(TObjectPtr<T>(&ArrayOfPtr)[Size])
+#if UE_DEPRECATE_MUTABLE_TOBJECTPTR
+FORCEINLINE T* const *
+#else
+UE_OBJPTR_DEPRECATED(5.3, "Mutable ToRawPtrArrayUnsafe() is deprecated. Use MutableView() or TArray<TObjectPtr<...>> instead.")
+FORCEINLINE T**
+#endif
+ToRawPtrArrayUnsafe(TObjectPtr<T>(&ArrayOfPtr)[Size])
 {
 #if UE_WITH_OBJECT_HANDLE_LATE_RESOLVE || UE_WITH_OBJECT_HANDLE_TRACKING
 	for (TObjectPtr<T>& Item : ArrayOfPtr)
@@ -701,7 +707,12 @@ FORCEINLINE T** ToRawPtrArrayUnsafe(TObjectPtr<T>(&ArrayOfPtr)[Size])
 		Item.Get();
 	}
 #endif
+
+	#if UE_DEPRECATE_MUTABLE_TOBJECTPTR
+	return reinterpret_cast<T* const *>(ArrayOfPtr);
+	#else
 	return reinterpret_cast<T**>(ArrayOfPtr);
+	#endif
 }
 
 template <typename T>
@@ -716,7 +727,13 @@ template <
 	typename ArrayTypeNoRef = std::remove_reference_t<ArrayType>,
 	std::enable_if_t<TIsTArray_V<ArrayTypeNoRef>>* = nullptr
 >
-decltype(auto) ToRawPtrTArrayUnsafe(ArrayType&& Array)
+#if UE_DEPRECATE_MUTABLE_TOBJECTPTR
+const auto&
+#else
+UE_OBJPTR_DEPRECATED(5.3, "Mutable ToRawPtrTArrayUnsafe() is deprecated. Use MutableView() or TArray<TObjectPtr<...>> instead.")
+decltype(auto)
+#endif
+ToRawPtrTArrayUnsafe(ArrayType&& Array)
 {
 	using ArrayElementType         = typename ArrayTypeNoRef::ElementType;
 	using ArrayAllocatorType       = typename ArrayTypeNoRef::AllocatorType;
@@ -876,6 +893,8 @@ inline void Swap(T*& A, TObjectPtr<T>& B)
 }
 
 /** Swap variants between TArray<TObjectPtr<T>> and TArray<T*> */
+#if !UE_DEPRECATE_MUTABLE_TOBJECTPTR
+UE_OBJPTR_DEPRECATED(5.3, "Swap between TObjectPtr arrays and raw pointer arrays is deprecated. Swap TArray<TObjectPtr<...>> values instead.")
 template <typename T>
 inline void Swap(TArray<TObjectPtr<T>>& A, TArray<T*>& B)
 {
@@ -886,6 +905,7 @@ inline void Swap(TArray<T*>& A, TArray<TObjectPtr<T>>& B)
 {
 	Swap(A, ToRawPtrTArrayUnsafe(B));
 }
+#endif
 
 /** Exchange variants between TObjectPtr<T> and raw pointer to T */
 template <typename T>
@@ -899,7 +919,9 @@ inline void Exchange(T*& A, TObjectPtr<T>& B)
 	Swap(A, (T*&)B);
 }
 
+#if !UE_DEPRECATE_MUTABLE_TOBJECTPTR
 /** Exchange variants between TArray<TObjectPtr<T>> and TArray<T*> */
+UE_OBJPTR_DEPRECATED(5.3, "Exchange between TObjectPtr arrays and raw pointer arrays is deprecated. Exchange TArray<TObjectPtr<...>> values instead.")
 template <typename T>
 inline void Exchange(TArray<TObjectPtr<T>>& A, TArray<T*>& B)
 {
@@ -910,6 +932,7 @@ inline void Exchange(TArray<T*>& A, TArray<TObjectPtr<T>>& B)
 {
 	Swap(A, ToRawPtrTArrayUnsafe(B));
 }
+#endif
 
 /**
  * Returns a pointer to a valid object if the Test object passes IsValid() tests, otherwise null
@@ -1031,9 +1054,10 @@ namespace UE::Core::Private // private facilities; not for direct use
 	class TMutableView
 	{
 		using TraitType = UE::Core::Private::TMutableViewTraits<T>;
-		using ViewType = typename TraitType::ViewType;
 
 	public:
+		using ViewType = typename TraitType::ViewType;
+
 		explicit TMutableView(T& Value)
 			: View{&TraitType::Open(Value)}
 		{
@@ -1072,6 +1096,13 @@ namespace UE::Core::Private // private facilities; not for direct use
 
 	MutatingFunc(MutableView(Array));			 // ok; Array will safely "catch up" on TObjectPtr
 																				 // semantics when MutatingFunc returns.
+
+	// it's generally preferable to pass references around (to avoid nullptr),
+	// but for compat with existing functions that take a pointer:
+
+	void NeedsAPointer(TArray<UObject*>* MutableArrayPtr);
+
+	NeedsAPointer(ToRawPtr(MutableView(Array)));
 */
 
 template <typename T>
@@ -1080,6 +1111,13 @@ template <typename T>
 	return UE::Core::Private::TMutableView<T>{A};
 }
 
+template <typename T>
+[[nodiscard]] decltype(auto) ToRawPtr(UE::Core::Private::TMutableView<T>&& X)
+{
+	typename UE::Core::Private::TMutableView<T>::ViewType& Ref = X;
+	return &Ref;
+}
+
 #if UE_ENABLE_INCLUDE_ORDER_DEPRECATED_IN_5_3
 #include "UObject/Class.h"
 #endif
diff --git a/Engine/Source/Runtime/Engine/Private/Components/StaticMeshComponent.cpp b/Engine/Source/Runtime/Engine/Private/Components/StaticMeshComponent.cpp
--- a/Engine/Source/Runtime/Engine/Private/Components/StaticMeshComponent.cpp
+++ b/Engine/Source/Runtime/Engine/Private/Components/StaticMeshComponent.cpp
@@ -2044,7 +2044,7 @@ void UStaticMeshComponent::PostLoad()
 		if (GetStaticMesh()->HasValidRenderData()
 			&& GetStaticMesh()->GetRenderData()->MaterialIndexToImportIndex.Num())
 		{
-			TArray<UMaterialInterface*> OldMaterials;
+			TArray<TObjectPtr<UMaterialInterface>> OldMaterials;
 			const TArray<int32>& MaterialIndexToImportIndex = GetStaticMesh()->GetRenderData()->MaterialIndexToImportIndex;
 
 			Exchange(OverrideMaterials, OldMaterials);
diff --git a/Engine/Source/Runtime/Engine/Private/DemoNetDriver.cpp b/Engine/Source/Runtime/Engine/Private/DemoNetDriver.cpp
--- a/Engine/Source/Runtime/Engine/Private/DemoNetDriver.cpp
+++ b/Engine/Source/Runtime/Engine/Private/DemoNetDriver.cpp
@@ -5328,7 +5328,7 @@ void UDemoNetDriver::QueueNetStartupActorForRollbackViaDeletion(AActor* Actor)
 				FRepShadowDataBuffer ShadowData(ReceivingRepState->StaticBuffer.GetData());
 				FConstRepObjectDataBuffer ActorData(Actor);
 
-				if (NewReplicator->RepLayout->DiffStableProperties(nullptr, &ToRawPtrTArrayUnsafe(RollbackActor.ObjReferences), ShadowData, ActorData))
+				if (NewReplicator->RepLayout->DiffStableProperties(nullptr, ToRawPtr(MutableView(RollbackActor.ObjReferences)), ShadowData, ActorData))
 				{
 					RollbackActor.RepState = MakeShareable(NewReplicator->RepState.Release());
 				}
@@ -5349,7 +5349,7 @@ void UDemoNetDriver::QueueNetStartupActorForRollbackViaDeletion(AActor* Actor)
 					FRepShadowDataBuffer ShadowData(ReceivingRepState->StaticBuffer.GetData());
 					FConstRepObjectDataBuffer ActorCompData(ActorComp);
 
-					if (SubObjReplicator->RepLayout->DiffStableProperties(nullptr, &ToRawPtrTArrayUnsafe(RollbackActor.ObjReferences), ShadowData, ActorCompData))
+					if (SubObjReplicator->RepLayout->DiffStableProperties(nullptr, ToRawPtr(MutableView(RollbackActor.ObjReferences)), ShadowData, ActorCompData))
 					{
 						RollbackActor.SubObjRepState.Add(ActorComp->GetFullName(), MakeShareable(SubObjReplicator->RepState.Release()));
 					}
diff --git a/Engine/Source/Runtime/Engine/Private/WorldSettings.cpp b/Engine/Source/Runtime/Engine/Private/WorldSettings.cpp
--- a/Engine/Source/Runtime/Engine/Private/WorldSettings.cpp
+++ b/Engine/Source/Runtime/Engine/Private/WorldSettings.cpp
@@ -378,9 +378,10 @@ PRAGMA_DISABLE_DEPRECATION_WARNINGS
 	{
 		if (Ar.CustomVer(FEnterpriseObjectVersion::GUID) < FEnterpriseObjectVersion::BookmarkExtensibilityUpgrade)
 		{
-			UBookmarkBase** LocalBookmarks = reinterpret_cast<UBookmarkBase**>(static_cast<UBookMark**>(ToRawPtrArrayUnsafe(BookMarks))); //-V777
+			const TObjectPtr<UBookmarkBase>* LocalBookmarks =
+				reinterpret_cast<TObjectPtr<UBookmarkBase>*>(BookMarks);
 			const int32 NumBookmarks = sizeof(BookMarks) / sizeof(UBookMark*);
-			BookmarkArray = TArray<UBookmarkBase*>(LocalBookmarks, NumBookmarks);
+			BookmarkArray = TArray<TObjectPtr<UBookmarkBase>>(LocalBookmarks, NumBookmarks);
 			AdjustNumberOfBookmarks();
 		}
 
diff --git a/Engine/Source/Runtime/Landscape/Private/Landscape.cpp b/Engine/Source/Runtime/Landscape/Private/Landscape.cpp
--- a/Engine/Source/Runtime/Landscape/Private/Landscape.cpp
+++ b/Engine/Source/Runtime/Landscape/Private/Landscape.cpp
@@ -650,8 +650,8 @@ void ULandscapeComponent::Serialize(FArchive& Ar)
 		{
 			// These are used for SM5 rendering
 			UTexture2D* BackupXYOffsetmapTexture = nullptr;
-			TArray<UMaterialInstanceConstant*> BackupMaterialInstances;
-			TArray<UTexture2D*> BackupWeightmapTextures;
+			TArray<TObjectPtr<UMaterialInstanceConstant>> BackupMaterialInstances;
+			TArray<TObjectPtr<UTexture2D>> BackupWeightmapTextures;
 
 			Exchange(BackupXYOffsetmapTexture, XYOffsetmapTexture);
 			Exchange(BackupMaterialInstances, MaterialInstances);
@@ -666,8 +666,8 @@ void ULandscapeComponent::Serialize(FArchive& Ar)
 		else if (!Ar.CookingTarget()->SupportsFeature(ETargetPlatformFeatures::MobileRendering))
 		{
 			// These properties are only for Mobile
-			TArray<UMaterialInterface*> BackupMobileMaterialInterfaces;
-			TArray<UTexture2D*> BackupMobileWeightmapTextures;
+			TArray<TObjectPtr<UMaterialInterface>> BackupMobileMaterialInterfaces;
+			TArray<TObjectPtr<UTexture2D>> BackupMobileWeightmapTextures;
 
 			Exchange(MobileMaterialInterfaces, BackupMobileMaterialInterfaces);
 			Exchange(MobileWeightmapTextures, BackupMobileWeightmapTextures);
