OnlyEditorOnly support: Add scopes around the load of undeclared dependencies so that the cooker knows whether they need to be added to the cook.
* Add more uses of the existing FSoftObjectPathSerializationScope, UE_TRACK_REFERENCING_PACKAGE_SCOPED , and FSerializedPropertyScope.
* Add a new FCookLoadScope which is similar to FSoftObjectPathSerializationScope but applies to undeclared loads rather than to SoftObjectPath references.
* Tunnel the scopes through to the eventual package creation when asynchronous loads are triggered through StreamingManager or AsyncLoading.
#rnx

diff --git a/Engine/Source/Editor/UnrealEd/Private/CookOnTheFlyServer.cpp b/Engine/Source/Editor/UnrealEd/Private/CookOnTheFlyServer.cpp
--- a/Engine/Source/Editor/UnrealEd/Private/CookOnTheFlyServer.cpp
+++ b/Engine/Source/Editor/UnrealEd/Private/CookOnTheFlyServer.cpp
@@ -6636,6 +6636,9 @@ TArray<UE::Cook::FInstigator> UCookOnTheFlyServer::GetInstigatorChain(FName Pack
 			case EInstigator::HardDependency: bGetNext = true; break;
 			case EInstigator::SoftDependency: bGetNext = true; break;
 			case EInstigator::Unsolicited: bGetNext = true; break;
+			case EInstigator::EditorOnlyLoad: bGetNext = true; break;
+			case EInstigator::SaveTimeHardDependency: bGetNext = true; break;
+			case EInstigator::SaveTimeSoftDependency: bGetNext = true; break;
 			case EInstigator::GeneratedPackage: bGetNext = true; break;
 			default: break;
 		}
diff --git a/Engine/Source/Runtime/CoreUObject/Private/Blueprint/BlueprintSupport.cpp b/Engine/Source/Runtime/CoreUObject/Private/Blueprint/BlueprintSupport.cpp
--- a/Engine/Source/Runtime/CoreUObject/Private/Blueprint/BlueprintSupport.cpp
+++ b/Engine/Source/Runtime/CoreUObject/Private/Blueprint/BlueprintSupport.cpp
@@ -13,6 +13,7 @@
 #include "Templates/Casts.h"
 #include "UObject/UnrealType.h"
 #include "Serialization/DuplicatedDataWriter.h"
+#include "Misc/PackageAccessTrackingOps.h"
 #include "Misc/PackageName.h"
 #include "UObject/ObjectResource.h"
 #include "UObject/GCObject.h"
@@ -1758,6 +1759,7 @@ void FLinkerLoad::PRIVATE_ForceLoadAllDependencies(UPackage* Package)
 
 void FLinkerLoad::ResolveAllImports()
 {
+	UE_TRACK_REFERENCING_PACKAGE_SCOPED(LinkerRoot->GetFName(), PackageAccessTrackingOps::NAME_Load);
 	for (int32 ImportIndex = 0; ImportIndex < ImportMap.Num() && IsBlueprintFinalizationPending(); ++ImportIndex)
 	{
 		// first, make sure every import object is available... just because 
diff --git a/Engine/Source/Runtime/CoreUObject/Private/Serialization/AsyncLoading.cpp b/Engine/Source/Runtime/CoreUObject/Private/Serialization/AsyncLoading.cpp
--- a/Engine/Source/Runtime/CoreUObject/Private/Serialization/AsyncLoading.cpp
+++ b/Engine/Source/Runtime/CoreUObject/Private/Serialization/AsyncLoading.cpp
@@ -19,6 +19,7 @@
 #include "Misc/CommandLine.h"
 #include "Misc/App.h"
 #include "Misc/MessageDialog.h"
+#include "Misc/PackageAccessTrackingOps.h"
 #include "Misc/PackageName.h"
 #include "Misc/ScopedSlowTask.h"
 #include "Misc/TrackedActivity.h"
@@ -6126,7 +6127,15 @@ EAsyncPackageState::Type FAsyncPackage::CreateLinker()
 		{
 			SCOPED_LOADTIMER(CreateLinker_CreatePackage);
 			FGCScopeGuard GCGuard;
-			Package = CreatePackage(*Desc.Name.ToString());
+			{
+#if UE_WITH_PACKAGE_ACCESS_TRACKING
+				UE_TRACK_REFERENCING_PACKAGE_SCOPED(Desc.ReferencerPackageName, Desc.ReferencerPackageOp);
+#endif
+#if WITH_EDITOR
+				FCookLoadScope CookLoadScope(Desc.CookLoadType);
+#endif
+				Package = CreatePackage(*Desc.Name.ToString());
+			}
 			if (!Package)
 			{
 				UE_LOG(LogStreaming, Error, TEXT("Failed to create package %s requested by async loading code. NameToLoad: %s"), *Desc.Name.ToString(), *Desc.PackagePath.GetDebugName());
@@ -6382,6 +6391,13 @@ void FAsyncPackage::AddImportDependency(const FName& PendingImport, const FName&
 	{
 		FPackagePath ImportPackagePath = FPackagePath::FromPackageNameChecked(!PackageToLoad.IsNone() ? PackageToLoad : PendingImport);
 		FAsyncPackageDesc Info(INDEX_NONE, PendingImport, ImportPackagePath);
+#if UE_WITH_PACKAGE_ACCESS_TRACKING
+		Info.ReferencerPackageName = GetPackageName();
+		Info.ReferencerPackageOp = PackageAccessTrackingOps::NAME_Load;
+#endif
+#if WITH_EDITOR
+		Info.CookLoadType = ECookLoadType::Unexpected;
+#endif
 		Info.SetInstancingContext(MoveTemp(InstancingContext));
 		PackageToStream = new FAsyncPackage(AsyncLoadingThread, Info, EDLBootNotificationManager);
 
@@ -7449,6 +7465,18 @@ int32 FAsyncLoadingThread::LoadPackage(const FPackagePath& InPackagePath, FName
 	{
 		PackageDesc.SetInstancingContext(*InInstancingContext);
 	}
+#if UE_WITH_PACKAGE_ACCESS_TRACKING
+	PackageAccessTracking_Private::FTrackedData* AccumulatedScopeData = PackageAccessTracking_Private::FPackageAccessRefScope::GetCurrentThreadAccumulatedData();
+	if (AccumulatedScopeData)
+	{
+		PackageDesc.ReferencerPackageName = AccumulatedScopeData->PackageName;
+		PackageDesc.ReferencerPackageOp = AccumulatedScopeData->OpName;
+	}
+#endif
+#if WITH_EDITOR
+	PackageDesc.CookLoadType = FCookLoadScope::GetCurrentValue();
+#endif
+
 	QueuePackage(PackageDesc);
 	return RequestID;
 }
diff --git a/Engine/Source/Runtime/CoreUObject/Private/Serialization/AsyncPackage.h b/Engine/Source/Runtime/CoreUObject/Private/Serialization/AsyncPackage.h
--- a/Engine/Source/Runtime/CoreUObject/Private/Serialization/AsyncPackage.h
+++ b/Engine/Source/Runtime/CoreUObject/Private/Serialization/AsyncPackage.h
@@ -9,7 +9,9 @@
 #include "CoreMinimal.h"
 #include "UObject/UObjectGlobals.h"
 #include "Misc/PackagePath.h"
+#include "Misc/PackageAccessTracking.h"
 #include "Templates/UniquePtr.h"
+#include "UObject/ICookInfo.h"
 #include "UObject/LinkerInstancingContext.h"
 
 struct FAsyncPackageDesc
@@ -32,7 +34,16 @@ struct FAsyncPackageDesc
 #if WITH_EDITORONLY_DATA
 	/** Instancing context, maps original package to their instanced counterpart, used to remap imports. */
 	FLinkerInstancingContext InstancingContext;
+#endif
+#if UE_WITH_PACKAGE_ACCESS_TRACKING
+	FName ReferencerPackageName;
+	FName ReferencerPackageOp;
+#endif
+#if WITH_EDITOR
+	ECookLoadType CookLoadType;
+#endif
 
+#if WITH_EDITORONLY_DATA
 	const FLinkerInstancingContext* GetInstancingContext() const { return &InstancingContext; }
 	void SetInstancingContext(FLinkerInstancingContext InInstancingContext) { InstancingContext = MoveTemp(InInstancingContext); }
 #else
@@ -48,6 +59,9 @@ struct FAsyncPackageDesc
 		, PackageFlags(InPackageFlags)
 		, Priority(InPriority)
 		, PIEInstanceID(InPIEInstanceID)
+#if WITH_EDITOR
+		, CookLoadType(ECookLoadType::Unexpected)
+#endif
 	{
 		check(!PackagePath.IsEmpty());
 	}
@@ -62,6 +76,13 @@ struct FAsyncPackageDesc
 		, PIEInstanceID(OldPackage.PIEInstanceID)
 #if WITH_EDITORONLY_DATA
 		, InstancingContext(OldPackage.InstancingContext)
+#endif
+#if UE_WITH_PACKAGE_ACCESS_TRACKING
+		, ReferencerPackageName(OldPackage.ReferencerPackageName)
+		, ReferencerPackageOp(OldPackage.ReferencerPackageOp)
+#endif
+#if WITH_EDITOR
+		, CookLoadType(OldPackage.CookLoadType)
 #endif
 	{
 	}
diff --git a/Engine/Source/Runtime/CoreUObject/Private/UObject/ICookInfo.cpp b/Engine/Source/Runtime/CoreUObject/Private/UObject/ICookInfo.cpp
--- a/Engine/Source/Runtime/CoreUObject/Private/UObject/ICookInfo.cpp
+++ b/Engine/Source/Runtime/CoreUObject/Private/UObject/ICookInfo.cpp
@@ -47,4 +47,22 @@ FString FInstigator::ToString() const
 
 }
 
+static thread_local ECookLoadType GCookLoadType = ECookLoadType::Unexpected;
+
+FCookLoadScope::FCookLoadScope(ECookLoadType ScopeType)
+	: PreviousScope(GCookLoadType)
+{
+	GCookLoadType = ScopeType;
+}
+
+FCookLoadScope::~FCookLoadScope()
+{
+	GCookLoadType = PreviousScope;
+}
+
+ECookLoadType FCookLoadScope::GetCurrentValue()
+{
+	return GCookLoadType;
+}
+
 #endif
diff --git a/Engine/Source/Runtime/CoreUObject/Private/UObject/LinkerLoad.cpp b/Engine/Source/Runtime/CoreUObject/Private/UObject/LinkerLoad.cpp
--- a/Engine/Source/Runtime/CoreUObject/Private/UObject/LinkerLoad.cpp
+++ b/Engine/Source/Runtime/CoreUObject/Private/UObject/LinkerLoad.cpp
@@ -22,6 +22,7 @@
 #include "UObject/UObjectHash.h"
 #include "Misc/PackageName.h"
 #include "Blueprint/BlueprintSupport.h"
+#include "Misc/PackageAccessTrackingOps.h"
 #include "Misc/PathViews.h"
 #include "Misc/PreloadableFile.h"
 #include "Misc/SecureHash.h"
@@ -54,6 +55,7 @@
 #include "Serialization/EditorBulkData.h"
 #include "HAL/FileManager.h"
 #include "UObject/CoreRedirects.h"
+#include "UObject/ICookInfo.h"
 #include "UObject/PackageRelocation.h"
 #include "Misc/AssetRegistryInterface.h"
 #include "Misc/StringBuilder.h"
@@ -2790,6 +2792,8 @@ void FLinkerLoad::Verify()
 				SlowTask.Emplace(Summary.ImportCount, LoadingImportsText);
 			}
 #endif
+			UE_TRACK_REFERENCING_PACKAGE_SCOPED(LinkerRoot->GetFName(), PackageAccessTrackingOps::NAME_Load);
+
 			// Validate all imports and map them to their remote linkers.
 			for (int32 ImportIndex = 0; ImportIndex < Summary.ImportCount; ImportIndex++)
 			{
@@ -3366,6 +3370,10 @@ bool FLinkerLoad::VerifyImportInner(const int32 ImportIndex, FString& WarningSuf
 				LocalInstancingContext = &GetInstancingContext();
 			}
 			FPackagePath PackagePathToLoad = FPackagePath::FromPackageNameChecked(PackageToLoad);
+#if WITH_EDITOR
+			FCookLoadScope ResetCookLoadScopeToUnexpected(ECookLoadType::Unexpected);
+#endif
+
 			Package = LoadPackageInternal(Package, PackagePathToLoad, InternalLoadFlags | LOAD_IsVerifying, this, nullptr /* InReaderOverride */,
 				LocalInstancingContext, nullptr /* DiffPackagePath */);
 		}
diff --git a/Engine/Source/Runtime/CoreUObject/Private/UObject/PropertySoftObjectPtr.cpp b/Engine/Source/Runtime/CoreUObject/Private/UObject/PropertySoftObjectPtr.cpp
--- a/Engine/Source/Runtime/CoreUObject/Private/UObject/PropertySoftObjectPtr.cpp
+++ b/Engine/Source/Runtime/CoreUObject/Private/UObject/PropertySoftObjectPtr.cpp
@@ -206,6 +206,7 @@ EConvertFromTypeResult FSoftObjectProperty::ConvertFromType(const FPropertyTag&
 		FSoftObjectPtr* PropertyValue = GetPropertyValuePtr_InContainer(Data, Tag.ArrayIndex);
 		check(PropertyValue);
 
+		FSerializedPropertyScope SerializedProperty(Archive, this);
 		return PropertyValue->GetUniqueID().SerializeFromMismatchedTag(Tag, Slot) ? EConvertFromTypeResult::Converted : EConvertFromTypeResult::UseSerializeItem;
 	}
 	else if (Tag.Type == NAME_StructProperty && (Tag.StructName == NAME_SoftObjectPath || Tag.StructName == NAME_SoftClassPath || Tag.StructName == NAME_StringAssetReference || Tag.StructName == NAME_StringClassReference))
@@ -213,6 +214,7 @@ EConvertFromTypeResult FSoftObjectProperty::ConvertFromType(const FPropertyTag&
 		// This property used to be a FSoftObjectPath but is now a TSoftObjectPtr<Foo>
 		FSoftObjectPath PreviousValue;
 		// explicitly call Serialize to ensure that the various delegates needed for cooking are fired
+		FSerializedPropertyScope SerializedProperty(Archive, this);
 		PreviousValue.Serialize(Slot);
 
 		// now copy the value into the object's address space
diff --git a/Engine/Source/Runtime/CoreUObject/Public/UObject/ICookInfo.h b/Engine/Source/Runtime/CoreUObject/Public/UObject/ICookInfo.h
--- a/Engine/Source/Runtime/CoreUObject/Public/UObject/ICookInfo.h
+++ b/Engine/Source/Runtime/CoreUObject/Public/UObject/ICookInfo.h
@@ -49,6 +49,9 @@ namespace UE::Cook
 	CallbackMacro(HardDependency, false) \
 	CallbackMacro(SoftDependency, false) \
 	CallbackMacro(Unsolicited, false) \
+	CallbackMacro(EditorOnlyLoad, false) \
+	CallbackMacro(SaveTimeHardDependency, false) \
+	CallbackMacro(SaveTimeSoftDependency, false) \
 	CallbackMacro(GeneratedPackage, false) \
 
 
@@ -96,4 +99,30 @@ public:
 
 }
 
+/**
+ * A scope around loads when cooking that indicates whether the loaded package is needed in game or not.
+ * The default is Unexpected. Declare an FCookLoadScope to set the value.
+ * Packages loads that are declared via AssetRegistry dependencies (imports or softobjectpaths in
+ * the editor package) change the default to UsedInGame or EditorOnly dependending on whether they were saved
+ * in an EditorOnly property or scope.
+ */
+enum class ECookLoadType : uint8
+{
+	Unexpected,
+	EditorOnly,
+	UsedInGame,
+};
+
+/** Set the ECookLoadType value in the current scope.*/
+struct FCookLoadScope
+{
+	COREUOBJECT_API explicit FCookLoadScope(ECookLoadType ScopeType);
+	COREUOBJECT_API ~FCookLoadScope();
+
+	COREUOBJECT_API static ECookLoadType GetCurrentValue();
+
+private:
+	ECookLoadType PreviousScope;
+};
+
 #endif
\ No newline at end of file
diff --git a/Engine/Source/Runtime/CoreUObject/Public/UObject/SoftObjectPath.h b/Engine/Source/Runtime/CoreUObject/Public/UObject/SoftObjectPath.h
--- a/Engine/Source/Runtime/CoreUObject/Public/UObject/SoftObjectPath.h
+++ b/Engine/Source/Runtime/CoreUObject/Public/UObject/SoftObjectPath.h
@@ -513,6 +513,11 @@ struct FSoftObjectPathSerializationScope
 		FSoftObjectPathThreadContext::Get().OptionStack.Emplace(SerializingPackageName, SerializingPropertyName, CollectType, SerializeType);
 	}
 
+	explicit FSoftObjectPathSerializationScope(ESoftObjectPathCollectType CollectType)
+	{
+		FSoftObjectPathThreadContext::Get().OptionStack.Emplace(NAME_None, NAME_None, CollectType, ESoftObjectPathSerializeType::AlwaysSerialize);
+	}
+
 	~FSoftObjectPathSerializationScope()
 	{
 		FSoftObjectPathThreadContext::Get().OptionStack.Pop();
diff --git a/Engine/Source/Runtime/Engine/Classes/Engine/StreamableManager.h b/Engine/Source/Runtime/Engine/Classes/Engine/StreamableManager.h
--- a/Engine/Source/Runtime/Engine/Classes/Engine/StreamableManager.h
+++ b/Engine/Source/Runtime/Engine/Classes/Engine/StreamableManager.h
@@ -5,6 +5,7 @@
 #include "CoreMinimal.h"
 #include "Containers/ArrayView.h"
 #include "Math/NumericLimits.h"
+#include "Misc/PackageAccessTracking.h"
 #include "Templates/Casts.h"
 #include "Templates/Function.h"
 #include "UObject/Class.h"
@@ -16,6 +17,9 @@
 DECLARE_DELEGATE(FStreamableDelegate);
 DECLARE_DELEGATE_OneParam(FStreamableUpdateDelegate, TSharedRef<struct FStreamableHandle>);
 
+#if WITH_EDITOR
+enum class ECookLoadType : uint8;
+#endif
 
 /** Storage class of the per-class id used in FStreamableHandleContextDataBase */
 typedef uint8 TStreamableHandleContextDataTypeID;
@@ -135,6 +139,23 @@ struct ENGINE_API FStreamableHandle : public TSharedFromThis<FStreamableHandle>
 		return DebugName;
 	}
 
+#if UE_WITH_PACKAGE_ACCESS_TRACKING
+	FName GetReferencerPackage() const
+	{
+		return ReferencerPackage;
+	}
+	FName GetRefencerPackageOp() const
+	{
+		return ReferencerPackageOp;
+	}
+#endif
+#if WITH_EDITOR
+	ECookLoadType GetCookLoadType() const
+	{
+		return CookLoadType;
+	}
+#endif
+
 	/** Returns the streaming priority. */
 	TAsyncLoadPriority GetPriority() const
 	{
@@ -415,24 +436,22 @@ private:
 	/** Backpointer to handles that depend on this */
 	TArray<TWeakPtr<FStreamableHandle> > ParentHandles;
 
+#if UE_WITH_PACKAGE_ACCESS_TRACKING
+	FName ReferencerPackage;
+	FName ReferencerPackageOp;
+#endif
+
 	/** This is set at the time of creation, and will be cleared when request completes or is canceled */
 	struct FStreamableManager* OwningManager;
 
 	/** Array of contextual data added by game/engine code */
 	TArray<TSharedPtr<FStreamableHandleContextDataBase>> AdditionalContextData;
 
-	FStreamableHandle()
-		: bLoadCompleted(false)
-		, bReleased(false)
-		, bCanceled(false)
-		, bStalled(false)
-		, bReleaseWhenLoaded(false)
-		, bIsCombinedHandle(false)
-		, Priority(0)
-		, StreamablesLoading(0)
-		, OwningManager(nullptr)
-	{
-	}
+#if WITH_EDITOR
+	ECookLoadType CookLoadType;
+#endif
+
+	FStreamableHandle();
 };
 
 enum class EStreamableManagerCombinedHandleOptions : uint8 
diff --git a/Engine/Source/Runtime/Engine/Classes/Engine/World.h b/Engine/Source/Runtime/Engine/Classes/Engine/World.h
--- a/Engine/Source/Runtime/Engine/Classes/Engine/World.h
+++ b/Engine/Source/Runtime/Engine/Classes/Engine/World.h
@@ -4198,6 +4198,9 @@ public:
 	DECLARE_MULTICAST_DELEGATE_OneParam(FOnWorldPIEEnded, UGameInstance*);
 	static FOnWorldPIEEnded OnPIEEnded;
 
+	// Callback to add references when Serialize is called during SavePackage
+	DECLARE_MULTICAST_DELEGATE_TwoParams(FWorldCollectSaveReferencesEvent, UWorld*, FArchive&);
+	static FWorldCollectSaveReferencesEvent OnCollectSaveReferences;
 #endif // WITH_EDITOR
 
 	// Post duplication event.
diff --git a/Engine/Source/Runtime/Engine/Private/StreamableManager.cpp b/Engine/Source/Runtime/Engine/Private/StreamableManager.cpp
--- a/Engine/Source/Runtime/Engine/Private/StreamableManager.cpp
+++ b/Engine/Source/Runtime/Engine/Private/StreamableManager.cpp
@@ -1,8 +1,10 @@
 // Copyright Epic Games, Inc. All Rights Reserved.
 
 #include "Engine/StreamableManager.h"
+#include "UObject/ICookInfo.h"
 #include "UObject/ObjectRedirector.h"
 #include "UObject/Package.h"
+#include "Misc/PackageAccessTrackingOps.h"
 #include "Misc/PackageName.h"
 #include "UObject/UObjectThreadContext.h"
 #include "HAL/IConsoleManager.h"
@@ -319,6 +321,22 @@ TStreamableHandleContextDataTypeID FStreamableHandleContextDataBase::AllocateCla
 }
 
 
+FStreamableHandle::FStreamableHandle()
+	: bLoadCompleted(false)
+	, bReleased(false)
+	, bCanceled(false)
+	, bStalled(false)
+	, bReleaseWhenLoaded(false)
+	, bIsCombinedHandle(false)
+	, Priority(0)
+	, StreamablesLoading(0)
+	, OwningManager(nullptr)
+#if WITH_EDITOR
+	, CookLoadType(ECookLoadType::Unexpected)
+#endif
+{
+}
+
 bool FStreamableHandle::BindCompleteDelegate(FStreamableDelegate NewDelegate)
 {
 	if (!IsLoadingInProgress())
@@ -1297,6 +1315,12 @@ FStreamable* FStreamableManager::StreamInternal(const FSoftObjectPath& InTargetN
 			{
 				Existing->bLoadFailed = false;
 				Existing->bAsyncLoadRequestOutstanding = true;
+#if UE_WITH_PACKAGE_ACCESS_TRACKING
+				UE_TRACK_REFERENCING_PACKAGE_SCOPED(Handle->GetReferencerPackage(), Handle->GetRefencerPackageOp());
+#endif
+#if WITH_EDITOR
+				FCookLoadScope CookLoadScope(Handle->GetCookLoadType());
+#endif
 				LoadPackageAsync(PackagePath,
 					NAME_None /* PackageNameToCreate */,
 					FLoadPackageAsyncDelegate::CreateSP(Handle, &FStreamableHandle::AsyncLoadCallbackWrapper, TargetName),
@@ -1322,6 +1346,17 @@ TSharedPtr<FStreamableHandle> FStreamableManager::RequestAsyncLoad(TArray<FSoftO
 	NewRequest->DebugName = MoveTemp(DebugName);
 #endif
 	NewRequest->Priority = Priority;
+#if UE_WITH_PACKAGE_ACCESS_TRACKING
+	PackageAccessTracking_Private::FTrackedData* AccumulatedScopeData = PackageAccessTracking_Private::FPackageAccessRefScope::GetCurrentThreadAccumulatedData();
+	if (AccumulatedScopeData)
+	{
+		NewRequest->ReferencerPackage = AccumulatedScopeData->PackageName;
+		NewRequest->ReferencerPackageOp = AccumulatedScopeData->OpName;
+	}
+#endif
+#if WITH_EDITOR
+	NewRequest->CookLoadType = FCookLoadScope::GetCurrentValue();
+#endif
 
 	int32 NumValidRequests = NewRequest->RequestedAssets.Num();
 	
diff --git a/Engine/Source/Runtime/Engine/Private/World.cpp b/Engine/Source/Runtime/Engine/Private/World.cpp
--- a/Engine/Source/Runtime/Engine/Private/World.cpp
+++ b/Engine/Source/Runtime/Engine/Private/World.cpp
@@ -543,6 +543,7 @@ FWorldDelegates::FWorldInitializationEvent FWorldDelegates::OnPostWorldInitializ
 FWorldDelegates::FWorldPreRenameEvent FWorldDelegates::OnPreWorldRename;
 FWorldDelegates::FWorldPostRenameEvent FWorldDelegates::OnPostWorldRename;
 FWorldDelegates::FWorldCurrentLevelChangedEvent FWorldDelegates::OnCurrentLevelChanged;
+FWorldDelegates::FWorldCollectSaveReferencesEvent FWorldDelegates::OnCollectSaveReferences;
 #endif // WITH_EDITOR
 FWorldDelegates::FWorldPostDuplicateEvent FWorldDelegates::OnPostDuplicate;
 FWorldDelegates::FWorldCleanupEvent FWorldDelegates::OnWorldCleanup;
@@ -702,6 +703,11 @@ void UWorld::Serialize( FArchive& Ar )
 		PersistedStreamingLevels.Reserve(StreamingLevels.Num());
 		Algo::CopyIf(StreamingLevels, PersistedStreamingLevels, [&](ULevelStreaming* LevelStreaming) { return LevelStreaming && !LevelStreaming->HasAnyFlags(RF_Transient); });
 		Ar << PersistedStreamingLevels;
+
+		if (Ar.IsObjectReferenceCollector())
+		{
+			FWorldDelegates::OnCollectSaveReferences.Broadcast(this, Ar);
+		}
 	}
 	else
 #endif
