- Fixed the root cause of GConfig and DefaultConfig.ini getting out of sync

diff --git a/Engine/Source/Runtime/Core/Private/Misc/ConfigContext.cpp b/Engine/Source/Runtime/Core/Private/Misc/ConfigContext.cpp
--- a/Engine/Source/Runtime/Core/Private/Misc/ConfigContext.cpp
+++ b/Engine/Source/Runtime/Core/Private/Misc/ConfigContext.cpp
@@ -653,6 +653,15 @@ bool FConfigContext::GenerateDestIniFile()
 
 	if (!FPlatformProperties::RequiresCookedData() || bAllowGeneratedIniWhenCooked)
 	{
+		if (bForceReload)
+		{
+			// if we are reloading from disk (probably to update GConfig's understanding after updating a Default ini file), we need to 
+			// flush any pending updates, replace the in memory version with the hierarchy, and then read the flushed file back on  
+			// this will make sure the in-memory version has what is on disk, so when Flush happens later, it does not write out an outdated value
+			// note: we only want to copy the TMap base class slice of the FConfigFile, none of the specific members of the FConfigFile class itself
+			ConfigFile->TMap<FString, FConfigSection>::operator=(*ConfigFile->SourceConfigFile);
+		}
+
 		LoadAnIniFile(*DestIniFilename, *ConfigFile);
 	}
 
diff --git a/Engine/Source/Runtime/CoreUObject/Private/UObject/Obj.cpp b/Engine/Source/Runtime/CoreUObject/Private/UObject/Obj.cpp
--- a/Engine/Source/Runtime/CoreUObject/Private/UObject/Obj.cpp
+++ b/Engine/Source/Runtime/CoreUObject/Private/UObject/Obj.cpp
@@ -3133,15 +3133,27 @@ void UObject::UpdateSingleSectionOfConfigFile(const FString& ConfigIniName)
 
 	// do we need to use a special platform hierarchy?
 	FString OverridePlatform = GetFinalOverridePlatform(this);
+	bool bUpdateGConfig = OverridePlatform.Len() == 0;
+
+	// if we are going to reload GConfig, we need to flush any pending writes to disk because we are going to
+	// blow away the contents of GConfig's in memory version with what it on disk, but we don't want to lose the
+	// local modifications
+	if (bUpdateGConfig)
+	{
+		GConfig->Flush(false, *GetClass()->ClassConfigName.ToString());
+	}
 
 	// make sure SaveConfig wrote only to the file we expected
 	NewFile.UpdateSections(*ConfigIniName, *GetClass()->ClassConfigName.ToString(), OverridePlatform.Len() ? *OverridePlatform : nullptr);
 
 	// reload the file, so that it refresh the cache internally, unless a non-standard platform was used,
 	// then we don't want to touch GConfig
-	if (OverridePlatform.Len() == 0)
+	if (bUpdateGConfig)
 	{
-		FConfigContext::ForceReloadIntoGConfig().Load(*GetClass()->ClassConfigName.ToString());
+		FConfigContext Context = FConfigContext::ForceReloadIntoGConfig();
+		// don't write the Saved out, as we just finished writing what we needed
+		Context.bWriteDestIni = false;
+		Context.Load(*GetClass()->ClassConfigName.ToString());
 	}
 }
 
@@ -3199,7 +3211,7 @@ void UObject::UpdateSinglePropertyInConfigFile(const FProperty* InProperty, cons
 
 		const FString SectionName = Keys[0];
 		FString PropertyKey = InProperty->GetFName().ToString();
-		
+
 #if WITH_EDITOR
 		static FName ConsoleVariableFName(TEXT("ConsoleVariable"));
 		const FString& CVarName = InProperty->GetMetaData(ConsoleVariableFName);
@@ -3210,13 +3222,21 @@ void UObject::UpdateSinglePropertyInConfigFile(const FProperty* InProperty, cons
 #endif // #if WITH_EDITOR
 
 		// do we need to use a special platform hierarchy?
-		FString OverridePlatform = GetFinalOverridePlatform(this);
+		bool bUpdateGConfig = GetFinalOverridePlatform(this).Len() == 0;
+
+		// if we are going to reload GConfig, we need to flush any pending writes to disk because we are going to
+		// blow away the contents of GConfig's in memory version with what it on disk, but we don't want to lose the
+		// local modifications
+		if (bUpdateGConfig)
+		{
+			GConfig->Flush(false, *GetClass()->ClassConfigName.ToString());
+		}
 
 		NewFile.UpdateSinglePropertyInSection(*InConfigIniName, *PropertyKey, *SectionName);
 
 		// reload the file, so that it refresh the cache internally, unless a non-standard platform was used,
 		// then we don't want to touch GConfig
-		if (OverridePlatform.Len() == 0)
+		if (bUpdateGConfig)
 		{
 			FConfigContext Context = FConfigContext::ForceReloadIntoGConfig();
 			// don't write the Saved out, as we just finished writing what we needed
