Add linker support to skip loading imports that are remapped to none from the instancing context.
This can be used to partially load an instanced package based on the caller's instancing context.

diff --git a/Engine/Source/Runtime/CoreUObject/Private/Serialization/AsyncLoading.cpp b/Engine/Source/Runtime/CoreUObject/Private/Serialization/AsyncLoading.cpp
--- a/Engine/Source/Runtime/CoreUObject/Private/Serialization/AsyncLoading.cpp
+++ b/Engine/Source/Runtime/CoreUObject/Private/Serialization/AsyncLoading.cpp
@@ -6535,8 +6535,9 @@ EAsyncPackageState::Type FAsyncPackage::LoadImports()
 		FName ImportToLoad = !Import->HasPackageName() ? Import->ObjectName : Import->GetPackageName();
 		FName ImportPackageFName = InstancingContext.RemapPackage(ImportToLoad);
 
-		// Don't try to import a package that is in an import table that we know is an invalid entry
-		if (FLinkerLoad::IsKnownMissingPackage(ImportPackageFName))
+		// Don't try to import a package that is in an import table that we know is an invalid entry or if it was filtered out
+		// by the instancing context.
+		if ((ImportPackageFName.IsNone() && !ImportToLoad.IsNone()) || FLinkerLoad::IsKnownMissingPackage(ImportPackageFName))
 		{
 			continue;
 		}
diff --git a/Engine/Source/Runtime/CoreUObject/Private/Serialization/AsyncLoading2.cpp b/Engine/Source/Runtime/CoreUObject/Private/Serialization/AsyncLoading2.cpp
--- a/Engine/Source/Runtime/CoreUObject/Private/Serialization/AsyncLoading2.cpp
+++ b/Engine/Source/Runtime/CoreUObject/Private/Serialization/AsyncLoading2.cpp
@@ -4611,6 +4611,12 @@ void FAsyncPackage2::ImportPackagesRecursiveInner(FAsyncLoadingThreadState2& Thr
 			ImportedPackageUPackageName = InstancingContext.RemapPackage(ImportedPackageNameToLoad);
 			if (ImportedPackageUPackageName != ImportedPackageNameToLoad)
 			{
+				if (ImportedPackageUPackageName.IsNone())
+				{
+					ImportedPackageIdToLoad = FPackageId::FromName(NAME_None);
+					ImportedPackageNameToLoad = NAME_None;
+				}
+
 				ImportedPackageId = FPackageId::FromName(ImportedPackageUPackageName);
 				// Rewrite the import table
 				ImportedPackageIds[LocalImportedPackageIndex] = ImportedPackageId;
@@ -4737,7 +4743,7 @@ void FAsyncPackage2::ImportPackagesRecursiveInner(FAsyncLoadingThreadState2& Thr
 		{
 			if (!ImportedPackageRef.GetPackage()) // If we found a package it's not actually missing but we can't load it anyway
 			{
-				UE_ASYNC_PACKAGE_LOG(Display, Desc, TEXT("ImportPackages: SkipPackage"),
+				UE_ASYNC_PACKAGE_CLOG(!ImportedPackageUPackageName.IsNone(), Display, Desc, TEXT("ImportPackages: SkipPackage"),
 					TEXT("Skipping non mounted imported package %s (0x%llX)"), *ImportedPackageNameToLoad.ToString(), ImportedPackageId.Value());
 				ImportedPackageRef.SetIsMissingPackage();
 			}
diff --git a/Engine/Source/Runtime/CoreUObject/Private/UObject/LinkerLoad.cpp b/Engine/Source/Runtime/CoreUObject/Private/UObject/LinkerLoad.cpp
--- a/Engine/Source/Runtime/CoreUObject/Private/UObject/LinkerLoad.cpp
+++ b/Engine/Source/Runtime/CoreUObject/Private/UObject/LinkerLoad.cpp
@@ -3479,6 +3479,18 @@ bool FLinkerLoad::VerifyImportInner(const int32 ImportIndex, FString& WarningSuf
 		return false;
 	}
 
+	if (Import.HasPackageName() || Import.OuterIndex.IsNull())
+	{
+		FName PackageToLoad = !Import.HasPackageName() ? Import.ObjectName : Import.GetPackageName();
+		FName PackageToLoadInto = InstancingContext.RemapPackage(PackageToLoad);
+
+		if (!PackageToLoad.IsNone() && PackageToLoadInto.IsNone())
+		{
+			// Import package was filtered out by instancing context
+			return false;
+		}
+	}
+
 	// Build the import object name on the stack and only once to avoid string temporaries
 	TStringBuilder<256> ImportObjectName;
 	Import.ObjectName.AppendString(ImportObjectName);
