UBlueprint::PostLoadAssetRegistryTags: When fixing up short path names, the short path name may no longer exist and have been redirected (e.g. GameplayAbilityBlueprintGeneratedClass). In that case,clear the tag; it will need to be recreated by recompiling the blueprint, which will happen on PostLoad.
Keeping the invalid shortname will trigger an assertion when saving the assetregistry if the asset is not loaded during the cook.
#rnx

diff --git a/Engine/Source/Runtime/AssetRegistry/Private/AssetRegistry.cpp b/Engine/Source/Runtime/AssetRegistry/Private/AssetRegistry.cpp
--- a/Engine/Source/Runtime/AssetRegistry/Private/AssetRegistry.cpp
+++ b/Engine/Source/Runtime/AssetRegistry/Private/AssetRegistry.cpp
@@ -4765,8 +4765,15 @@ void FAssetRegistryImpl::PostLoadAssetRegistryTags(FAssetData* AssetData)
 					FAssetDataTagMap TagsAndValues = AssetData->TagsAndValues.CopyMap();
 					for (const UObject::FAssetRegistryTag& Tag : TagsToModify)
 					{
-						FString& Value = TagsAndValues.FindOrAdd(Tag.Name);
-						Value = Tag.Value;
+						if (!Tag.Value.IsEmpty())
+						{
+							FString& Value = TagsAndValues.FindOrAdd(Tag.Name);
+							Value = Tag.Value;
+						}
+						else
+						{
+							TagsAndValues.Remove(Tag.Name);
+						}
 					}
 					AssetData->TagsAndValues = FAssetDataTagMapSharedView(MoveTemp(TagsAndValues));
 				}
diff --git a/Engine/Source/Runtime/CoreUObject/Private/UObject/Class.cpp b/Engine/Source/Runtime/CoreUObject/Private/UObject/Class.cpp
--- a/Engine/Source/Runtime/CoreUObject/Private/UObject/Class.cpp
+++ b/Engine/Source/Runtime/CoreUObject/Private/UObject/Class.cpp
@@ -6066,7 +6066,9 @@ FTopLevelAssetPath UClass::TryConvertShortTypeNameToPathName(UClass* TypeClass,
 	}
 }
 
-bool UClass::TryFixShortClassNameExportPath(FString& InOutExportPathToFix, ELogVerbosity::Type AmbiguousMessageVerbosity /*= ELogVerbosity::NoLogging*/, const TCHAR* AmbiguousClassMessage /*= nullptr*/)
+bool UClass::TryFixShortClassNameExportPath(FString& InOutExportPathToFix,
+	ELogVerbosity::Type AmbiguousMessageVerbosity /*= ELogVerbosity::NoLogging*/,
+	const TCHAR* AmbiguousClassMessage /*= nullptr*/, bool bClearOnError /* = false */)
 {
 	FString ClassName;
 	FString ObjectPath;
@@ -6075,9 +6077,9 @@ bool UClass::TryFixShortClassNameExportPath(FString& InOutExportPathToFix, ELogV
 		if (FPackageName::IsShortPackageName(ClassName))
 		{
 			FTopLevelAssetPath ClassPathName = UClass::TryConvertShortTypeNameToPathName<UClass>(ClassName, AmbiguousMessageVerbosity, AmbiguousClassMessage);
-			if (!ClassPathName.IsNull())
+			if (!ClassPathName.IsNull() || bClearOnError)
 			{
-				InOutExportPathToFix = FObjectPropertyBase::GetExportPath(ClassPathName, ObjectPath);
+				InOutExportPathToFix = !ClassPathName.IsNull() ? FObjectPropertyBase::GetExportPath(ClassPathName, ObjectPath) : FString();
 				return true;
 			}
 		}
diff --git a/Engine/Source/Runtime/CoreUObject/Public/UObject/Class.h b/Engine/Source/Runtime/CoreUObject/Public/UObject/Class.h
--- a/Engine/Source/Runtime/CoreUObject/Public/UObject/Class.h
+++ b/Engine/Source/Runtime/CoreUObject/Public/UObject/Class.h
@@ -3066,14 +3066,19 @@ public:
 	}
 
 	/**
-	 * Tries to fix an export path containing short class name. Will not modify the input InOutExportPathToFix if a fixup
-	 * was not necessary or unsuccessful
+	 * Tries to fix an export path containing short class name. Will not modify the input InOutExportPathToFix if a
+	 * fixup was not necessary. Optionally modifies it if a fixup is necessary but unsuccessful.
 	 * @param InOutExportPathToFix Export path (Class'/Path/To.Object')
 	 * @param AmbiguousMessageVerbosity Verbosity with which to log a message if class name is ambiguous
 	 * @param AmbiguousClassMessage Additional message to log when class name is ambiguous (e.g. current operation)
-	 * @return True if the short path was successfully fixed. False if the provided export path did not contain short class name or the short path could not be fixed.
-	 */
-	static bool TryFixShortClassNameExportPath(FString& InOutExportPathToFix, ELogVerbosity::Type AmbiguousMessageVerbosity = ELogVerbosity::NoLogging, const TCHAR* AmbiguousClassMessage = nullptr);
+	 * @param bClearOnError If the fixup is necessary but unsuccessful, if bClearOnError, clear the value and return
+	 *        true, otherwise leave the value unchanged and return false.
+	 * @return True if the short path was successfully fixed. False if the provided export path did not contain short
+	 *         class name or the short path could not be fixed.
+	 */
+	static bool TryFixShortClassNameExportPath(FString& InOutExportPathToFix,
+		ELogVerbosity::Type AmbiguousMessageVerbosity = ELogVerbosity::NoLogging,
+		const TCHAR* AmbiguousClassMessage = nullptr, bool bClearOnError = false);
 
 	/**
 	 * Returns the ObjectName portion of a ClassPath name: "/Path/To.Object" is converted to "Object". 
diff --git a/Engine/Source/Runtime/Engine/Private/Blueprint.cpp b/Engine/Source/Runtime/Engine/Private/Blueprint.cpp
--- a/Engine/Source/Runtime/Engine/Private/Blueprint.cpp
+++ b/Engine/Source/Runtime/Engine/Private/Blueprint.cpp
@@ -1105,7 +1105,8 @@ void UBlueprint::PostLoadBlueprintAssetRegistryTags(const FAssetData& InAssetDat
 		FString TagValue = InAssetData.GetTagValueRef<FString>(TagName);
 		if (!TagValue.IsEmpty() && TagValue != TEXT("None"))
 		{
-			if (UClass::TryFixShortClassNameExportPath(TagValue, ELogVerbosity::Warning, TEXT("UBlueprint::PostLoadAssetRegistryTags")))
+			if (UClass::TryFixShortClassNameExportPath(TagValue, ELogVerbosity::Warning,
+				TEXT("UBlueprint::PostLoadAssetRegistryTags"), true /* bClearOnError */))
 			{
 				OutTagsAndValuesToUpdate.Add(FAssetRegistryTag(TagName, TagValue, TagType));
 			}
diff --git a/Engine/Source/Runtime/Engine/Private/BlueprintGeneratedClass.cpp b/Engine/Source/Runtime/Engine/Private/BlueprintGeneratedClass.cpp
--- a/Engine/Source/Runtime/Engine/Private/BlueprintGeneratedClass.cpp
+++ b/Engine/Source/Runtime/Engine/Private/BlueprintGeneratedClass.cpp
@@ -308,7 +308,8 @@ void UBlueprintGeneratedClass::PostLoadAssetRegistryTags(const FAssetData& InAss
 		FString TagValue = InAssetData.GetTagValueRef<FString>(TagName);
 		if (!TagValue.IsEmpty() && TagValue != TEXT("None"))
 		{
-			if (UClass::TryFixShortClassNameExportPath(TagValue, ELogVerbosity::Warning, TEXT("UBlueprintGeneratedClass::PostLoadAssetRegistryTags")))
+			if (UClass::TryFixShortClassNameExportPath(TagValue, ELogVerbosity::Warning,
+				TEXT("UBlueprintGeneratedClass::PostLoadAssetRegistryTags"), true /* bClearOnError */))
 			{
 				OutTagsAndValuesToUpdate.Add(FAssetRegistryTag(TagName, TagValue, TagType));
 			}
