SavePackage TransientPropertyOverrides: Fix harvesting of FNames for the transient override properties to be included. They are needed during serialization of the properties when the exports are saved into FLinkerSave.
#rnx

diff --git a/Engine/Source/Runtime/CoreUObject/Private/UObject/SavePackage/PackageHarvester.cpp b/Engine/Source/Runtime/CoreUObject/Private/UObject/SavePackage/PackageHarvester.cpp
--- a/Engine/Source/Runtime/CoreUObject/Private/UObject/SavePackage/PackageHarvester.cpp
+++ b/Engine/Source/Runtime/CoreUObject/Private/UObject/SavePackage/PackageHarvester.cpp
@@ -824,6 +824,9 @@ bool FPackageHarvester::ShouldSkipProperty(const FProperty* InProperty) const
 	const TSet<FProperty*>* Props = TransientPropertyOverrides.Find(CurrentExportDependencies.CurrentExport);
 	if (Props && Props->Contains(InProperty))
 	{
+		// Hack fix: harvest the property name so it can be found when the property is serialized by LinkerSave
+		// TODO: Skip the property during LinkerSave as well and remove this hack.
+		const_cast<FPackageHarvester*>(this)->HarvestExportDataName(InProperty->GetFName()); 
 		return true;
 	}
 	return false;
