Fix newly created flags not being properly stripped from packages on save

diff --git a/Engine/Source/Runtime/CoreUObject/Private/UObject/SavePackage/SaveContext.h b/Engine/Source/Runtime/CoreUObject/Private/UObject/SavePackage/SaveContext.h
--- a/Engine/Source/Runtime/CoreUObject/Private/UObject/SavePackage/SaveContext.h
+++ b/Engine/Source/Runtime/CoreUObject/Private/UObject/SavePackage/SaveContext.h
@@ -101,20 +101,6 @@ struct FIllegalReference
 	FString FormatStringArg;
 };
 
-enum class ESaveableStatus
-{
-	Success,
-	PendingKill,
-	Transient,
-	AbstractClass,
-	DeprecatedClass,
-	NewerVersionExistsClass,
-	OuterUnsaveable,
-	ClassUnsaveable,
-	ExcludedByPlatform,
-	__Count,
-};
-
 /** Hold the harvested exports and imports for a realm */
 struct FHarvestedRealm
 {
@@ -163,10 +149,6 @@ struct FHarvestedRealm
 		return IsImport(InObject) || IsExport(InObject);
 	}
 
-	/**
-	 * Used during harvesting to early exit from objects we have found referenced earlier but excluded because of
-	 * editoronly or unsaveable or otherwise.
-	 */
 	bool IsExcluded(UObject* InObject) const
 	{
 		return Excluded.Contains(InObject);
@@ -420,14 +402,14 @@ public:
 		, SaveArgs(InSaveArgs)
 		, PackageWriter(InSaveArgs.SavePackageContext ? InSaveArgs.SavePackageContext->PackageWriter : nullptr)
 		, SerializeContext(InSerializeContext)
-		, GameRealmExcludedObjectMarks(GetExcludedObjectMarksForGameRealm(SaveArgs.GetTargetPlatform()))
+		, ExcludedObjectMarks(SavePackageUtilities::GetExcludedObjectMarksForTargetPlatform(SaveArgs.GetTargetPlatform()))
 	{
 		// Assumptions & checks
 		check(InPackage);
 		check(InFilename);
-		// if we are cooking we should be doing it in the editor and with a PackageWriter
+		// if we are cooking we should be doing it in the editor and with a CookedPackageWriter
 		check(!IsCooking() || WITH_EDITOR);
-		checkf(!IsCooking() || PackageWriter, TEXT("Cook saves require an IPackageWriter"));
+		checkf(!IsCooking() || (PackageWriter && PackageWriter->AsCookedPackageWriter()), TEXT("Cook saves require an ICookedPackageWriter"));
 
 		// Store initial state
 		InitialPackageFlags = Package->GetPackageFlags();
@@ -519,19 +501,8 @@ public:
 
 	EObjectMark GetExcludedObjectMarks(ESaveRealm HarvestingRealm) const
 	{
-		switch (HarvestingRealm)
-		{
-		case ESaveRealm::Optional:
-			// When considering excluded objects for a platform, do not consider editor only object and not for platform objects in the optional context as excluded
-			return (EObjectMark)(GameRealmExcludedObjectMarks & ~(EObjectMark::OBJECTMARK_EditorOnly | EObjectMark::OBJECTMARK_NotForTargetPlatform));
-		case ESaveRealm::Game:
-			return GameRealmExcludedObjectMarks;
-		case ESaveRealm::Editor:
-			return EObjectMark::OBJECTMARK_NOMARKS;
-		default:
-			checkNoEntry();
-			return EObjectMark::OBJECTMARK_NOMARKS;
-		}
+		// When considering excluded objects for a platform, do not consider editor only object and not for platform objects in the optional context as excluded
+		return (HarvestingRealm == ESaveRealm::Optional) ? (EObjectMark)(ExcludedObjectMarks & ~(EObjectMark::OBJECTMARK_EditorOnly|EObjectMark::OBJECTMARK_NotForTargetPlatform)) : ExcludedObjectMarks;
 	}
 
 	EObjectFlags GetTopLevelFlags() const
@@ -745,6 +716,17 @@ public:
 	void MarkUnsaveable(UObject* InObject);
 
 	bool IsUnsaveable(UObject* InObject, bool bEmitWarning = true) const;
+	enum class ESaveableStatus
+	{
+		Success,
+		PendingKill,
+		Transient,
+		AbstractClass,
+		DeprecatedClass,
+		NewerVersionExistsClass,
+		OuterUnsaveable,
+		__Count,
+	};
 	ESaveableStatus GetSaveableStatus(UObject* InObject, UObject** OutCulprit = nullptr, ESaveableStatus* OutCulpritStatus = nullptr) const;
 	ESaveableStatus GetSaveableStatusNoOuter(UObject* InObject) const;
 
@@ -788,6 +770,11 @@ public:
 		return GetHarvestedRealm().IsIncluded(InObject);
 	}
 
+	bool IsExcluded(UObject* InObject) const
+	{
+		return GetHarvestedRealm().IsExcluded(InObject);
+	}
+
 	TSet<FTaggedExport>& GetExports()
 	{
 		return GetHarvestedRealm().GetExports();
@@ -985,7 +972,19 @@ public:
 		GetHarvestedRealm().SetTextFormatTempFilename(MoveTemp(InTemp));
 	}
 
-	FSavePackageResultStruct GetFinalResult();
+	FSavePackageResultStruct GetFinalResult()
+	{
+		if (Result != ESavePackageResult::Success)
+		{
+			return Result;
+		}
+
+		ESavePackageResult FinalResult = IsStubRequested() ? ESavePackageResult::GenerateStub : ESavePackageResult::Success;
+		FSavePackageResultStruct FinalResultStruct(FinalResult, TotalPackageSizeUncompressed,
+			SerializedPackageFlags, IsCompareLinker() ? MoveTemp(GetHarvestedRealm().Linker) : nullptr);
+		FinalResultStruct.SavedAssets = MoveTemp(SavedAssets);
+		return FinalResultStruct;
+	}
 
 	FObjectSaveContextData& GetObjectSaveContext()
 	{
@@ -1069,7 +1068,6 @@ private:
 
 	// Create the harvesting contexts and automatic optional context gathering options
 	void SetupHarvestingRealms();
-	static EObjectMark GetExcludedObjectMarksForGameRealm(const ITargetPlatform* TargetPlatform);
 		
 	friend class FPackageHarvester;
 
@@ -1103,8 +1101,8 @@ private:
 	// Pointer to the EDLCookChecker associated with this context
 	FEDLCookCheckerThreadState* EDLCookChecker = nullptr;
 
-	// An object matching any GameRealmExcludedObjectMarks should be excluded from imports or exports in the game realm
-	const EObjectMark GameRealmExcludedObjectMarks;
+	// Matching any mark in ExcludedObjectMarks indicates that an object should be excluded from being either an import or an export for this save
+	const EObjectMark ExcludedObjectMarks;
 
 	// Harvested custom versions
 	FCustomVersionContainer CustomVersions;
@@ -1128,4 +1126,4 @@ private:
 	TMap<UObject*, TSet<FProperty*>> TransientPropertyOverrides;
 };
 
-const TCHAR* LexToString(ESaveableStatus Status);
+const TCHAR* LexToString(FSaveContext::ESaveableStatus Status);
diff --git a/Engine/Source/Runtime/CoreUObject/Private/UObject/SavePackage2.cpp b/Engine/Source/Runtime/CoreUObject/Private/UObject/SavePackage2.cpp
--- a/Engine/Source/Runtime/CoreUObject/Private/UObject/SavePackage2.cpp
+++ b/Engine/Source/Runtime/CoreUObject/Private/UObject/SavePackage2.cpp
@@ -157,15 +157,7 @@ ESavePackageResult ValidatePackage(FSaveContext& SaveContext)
 					ErrorText = FText::Format(NSLOCTEXT("SavePackage2", "AssetSaveMissingTopLevelFlags", "The Asset {Name} being saved does not have any of the provided object flags (0x{Flags}); saving the package would cause data loss."),	
 						Arguments);
 				}
-				if (SaveContext.IsGenerateSaveError() && SaveContext.GetError())
-				{
-					SaveContext.GetError()->Logf(ELogVerbosity::Warning, TEXT("%s"), *ErrorText.ToString());
-				}
-				else
-				{
-					UE_LOG(LogSavePackage, Warning, TEXT("%s"), *ErrorText.ToString());
-				}
-
+				SaveContext.GetError()->Logf(ELogVerbosity::Warning, TEXT("%s"), *ErrorText.ToString());
 				return ESavePackageResult::Error;
 			}
 		}
@@ -385,7 +377,7 @@ ESavePackageResult HarvestPackage(FSaveContext& SaveContext)
 		FHarvestedRealm& OptionalContext = SaveContext.GetHarvestedRealm(ESaveRealm::Optional);
 		for (auto It = OptionalContext.GetExports().CreateIterator(); It; ++It)
 		{
-			if (!FPackageHarvester::ShouldObjectBeHarvestedInOptionalRealm(It->Obj, SaveContext))
+			if (!It->Obj->GetClass()->HasAnyClassFlags(CLASS_Optional))
 			{
 				// Make sure the export is found in the game context as well
 				if (FTaggedExport* GameExport = SaveContext.GetHarvestedRealm(ESaveRealm::Game).GetExports().Find(It->Obj))
@@ -413,33 +405,30 @@ ESavePackageResult HarvestPackage(FSaveContext& SaveContext)
 		}
 	}
 
+	// Trim PrestreamPackage list 
+	TSet<UPackage*>& PrestreamPackages = SaveContext.GetPrestreamPackages();
+	TSet<UPackage*> KeptPrestreamPackages;
+	for (UPackage* Pkg : PrestreamPackages)
 	{
-		FPackageHarvester::FHarvestScope RootReferenceScope = Harvester.EnterRootReferencesScope();
-		// Trim PrestreamPackage list 
-		TSet<UPackage*>& PrestreamPackages = SaveContext.GetPrestreamPackages();
-		TSet<UPackage*> KeptPrestreamPackages;
-		for (UPackage* Pkg : PrestreamPackages)
+		// If the prestream package hasn't been otherwise already marked as an import, keep it as such and mark it as an import
+		if (!SaveContext.IsImport(Pkg))
 		{
-			// If the prestream package hasn't been otherwise already marked as an import, keep it as such and mark it as an import
-			if (!SaveContext.IsImport(Pkg))
-			{
-				KeptPrestreamPackages.Add(Pkg);
-				Harvester << Pkg;
-			}
+			KeptPrestreamPackages.Add(Pkg);
+			SaveContext.AddImport(Pkg);
 		}
-		Exchange(PrestreamPackages, KeptPrestreamPackages);
+	}
+	Exchange(PrestreamPackages, KeptPrestreamPackages);
 
-		// Harvest the PrestreamPackage class name if needed
-		if (PrestreamPackages.Num() > 0)
-		{
-			Harvester.HarvestPackageHeaderName(SavePackageUtilities::NAME_PrestreamPackage);
-		}
+	// Harvest the PrestreamPackage class name if needed
+	if (PrestreamPackages.Num() > 0)
+	{
+		Harvester.HarvestPackageHeaderName(SavePackageUtilities::NAME_PrestreamPackage);		
+	}
 
-		// if we have a WorldTileInfo, we need to harvest its dependencies as well, i.e. Custom Version
-		if (FWorldTileInfo* WorldTileInfo = SaveContext.GetPackage()->GetWorldTileInfo())
-		{
-			Harvester << *WorldTileInfo;
-		}
+	// if we have a WorldTileInfo, we need to harvest its dependencies as well, i.e. Custom Version
+	if (FWorldTileInfo* WorldTileInfo = SaveContext.GetPackage()->GetWorldTileInfo())
+	{
+		Harvester << *WorldTileInfo;
 	}
 
 	// The Editor version is used as part of the check to see if a package is too old to use the gather cache, so we always have to add it if we have gathered loc for this asset
@@ -1423,13 +1412,13 @@ void SavePreloadDependencies(FStructuredArchive::FRecord& StructuredArchiveRoot,
 			{
 				UE_LOG(LogSavePackage, Warning, TEXT("A dependency '%s' of '%s' is in the linker table, but is pending kill or garbage. We will keep the dependency anyway (%d)."), *ToTest->GetFullName(), *ForObj->GetFullName(), CallSite);
 			}
-			bool bIncludedInHarvest = SaveContext.IsIncluded(ToTest);
-			if (bMandatory && !bIncludedInHarvest)
+			bool bNotFiltered = !SaveContext.IsExcluded(ToTest);
+			if (bMandatory && !bNotFiltered)
 			{
 				UE_LOG(LogSavePackage, Warning, TEXT("A dependency '%s' of '%s' was filtered, but is mandatory. This indicates a problem with editor only stripping. We will keep the dependency anyway (%d)."), *ToTest->GetFullName(), *ForObj->GetFullName(), CallSite);
-				bIncludedInHarvest = true;
+				bNotFiltered = true;
 			}
-			if (bIncludedInHarvest)
+			if (bNotFiltered)
 			{
 				if (!Index.IsNull())
 				{
@@ -1538,8 +1527,8 @@ void SavePreloadDependencies(FStructuredArchive::FRecord& StructuredArchiveRoot,
 						// Only include subobject archetypes
 						if (SubObj->HasAnyFlags(RF_DefaultSubObject | RF_ArchetypeObject))
 						{
-							// Don't include the archetype of SubObjects that were not included in the harvesting phase; we didn't add their archetypes
-							if (!SaveContext.IsIncluded(SubObj))
+							// Don't include the archetype of SubObjects that we excluded in the harvesting phase; we didn't add their archetypes
+							if (SaveContext.IsExcluded(SubObj))
 							{
 								continue;
 							}
