Fixing option property GC processing code

diff --git a/Engine/Source/Runtime/CoreUObject/Public/UObject/FastReferenceCollector.h b/Engine/Source/Runtime/CoreUObject/Public/UObject/FastReferenceCollector.h
--- a/Engine/Source/Runtime/CoreUObject/Public/UObject/FastReferenceCollector.h
+++ b/Engine/Source/Runtime/CoreUObject/Public/UObject/FastReferenceCollector.h
@@ -939,23 +939,27 @@ StoleContext:
 							uint32 ValueSize = TokenStream.ReadStride(TokenStreamIndex); // Size of value in bytes. This is also the offset to the bIsSet variable stored thereafter.
 							const FGCSkipInfo SkipInfo = TokenStream.ReadSkipInfo(TokenStreamIndex);
 							const bool& bIsSet = *((bool*)(StackEntryData + ReferenceInfo.Offset + ValueSize));
+
+							StackEntry = Stack.Push();
+							StackEntryData += ReferenceInfo.Offset;
+							StackEntry->Data = StackEntryData;
+							StackEntry->Stride = ValueSize;
+							StackEntry->Count = 1;
+							StackEntry->LoopStartIndex = TokenStreamIndex;
+
 							if (bIsSet)
 							{
 								// It's set - push a stack entry for processing the value
 								// This is somewhat suboptimal since there is only ever just one value, but this approach avoids any changes to the surrounding code
-								StackEntry = Stack.Push();
-								StackEntryData += ReferenceInfo.Offset;
-								StackEntry->Data = StackEntryData;
-								StackEntry->Stride = ValueSize;
-								StackEntry->Count = 1;
-								StackEntry->LoopStartIndex = TokenStreamIndex;
+								check(StackEntry->Data);
+								TokenReturnCount = 0;
 							}
 							else
 							{
 								// It's unset - keep going by jumping to skip index
 								TokenStreamIndex = SkipInfo.SkipIndex;
+								TokenReturnCount = TokenStream.GetSkipReturnCount(SkipInfo);
 							}
-							TokenReturnCount = 0;
 						}
 						break;
 						case GCRT_EndOfPointer:
