Added a CBaseStructureProvider concept for detecting if a type is supported by TBaseStructure.
Deleted the second parameter from TStructOnScope, which doesn't work so well with forward-declared types.

diff --git a/Engine/Source/Runtime/CoreUObject/Public/Concepts/BaseStructureProvider.h b/Engine/Source/Runtime/CoreUObject/Public/Concepts/BaseStructureProvider.h
new file mode 100644
--- /dev/null
+++ b/Engine/Source/Runtime/CoreUObject/Public/Concepts/BaseStructureProvider.h
@@ -0,0 +1,19 @@
+// Copyright Epic Games, Inc. All Rights Reserved.
+
+#pragma once
+
+class UScriptStruct;
+
+template <typename T>
+struct TBaseStructure;
+
+/**
+ * Describes a type for which TBaseStructure<T>::Get() returning a UScriptStruct* is defined.
+ */
+struct CBaseStructureProvider
+{
+	template <typename T>
+	auto Requires(UScriptStruct*& StructRef) -> decltype(
+		StructRef = TBaseStructure<T>::Get()
+	);
+};
diff --git a/Engine/Source/Runtime/CoreUObject/Public/UObject/Class.h b/Engine/Source/Runtime/CoreUObject/Public/UObject/Class.h
--- a/Engine/Source/Runtime/CoreUObject/Public/UObject/Class.h
+++ b/Engine/Source/Runtime/CoreUObject/Public/UObject/Class.h
@@ -7,6 +7,7 @@
 #pragma once
 
 #include "Concepts/GetTypeHashable.h"
+#include "Concepts/StaticStructProvider.h"
 #include "Containers/Array.h"
 #include "Containers/EnumAsByte.h"
 #include "Containers/Map.h"
@@ -3977,7 +3978,8 @@ struct FTestUninitializedScriptStructMembersTest
 	odd types.
 -----------------------------------------------------------------------------*/
 
-template< class T > struct TBaseStructure
+template <typename T, bool bHasStaticStruct = TModels_V<CStaticStructProvider, T>>
+struct TBaseStructureBase
 {
 	static UScriptStruct* Get()
 	{
@@ -3985,6 +3987,16 @@ template< class T > struct TBaseStructure
 	}
 };
 
+template <typename T>
+struct TBaseStructureBase<T, false>
+{
+};
+
+template <typename T>
+struct TBaseStructure : TBaseStructureBase<T>
+{
+};
+
 template<> struct TBaseStructure<FIntPoint> 
 {
 	COREUOBJECT_API static UScriptStruct* Get(); 
diff --git a/Engine/Source/Runtime/CoreUObject/Public/UObject/StructOnScope.h b/Engine/Source/Runtime/CoreUObject/Public/UObject/StructOnScope.h
--- a/Engine/Source/Runtime/CoreUObject/Public/UObject/StructOnScope.h
+++ b/Engine/Source/Runtime/CoreUObject/Public/UObject/StructOnScope.h
@@ -125,10 +125,8 @@ public:
 
 /**
  * Typed FStructOnScope that exposes type-safe access to the wrapped struct
- * @note The second template argument is there to restrict to type that are actually USTRUCT.
- *		It will be replaced to a static_assert with a better compiler error with a "is ustruct" type trait
  */
-template<typename T, typename = decltype(TBaseStructure<T>::Get())>
+template<typename T>
 class TStructOnScope final : public FStructOnScope
 {
 public:
