reducing the number of items in the array conversion tests as the 100k items was unnecessary
https://p4-swarm.epicgames.net/reviews/25165292

diff --git a/Engine/Source/Runtime/CoreUObject/Tests/ObjectPtrTest.cpp b/Engine/Source/Runtime/CoreUObject/Tests/ObjectPtrTest.cpp
--- a/Engine/Source/Runtime/CoreUObject/Tests/ObjectPtrTest.cpp
+++ b/Engine/Source/Runtime/CoreUObject/Tests/ObjectPtrTest.cpp
@@ -796,7 +796,7 @@ void TestArrayConversion()
 	};
 #endif
 	TArray<TObjectPtr<TObj>> PtrArray;
-	uint32  NumObjs = 1000000;
+	uint32  NumObjs = 5;
 	for (uint32 i = 0; i < NumObjs; ++i)
 	{
 		PtrArray.Add(Obj1);
