ZenLoader: Now that we only finish loading packages that are part of the current sync load (see 25024411) we need to let IncludePackageInSyncLoadContextRecursive process all packages that has not reached the Complete state yet.
Fix infinite load problem when flushing a package that is already in the LoadedPackagesToProcess array.

#rnx

diff --git a/Engine/Source/Runtime/CoreUObject/Private/Serialization/AsyncLoading2.cpp b/Engine/Source/Runtime/CoreUObject/Private/Serialization/AsyncLoading2.cpp
--- a/Engine/Source/Runtime/CoreUObject/Private/Serialization/AsyncLoading2.cpp
+++ b/Engine/Source/Runtime/CoreUObject/Private/Serialization/AsyncLoading2.cpp
@@ -3908,7 +3908,7 @@ void FAsyncLoadingThread2::IncludePackageInSyncLoadContextRecursive(FAsyncLoadin
 		return;
 	}
 
-	if (Package->AsyncPackageLoadingState > EAsyncPackageLoadingState2::DeferredPostLoadDone)
+	if (Package->AsyncPackageLoadingState >= EAsyncPackageLoadingState2::Complete)
 	{
 		return;
 	}
