ZenLoader: Minor code modification to fix pvsstudio warning V595

Engine/Source/Runtime/CoreUObject/Private/Serialization/AsyncLoading2.cpp(4723): warning V595: The 'ImportedPackageRef.GetPackage()' pointer was utilized before it was verified against nullptr. Check lines: 4723, 4741.

#rnx

diff --git a/Engine/Source/Runtime/CoreUObject/Private/Serialization/AsyncLoading2.cpp b/Engine/Source/Runtime/CoreUObject/Private/Serialization/AsyncLoading2.cpp
--- a/Engine/Source/Runtime/CoreUObject/Private/Serialization/AsyncLoading2.cpp
+++ b/Engine/Source/Runtime/CoreUObject/Private/Serialization/AsyncLoading2.cpp
@@ -4741,7 +4741,7 @@ void FAsyncPackage2::ImportPackagesRecursiveInner(FAsyncLoadingThreadState2& Thr
 		}
 		else if (ImportedPackageStatus == EPackageStoreEntryStatus::Missing)
 		{
-			if (!ImportedPackageRef.GetPackage()) // If we found a package it's not actually missing but we can't load it anyway
+			if (!ImportedPackageRef.HasPackage()) // If we found a package it's not actually missing but we can't load it anyway
 			{
 				UE_ASYNC_PACKAGE_CLOG(!ImportedPackageUPackageName.IsNone(), Display, Desc, TEXT("ImportPackages: SkipPackage"),
 					TEXT("Skipping non mounted imported package %s (0x%llX)"), *ImportedPackageNameToLoad.ToString(), ImportedPackageId.Value());
