Fix some reentrancy problems during flush in zenloader

#rnx

diff --git a/Engine/Source/Runtime/CoreUObject/Private/Serialization/AsyncLoading2.cpp b/Engine/Source/Runtime/CoreUObject/Private/Serialization/AsyncLoading2.cpp
--- a/Engine/Source/Runtime/CoreUObject/Private/Serialization/AsyncLoading2.cpp
+++ b/Engine/Source/Runtime/CoreUObject/Private/Serialization/AsyncLoading2.cpp
@@ -2190,6 +2190,35 @@ uint64 FAsyncLoadingSyncLoadContext::NextContextId = 1;
 
 struct FAsyncLoadingThreadState2
 {
+	struct FTimeLimitScope
+	{
+		bool bUseTimeLimit = false;
+		double TimeLimit = 0.0;
+		double StartTime = 0.0;
+
+		FTimeLimitScope(bool bInUseTimeLimit, double InTimeLimit)
+		{
+			FAsyncLoadingThreadState2* ThreadState = FAsyncLoadingThreadState2::Get();
+
+			bUseTimeLimit = ThreadState->bUseTimeLimit;
+			TimeLimit = ThreadState->TimeLimit;
+			StartTime = ThreadState->StartTime;
+
+			ThreadState->bUseTimeLimit = bInUseTimeLimit;
+			ThreadState->TimeLimit = InTimeLimit;
+			ThreadState->StartTime = FPlatformTime::Seconds();
+		}
+		
+		~FTimeLimitScope()
+		{
+			FAsyncLoadingThreadState2* ThreadState = FAsyncLoadingThreadState2::Get();
+
+			ThreadState->bUseTimeLimit = bUseTimeLimit;
+			ThreadState->TimeLimit = TimeLimit;
+			ThreadState->StartTime = StartTime;
+		}
+	};
+
 	static void Set(FAsyncLoadingThreadState2* State)
 	{
 		check(TlsSlot != 0);
@@ -2206,14 +2235,6 @@ struct FAsyncLoadingThreadState2
 	FAsyncLoadingThreadState2(FAsyncLoadEventGraphAllocator& InGraphAllocator, FIoDispatcher& InIoDispatcher)
 		: GraphAllocator(InGraphAllocator)
 	{
-
-	}
-
-	void SetTimeLimit(bool bInUseTimeLimit, double InTimeLimit)
-	{
-		bUseTimeLimit = bInUseTimeLimit;
-		TimeLimit = InTimeLimit;
-		StartTime = FPlatformTime::Seconds();
 	}
 
 	bool IsTimeLimitExceeded(const TCHAR* InLastTypeOfWorkPerformed = nullptr, UObject* InLastObjectWorkWasPerformedOn = nullptr)
@@ -6829,6 +6850,8 @@ FEventLoadNode2& FAsyncPackage2::GetExportBundleNode(EEventLoadNode2 Phase, uint
 
 void FAsyncLoadingThread2::UpdateSyncLoadContext(FAsyncLoadingThreadState2& ThreadState)
 {
+	TRACE_CPUPROFILER_EVENT_SCOPE(FAsyncLoadingThread2::UpdateSyncLoadContext);
+
 	if (ThreadState.bIsAsyncLoadingThread)
 	{
 		FAsyncLoadingSyncLoadContext* CreatedOnMainThread;
@@ -7177,9 +7200,10 @@ EAsyncPackageState::Type FAsyncLoadingThread2::ProcessLoadedPackagesFromGameThre
 			}
 		}
 		LocalCompletedAsyncPackages.Reset();
-		
-		// Call callbacks in a batch in a stack-local array after all other work has been done to handle
-		// callbacks that may call FlushAsyncLoading
+
+		TArray<FCompletedPackageRequest> RequestsToProcess;
+
+		// Move CompletedPackageRequest out of the global collection to prevent it from changing from within the callbacks.
 		// If we're flushing a specific request only call callbacks for that request
 		for (int32 CompletedPackageRequestIndex = CompletedPackageRequests.Num() - 1; CompletedPackageRequestIndex >= 0; --CompletedPackageRequestIndex)
 		{
@@ -7187,22 +7211,29 @@ EAsyncPackageState::Type FAsyncLoadingThread2::ProcessLoadedPackagesFromGameThre
 			if (FlushRequestID == INDEX_NONE || CompletedPackageRequest.RequestIDs.Contains(FlushRequestID))
 			{
 				RemovePendingRequests(CompletedPackageRequest.RequestIDs);
-				TRACE_COUNTER_SET(AsyncLoadingPackagesWithRemainingWork, PackagesWithRemainingWorkCounter);
-				CompletedPackageRequest.CallCompletionCallbacks();
-				if (CompletedPackageRequest.AsyncPackage)
-				{
-					CompletedPackageRequest.AsyncPackage->ReleaseRef();
-				}
-				else
-				{
-					// Requests for missing packages have no AsyncPackage but they count as packages with remaining work
-					--PackagesWithRemainingWorkCounter;
-				}
+				RequestsToProcess.Emplace(MoveTemp(CompletedPackageRequest));
 				CompletedPackageRequests.RemoveAt(CompletedPackageRequestIndex);
 				bLocalDidSomething = true;
 			}
 		}
 
+		// Call callbacks in a batch in a stack-local array after all other work has been done to handle
+		// callbacks that may call FlushAsyncLoading
+		for (FCompletedPackageRequest& CompletedPackageRequest : RequestsToProcess)
+		{
+			TRACE_COUNTER_SET(AsyncLoadingPackagesWithRemainingWork, PackagesWithRemainingWorkCounter);
+			CompletedPackageRequest.CallCompletionCallbacks();
+			if (CompletedPackageRequest.AsyncPackage)
+			{
+				CompletedPackageRequest.AsyncPackage->ReleaseRef();
+			}
+			else
+			{
+				// Requests for missing packages have no AsyncPackage but they count as packages with remaining work
+				--PackagesWithRemainingWorkCounter;
+			}
+		}
+
 		if (!bLocalDidSomething)
 		{
 			break;
@@ -7247,7 +7278,9 @@ EAsyncPackageState::Type FAsyncLoadingThread2::TickAsyncLoadingFromGameThread(FA
 
 	if (!bLoadingSuspended)
 	{
-		ThreadState.SetTimeLimit(bUseTimeLimit, TimeLimit);
+		// Use a time limit scope to restore the time limit to the old values when exiting the scope
+		// This is required to ensure a reentrant call here doesn't overwrite time limits permanently.
+		FAsyncLoadingThreadState2::FTimeLimitScope TimeLimitScope(bUseTimeLimit, TimeLimit);
 
 		const bool bIsMultithreaded = FAsyncLoadingThread2::IsMultithreaded();
 		double TickStartTime = FPlatformTime::Seconds();
@@ -8271,6 +8304,8 @@ void FAsyncPackage2::FinishUPackage()
 #if WITH_EDITOR
 void FAsyncLoadingThread2::ConditionalProcessEditorCallbacks()
 {
+	TRACE_CPUPROFILER_EVENT_SCOPE(FAsyncLoadingThread2::ConditionalProcessEditorCallbacks);
+
 	check(IsInGameThread());
 	if (!GameThreadState->SyncLoadContextStack.IsEmpty())
 	{
@@ -8418,6 +8453,7 @@ void FAsyncLoadingThread2::FlushLoading(int32 RequestId)
 		FUObjectThreadContext& ThreadContext = FUObjectThreadContext::Get();
 		if (ThreadContext.SyncLoadUsingAsyncLoaderCount == 0)
 		{
+			TRACE_CPUPROFILER_EVENT_SCOPE(FCoreDelegates::OnAsyncLoadingFlush);
 			FCoreDelegates::OnAsyncLoadingFlush.Broadcast();
 		}
 
@@ -8427,6 +8463,8 @@ void FAsyncLoadingThread2::FlushLoading(int32 RequestId)
 		FAsyncPackage2* CurrentlyExecutingPackage = nullptr;
 		if (GameThreadState->CurrentlyExecutingEventNode)
 		{
+			TRACE_CPUPROFILER_EVENT_SCOPE(HandleCurrentlyExecutingEventNode);
+
 			UE_CLOG(RequestId == INDEX_NONE, LogStreaming, Fatal, TEXT("Flushing async loading while creating, serializing or postloading an object is not permitted"));
 			CurrentlyExecutingPackage = GameThreadState->CurrentlyExecutingEventNode->GetPackage();
 			GameThreadState->PackagesOnStack.Push(CurrentlyExecutingPackage);
