Fix potential truncation in FUnversionedPropertySerializer constructor

#rnx

- It is expected that the value of GetIntNum will fit in a uint8 so we can just cast.

diff --git a/Engine/Source/Runtime/CoreUObject/Private/Serialization/UnversionedPropertySerialization.cpp b/Engine/Source/Runtime/CoreUObject/Private/Serialization/UnversionedPropertySerialization.cpp
--- a/Engine/Source/Runtime/CoreUObject/Private/Serialization/UnversionedPropertySerialization.cpp
+++ b/Engine/Source/Runtime/CoreUObject/Private/Serialization/UnversionedPropertySerialization.cpp
@@ -50,7 +50,7 @@ public:
 		, Offset(Property->GetOffset_ForInternal() + Property->ElementSize * InArrayIndex)
 		, bSerializeAsInteger(CanSerializeAsInteger(Property))
 		, IntType(GetIntType(Property->GetMinAlignment()))
-		, FastZeroIntNum(CanSerializeAsZero(InProperty, IntType) ? GetIntNum(InProperty, IntType) : uint8(0))
+		, FastZeroIntNum(CanSerializeAsZero(InProperty, IntType) ? static_cast<uint8>(GetIntNum(InProperty, IntType)) : uint8(0))
 #else
 		, ArrayIndex(InArrayIndex)
 #endif
