Fixed some potential truncations errors in UObjectHash

#rnx

- Fixed inconsistency in how allocation sizes were being handled by storing everything as SIZE_T, removing the need to cast.
- Fixed some whitespace issues (spaces instead of tabs)

diff --git a/Engine/Source/Runtime/CoreUObject/Private/UObject/UObjectHash.cpp b/Engine/Source/Runtime/CoreUObject/Private/UObject/UObjectHash.cpp
--- a/Engine/Source/Runtime/CoreUObject/Private/UObject/UObjectHash.cpp
+++ b/Engine/Source/Runtime/CoreUObject/Private/UObject/UObjectHash.cpp
@@ -161,12 +161,12 @@ struct FHashBucket
 		return !!ElementsOrSetPtr[0] + !!ElementsOrSetPtr[1];
 	}
 	/** Returns the amount of memory allocated for and by Items TSet */
-	FORCEINLINE uint32 GetItemsSize() const
+	FORCEINLINE SIZE_T GetItemsSize() const
 	{
 		const TSet<UObjectBase*>* Items = GetSet();
 		if (Items)
 		{
-			return (uint32)sizeof(*Items) + Items->GetAllocatedSize();
+			return sizeof(*Items) + Items->GetAllocatedSize();
 		}
 		return 0;
 	}
@@ -1677,8 +1677,8 @@ void LogHashStatisticsInternal(TMultiMap<int32, uint32>& Hash, FOutputDevice& Ar
 		MaxCollisions);
 
 	// Calculate Hashtable size
-	const uint32 HashtableAllocatedSize = Hash.GetAllocatedSize();
-	Ar.Logf(TEXT("Total memory allocated for Object Outer Hash: %u bytes."), HashtableAllocatedSize);
+	const SIZE_T HashtableAllocatedSize = Hash.GetAllocatedSize();
+	Ar.Logf(TEXT("Total memory allocated for Object Outer Hash: %" SIZE_T_FMT " bytes."), HashtableAllocatedSize);
 }
 
 void LogHashStatisticsInternal(TBucketMap<int32>& Hash, FOutputDevice& Ar, const bool bShowHashBucketCollisionInfo)
@@ -1744,13 +1744,13 @@ void LogHashStatisticsInternal(TBucketMap<int32>& Hash, FOutputDevice& Ar, const
 		SlotsInUse);
 
 	// Calculate Hashtable size
-	uint32 HashtableAllocatedSize = Hash.GetAllocatedSize();
+	SIZE_T HashtableAllocatedSize = Hash.GetAllocatedSize();
 	// Calculate the size of a all Allocations inside of the buckets (TSet Items)
 	for (auto& Pair : Hash)
 	{
 		HashtableAllocatedSize += Pair.Value.GetItemsSize();
 	}
-	Ar.Logf(TEXT("Total memory allocated for and by Object Hash: %u bytes."), HashtableAllocatedSize);
+	Ar.Logf(TEXT("Total memory allocated for and by Object Hash: %" SIZE_T_FMT " bytes."), HashtableAllocatedSize);
 }
 
 void LogHashStatistics(FOutputDevice& Ar, const bool bShowHashBucketCollisionInfo)
@@ -1772,12 +1772,12 @@ void LogHashOuterStatistics(FOutputDevice& Ar, const bool bShowHashBucketCollisi
 	LogHashStatisticsInternal(FUObjectHashTables::Get().HashOuter, Ar, bShowHashBucketCollisionInfo);
 	Ar.Logf(TEXT(""));
 
-	uint32 HashOuterMapSize = 0;
+	SIZE_T HashOuterMapSize = 0;
 	for (TPair<UObjectBase*, FHashBucket>& OuterMapEntry : FUObjectHashTables::Get().ObjectOuterMap)
 	{
 		HashOuterMapSize += OuterMapEntry.Value.GetItemsSize();
 	}
-	Ar.Logf(TEXT("Total memory allocated for Object Outer Map: %u bytes."), HashOuterMapSize);
+	Ar.Logf(TEXT("Total memory allocated for Object Outer Map: %" SIZE_T_FMT " bytes."), HashOuterMapSize);
 	Ar.Logf(TEXT(""));
 }
 
@@ -1789,26 +1789,26 @@ void LogHashMemoryOverheadStatistics(FOutputDevice& Ar, const bool bShowIndividu
 	FUObjectHashTables& HashTables = FUObjectHashTables::Get();
 	FHashTableLock HashLock(HashTables);
 
-	int64 TotalSize = 0;
+	SIZE_T TotalSize = 0;
 	
 	{
-		int64 Size = HashTables.Hash.GetAllocatedSize();
+		SIZE_T Size = HashTables.Hash.GetAllocatedSize();
 		for (const TPair<int32, FHashBucket>& Pair : HashTables.Hash)
 		{
 			Size += Pair.Value.GetItemsSize();
 		}
 		if (bShowIndividualStats)
 		{
-			Ar.Logf(TEXT("Memory used by UObject Hash: %lld bytes."), Size);
+			Ar.Logf(TEXT("Memory used by UObject Hash: %" SIZE_T_FMT " bytes."), Size);
 		}
 		TotalSize += Size;
 	}
 
 	{
-		int64 Size = HashTables.HashOuter.GetAllocatedSize();
+		const SIZE_T Size = HashTables.HashOuter.GetAllocatedSize();
 		if (bShowIndividualStats)
 		{
-			Ar.Logf(TEXT("Memory used by UObject Outer Hash: %lld bytes."), Size);
+			Ar.Logf(TEXT("Memory used by UObject Outer Hash: %" SIZE_T_FMT " bytes."), Size);
 		}
 		TotalSize += Size;
 	}
@@ -1827,62 +1827,62 @@ void LogHashMemoryOverheadStatistics(FOutputDevice& Ar, const bool bShowIndividu
 	}
 
 	{
-		int64 Size = HashTables.ClassToObjectListMap.GetAllocatedSize();
+		SIZE_T Size = HashTables.ClassToObjectListMap.GetAllocatedSize();
 		for (const TPair<UClass*, FHashBucket>& Pair : HashTables.ClassToObjectListMap)
 		{
 			Size += Pair.Value.GetItemsSize();
 		}
 		if (bShowIndividualStats)
 		{
-			Ar.Logf(TEXT("Memory used by UClass To UObject List Map: %lld bytes."), Size);
+			Ar.Logf(TEXT("Memory used by UClass To UObject List Map: %" SIZE_T_FMT " bytes."), Size);
 		}
 		TotalSize += Size;
 	}
 
 	{
-		int64 Size = HashTables.ClassToChildListMap.GetAllocatedSize();
+		SIZE_T Size = HashTables.ClassToChildListMap.GetAllocatedSize();
 		for (const TPair<UClass*, TSet<UClass*> >& Pair : HashTables.ClassToChildListMap)
 		{
 			Size += Pair.Value.GetAllocatedSize();
 		}
 		if (bShowIndividualStats)
 		{
-			Ar.Logf(TEXT("Memory used by UClass To Child UClass List Map: %lld bytes."), Size);
+			Ar.Logf(TEXT("Memory used by UClass To Child UClass List Map: %" SIZE_T_FMT " bytes."), Size);
 		}
 		TotalSize += Size;
 	}
 
 	{
-		int64 Size = HashTables.PackageToObjectListMap.GetAllocatedSize();
+		SIZE_T Size = HashTables.PackageToObjectListMap.GetAllocatedSize();
 		for (const TPair<UPackage*, FHashBucket>& Pair : HashTables.PackageToObjectListMap)
 		{
 			Size += Pair.Value.GetItemsSize();
 		}
 		if (bShowIndividualStats)
 		{
-			Ar.Logf(TEXT("Memory used by UPackage To UObject List Map: %lld bytes."), Size);
+			Ar.Logf(TEXT("Memory used by UPackage To UObject List Map: %" SIZE_T_FMT " bytes."), Size);
 		}
 		TotalSize += Size;
 	}
 
 	{
-		int64 Size = HashTables.ObjectToPackageMap.GetAllocatedSize();
+		const SIZE_T Size = HashTables.ObjectToPackageMap.GetAllocatedSize();
 		if (bShowIndividualStats)
 		{
-			Ar.Logf(TEXT("Memory used by UObject To External Package Map: %lld bytes."), Size);
+			Ar.Logf(TEXT("Memory used by UObject To External Package Map: %" SIZE_T_FMT " bytes."), Size);
 		}
 		TotalSize += Size;
 	}
 
-    {
-        int64 Size = GUObjectArray.GetAllocatedSize();
-        if (bShowIndividualStats)
-        {
-            Ar.Logf(TEXT("Memory used by UObjectArray: %lld bytes."), Size);
-        }
-        TotalSize += Size;
-    }
-    
-	Ar.Logf(TEXT("Total memory allocated by Object hash tables and maps: %lld bytes (%.2f MB)."), TotalSize, (double)TotalSize / 1024.0 / 1024.0);
+	{
+		const SIZE_T Size = GUObjectArray.GetAllocatedSize();
+		if (bShowIndividualStats)
+		{
+			Ar.Logf(TEXT("Memory used by UObjectArray: %" SIZE_T_FMT " bytes."), Size);
+		}
+		TotalSize += Size;
+	}
+
+	Ar.Logf(TEXT("Total memory allocated by Object hash tables and maps: %" SIZE_T_FMT " bytes(% .2f MB)."), TotalSize, (double)TotalSize / 1024.0 / 1024.0);
 	Ar.Logf(TEXT(""));
 }
