Add an explicit cast to GetObjectOuterHash to supress a truncation warning and show intent.

#rnx

- Also fixed up some comments that had gotten out of date with the code.

diff --git a/Engine/Source/Runtime/CoreUObject/Private/UObject/UObjectHash.cpp b/Engine/Source/Runtime/CoreUObject/Private/UObject/UObjectHash.cpp
--- a/Engine/Source/Runtime/CoreUObject/Private/UObject/UObjectHash.cpp
+++ b/Engine/Source/Runtime/CoreUObject/Private/UObject/UObjectHash.cpp
@@ -518,16 +518,16 @@ static FORCEINLINE int32 GetObjectHash(FName ObjName)
 
 /**
  * Calculates the object's hash just using the object's name index
- * XORed with the outer. Yields much better spread in the hash
+ * added with the outer. Yields much better spread in the hash
  * buckets, but requires knowledge of the outer, which isn't available
  * in all cases.
  *
  * @param ObjName the object's name to use the index of
  * @param Outer the object's outer pointer treated as an int32
  */
-static FORCEINLINE int32 GetObjectOuterHash(FName ObjName,PTRINT Outer)
+static int32 GetObjectOuterHash(FName ObjName, PTRINT Outer)
 {
-	return GetTypeHash(ObjName) + (Outer >> 6);
+	return GetTypeHash(ObjName) + static_cast<int32>(Outer >> 6);
 }
 
 UObject* StaticFindObjectFastExplicitThreadSafe(FUObjectHashTables& ThreadHash, const UClass* ObjectClass, FName ObjectName, const FString& ObjectPathName, bool bExactClass, EObjectFlags ExcludeFlags/*=0*/)
