[Bug Fix] Map element properties were being handled incorrectly causing garbage memory reads

diff --git a/Engine/Source/Editor/PropertyEditor/Private/AsyncDetailViewDiff.cpp b/Engine/Source/Editor/PropertyEditor/Private/AsyncDetailViewDiff.cpp
--- a/Engine/Source/Editor/PropertyEditor/Private/AsyncDetailViewDiff.cpp
+++ b/Engine/Source/Editor/PropertyEditor/Private/AsyncDetailViewDiff.cpp
@@ -77,9 +77,10 @@ static bool MapKeysMatch(const TSharedPtr<FDetailTreeNode>& TreeNodeA, const TSh
 	const FResolvedProperty ResolvedMapB = GetResolvedProperty(MapPropertyNodeB, GetObject(TreeNodeB));
 	FScriptMapHelper MapHelperA(MapPropertyA, MapPropertyA->ContainerPtrToValuePtr<UObject*>(ResolvedMapA.Object));
 	FScriptMapHelper MapHelperB(MapPropertyB, MapPropertyB->ContainerPtrToValuePtr<UObject*>(ResolvedMapB.Object));
-
-	const void* KeyA = MapHelperA.GetKeyPtr(KeyIndexA);
-	const void* KeyB = MapHelperB.GetKeyPtr(KeyIndexB);
+	
+	const void* KeyA = MapHelperA.FindNthKeyPtr(KeyIndexA);
+	const void* KeyB = MapHelperB.FindNthKeyPtr(KeyIndexB);
+	
 	if (MapPropertyA->KeyProp->SameType(MapPropertyB->KeyProp))
 	{
 		return MapPropertyA->KeyProp->Identical(KeyA, KeyB, PPF_DeepComparison);
@@ -107,8 +108,9 @@ static bool SetKeysMatch(const TSharedPtr<FDetailTreeNode>& TreeNodeA, const TSh
 	FScriptSetHelper SetHelperA(SetPropertyA, SetPropertyA->ContainerPtrToValuePtr<UObject*>(ResolvedSetA.Object));
 	FScriptSetHelper SetHelperB(SetPropertyB, SetPropertyB->ContainerPtrToValuePtr<UObject*>(ResolvedSetB.Object));
 	
-	const void* KeyA = SetHelperA.GetElementPtr(KeyIndexA);
-	const void* KeyB = SetHelperB.GetElementPtr(KeyIndexB);
+	const void* KeyA = SetHelperA.FindNthElementPtr(KeyIndexA);
+	const void* KeyB = SetHelperB.FindNthElementPtr(KeyIndexB);
+	
 	if (SetPropertyA->ElementProp->SameType(SetPropertyB->ElementProp))
 	{
 		return SetPropertyA->ElementProp->Identical(KeyA, KeyB, PPF_DeepComparison);
diff --git a/Engine/Source/Editor/PropertyEditor/Private/SDetailsSplitter.cpp b/Engine/Source/Editor/PropertyEditor/Private/SDetailsSplitter.cpp
--- a/Engine/Source/Editor/PropertyEditor/Private/SDetailsSplitter.cpp
+++ b/Engine/Source/Editor/PropertyEditor/Private/SDetailsSplitter.cpp
@@ -119,7 +119,7 @@ static void CopyPropertyValueForInsert(TSharedPtr<FDetailTreeNode> SourceDetails
 		FPropertyNode* DestinationPropertyNode;
 		if (ensure(TryGetDestinationContainer(DestinationDetailsNode, DestinationPropertyNode, DestinationSet, InsertIndex)))
 		{
-			const void* SourceData = SourceSet->GetElementPtr(SourceDetailsNode->GetPropertyNode()->GetArrayIndex());
+			const void* SourceData = SourceSet->FindNthElementPtr(SourceDetailsNode->GetPropertyNode()->GetArrayIndex());
 			DestinationSet->AddElement(SourceData);
 		}
 	}
@@ -133,8 +133,9 @@ static void CopyPropertyValueForInsert(TSharedPtr<FDetailTreeNode> SourceDetails
 		FPropertyNode* DestinationPropertyNode;
 		if (ensure(TryGetDestinationContainer(DestinationDetailsNode, DestinationPropertyNode, DestinationMap, InsertIndex)))
 		{
-			const void* SourceKey = SourceMap->GetKeyPtr(SourceDetailsNode->GetPropertyNode()->GetArrayIndex());
-			const void* SourceVal = SourceMap->GetValuePtr(SourceDetailsNode->GetPropertyNode()->GetArrayIndex());
+			const int32 Index = SourceMap->FindInternalIndex(SourceDetailsNode->GetPropertyNode()->GetArrayIndex());
+			const void* SourceKey = SourceMap->GetKeyPtr(Index);
+			const void* SourceVal = SourceMap->GetValuePtr(Index);
 			DestinationMap->AddPair(SourceKey, SourceVal);
 		}
 	}
diff --git a/Engine/Source/Editor/UnrealEd/Private/DiffUtils.cpp b/Engine/Source/Editor/UnrealEd/Private/DiffUtils.cpp
--- a/Engine/Source/Editor/UnrealEd/Private/DiffUtils.cpp
+++ b/Engine/Source/Editor/UnrealEd/Private/DiffUtils.cpp
@@ -128,20 +128,14 @@ FResolvedProperty FPropertySoftPath::Resolve(const UStruct* Struct, const void*
 			else if( const FSetProperty* SetProperty = CastField<FSetProperty>(Property) )
 			{
 				FScriptSetHelper SetHelper(SetProperty, Property->ContainerPtrToValuePtr<UObject*>(CurrentBlock));
-				if (SetHelper.IsValidIndex(PropertyIndex))
-				{
-					NextProperty = SetProperty->ElementProp;
-					NextBlock = SetHelper.GetElementPtr(PropertyIndex);
-				}
+				NextProperty = SetHelper.GetElementProperty();
+				NextBlock = SetHelper.FindNthElementPtr(PropertyIndex);
 			}
 			else if( const FMapProperty* MapProperty = CastField<FMapProperty>(Property) )
 			{
 				FScriptMapHelper MapHelper(MapProperty, Property->ContainerPtrToValuePtr<UObject*>(CurrentBlock));
-				if (MapHelper.IsValidIndex(PropertyIndex))
-				{
-					NextProperty = MapProperty->ValueProp;
-					NextBlock = MapHelper.GetValuePtr(PropertyIndex);
-				}
+				NextProperty = MapHelper.GetValueProperty();
+				NextBlock = MapHelper.FindNthPairPtr(PropertyIndex);
 			}
 		}
 		
diff --git a/Engine/Source/Runtime/CoreUObject/Public/UObject/UnrealType.h b/Engine/Source/Runtime/CoreUObject/Public/UObject/UnrealType.h
--- a/Engine/Source/Runtime/CoreUObject/Public/UObject/UnrealType.h
+++ b/Engine/Source/Runtime/CoreUObject/Public/UObject/UnrealType.h
@@ -4662,6 +4662,62 @@ public:
 	{
 		return const_cast<FScriptMapHelper*>(this)->GetPairPtr(Index);
 	}
+	
+	/**
+	* Returns a uint8 pointer to the the Nth valid pair in the map (skipping invalid entries).
+	* NOTE: This is slow, do not use this for iteration! Use CreateIterator() instead.
+	*
+	* @return Pointer to the element, or nullptr if the index is invalid.
+	*/
+	uint8* FindNthPairPtr(int32 N)
+	{
+		const int32 Index = FindInternalIndex(N);
+		
+		checkSlow(IsValidIndex(Index));
+		return (Index != INDEX_NONE) ? GetPairPtr(FindInternalIndex(Index)) : nullptr;
+	}
+	
+	/**
+	* Returns a uint8 pointer to the the Nth valid key in the map (skipping invalid entries).
+	* NOTE: This is slow, do not use this for iteration! Use CreateIterator() instead.
+	*
+	* @return Pointer to the element, or nullptr if the index is invalid.
+	*/
+	uint8* FindNthKeyPtr(int32 N)
+	{
+		const int32 Index = FindInternalIndex(N);
+		
+		checkSlow(IsValidIndex(Index));
+		return (Index != INDEX_NONE) ? GetKeyPtr(FindInternalIndex(Index)) : nullptr;
+	}
+	
+	/**
+	* Returns a uint8 pointer to the the Nth valid value in the map (skipping invalid entries).
+	* NOTE: This is slow, do not use this for iteration! Use CreateIterator() instead.
+	*
+	* @return Pointer to the element, or nullptr if the index is invalid.
+	*/
+	uint8* FindNthValuePtr(int32 N)
+	{
+		const int32 Index = FindInternalIndex(N);
+		
+		checkSlow(IsValidIndex(Index));
+		return (Index != INDEX_NONE) ? GetValuePtr(FindInternalIndex(Index)) : nullptr;
+	}
+	
+	/**
+	* Returns a uint8 pointer to the the Nth valid pair in the map (skipping invalid entries).
+	* NOTE: This is slow, do not use this for iteration! Use CreateIterator() instead.
+	*
+	* @return Pointer to the element, or nullptr if the index is invalid.
+	*/
+	const uint8* FindNthPairPtr(int32 N) const
+	{
+		const int32 Index = FindInternalIndex(N);
+		
+		checkSlow(IsValidIndex(Index));
+		return (Index != INDEX_NONE) ? GetPairPtr(FindInternalIndex(Index)) : nullptr;
+	}
 
 	/**
 	 * Move the allocation from another map and make it our own.
@@ -4800,6 +4856,12 @@ public:
 				return INDEX_NONE;
 			}
 
+			// if map is compact, use random access
+			if (Num() == GetMaxIndex())
+			{
+				return LogicalIdx;
+			}
+
 			int32 MaxIndex = Map->GetMaxIndex();
 			for (int32 Actual = 0; Actual < MaxIndex; ++Actual)
 			{
@@ -5351,6 +5413,32 @@ public:
 		return const_cast<FScriptSetHelper*>(this)->GetElementPtr(Index);
 	}
 
+	/**
+	* Returns a uint8 pointer to the the Nth valid element in the set (skipping invalid entries).
+	* NOTE: This is slow, do not use this for iteration! Use CreateIterator() instead.
+	*
+	* @return Pointer to the element, or nullptr if the index is invalid.
+	*/
+	uint8* FindNthElementPtr(int32 N)
+	{
+		const int32 Index = FindInternalIndex(N);
+		
+		checkSlow(IsValidIndex(Index));
+		return (Index != INDEX_NONE) ? GetElementPtr(FindInternalIndex(Index)) : nullptr;
+	}
+
+	/**
+	* Returns a uint8 pointer to the the Nth valid element in the set (skipping invalid entries).
+	* NOTE: This is slow, do not use this for iteration! Use CreateIterator() instead.
+	*
+	* @return Pointer to the element, or nullptr if the index is invalid.
+	*/
+	const uint8* FindNthElementPtr(int32 N) const
+	{
+		const int32 Index = FindInternalIndex(N);
+		return (Index != INDEX_NONE) ? GetElementPtr(FindInternalIndex(Index)) : nullptr;
+	}
+
 	/**
 	* Move the allocation from another set and make it our own.
 	* @note The sets MUST be of the same type, and this function will NOT validate that!
@@ -5457,6 +5545,12 @@ public:
 			return INDEX_NONE;
 		}
 
+		// if set is compact, use random access
+		if (Num() == GetMaxIndex())
+		{
+			return LogicalIdx;
+		}
+
 		int32 MaxIndex = GetMaxIndex();
 		for (int32 Actual = 0; Actual < MaxIndex; ++Actual)
 		{
