Ensure that ObjectHandles can compute hashes when referencing objects with an invalid index.  This can happen during the early initialization phase of the engine before UObjectBaseInit|UObjectProcessRegistrations get called and register deferred objects.

diff --git a/Engine/Source/Runtime/CoreUObject/Private/UObject/ObjectHandleIndex.cpp b/Engine/Source/Runtime/CoreUObject/Private/UObject/ObjectHandleIndex.cpp
--- a/Engine/Source/Runtime/CoreUObject/Private/UObject/ObjectHandleIndex.cpp
+++ b/Engine/Source/Runtime/CoreUObject/Private/UObject/ObjectHandleIndex.cpp
@@ -522,9 +522,13 @@ namespace UE::CoreUObject::Private
 		return PackedObjectRef;
 	}
 
-	FPackedObjectRef GetPackedObjectRef(const UObject& Object)
+	FPackedObjectRef FindExistingPackedObjectRef(const UObject& Object)
 	{
 		int32 ObjectIndex = GUObjectArray.ObjectToIndex(&Object);
+		if (ObjectIndex == INDEX_NONE)
+		{
+			return FPackedObjectRef {0};
+		}
 		return UE::CoreUObject::Private::GObjectHandleIndex.ObjectIndexToPackedObjectRef[ObjectIndex];
 	}
 }
diff --git a/Engine/Source/Runtime/CoreUObject/Public/UObject/ObjectHandle.h b/Engine/Source/Runtime/CoreUObject/Public/UObject/ObjectHandle.h
--- a/Engine/Source/Runtime/CoreUObject/Public/UObject/ObjectHandle.h
+++ b/Engine/Source/Runtime/CoreUObject/Public/UObject/ObjectHandle.h
@@ -11,10 +11,10 @@
 #include "Templates/TypeHash.h"
 #include "UObject/NameTypes.h"
 #include "UObject/Object.h"
+#include "UObject/ObjectHandleTracking.h"
 #include "UObject/ObjectMacros.h"
 #include "UObject/ObjectRef.h"
 #include "UObject/PackedObjectRef.h"
-#include "UObject/ObjectHandleTracking.h"
 
 
 class UClass;
@@ -60,7 +60,7 @@ namespace UE::CoreUObject::Private
 	};
 
 	/* Returns the packed object ref for this object IF one exists otherwise returns a null PackedObjectRef */
-	COREUOBJECT_API FPackedObjectRef GetPackedObjectRef(const UObject& Object);
+	COREUOBJECT_API FPackedObjectRef FindExistingPackedObjectRef(const UObject& Object);
 
 	/* Creates and ObjectRef from a packed object ref*/
 	COREUOBJECT_API FObjectRef MakeObjectRef(FPackedObjectRef Handle);
@@ -141,7 +141,7 @@ inline bool operator==(UE::CoreUObject::Private::FObjectHandlePrivate LHS, UE::C
 		}
 
 		//if packed ref empty then can't be equal as RHS is an unresolved pointer
-		FPackedObjectRef PackedLhs = GetPackedObjectRef(*Obj);
+		FPackedObjectRef PackedLhs = FindExistingPackedObjectRef(*Obj);
 		if (PackedLhs.EncodedRef == 0)
 		{
 			return false;
@@ -159,7 +159,7 @@ inline bool operator==(UE::CoreUObject::Private::FObjectHandlePrivate LHS, UE::C
 		}
 
 		//if packed ref empty then can't be equal as RHS is an unresolved pointer
-		FPackedObjectRef PackedRhs = GetPackedObjectRef(*Obj);
+		FPackedObjectRef PackedRhs = FindExistingPackedObjectRef(*Obj);
 		if (PackedRhs.EncodedRef == 0)
 		{
 			return false;
@@ -186,7 +186,8 @@ inline uint32 GetTypeHash(UE::CoreUObject::Private::FObjectHandlePrivate Handle)
 	if (IsObjectHandleResolved(Handle))
 	{
 		const UObject* Obj = ReadObjectHandlePointerNoCheck(Handle);
-		FPackedObjectRef PackedObjectRef = GetPackedObjectRef(*Obj);
+
+		FPackedObjectRef PackedObjectRef = FindExistingPackedObjectRef(*Obj);
 		if (PackedObjectRef.EncodedRef == 0)
 		{
 			return GetTypeHash(Obj);
diff --git a/Engine/Source/Runtime/CoreUObject/Tests/ObjectHandleTest.cpp b/Engine/Source/Runtime/CoreUObject/Tests/ObjectHandleTest.cpp
--- a/Engine/Source/Runtime/CoreUObject/Tests/ObjectHandleTest.cpp
+++ b/Engine/Source/Runtime/CoreUObject/Tests/ObjectHandleTest.cpp
@@ -352,4 +352,13 @@ TEST_CASE_METHOD(FObjectHandleTestBase, "CoreUObject::FObjectHandle::Resolve Mal
 }
 #endif // UE_WITH_OBJECT_HANDLE_LATE_RESOLVE
 
+TEST_CASE_METHOD(FObjectHandleTestBase, "CoreUObject::FObjectHandle::Hash Object Without Index", "[CoreUObject][ObjectHandle]")
+{
+	UObject DummyObjectWithInvalidIndex(EC_StaticConstructor, RF_NoFlags);
+	CHECK(DummyObjectWithInvalidIndex.GetUniqueID() == -1);
+
+	FObjectHandle DummyObjectHandle = UE::CoreUObject::Private::MakeObjectHandle(&DummyObjectWithInvalidIndex);
+	CHECK(GetTypeHash(TargetHandle) == GetTypeHash(&DummyObjectWithInvalidIndex));
+}
+
 #endif
