Fix compile error caused by variable rename.

diff --git a/Engine/Source/Runtime/CoreUObject/Tests/ObjectHandleTest.cpp b/Engine/Source/Runtime/CoreUObject/Tests/ObjectHandleTest.cpp
--- a/Engine/Source/Runtime/CoreUObject/Tests/ObjectHandleTest.cpp
+++ b/Engine/Source/Runtime/CoreUObject/Tests/ObjectHandleTest.cpp
@@ -358,7 +358,7 @@ TEST_CASE_METHOD(FObjectHandleTestBase, "CoreUObject::FObjectHandle::Hash Object
 	CHECK(DummyObjectWithInvalidIndex.GetUniqueID() == -1);
 
 	FObjectHandle DummyObjectHandle = UE::CoreUObject::Private::MakeObjectHandle(&DummyObjectWithInvalidIndex);
-	CHECK(GetTypeHash(TargetHandle) == GetTypeHash(&DummyObjectWithInvalidIndex));
+	CHECK(GetTypeHash(DummyObjectHandle) == GetTypeHash(&DummyObjectWithInvalidIndex));
 }
 
 #endif
