Fix potential truncations in ScriptStackTracker

#rnx

- Even though the frame counts will never be high enough to need it, we might as well preserve the uint64 type here. Fixed up the string formatting to account for the type.
- This code does not seem to be used anymore unless LINE_CHECK_TRACING is enabled in LevelActor.cpp and even then the code will not compile.

diff --git a/Engine/Source/Runtime/CoreUObject/Private/UObject/ScriptStackTracker.cpp b/Engine/Source/Runtime/CoreUObject/Private/UObject/ScriptStackTracker.cpp
--- a/Engine/Source/Runtime/CoreUObject/Private/UObject/ScriptStackTracker.cpp
+++ b/Engine/Source/Runtime/CoreUObject/Private/UObject/ScriptStackTracker.cpp
@@ -75,7 +75,7 @@ void FScriptStackTracker::DumpStackTraces( int32 StackThreshold, FOutputDevice&
 	}
 
 	// Calculate the number of frames we captured.
-	int32 FramesCaptured = 0;
+	uint64 FramesCaptured = 0;
 	if( bIsEnabled )
 	{
 		FramesCaptured = GFrameCounter - StartFrameCounter;
@@ -86,7 +86,8 @@ void FScriptStackTracker::DumpStackTraces( int32 StackThreshold, FOutputDevice&
 	}
 
 	// Log quick summary as we don't log each individual so totals in CSV won't represent real totals.
-	Ar.Logf(TEXT("Captured %i unique callstacks totalling %i function calls over %i frames, averaging %5.2f calls/frame"), SortedCallStacks.Num(), (int32)TotalStackCount, FramesCaptured, (float) TotalStackCount / FramesCaptured);
+	Ar.Logf(TEXT("Captured %i unique callstacks totalling %" UINT64_FMT " function calls over %" UINT64_FMT " frames, averaging %5.2f calls/frame"), 
+		SortedCallStacks.Num(), TotalStackCount, FramesCaptured, (float)TotalStackCount / FramesCaptured);
 
 	// Iterate over each callstack and write out info in human readable form in CSV format
 	for( int32 CallStackIndex=0; CallStackIndex<SortedCallStacks.Num(); CallStackIndex++ )
