Excluding relevant unit tests when mutable TObjectPtr containers are disabled

diff --git a/Engine/Source/Runtime/CoreUObject/Tests/ObjectPtrTest.cpp b/Engine/Source/Runtime/CoreUObject/Tests/ObjectPtrTest.cpp
--- a/Engine/Source/Runtime/CoreUObject/Tests/ObjectPtrTest.cpp
+++ b/Engine/Source/Runtime/CoreUObject/Tests/ObjectPtrTest.cpp
@@ -536,6 +536,7 @@ TEST_CASE("CoreUObject::TObjectPtr::Swap")
 
 }
 
+#if !UE_DEPRECATE_MUTABLE_TOBJECTPTR
 TEST_CASE("CoreUObject::TObjectPtr::SwapArray")
 {
 	const FName TestPackageName(TEXT("/Engine/TestPackage"));
@@ -564,7 +565,7 @@ TEST_CASE("CoreUObject::TObjectPtr::SwapArray")
 	CHECK(ArrayRaw[0] == RawPtrA);
 	CHECK(ArrayPtr[0] == RawPtrB);
 }
-
+#endif
 
 TEST_CASE("CoreUObject::TObjectPtr::Move")
 {
@@ -892,6 +893,7 @@ void TestArrayConversion()
 		CHECK(RawArray.Num() == 0);
 	}
 
+#if !UE_DEPRECATE_MUTABLE_TOBJECTPTR
 	//ArrayView
 	{
 		TArrayView<TObj*> RawArray = PtrArray;
@@ -907,6 +909,9 @@ void TestArrayConversion()
 			CHECK(RawArray[i] == Obj1);
 		}
 	}
+#endif
+
+#if !UE_DEPRECATE_MUTABLE_TOBJECTPTR
 	{
 		TArray<TObjectPtr<TObj>> EmptyArray;
 		TArrayView<TObj*> RawArray = EmptyArray;
@@ -918,6 +923,8 @@ void TestArrayConversion()
 #endif
 		CHECK(RawArray.Num() == 0);
 	}
+#endif
+  
 	{
 		const TArray<TObjectPtr<TObj>>& ConstPtrArray = PtrArray;
 		const TArrayView<TObj* const> RawArray = ConstPtrArray;
@@ -1171,4 +1178,4 @@ TEST_CASE("CoreUObject::TObjectPtr::TestEquals")
 	CHECK(BasePtr == ObjPtr);
 }
 
-#endif
\ No newline at end of file
+#endif
