changed ArchiveObjectCrc32 to not resolve unresolved FObjectPtr the path name of an object can be obtained from a FObjectPtr without resolving
https://p4-swarm.epicgames.net/reviews/25174152

diff --git a/Engine/Source/Runtime/CoreUObject/Private/Serialization/ArchiveObjectCrc32.cpp b/Engine/Source/Runtime/CoreUObject/Private/Serialization/ArchiveObjectCrc32.cpp
--- a/Engine/Source/Runtime/CoreUObject/Private/Serialization/ArchiveObjectCrc32.cpp
+++ b/Engine/Source/Runtime/CoreUObject/Private/Serialization/ArchiveObjectCrc32.cpp
@@ -2,6 +2,7 @@
 
 #include "Serialization/ArchiveObjectCrc32.h"
 #include "UObject/Object.h"
+#include "UObject/ObjectPtr.h"
 
 DEFINE_LOG_CATEGORY_STATIC(LogArchiveObjectCrc32, Log, All);
 
@@ -60,6 +61,22 @@ FArchive& FArchiveObjectCrc32::operator<<(class UObject*& Object)
 	return Ar;
 }
 
+FArchive& FArchiveObjectCrc32::operator<<(FObjectPtr& ObjectPtr)
+{
+	FArchive& Ar = *this;
+	if (!ObjectPtr.IsResolved())
+	{
+		auto UniqueName = ObjectPtr.GetPathName();
+		Ar << UniqueName;
+	}
+	else
+	{
+		UObject* Object = ObjectPtr.Get();
+		Ar << Object;
+	}
+	return Ar;
+}
+
 uint32 FArchiveObjectCrc32::Crc32(UObject* Object, UObject* Root, uint32 CRC)
 {
 #ifdef DEBUG_ARCHIVE_OBJECT_CRC32
diff --git a/Engine/Source/Runtime/CoreUObject/Public/Serialization/ArchiveObjectCrc32.h b/Engine/Source/Runtime/CoreUObject/Public/Serialization/ArchiveObjectCrc32.h
--- a/Engine/Source/Runtime/CoreUObject/Public/Serialization/ArchiveObjectCrc32.h
+++ b/Engine/Source/Runtime/CoreUObject/Public/Serialization/ArchiveObjectCrc32.h
@@ -32,6 +32,7 @@ public:
 	virtual void Serialize(void* Data, int64 Length);
 	virtual FArchive& operator<<(class FName& Name);
 	virtual FArchive& operator<<(class UObject*& Object);
+	virtual FArchive& operator<<(FObjectPtr& ObjectPtr) override;
 	virtual FString GetArchiveName() const { return TEXT("FArchiveObjectCrc32"); }
 	//~ End FArchive Interface
 
