changed FSoftObjectPathFixupArchive to ignore FObjectPtr to avoid resolving unresolved FObjectPtrs
https://p4-swarm.epicgames.net/reviews/25174125

diff --git a/Engine/Source/Runtime/CoreUObject/Public/UObject/SoftObjectPath.h b/Engine/Source/Runtime/CoreUObject/Public/UObject/SoftObjectPath.h
--- a/Engine/Source/Runtime/CoreUObject/Public/UObject/SoftObjectPath.h
+++ b/Engine/Source/Runtime/CoreUObject/Public/UObject/SoftObjectPath.h
@@ -570,6 +570,12 @@ struct FSoftObjectPathFixupArchive : public FArchiveUObject
 		return *this;
 	}
 
+	virtual FArchive& operator<<(FObjectPtr& Value) override
+	{
+		//do nothing to avoid resolving
+		return *this;
+	}
+
 	void Fixup(UObject* Root)
 	{
 		Root->Serialize(*this);
