Update TObjectPtr unit tests to assert against specific pointers (instead of all pointers) due to upcoming ARO changes.

diff --git a/Engine/Source/Runtime/CoreUObject/Tests/ObjectPropertyTest.cpp b/Engine/Source/Runtime/CoreUObject/Tests/ObjectPropertyTest.cpp
--- a/Engine/Source/Runtime/CoreUObject/Tests/ObjectPropertyTest.cpp
+++ b/Engine/Source/Runtime/CoreUObject/Tests/ObjectPropertyTest.cpp
@@ -14,19 +14,18 @@
 #include "LowLevelTestsRunner/WarnFilterScope.h"
 #include "Misc/AssetRegistryInterface.h"
 #include "AssetRegistry/AssetData.h"
+#include "ObjectRefTrackingTestBase.h"
 
 TEST_CASE("UE::CoreUObject::FObjectProperty::CheckValidAddress")
 {
-	bool bAllowRead = false;
 #if UE_WITH_OBJECT_HANDLE_TRACKING
-	auto CallbackHandle = UE::CoreUObject::AddObjectHandleReadCallback([&bAllowRead](TArrayView<const UObject* const> Objects)
+	auto CallbackHandle = UE::CoreUObject::AddObjectHandleReferenceResolvedCallback([](const FObjectRef&, UPackage*, UObject*)
 		{
-			if (!bAllowRead)
-				FAIL("Unexpected read during CheckValidObject");
+			FAIL("Unexpected resolve during CheckValidObject");
 		});
 	ON_SCOPE_EXIT
 	{
-		UE::CoreUObject::RemoveObjectHandleReadCallback(CallbackHandle);
+		UE::CoreUObject::RemoveObjectHandleReferenceResolvedCallback(CallbackHandle);
 	};
 #endif
 	UClass* Class = UObjectPtrTestClassWithRef::StaticClass();
@@ -42,17 +41,23 @@ TEST_CASE("UE::CoreUObject::FObjectProperty::CheckValidAddress")
 	UObjectPtrTestClassWithRef* Obj = NewObject<UObjectPtrTestClassWithRef>(TestPackage, TEXT("Object"));
 	UObjectPtrTestClass* Other = NewObject<UObjectPtrTestClass>(Obj, TEXT("Other"));
 
+#if UE_WITH_OBJECT_HANDLE_TRACKING
+	FObjectHandle Handle = MakeUnresolvedHandle(Other);
+	TObjectPtr<UObjectPtrTestClass> ObjectPtr = *reinterpret_cast<TObjectPtr<UObjectPtrTestClass>*>(&Handle);
+#else
+	TObjectPtr<UObjectPtrTestClass> ObjectPtr = Other;
+#endif
+
 	//verify nothing happens by default
 	CHECK(!Obj->ObjectPtr);
 	Property->CheckValidObject(&Obj->ObjectPtr, nullptr);
 	CHECK(!Obj->ObjectPtr);
 
 	//valid assignment
-	Obj->ObjectPtr = Other;
+	Obj->ObjectPtr = ObjectPtr;
 	Property->CheckValidObject(&Obj->ObjectPtr, nullptr);
-	CHECK(Obj->ObjectPtr == Other);
+	CHECK(Obj->ObjectPtr == ObjectPtr);
 
-	bAllowRead = true; //has to read the fullpath from the object to make the error message
 	//assign a bad value to the pointer
 	Obj->ObjectPtr = reinterpret_cast<UObjectPtrTestClass*>(Obj);
 	CHECK(Obj->ObjectPtr != nullptr);
@@ -65,7 +70,7 @@ TEST_CASE("UE::CoreUObject::FObjectProperty::CheckValidAddress")
 			}
 			return false;
 		});
-	Property->CheckValidObject(&Obj->ObjectPtr, Other);
+	Property->CheckValidObject(&Obj->ObjectPtr, ObjectPtr);
 	CHECK(!Obj->ObjectPtr); //value should be nulled since the type was not compatible	
 }
 
@@ -73,16 +78,16 @@ TEST_CASE("UE::CoreUObject::FObjectProperty::CheckValidAddressNonNullable")
 {
 	bool bAllowRead = false;
 #if UE_WITH_OBJECT_HANDLE_TRACKING
-	auto CallbackHandle = UE::CoreUObject::AddObjectHandleReadCallback([&bAllowRead](TArrayView<const UObject* const> Objects)
+	auto CallbackHandle = UE::CoreUObject::AddObjectHandleReferenceResolvedCallback([&bAllowRead](const FObjectRef & SourceRef, UPackage * ClassPackage, UObject* Object)
 		{
 			if (!bAllowRead)
 			{
-				FAIL("Unexpected read during CheckValidObject");
+				FAIL("Unexpected resolve during CheckValidObject");
 			}
 		});
 	ON_SCOPE_EXIT
 	{
-		UE::CoreUObject::RemoveObjectHandleReadCallback(CallbackHandle);
+		UE::CoreUObject::RemoveObjectHandleReferenceResolvedCallback(CallbackHandle);
 	};
 #endif
 	UClass* Class = UObjectPtrTestClassWithRef::StaticClass();
@@ -98,17 +103,24 @@ TEST_CASE("UE::CoreUObject::FObjectProperty::CheckValidAddressNonNullable")
 	UObjectPtrTestClassWithRef* Obj = NewObject<UObjectPtrTestClassWithRef>(TestPackage, TEXT("Object"));
 	UObjectPtrTestClass* Other = NewObject<UObjectPtrTestClass>(Obj, TEXT("Other"));
 
+#if UE_WITH_OBJECT_HANDLE_TRACKING
+	FObjectHandle Handle = MakeUnresolvedHandle(Other);
+	TObjectPtr<UObjectPtrTestClass> ObjectPtr = *reinterpret_cast<TObjectPtr<UObjectPtrTestClass>*>(&Handle);
+#else
+	TObjectPtr<UObjectPtrTestClass> ObjectPtr = Other;
+#endif
+
 	//property is already null should stay null
 	CHECK(!Obj->ObjectPtrNonNullable);
-	Property->CheckValidObject(&Obj->ObjectPtrNonNullable, Other);
+	Property->CheckValidObject(&Obj->ObjectPtrNonNullable, ObjectPtr);
 	CHECK(!Obj->ObjectPtrNonNullable);
 
 	//valid assignment
-	Obj->ObjectPtrNonNullable = Other;
+	Obj->ObjectPtrNonNullable = ObjectPtr;
 	Property->CheckValidObject(&Obj->ObjectPtrNonNullable, nullptr);
-	CHECK(Obj->ObjectPtrNonNullable == Other);
+	CHECK(Obj->ObjectPtrNonNullable == ObjectPtr);
 
-	bAllowRead = true; //has to read the fullpath from the object to make the error message
+	bAllowRead = true; //has resolve the old value to construct a new default value for the property
 	//assign a bad value to the pointer
 	Obj->ObjectPtrNonNullable = reinterpret_cast<UObjectPtrTestClass*>(Obj);
 	CHECK(Obj->ObjectPtrNonNullable != nullptr);
@@ -121,8 +133,8 @@ TEST_CASE("UE::CoreUObject::FObjectProperty::CheckValidAddressNonNullable")
 			}
 			return false;
 		});
-	Property->CheckValidObject(&Obj->ObjectPtrNonNullable, Other);
-	CHECK(Obj->ObjectPtrNonNullable == Other); //non nullable properties should be assigned the old value
+	Property->CheckValidObject(&Obj->ObjectPtrNonNullable, ObjectPtr);
+	CHECK(Obj->ObjectPtrNonNullable == ObjectPtr); //non nullable properties should be assigned the old value
 
 	//assign a bad value to the pointer
 	Obj->ObjectPtrNonNullable = reinterpret_cast<UObjectPtrTestClass*>(Obj);
@@ -165,15 +177,15 @@ static void TestSerializeItem(FName ObjectName)
 	ULinkerPlaceholderClass* PlaceHolderClass = NewObject<ULinkerPlaceholderClass>(TestPackage, TEXT("PlaceHolderClass"));
 	PlaceHolderClass->Bind(); //must call bind or crashes on shutdown
 
-	int ReadCount = 0;
+	int ResolveCount = 0;
 #if UE_WITH_OBJECT_HANDLE_TRACKING
-	auto CallbackId = UE::CoreUObject::AddObjectHandleReadCallback([&ReadCount](TArrayView<const UObject* const> Objects)
+	auto CallbackId = UE::CoreUObject::AddObjectHandleReferenceResolvedCallback([&ResolveCount](const FObjectRef&, UPackage*, UObject*)
 		{
-			++ReadCount;
+			++ResolveCount;
 		});
 	ON_SCOPE_EXIT
 	{
-		UE::CoreUObject::RemoveObjectHandleReadCallback(CallbackId);
+		UE::CoreUObject::RemoveObjectHandleReferenceResolvedCallback(CallbackId);
 	};
 #endif
 	{
@@ -188,7 +200,7 @@ static void TestSerializeItem(FName ObjectName)
 		FStructuredArchive Ar(Formatter);
 		FStructuredArchiveSlot Slot = Ar.Open();
 		FObjectPtrProperty::StaticSerializeItem(Property, Slot, &Obj->ObjectPtr, nullptr);
-		CHECK(ReadCount == 0);
+		CHECK(ResolveCount == 0);
 		CHECK(*reinterpret_cast<ULinkerPlaceholderExportObject**>(&Obj->ObjectPtr) == PlaceHolderExport);
 
 	}
@@ -205,7 +217,7 @@ static void TestSerializeItem(FName ObjectName)
 		FStructuredArchiveSlot Slot = Ar.Open();
 		Obj->ObjectPtr = Other;
 		FObjectPtrProperty::StaticSerializeItem(Property, Slot, &Obj->ObjectPtr, nullptr);
-		CHECK(ReadCount == 0);
+		CHECK(ResolveCount == 0);
 		CHECK(*reinterpret_cast<ULinkerPlaceholderClass**>(&Obj->ObjectPtr) == PlaceHolderClass);
 	}
 }
diff --git a/Engine/Source/Runtime/CoreUObject/Tests/ObjectPtrTest.cpp b/Engine/Source/Runtime/CoreUObject/Tests/ObjectPtrTest.cpp
--- a/Engine/Source/Runtime/CoreUObject/Tests/ObjectPtrTest.cpp
+++ b/Engine/Source/Runtime/CoreUObject/Tests/ObjectPtrTest.cpp
@@ -187,20 +187,29 @@ TEST_CASE_METHOD(FObjectPtrTestBase, "CoreUObject::TObjectPtr::Default Serialize
 
 	ObjectRefMetrics.TestNumResolves(TEXT("Unexpected resolve count after initializing an FObjectPtr"), 0);
 	ObjectRefMetrics.TestNumFailedResolves(TEXT("Unexpected resolve failure after initializing an FObjectPtr"), 0);
-	ObjectRefMetrics.TestNumReads(TEXT("NumReads should not change when initializing an FObjectPtr"), 0);
+	ObjectRefMetrics.TestNumReads(TEXT("NumReads should not change when initializing an FObjectPtr"),
+																TestPackage,
+																0);
+	ObjectRefMetrics.TestNumReads(TEXT("NumReads should not change when initializing an FObjectPtr"),
+																TestSoftObject,
+																0);
 
 	FArchiveUObject Writer;
 	Writer << DefaultSerializeObjectPtr;
 
 	ObjectRefMetrics.TestNumResolves(TEXT("Serializing an FObjectPtr should force it to resolve"), UE_WITH_OBJECT_HANDLE_LATE_RESOLVE ? 1 : 0);
 	ObjectRefMetrics.TestNumFailedResolves(TEXT("Unexpected resolve failure after serializing an FObjectPtr"), 0);
-	ObjectRefMetrics.TestNumReads(TEXT("NumReads should increase after serializing an FObjectPtr"), 1);
+	ObjectRefMetrics.TestNumReads(TEXT("NumReads should increase after serializing an FObjectPtr"),
+																TestSoftObject,
+																1);
 
 	Writer << DefaultSerializeObjectPtr;
 
 	ObjectRefMetrics.TestNumResolves(TEXT("Serializing an FObjectPtr twice should only require it to resolve once"), UE_WITH_OBJECT_HANDLE_LATE_RESOLVE ? 1 : 0);
 	ObjectRefMetrics.TestNumFailedResolves(TEXT("Unexpected resolve failure after serializing an FObjectPtr"), 0);
-	ObjectRefMetrics.TestNumReads(TEXT("NumReads should increase after serializing an FObjectPtr"), 2);
+	ObjectRefMetrics.TestNumReads(TEXT("NumReads should increase after serializing an FObjectPtr"),
+																TestSoftObject,
+																2);
 
 	TestPackage->RemoveFromRoot();
 }
@@ -224,12 +233,16 @@ TEST_CASE_METHOD(FObjectPtrTestBase, "CoreUObject::TObjectPtr::Soft Object Path"
 #endif
 	ObjectRefMetrics.TestNumResolves(TEXT("Unexpected resolve count after initializing an FObjectPtr"), 0);
 	ObjectRefMetrics.TestNumFailedResolves(TEXT("Unexpected resolve failure after initializing an FObjectPtr"), 0);
-	ObjectRefMetrics.TestNumReads(TEXT("NumReads should not change when initializing an FObjectPtr"), 0);
+	ObjectRefMetrics.TestNumReads(TEXT("NumReads should not change when initializing an FObjectPtr"),
+																TestSoftObject,
+																0);
 
 	// Initializing a soft object path from a TObjectPtr that's unresolved should stay unresolved.
 	FSoftObjectPath DefaultSoftObjPath(DefaultSoftObjPtr);
 	ObjectRefMetrics.TestNumResolves(TEXT("Unexpected resolve count after initializing an FSoftObjectPath from an FObjectPtr"), 0);
-	ObjectRefMetrics.TestNumReads(TEXT("NumReads should have changed when initializing a FSoftObjectPath"), UE_WITH_OBJECT_HANDLE_LATE_RESOLVE ? 0 : 1);
+	ObjectRefMetrics.TestNumReads(TEXT("NumReads should have changed when initializing a FSoftObjectPath"),
+																TestSoftObject,
+																UE_WITH_OBJECT_HANDLE_LATE_RESOLVE ? 0 : 1);
 
 	TEST_EQUAL_STR(TEXT("Soft object path constructed from an FObjectPtr does not have the expected path value"), TEXT("/Engine/Test/ObjectPtrSoftObjectPath/Transient.TestSoftObject"), *DefaultSoftObjPath.ToString());
 
diff --git a/Engine/Source/Runtime/CoreUObject/Tests/ObjectRefTrackingTestBase.cpp b/Engine/Source/Runtime/CoreUObject/Tests/ObjectRefTrackingTestBase.cpp
--- a/Engine/Source/Runtime/CoreUObject/Tests/ObjectRefTrackingTestBase.cpp
+++ b/Engine/Source/Runtime/CoreUObject/Tests/ObjectRefTrackingTestBase.cpp
@@ -11,6 +11,7 @@
 thread_local uint32 FObjectRefTrackingTestBase::NumResolves = 0;
 thread_local uint32 FObjectRefTrackingTestBase::NumFailedResolves = 0;
 thread_local uint32 FObjectRefTrackingTestBase::NumReads = 0;
+thread_local TMap<const UObject* const, uint32> FObjectRefTrackingTestBase::NumReadsPerObject{};
 
 #if UE_WITH_OBJECT_HANDLE_TRACKING
 UE::CoreUObject::FObjectHandleTrackingCallbackId FObjectRefTrackingTestBase::ResolvedCallbackHandle;
@@ -235,6 +236,9 @@ TEST_CASE("UE::CoreUObject::ObjectHandleTracking::Callbacks::Verify")
 	//get fires read event
 	CHECK(ObjPtr.Get() != nullptr);
 	CHECK(ReadCount == 1);
+
+	RemoveObjectHandleReadCallback(ReadCallbackHandle1);
+	
 	FObjectRef ObjectRef(Obj1);
 	int32 ClassCount = 0;
 	auto ClassResolvedHandle1 = AddObjectHandleClassResolvedCallback([&](const FObjectRef& SourceRef, UPackage* ClassPackage, UClass* Class)
@@ -274,7 +278,6 @@ TEST_CASE("UE::CoreUObject::ObjectHandleTracking::Callbacks::Verify")
 
 	ON_SCOPE_EXIT
 	{
-		RemoveObjectHandleReadCallback(ReadCallbackHandle1);
 		RemoveObjectHandleClassResolvedCallback(ClassResolvedHandle1);
 		RemoveObjectHandleReferenceResolvedCallback(ReferenceResolvedHandle1);
 		RemoveObjectHandleReferenceLoadedCallback(ReferenceLoadedHandle1);
@@ -282,4 +285,4 @@ TEST_CASE("UE::CoreUObject::ObjectHandleTracking::Callbacks::Verify")
 	
 }
 #endif
-#endif
\ No newline at end of file
+#endif
diff --git a/Engine/Source/Runtime/CoreUObject/Tests/ObjectRefTrackingTestBase.h b/Engine/Source/Runtime/CoreUObject/Tests/ObjectRefTrackingTestBase.h
--- a/Engine/Source/Runtime/CoreUObject/Tests/ObjectRefTrackingTestBase.h
+++ b/Engine/Source/Runtime/CoreUObject/Tests/ObjectRefTrackingTestBase.h
@@ -25,7 +25,16 @@ public:
 	uint32 GetNumResolves() const { return NumResolves; }
 	uint32 GetNumFailedResolves() const { return NumFailedResolves; }
 	uint32 GetNumReads() const { return NumReads; }
-
+	uint32 GetNumReads(const UObject* const Obj)
+	{
+		const auto* N = NumReadsPerObject.Find(Obj);
+		return N ? *N : 0;
+	}
+	TMap<const UObject* const, uint32>& GetNumReadsPerObject()
+	{
+		return NumReadsPerObject;
+	}
+	
 	struct FSnapshotObjectRefMetrics
 	{
 	public:
@@ -34,6 +43,7 @@ public:
 		, OriginalNumResolves(Test.GetNumResolves())
 		, OriginalNumFailedResolves(Test.GetNumFailedResolves())
 		, OriginalNumReads(Test.GetNumReads())
+		, OriginalNumReadsPerObject(Test.GetNumReadsPerObject())
 		{
 			Test.InstallCallbacks();
 		}
@@ -41,6 +51,7 @@ public:
 		~FSnapshotObjectRefMetrics()
 		{
 			Test.RemoveCallbacks();
+			Test.NumReadsPerObject.Empty();
 		}
 
 		void TestNumResolves(const TCHAR* What, uint32 ExpectedDelta)
@@ -73,11 +84,30 @@ public:
 			}
 #endif
 		}
+
+		void TestNumReads(const TCHAR* What, const UObject* const Obj, uint32 ExpectedDelta, bool bAllowAdditionalReads = false)
+		{
+#if UE_WITH_OBJECT_HANDLE_TRACKING
+			bool bValue = false;
+			if (bAllowAdditionalReads)
+			{
+				INFO(What);
+				CHECK(Test.GetNumReads(Obj) >= OriginalNumReadsPerObject.FindOrAdd(Obj, 0) + ExpectedDelta);
+			}
+			else
+			{
+				bValue = OriginalNumReadsPerObject.FindOrAdd(Obj, 0) + ExpectedDelta == Test.GetNumReads(Obj);
+				TEST_TRUE(What, bValue);
+			}
+#endif
+		}
+		
 	private:
 		FObjectRefTrackingTestBase& Test;
 		uint32 OriginalNumResolves;
 		uint32 OriginalNumFailedResolves;
 		uint32 OriginalNumReads;
+		TMap<const UObject* const, uint32> OriginalNumReadsPerObject;
 	};
 
 private:
@@ -93,6 +123,10 @@ private:
 	static void OnRefRead(TArrayView<const UObject* const> Objects)
 	{
 		NumReads++;
+		for (auto Obj : Objects)
+		{
+			++NumReadsPerObject.FindOrAdd(Obj, 0);
+		}
 	}
 #endif
 	
@@ -119,6 +153,7 @@ private:
 	static thread_local uint32 NumResolves;
 	static thread_local uint32 NumFailedResolves;
 	static thread_local uint32 NumReads;
+	static thread_local TMap<const UObject* const, uint32> NumReadsPerObject;
 };
 
-#endif
\ No newline at end of file
+#endif
diff --git a/Engine/Source/Runtime/CoreUObject/Tests/PropertyClassTest.cpp b/Engine/Source/Runtime/CoreUObject/Tests/PropertyClassTest.cpp
--- a/Engine/Source/Runtime/CoreUObject/Tests/PropertyClassTest.cpp
+++ b/Engine/Source/Runtime/CoreUObject/Tests/PropertyClassTest.cpp
@@ -15,16 +15,14 @@
 
 TEST_CASE("UE::CoreUObject::FClassProperty::Identical")
 {
-	bool bAllowRead = false;
 #if UE_WITH_OBJECT_HANDLE_TRACKING
-	auto CallbackHandle = UE::CoreUObject::AddObjectHandleReadCallback([&bAllowRead](TArrayView<const UObject* const> Objects)
+	auto CallbackHandle = UE::CoreUObject::AddObjectHandleReferenceResolvedCallback([](const FObjectRef&, UPackage*, UObject*)
 		{
-			if (!bAllowRead)
-				FAIL("Unexpected read during CheckValidObject");
+			FAIL("Unexpected resolve during CheckValidObject");
 		});
 	ON_SCOPE_EXIT
 	{
-		UE::CoreUObject::RemoveObjectHandleReadCallback(CallbackHandle);
+		UE::CoreUObject::RemoveObjectHandleReferenceResolvedCallback(CallbackHandle);
 	};
 #endif
 	UClass* Class = UObjectWithClassProperty::StaticClass();
@@ -43,4 +41,4 @@ TEST_CASE("UE::CoreUObject::FClassProperty::Identical")
 	CHECK(Property->Identical(Obj, Obj, 0u));
 
 }
-#endif
\ No newline at end of file
+#endif
