Update the experimental FEditorBulkData sidecar saving feature to use FPackageTrailer

#rnx

- The functionality remains mostly the same and is still disabled by default.
- We use the older legacy serialization path if sidecar saving is enabled.
- Updating the code to use FPackageTrailer allows us to finally delete the older which was the precursor to the FPackageTrailer system. This will allow us to combine some code paths in the future and reduce complexity.

diff --git a/Engine/Source/Runtime/CoreUObject/Private/Serialization/EditorBulkData.cpp b/Engine/Source/Runtime/CoreUObject/Private/Serialization/EditorBulkData.cpp
--- a/Engine/Source/Runtime/CoreUObject/Private/Serialization/EditorBulkData.cpp
+++ b/Engine/Source/Runtime/CoreUObject/Private/Serialization/EditorBulkData.cpp
@@ -1393,50 +1393,20 @@ FCompressedBuffer FEditorBulkData::LoadFromSidecarFileInternal(ErrorVerbosity Ve
 	FOpenPackageResult Result = IPackageResourceManager::Get().OpenReadPackage(PackagePath, EPackageSegment::PayloadSidecar);
 	if (Result.Archive.IsValid() && Result.Format == EPackageFormat::Binary)
 	{
-		uint32 Version = INDEX_NONE;
-		*Result.Archive << Version;
-
-		if (Version != FTocEntry::PayloadSidecarFileVersion)
+		FPackageTrailer Trailer;
+		if (!FPackageTrailer::TryLoadFromArchive(*Result.Archive, Trailer))
 		{
-			UE_CLOG(Verbosity > ErrorVerbosity::None, LogSerialization, Error, TEXT("Unknown payload sidecar version (%u) found in '%s'"), Version, *PackagePath.GetLocalFullPath(EPackageSegment::PayloadSidecar));
+			UE_CLOG(Verbosity > ErrorVerbosity::None, LogSerialization, Error, TEXT("Unable to parse FPackageTrailer in '%s'"), *PackagePath.GetLocalFullPath(EPackageSegment::PayloadSidecar));
 			return FCompressedBuffer();
 		}
 
-		// First we load the table of contents so we can find the payload in the file
-		TArray<FTocEntry> TableOfContents;
-		*Result.Archive << TableOfContents;
-
-		const FTocEntry* Entry = TableOfContents.FindByPredicate([&PayloadContentId = PayloadContentId](const FTocEntry& Entry)
-			{
-				return Entry.Identifier == PayloadContentId;
-			});
+		FCompressedBuffer CompressedPayload = Trailer.LoadLocalPayload(PayloadContentId, *Result.Archive);
 
-		if (Entry != nullptr)
-		{
-			if (Entry->OffsetInFile != INDEX_NONE)
-			{
-				// Move the correct location of the data in the file
-				Result.Archive->Seek(Entry->OffsetInFile);
-
-				// Now we can actually serialize it
-				FCompressedBuffer PayloadFromDisk;
-				SerializeData(*Result.Archive, PayloadFromDisk, EFlags::None);
+		UE_CLOG(CompressedPayload.IsNull(), LogSerialization, Error, TEXT("Could not find the payload '%s' in '%s'"),
+			*LexToString(PayloadContentId),
+			*PackagePath.GetLocalFullPath(EPackageSegment::PayloadSidecar));
 
-				return PayloadFromDisk;
-			}
-			else if(Verbosity > ErrorVerbosity::None)
-			{
-				UE_LOG(LogSerialization, Error, TEXT("Payload '%s' in '%s' has an invalid OffsetInFile!"), 
-					*LexToString(PayloadContentId), 
-					*PackagePath.GetLocalFullPath(EPackageSegment::PayloadSidecar));
-			}
-		}
-		else if(Verbosity > ErrorVerbosity::None)
-		{
-			UE_LOG(LogSerialization, Error, TEXT("Unable to find payload '%s' in '%s'"), 
-				*LexToString(PayloadContentId), 
-				*PackagePath.GetLocalFullPath(EPackageSegment::PayloadSidecar));
-		}
+		return CompressedPayload;
 	}
 	else if(Verbosity > ErrorVerbosity::None)
 	{
@@ -2388,72 +2358,12 @@ bool FEditorBulkData::ShouldUseLegacySerialization(const FLinkerSave* LinkerSave
 		return true;
 	}
 
-	return !LinkerSave->PackageTrailerBuilder.IsValid();
-}
-	
-FArchive& operator<<(FArchive& Ar, FTocEntry& Entry)
-{
-	Ar << Entry.Identifier;
-	Ar << Entry.OffsetInFile;
-	Ar << Entry.UncompressedSize;
-
-	return Ar;
-}
-
- void operator<<(FStructuredArchive::FSlot Slot, FTocEntry& Entry)
-{
-	FStructuredArchive::FRecord Record = Slot.EnterRecord();
-
-	Record << SA_VALUE(TEXT("Identifier"), Entry.Identifier);
-	Record << SA_VALUE(TEXT("OffsetInFile"), Entry.OffsetInFile);
-	Record << SA_VALUE(TEXT("UncompressedSize"), Entry.UncompressedSize);
-}
-
-void FPayloadToc::AddEntry(const FEditorBulkData& BulkData)
-{
-	if (!BulkData.GetPayloadId().IsZero())
-	{
-		Contents.Emplace(BulkData);
-	}
-}
-
-bool FPayloadToc::FindEntry(const FIoHash& Identifier, FTocEntry& OutEntry)
-{
-	for (const FTocEntry& Entry : Contents)
+	if (ShouldSaveToPackageSidecar())
 	{
-		if (Entry.Identifier == Identifier)
-		{
-			OutEntry = Entry;
-			return true;
-		}
+		return true;
 	}
 
-	return false;
-}
-
-const TArray<FTocEntry>& FPayloadToc::GetContents() const
-{
-	return Contents;
-}
-
-FArchive& operator<<(FArchive& Ar, FPayloadToc& TableOfContents)
-{
-	FPayloadToc::EVersion Version = FPayloadToc::EVersion::AUTOMATIC_VERSION;
-	Ar << Version;
-
-	Ar << TableOfContents.Contents;
-
-	return Ar;
-}
-
-void operator<<(FStructuredArchive::FSlot Slot, FPayloadToc& TableOfContents)
-{
-	FStructuredArchive::FRecord Record = Slot.EnterRecord();
-
-	FPayloadToc::EVersion Version = FPayloadToc::EVersion::AUTOMATIC_VERSION;
-
-	Record << SA_VALUE(TEXT("Version"), Version);
-	Record << SA_VALUE(TEXT("Entries"), TableOfContents.Contents);
+	return !LinkerSave->PackageTrailerBuilder.IsValid();
 }
 
 } // namespace UE::Serialization
diff --git a/Engine/Source/Runtime/CoreUObject/Private/UObject/PackageTrailer.cpp b/Engine/Source/Runtime/CoreUObject/Private/UObject/PackageTrailer.cpp
--- a/Engine/Source/Runtime/CoreUObject/Private/UObject/PackageTrailer.cpp
+++ b/Engine/Source/Runtime/CoreUObject/Private/UObject/PackageTrailer.cpp
@@ -246,6 +246,12 @@ void FPackageTrailerBuilder::AddPayload(const FIoHash& Identifier, FCompressedBu
 {
 	Callbacks.Emplace(MoveTemp(Callback));
 
+	AddPayload(Identifier, Payload, FilterFlags);
+
+}
+
+void FPackageTrailerBuilder::AddPayload(const FIoHash& Identifier, FCompressedBuffer Payload, UE::Virtualization::EPayloadFilterReason FilterFlags)
+{
 	if (!Identifier.IsZero())
 	{
 		// If the payload already exists and the DisableVirtualization flag has been passed in
diff --git a/Engine/Source/Runtime/CoreUObject/Private/UObject/SavePackage/SavePackageUtilities.cpp b/Engine/Source/Runtime/CoreUObject/Private/UObject/SavePackage/SavePackageUtilities.cpp
--- a/Engine/Source/Runtime/CoreUObject/Private/UObject/SavePackage/SavePackageUtilities.cpp
+++ b/Engine/Source/Runtime/CoreUObject/Private/UObject/SavePackage/SavePackageUtilities.cpp
@@ -1934,49 +1934,24 @@ ESavePackageResult CreatePayloadSidecarFile(FLinkerSave& Linker, const FPackageP
 	checkf(!Linker.IsCooking(), TEXT("Cannot write a sidecar file during cooking! (%s)"), *PackagePath.GetDebugName());
 	IPackageWriter* PackageWriter = SavePackageContext ? SavePackageContext->PackageWriter : nullptr;
 
-	FLargeMemoryWriter Ar(0, true /* bIsPersistent */);
-
-	uint32 VersionNumber = UE::Serialization::FTocEntry::PayloadSidecarFileVersion;
-	Ar << VersionNumber;
-
-	int64 TocPosition = Ar.Tell();
-
-	TArray<UE::Serialization::FTocEntry> TableOfContents;
-	TableOfContents.SetNum(Linker.SidecarDataToAppend.Num());
+	UE::FPackageTrailerBuilder Builder;
 
-	// First we write an empty table of contents to the file to reserve the space
-	Ar << TableOfContents;
-
-	int32 Index = 0;
 	for (FLinkerSave::FSidecarStorageInfo& Info : Linker.SidecarDataToAppend)
 	{
-		// Fill out the entry to the table of contents
-		TableOfContents[Index].Identifier = Info.Identifier;
-		TableOfContents[Index].OffsetInFile = Ar.Tell();
-		TableOfContents[Index].UncompressedSize = Info.Payload.GetRawSize();
-
-		Index++;
+		Builder.AddPayload(Info.Identifier, Info.Payload, UE::Virtualization::EPayloadFilterReason::None);
+	}
 
-		// Now write the payload to the archive
-		for (const FSharedBuffer& Buffer : Info.Payload.GetCompressed().GetSegments())
-		{
-			// Const cast because FArchive requires a non-const pointer!
-			Ar.Serialize(const_cast<void*>(Buffer.GetData()), static_cast<int64>(Buffer.GetSize()));
-		}
+	Linker.SidecarDataToAppend.Empty();
 
-		// Reset each payload reference after it has been written to the archive, this could
-		// potentially release memory and keep our high water mark down.
-		Info.Payload.Reset();
+	FLargeMemoryWriter Ar(0, true /* bIsPersistent */);
+	if (!Builder.BuildAndAppendTrailer(nullptr, Ar))
+	{
+		UE_LOG(LogSavePackage, Error, TEXT("Failed to build sidecar package trailer for '%s'"), *PackagePath.GetDebugName());
+		return ESavePackageResult::Error;
 	}
 
-	// Now write out the table of contents again but with valid data
-	int64 EndPos = Ar.Tell();
-	Ar.Seek(TocPosition);
-	Ar << TableOfContents;
-	Ar.Seek(EndPos); 
-
 	const int64 DataSize = Ar.TotalSize();
-	checkf(DataSize > 0, TEXT("The archive should not be empty at this point!"));
+	checkf(DataSize > 0, TEXT("Sidecar archive should not be empty! (%s)"), *PackagePath.GetDebugName());
 
 	FString TargetFilePath = PackagePath.GetLocalFullPath(EPackageSegment::PayloadSidecar);
 
@@ -1984,7 +1959,7 @@ ESavePackageResult CreatePayloadSidecarFile(FLinkerSave& Linker, const FPackageP
 	{
 		IPackageWriter::FAdditionalFileInfo SidecarSegmentInfo;
 		SidecarSegmentInfo.PackageName = PackagePath.GetPackageFName();
-		SidecarSegmentInfo.Filename = TargetFilePath;
+		SidecarSegmentInfo.Filename = MoveTemp(TargetFilePath);
 		FIoBuffer FileData(FIoBuffer::AssumeOwnership, Ar.ReleaseOwnership(), DataSize);
 		PackageWriter->WriteAdditionalFile(SidecarSegmentInfo, FileData);
 	}
@@ -2003,8 +1978,6 @@ ESavePackageResult CreatePayloadSidecarFile(FLinkerSave& Linker, const FPackageP
 		AdditionalPackageFiles.Emplace(MoveTemp(TargetFilePath), MoveTemp(TempFilePath), DataSize);		
 	}
 
-	Linker.SidecarDataToAppend.Empty();
-
 	return ESavePackageResult::Success;
 }
 
diff --git a/Engine/Source/Runtime/CoreUObject/Public/Serialization/EditorBulkData.h b/Engine/Source/Runtime/CoreUObject/Public/Serialization/EditorBulkData.h
--- a/Engine/Source/Runtime/CoreUObject/Public/Serialization/EditorBulkData.h
+++ b/Engine/Source/Runtime/CoreUObject/Public/Serialization/EditorBulkData.h
@@ -560,70 +560,6 @@ private:
 
 ENUM_CLASS_FLAGS(FEditorBulkData::EFlags);
 
-/** 
- * NOTE: FPayloadToc/FTocEntry are now strictly legacy and are only used by the experimental sidecar system.
- * They will be removed when the sidecar storage is changed to use FPackageTrailer instead.
- */
-
-/** 
-  * Represents an entry to the table of contents found at the start of a payload sidecar file.
-  * This might be moved to it's own header and the table of contents made into a proper class 
-  * if we decide that we want to make access of the payload sidecar file a generic feature.
-  */
-struct COREUOBJECT_API FTocEntry
-{
-	FTocEntry() = default;
-	FTocEntry(const FEditorBulkData& BulkData)
-		: Identifier(BulkData.PayloadContentId)
-		, OffsetInFile(BulkData.OffsetInFile)
-		, UncompressedSize(BulkData.PayloadSize)
-	{
-
-	}
-
-	friend FArchive& operator<<(FArchive& Ar, FTocEntry& Entry);
-	friend void operator<<(FStructuredArchive::FSlot Slot, FTocEntry& Entry);
-
-	static constexpr uint32 PayloadSidecarFileVersion = 1;
-
-	/** Identifier for the payload */
-	FIoHash Identifier;
-	/** The offset into the file where we can find the payload, note that a virtualized payload will have an offset of INDEX_NONE */
-	int64 OffsetInFile = INDEX_NONE;
-	/** The size of the payload when uncompressed. */
-	int64 UncompressedSize = INDEX_NONE;
-};
-
-/** 
- * A table of contents showing the location of all virtualized bulkdata payloads in a file. 
- * NOTE: This is currently only used by the sidecar functionality.
- */
-class COREUOBJECT_API FPayloadToc
-{
-public:
-
-	void AddEntry(const FEditorBulkData& BulkData);
-	bool FindEntry(const FIoHash& Identifier, FTocEntry& OutEntry);
-
-	const TArray<FTocEntry>& GetContents() const;
-
-	friend FArchive& operator<<(FArchive& Ar, FPayloadToc& TableOfContents);
-	friend void operator<<(FStructuredArchive::FSlot Slot, FPayloadToc& TableOfContents);
-
-private:
-	enum class EVersion : uint32
-	{
-		INITIAL = 0,
-
-		// -----<new versions can be added before this line>-------------------------------------------------
-		// - this needs to be the last line (see note below)
-		AUTOMATIC_VERSION_PLUS_ONE,
-		AUTOMATIC_VERSION = AUTOMATIC_VERSION_PLUS_ONE - 1
-	};
-
-	TArray<FTocEntry> Contents;
-};
-
 /** 
  * Returns a FGuid representation of a FIoHash, to be used in existing code paths that require the id of 
  * an FEditorBulkData payload to be in FGuid form 
diff --git a/Engine/Source/Runtime/CoreUObject/Public/UObject/PackageTrailer.h b/Engine/Source/Runtime/CoreUObject/Public/UObject/PackageTrailer.h
--- a/Engine/Source/Runtime/CoreUObject/Public/UObject/PackageTrailer.h
+++ b/Engine/Source/Runtime/CoreUObject/Public/UObject/PackageTrailer.h
@@ -252,7 +252,7 @@ public:
 	 * @param Callback		This callback will be invoked once the FPackageTrailer has been built and appended to disk.
 	 */
 	void AddPayload(const FIoHash& Identifier, FCompressedBuffer Payload, UE::Virtualization::EPayloadFilterReason Filter, AdditionalDataCallback&& Callback);
-
+	void AddPayload(const FIoHash& Identifier, FCompressedBuffer Payload, UE::Virtualization::EPayloadFilterReason Filter);
 	/**
 	 * Adds an already virtualized payload to the builder to be written to the trailer. When the trailer is written
 	 * the payload will have EPayloadAccessMode::Virtualized set as it's access mode. It is assumed that the payload
