[BugFix] Optimizations to FScriptMapHelper::FindInternalIndex and FScriptSetHelper::FindInternalIndex aren't returning INDEX_NONE when someone passes in an index outside the bounds.

diff --git a/Engine/Source/Runtime/CoreUObject/Public/UObject/UnrealType.h b/Engine/Source/Runtime/CoreUObject/Public/UObject/UnrealType.h
--- a/Engine/Source/Runtime/CoreUObject/Public/UObject/UnrealType.h
+++ b/Engine/Source/Runtime/CoreUObject/Public/UObject/UnrealType.h
@@ -4859,7 +4859,7 @@ public:
 			// if map is compact, use random access
 			if (Num() == GetMaxIndex())
 			{
-				return LogicalIdx;
+				return IsValidIndex(LogicalIdx) ? LogicalIdx : INDEX_NONE;
 			}
 
 			int32 MaxIndex = Map->GetMaxIndex();
@@ -5548,7 +5548,7 @@ public:
 		// if set is compact, use random access
 		if (Num() == GetMaxIndex())
 		{
-			return LogicalIdx;
+			return IsValidIndex(LogicalIdx) ? LogicalIdx : INDEX_NONE;
 		}
 
 		int32 MaxIndex = GetMaxIndex();
