Add a callback to get notified when the UObject GC finishes purging objects. This is a copy of the changes in `25150493`.

diff --git a/Engine/Source/Runtime/CoreUObject/Private/UObject/GarbageCollection.cpp b/Engine/Source/Runtime/CoreUObject/Private/UObject/GarbageCollection.cpp
--- a/Engine/Source/Runtime/CoreUObject/Private/UObject/GarbageCollection.cpp
+++ b/Engine/Source/Runtime/CoreUObject/Private/UObject/GarbageCollection.cpp
@@ -3952,6 +3952,14 @@ void IncrementalPurgeGarbage(bool bUseTimeLimit, double TimeLimit)
 		{
 			bCompleted = IncrementalDestroyGarbage(bUseTimeLimit, TimeLimit);
 		}
+		
+		if (bCompleted)
+		{
+			// Broadcast the post-purge garbage delegate to give systems a chance to clean up things
+			// that might have been referenced by purged objects.
+			TRACE_CPUPROFILER_EVENT_SCOPE(BroadcastPostPurgeGarbage);
+			FCoreUObjectDelegates::GetPostPurgeGarbageDelegate().Broadcast();
+		}
 
 		// when running incrementally using a time limit, add one last tick for the memory trim
 		bCompleted = bCompleted && !bUseTimeLimit;
diff --git a/Engine/Source/Runtime/CoreUObject/Private/UObject/UObjectGlobals.cpp b/Engine/Source/Runtime/CoreUObject/Private/UObject/UObjectGlobals.cpp
--- a/Engine/Source/Runtime/CoreUObject/Private/UObject/UObjectGlobals.cpp
+++ b/Engine/Source/Runtime/CoreUObject/Private/UObject/UObjectGlobals.cpp
@@ -187,6 +187,12 @@ FSimpleMulticastDelegate& FCoreUObjectDelegates::GetPostGarbageCollect()
 	return Delegate;
 }
 
+FSimpleMulticastDelegate& FCoreUObjectDelegates::GetPostPurgeGarbageDelegate()
+{
+	static FSimpleMulticastDelegate Delegate;
+	return Delegate;
+}
+
 #if !UE_BUILD_SHIPPING
 FCoreUObjectDelegates::FOnReportGarbageReferencers& FCoreUObjectDelegates::GetGarbageCollectReportGarbageReferencers()
 {
diff --git a/Engine/Source/Runtime/CoreUObject/Public/UObject/UObjectGlobals.h b/Engine/Source/Runtime/CoreUObject/Public/UObject/UObjectGlobals.h
--- a/Engine/Source/Runtime/CoreUObject/Public/UObject/UObjectGlobals.h
+++ b/Engine/Source/Runtime/CoreUObject/Public/UObject/UObjectGlobals.h
@@ -2886,6 +2886,9 @@ struct COREUOBJECT_API FCoreUObjectDelegates
 	/** Called after garbage collection */
 	static FSimpleMulticastDelegate& GetPostGarbageCollect();
 
+	/** Called after purging unreachable objects during garbage collection */
+	static FSimpleMulticastDelegate& GetPostPurgeGarbageDelegate();
+
 #if !UE_BUILD_SHIPPING
 	/** Called when garbage collection detects references to objects that are marked for explicit destruction by MarkAsGarbage */
 	DECLARE_MULTICAST_DELEGATE_OneParam(FOnReportGarbageReferencers, TConstArrayView<struct FGarbageReferenceInfo>);
