GC cleanup in preparation for new GC schema

#rnx

diff --git a/Engine/Source/Runtime/CoreUObject/Private/UObject/CoreNative.cpp b/Engine/Source/Runtime/CoreUObject/Private/UObject/CoreNative.cpp
--- a/Engine/Source/Runtime/CoreUObject/Private/UObject/CoreNative.cpp
+++ b/Engine/Source/Runtime/CoreUObject/Private/UObject/CoreNative.cpp
@@ -72,9 +72,6 @@ UClass* Z_Construct_UClass_UObject()
 	{
 		Z_Registration_Info_UClass_UObject.OuterSingleton = UObject::StaticClass();
 		UObjectForceRegistration(Z_Registration_Info_UClass_UObject.OuterSingleton);
-#if WITH_EDITOR
-		static_assert(false, "Need to EmitBaseReferences here");
-#endif
 		Z_Registration_Info_UClass_UObject.OuterSingleton->StaticLink();
 	}
 	check(Z_Registration_Info_UClass_UObject.OuterSingleton->GetClass());
diff --git a/Engine/Source/Runtime/CoreUObject/Private/UObject/GarbageCollection.cpp b/Engine/Source/Runtime/CoreUObject/Private/UObject/GarbageCollection.cpp
--- a/Engine/Source/Runtime/CoreUObject/Private/UObject/GarbageCollection.cpp
+++ b/Engine/Source/Runtime/CoreUObject/Private/UObject/GarbageCollection.cpp
@@ -316,7 +316,6 @@ namespace TokenDebugNames
 		switch (Type)
 		{
 		case GCRT_Object:								return TEXT("Reference");
-		case GCRT_ExternalPackage:						return TEXT("ExternalPackage");
 		case GCRT_ArrayObject:							return TEXT("Array<Reference>");
 		case GCRT_ArrayStruct:							return TEXT("Array<Struct>");
 		case GCRT_FixedArray:							return TEXT("FixedArray");
@@ -1462,23 +1461,18 @@ struct FReferenceMetadata
 class FReferenceBatcherBase
 {
 protected:
-	template<class ReferenceType>
-	struct TEntry
-	{
-		ReferenceType Reference;
-	};
 
 	template<class ReferenceType, uint32 Capacity, uint32 PrefetchCapacity = 0>
 	struct TBatchQueue
 	{
 		uint32 Num = 0;
-		TEntry<ReferenceType> Entries[Capacity + PrefetchCapacity];
+		ReferenceType Entries[Capacity + PrefetchCapacity];
 
 		FORCEINLINE static constexpr uint32 Max() { return Capacity; }
 		FORCEINLINE bool IsFull() const { return Num == Capacity; }
-		FORCEINLINE void Push(TEntry<ReferenceType> Entry) { Entries[Num++] = Entry; }
+		FORCEINLINE void Push(ReferenceType Entry) { Entries[Num++] = Entry; }
 		FORCEINLINE uint32 Slack() const { return Capacity - Num; }
-		FORCEINLINE TEntry<ReferenceType>& operator[](uint32 Idx) { return Entries[Idx]; }
+		FORCEINLINE ReferenceType& operator[](uint32 Idx) { return Entries[Idx]; }
 
 		TBatchQueue() { FMemory::Memzero(Entries); } // Zero out so prefetching ahead fetches same (near) null addresses
 		~TBatchQueue() { checkf(Num == 0, TEXT("Failed to flush")); }
@@ -1618,14 +1612,14 @@ private:
 	{
 		check(Num <= ArrayBatchSize);
 
-		for (TEntry<FReferenceArray> Entry : MakeArrayView(UnvalidatedArrays.Entries, Num))
+		for (FReferenceArray Entry : MakeArrayView(UnvalidatedArrays.Entries, Num))
 		{
-			FPlatformMisc::Prefetch(Entry.Reference.Array->GetData());
+			FPlatformMisc::Prefetch(Entry.Array->GetData());
 		}
 
-		for (TEntry<FReferenceArray> Entry : MakeArrayView(UnvalidatedArrays.Entries, Num))
+		for (FReferenceArray Entry : MakeArrayView(UnvalidatedArrays.Entries, Num))
 		{
-			PushReferences(*Entry.Reference.Array);
+			PushReferences(*Entry.Array);
 		}
 
 		UnvalidatedArrays.Num = 0;
@@ -1641,7 +1635,7 @@ private:
 	{
 		check((uint32)Range.Num() <= UnvalidatedReferences.Slack());
 
-		TEntry<UnvalidatedReferenceType>* Unvalidated = UnvalidatedReferences.Entries + UnvalidatedReferences.Num;
+		UnvalidatedReferenceType* Unvalidated = UnvalidatedReferences.Entries + UnvalidatedReferences.Num;
 		for (int32 Idx = 0; Idx < Range.Num(); ++Idx)
 		{
 			Unvalidated[Idx] = { MakeReference<UnvalidatedReferenceType>(Range[Idx]) };
@@ -1660,14 +1654,14 @@ private:
 		FValidatedBitmask ValidsA, ValidsB;
 		for (uint32 Idx = 0; Idx < Num; ++Idx)
 		{
-			UObject* Object = GetObject(UnvalidatedReferences[Idx].Reference);
+			UObject* Object = GetObject(UnvalidatedReferences[Idx]);
 			uint64 ObjectAddress = reinterpret_cast<uint64>(Object);
 			ValidsA.Set(Idx, !Permanent.Contains(Object));
 		}
 
 		for (uint32 Idx = 0; Idx < Num; ++Idx)
 		{
-			UObject* Object = GetObject(UnvalidatedReferences[Idx].Reference);
+			UObject* Object = GetObject(UnvalidatedReferences[Idx]);
 			ValidsB.Set(Idx, (!!Object) & IsObjectHandleResolved(reinterpret_cast<FObjectHandle&>(Object))); //-V792
 		}
 		
@@ -1702,7 +1696,7 @@ private:
 		for (uint32 QueueIdx = ValidatedReferences.Num; QueueIdx < NewQueueNum; ++InOutIdx)
 		{
 			bool bIsValid = !!Validations.Get(InOutIdx);
-			ValidatedReferences[QueueIdx].Reference = ToResolvedReference(UnvalidatedReferences[InOutIdx].Reference);
+			ValidatedReferences[QueueIdx] = ToResolvedReference(UnvalidatedReferences[InOutIdx]);
 			QueueIdx += bIsValid;
 		}
 
@@ -1726,23 +1720,23 @@ private:
 		{
 			for (uint32 Idx = 0; Idx < PrefetchAhead; ++Idx)
 			{
-				FPlatformMisc::Prefetch(ValidatedReferences[Idx].Reference.Object, InternalIndexPrefetchOffset);
+				FPlatformMisc::Prefetch(ValidatedReferences[Idx].Object, InternalIndexPrefetchOffset);
 			}
 			for (uint32 Idx = 0; Idx < Num; ++Idx)
 			{
-				ObjectIndices[Idx] = GUObjectArray.ObjectToIndex(ValidatedReferences[Idx].Reference.Object);
-				FPlatformMisc::Prefetch(ValidatedReferences[Idx + PrefetchAhead].Reference.Object, InternalIndexPrefetchOffset);
+				ObjectIndices[Idx] = GUObjectArray.ObjectToIndex(ValidatedReferences[Idx].Object);
+				FPlatformMisc::Prefetch(ValidatedReferences[Idx + PrefetchAhead].Object, InternalIndexPrefetchOffset);
 			}
 		}
 		else
 		{
 			for (uint32 Idx = 0; Idx < Num; ++Idx)
 			{
-				FPlatformMisc::Prefetch(ValidatedReferences[Idx].Reference.Object, InternalIndexPrefetchOffset);
+				FPlatformMisc::Prefetch(ValidatedReferences[Idx].Object, InternalIndexPrefetchOffset);
 			}
 			for (uint32 Idx = 0; Idx < Num; ++Idx)
 			{
-				ObjectIndices[Idx] = GUObjectArray.ObjectToIndex(ValidatedReferences[Idx].Reference.Object);
+				ObjectIndices[Idx] = GUObjectArray.ObjectToIndex(ValidatedReferences[Idx].Object);
 			}
 		}
 
@@ -1781,7 +1775,7 @@ private:
 
 		for (uint32 Idx = 0; Idx < Num; ++Idx)
 		{
-			ProcessorType::HandleBatchedReference(Context, ValidatedReferences[Idx].Reference, Metadatas[Idx]);
+			ProcessorType::HandleBatchedReference(Context, ValidatedReferences[Idx], Metadatas[Idx]);
 		}
 
 		ValidatedReferences.Num = 0;
@@ -2962,12 +2956,6 @@ struct TBatchDispatcher
 		QueueReferences(Context.GetReferencingObject(), Objects, TokenId, ProcessorType::MayKill(TokenType, true));
 	}
 
-	FORCEINLINE void HandleWeakReference(FWeakObjectPtr& WeakPtr, const UObject* ReferencingObject, FTokenId TokenId, EGCTokenType)
-	{
-		UObject* WeakObject = WeakPtr.Get(true);
-		HandleReferenceDirectly(ReferencingObject, WeakObject, TokenId, EKillable::No);
-	}
-
 	FORCENOINLINE void FlushQueuedReferences()
 	{
 		KillableBatcher.FlushQueues();
@@ -5438,14 +5426,6 @@ void FTokenStreamBuilder::EmitFixedArrayEnd()
 	EmitReturn();
 }
 
-void FTokenStreamBuilder::EmitExternalPackageReference()
-{
-#if WITH_EDITOR
-	static const FName TokenName("ExternalPackageToken");
-	EmitReferenceInfo(FGCReferenceInfo(GCRT_ExternalPackage, 0), TokenName);
-#endif
-}
-
 // FTokenStreamBuilderIterator helper
 struct FNamedReferenceInfo
 {
@@ -5550,10 +5530,6 @@ void UClass::AssembleReferenceTokenStreamInternal(bool bForce)
 			SuperClass->AssembleReferenceTokenStreamInternal();
 			SuperTokens = FTokenStreamBuilder::DropFinalTokens(SuperClass->ReferenceTokens.Strong, GetAROFunc(SuperClass));
 		}
-		else
-		{
-			UObjectBase::EmitBaseReferences(Builder);
-		}
 
 		// Make sure all Blueprint properties are marked as non-native
 		Builder.SetTokenType(GetClass()->HasAnyClassFlags(CLASS_NeedsDeferredDependencyLoading) ? EGCTokenType::NonNative : EGCTokenType::Native);
@@ -5809,8 +5785,8 @@ static FName ToName(ETokenlessId Id)
 		FName("Collector"), // ETokenlessId::Collector
 		FName("Class"), // ETokenlessId::Class
 		FName("Outer"), // ETokenlessId::Outer
+		FName("ExternalPackage"), // ETokenlessId::ExternalPackage
 		FName("ClassOuter"), // ETokenlessId::ClassOuter
-		FName("Cluster"), // ETokenlessId::Cluster
 		FName("InitialReference"), // ETokenlessId::InitialReference
 	};
 
diff --git a/Engine/Source/Runtime/CoreUObject/Private/UObject/UObjectBase.cpp b/Engine/Source/Runtime/CoreUObject/Private/UObject/UObjectBase.cpp
--- a/Engine/Source/Runtime/CoreUObject/Private/UObject/UObjectBase.cpp
+++ b/Engine/Source/Runtime/CoreUObject/Private/UObject/UObjectBase.cpp
@@ -388,11 +388,6 @@ bool UObjectBase::IsValidLowLevelFast(bool bRecursive /*= true*/) const
 	return true;
 }
 
-void UObjectBase::EmitBaseReferences(UE::GC::FTokenStreamBuilder& TokenStream)
-{
-	TokenStream.EmitExternalPackageReference();
-}
-
 #if USE_PER_MODULE_UOBJECT_BOOTSTRAP
 static void UObjectReleaseModuleRegistrants(FName Module)
 {
diff --git a/Engine/Source/Runtime/CoreUObject/Public/UObject/FastReferenceCollector.h b/Engine/Source/Runtime/CoreUObject/Public/UObject/FastReferenceCollector.h
--- a/Engine/Source/Runtime/CoreUObject/Public/UObject/FastReferenceCollector.h
+++ b/Engine/Source/Runtime/CoreUObject/Public/UObject/FastReferenceCollector.h
@@ -44,9 +44,6 @@ struct FGCInternals
 {
 	FORCEINLINE static FUObjectItem* GetResolvedOwner(FFieldPath& Path) { return Path.GetResolvedOwnerItemInternal(); }
 	FORCEINLINE static void ClearCachedField(FFieldPath& Path) { Path.ClearCachedFieldInternal(); }
-
-	template<class ObjectID>
-	FORCEINLINE static FWeakObjectPtr& AccessWeakPtr(TPersistentObjectPtr<ObjectID>& Ptr) { return Ptr.WeakPtr; }
 };
 
 /** Interface to allow external systems to trace additional object references, used for bridging GCs */
@@ -310,8 +307,8 @@ enum class ETokenlessId
 	Collector = 1,
 	Class,
 	Outer,
+	ExternalPackage,
 	ClassOuter,
-	Cluster,
 	InitialReference,
 	Max = InitialReference
 };
@@ -488,12 +485,6 @@ struct TDirectDispatcher
 		}
 	}
 
-	FORCEINLINE void HandleWeakReference(FWeakObjectPtr& WeakPtr, UObject* ReferencingObject, FTokenId TokenId, EGCTokenType TokenType) const
-	{
-		UObject* WeakObject = WeakPtr.Get(true);
-		HandleReferenceDirectly(ReferencingObject, WeakObject, TokenId, TokenType, false);
-	}
-
 	FORCEINLINE void FlushQueuedReferences() const {}
 };
 
@@ -581,7 +572,11 @@ StoleContext:
 				// Emit base references
 				Dispatcher.HandleImmutableReference(Class, ETokenlessId::Class, EGCTokenType::Native);
 				Dispatcher.HandleImmutableReference(Outer, ETokenlessId::Outer, EGCTokenType::Native);
-
+#if WITH_EDITOR
+				UObject* Package = CurrentObject->GetExternalPackageInternal();
+				Package = Package != CurrentObject ? Package : nullptr;
+				Dispatcher.HandleImmutableReference(Package, ETokenlessId::ExternalPackage, EGCTokenType::Native);
+#endif
 				if (TokenStream.IsEmpty())
 				{
 					continue;
@@ -774,15 +769,6 @@ StoleContext:
 							// ARO is always last and an implicit terminator
 							goto TokensDone;
 						}
-						case GCRT_ExternalPackage:
-						{
-							TokenReturnCount = ReferenceInfo.ReturnCount;
-							// Test if the object isn't itself, since currently package are their own external and tracking that reference is pointless
-							UObject* Object = CurrentObject->GetExternalPackageInternal();
-							Object = Object != CurrentObject ? Object : nullptr;
-							Dispatcher.HandleImmutableReference( Object, FTokenId(ReferenceTokenStreamIndex), TokenType);
-						}
-						break;
 						case GCRT_FixedArray:
 						{
 							uint8* PreviousData = StackEntryData;
diff --git a/Engine/Source/Runtime/CoreUObject/Public/UObject/GarbageCollection.h b/Engine/Source/Runtime/CoreUObject/Public/UObject/GarbageCollection.h
--- a/Engine/Source/Runtime/CoreUObject/Public/UObject/GarbageCollection.h
+++ b/Engine/Source/Runtime/CoreUObject/Public/UObject/GarbageCollection.h
@@ -325,7 +325,6 @@ public:
 	/** All tokens issues between Begin and End will be replayed Count times. */
 	COREUOBJECT_API void EmitFixedArrayBegin(int32 Offset, FName DebugName, int32 Stride, int32 Count);
 	COREUOBJECT_API void EmitFixedArrayEnd();
-	COREUOBJECT_API void EmitExternalPackageReference();
 
 	void EmitFinalTokens(void(*AddReferencedObjects)(UObject*, FReferenceCollector&));	
 	static FTokenStreamView DropFinalTokens(FTokenStreamView Tokens, void(*DropARO)(UObject*, FReferenceCollector&));
diff --git a/Engine/Source/Runtime/CoreUObject/Public/UObject/PersistentObjectPtr.h b/Engine/Source/Runtime/CoreUObject/Public/UObject/PersistentObjectPtr.h
--- a/Engine/Source/Runtime/CoreUObject/Public/UObject/PersistentObjectPtr.h
+++ b/Engine/Source/Runtime/CoreUObject/Public/UObject/PersistentObjectPtr.h
@@ -14,9 +14,6 @@
 template<class TObjectID>
 struct TPersistentObjectPtr
 {
-	friend struct FGCInternals;
-
-public:
 	using ElementType = TObjectID;
 
 	/** Default constructor, will be null */
diff --git a/Engine/Source/Runtime/CoreUObject/Public/UObject/UObjectBase.h b/Engine/Source/Runtime/CoreUObject/Public/UObject/UObjectBase.h
--- a/Engine/Source/Runtime/CoreUObject/Public/UObject/UObjectBase.h
+++ b/Engine/Source/Runtime/CoreUObject/Public/UObject/UObjectBase.h
@@ -83,11 +83,6 @@ public:
 	 **/
 	virtual ~UObjectBase();
 
-	/**
-	 * Emit GC tokens for UObjectBase, this might be UObject::StaticClass or Default__Class
-	 **/
-	static void EmitBaseReferences(UE::GC::FTokenStreamBuilder& TokenStream);
-
 protected:
 	/**
 	 * Just change the FName and Outer and rehash into name hash tables. For use by higher level rename functions.
diff --git a/Engine/Source/Runtime/Engine/Public/Model.h b/Engine/Source/Runtime/Engine/Public/Model.h
--- a/Engine/Source/Runtime/Engine/Public/Model.h
+++ b/Engine/Source/Runtime/Engine/Public/Model.h
@@ -204,20 +204,17 @@ struct FLeaf
 //
 struct FBspSurf
 {
-public:
-
-	UMaterialInterface*	Material;		// 4 Material.
-	uint32				PolyFlags;		// 4 Polygon flags.
-	int32					pBase;			// 4 Polygon & texture base point index (where U,V==0,0).
-	int32					vNormal;		// 4 Index to polygon normal.
-	int32					vTextureU;		// 4 Texture U-vector index.
-	int32					vTextureV;		// 4 Texture V-vector index.
-	int32					iBrushPoly;		// 4 Editor brush polygon index.
-	ABrush*				Actor;			// 4 Brush actor owning this Bsp surface.
-	FPlane4f				Plane;			// 16 The plane this surface lies on.
-	float				LightMapScale;	// 4 The number of units/lightmap texel on this surface.
-
-	int32					iLightmassIndex;// 4 Index to the lightmass settings
+	TObjectPtr<UMaterialInterface>	Material;		// 4 Material.
+	uint32							PolyFlags;		// 4 Polygon flags.
+	int32							pBase;			// 4 Polygon & texture base point index (where U,V==0,0).
+	int32							vNormal;		// 4 Index to polygon normal.
+	int32							vTextureU;		// 4 Texture U-vector index.
+	int32							vTextureV;		// 4 Texture V-vector index.
+	int32							iBrushPoly;		// 4 Editor brush polygon index.
+	TObjectPtr<ABrush>				Actor;			// 4 Brush actor owning this Bsp surface.
+	FPlane4f						Plane;			// 16 The plane this surface lies on.
+	float							LightMapScale;	// 4 The number of units/lightmap texel on this surface.
+	int32							iLightmassIndex;// 4 Index to the lightmass settings
 
 	bool				bHiddenEdTemporary;	// 1 Marks whether this surface is temporarily hidden in the editor or not. Not serialized.
 	bool				bHiddenEdLevel;		// 1 Marks whether this surface is hidden by the level browser or not. Not serialized.
