Reverse the post load order for PostLoadGroups start with dependencies in ZenLoader
Also use the PostLoadGroup to run the deferred post load using the group to prevent race with ConditionalPostLoad with async and non async PostLoad

diff --git a/Engine/Source/Runtime/CoreUObject/Private/Serialization/AsyncLoading2.cpp b/Engine/Source/Runtime/CoreUObject/Private/Serialization/AsyncLoading2.cpp
--- a/Engine/Source/Runtime/CoreUObject/Private/Serialization/AsyncLoading2.cpp
+++ b/Engine/Source/Runtime/CoreUObject/Private/Serialization/AsyncLoading2.cpp
@@ -2689,6 +2689,7 @@ private:
 	// The sync load context associated with this package
 	std::atomic<uint64>			SyncLoadContextId = 0;
 	FAsyncLoadingPostLoadGroup* PostLoadGroup = nullptr;
+	FAsyncLoadingPostLoadGroup* DeferredPostLoadGroup = nullptr;
 	/** Time load begun. This is NOT the time the load was requested in the case of pending requests.	*/
 	double						LoadStartTime = 0.0;
 	std::atomic<int32>			RefCount = 0;
@@ -2953,6 +2954,7 @@ struct FAsyncLoadingPostLoadGroup
 	uint64 SyncLoadContextId = 0;
 	TArray<FAsyncPackage2*> Packages;
 	int32 PackagesWithExportsToSerializeCount = 0;
+	int32 PackagesWithExportsToPostLoadCount = 0;
 };
 
 class FAsyncLoadingThread2 final
@@ -3355,6 +3357,7 @@ private:
 	void ConditionalProcessEditorCallbacks();
 #endif
 	void ConditionalBeginPostLoad(FAsyncLoadingThreadState2& ThreadState, FAsyncLoadingPostLoadGroup* PostLoadGroup);
+	void ConditionalBeginDeferredPostLoad(FAsyncLoadingThreadState2& ThreadState, FAsyncLoadingPostLoadGroup* DeferredPostLoadGroup);
 	void MergePostLoadGroups(FAsyncLoadingThreadState2& ThreadState, FAsyncLoadingPostLoadGroup* Target, FAsyncLoadingPostLoadGroup* Source);
 
 	bool ProcessDeferredDeletePackagesQueue(int32 MaxCount = MAX_int32)
@@ -3801,18 +3804,50 @@ void FAsyncLoadingThread2::ConditionalBeginPostLoad(FAsyncLoadingThreadState2& T
 	check(ThreadState.bCanAccessAsyncLoadingThreadData);
 	if (PostLoadGroup->PackagesWithExportsToSerializeCount == 0)
 	{
-		for (FAsyncPackage2* Package : PostLoadGroup->Packages)
+		// Release the post load node of packages in the post load group in reverse order that they were added to the group
+		// This usually means that dependencies will be post load first, similarly to how they are also serialized first
+		for (int32 Index = PostLoadGroup->Packages.Num() - 1; Index >= 0; --Index)
 		{
+			FAsyncPackage2* Package = PostLoadGroup->Packages[Index];
 			check(Package->PostLoadGroup == PostLoadGroup);
 			check(Package->AsyncPackageLoadingState == EAsyncPackageLoadingState2::ExportsDone);
-			Package->AsyncPackageLoadingState = EAsyncPackageLoadingState2::PostLoad;
+
+			// Move the PostLoadGroup to the DeferredPostLoadGroup so that we do not mistakenly consider post load as not having being triggered yet
 			Package->PostLoadGroup = nullptr;
+			Package->DeferredPostLoadGroup = PostLoadGroup;
+			
+			Package->AsyncPackageLoadingState = EAsyncPackageLoadingState2::PostLoad;
 			for (int32 ExportBundleIndex = 0; ExportBundleIndex < Package->Data.TotalExportBundleCount; ++ExportBundleIndex)
 			{
 				Package->GetExportBundleNode(EEventLoadNode2::ExportBundle_PostLoad, ExportBundleIndex).ReleaseBarrier(&ThreadState);
 			}
 		}
-		delete PostLoadGroup;
+		PostLoadGroup->PackagesWithExportsToPostLoadCount = PostLoadGroup->Packages.Num();
+	}
+}
+
+void FAsyncLoadingThread2::ConditionalBeginDeferredPostLoad(FAsyncLoadingThreadState2& ThreadState, FAsyncLoadingPostLoadGroup* DeferredPostLoadGroup)
+{
+	TRACE_CPUPROFILER_EVENT_SCOPE(ConditionalBeginDeferredPostLoad);
+	check(DeferredPostLoadGroup);
+	check(ThreadState.bCanAccessAsyncLoadingThreadData);
+	if (DeferredPostLoadGroup->PackagesWithExportsToPostLoadCount == 0)
+	{
+		// Release the post load node of packages in the post load group in reverse order that they were added to the group
+		// This usually means that dependencies will be post load first, similarly to how they are also serialized first
+		for (int32 Index = DeferredPostLoadGroup->Packages.Num() - 1; Index >= 0; --Index)
+		{
+			FAsyncPackage2* Package = DeferredPostLoadGroup->Packages[Index];
+			check(Package->DeferredPostLoadGroup == DeferredPostLoadGroup);
+			check(Package->AsyncPackageLoadingState == EAsyncPackageLoadingState2::PostLoad);
+			Package->DeferredPostLoadGroup = nullptr;
+			Package->AsyncPackageLoadingState = EAsyncPackageLoadingState2::DeferredPostLoad;
+			for (int32 ExportBundleIndex = 0; ExportBundleIndex < Package->Data.TotalExportBundleCount; ++ExportBundleIndex)
+			{
+				Package->GetExportBundleNode(EEventLoadNode2::ExportBundle_DeferredPostLoad, ExportBundleIndex).ReleaseBarrier(&ThreadState);
+			}
+		}
+		delete DeferredPostLoadGroup;
 	}
 }
 
@@ -3824,10 +3859,6 @@ void FAsyncLoadingThread2::MergePostLoadGroups(FAsyncLoadingThreadState2& Thread
 	}
 	TRACE_CPUPROFILER_EVENT_SCOPE(MergePostLoadGroups);
 	check(ThreadState.bCanAccessAsyncLoadingThreadData);
-	if (Source->Packages.Num() > Target->Packages.Num())
-	{
-		Swap(Source, Target);
-	}
 	for (FAsyncPackage2* Package : Source->Packages)
 	{
 		check(Package->PostLoadGroup == Source);
@@ -3835,6 +3866,7 @@ void FAsyncLoadingThread2::MergePostLoadGroups(FAsyncLoadingThreadState2& Thread
 	}
 	Target->Packages.Append(MoveTemp(Source->Packages));
 	Target->PackagesWithExportsToSerializeCount += Source->PackagesWithExportsToSerializeCount;
+	check(Target->PackagesWithExportsToPostLoadCount == 0 && Source->PackagesWithExportsToPostLoadCount == 0);
 
 	const uint64 SyncLoadContextId = FMath::Max(Source->SyncLoadContextId, Target->SyncLoadContextId);
 	if (SyncLoadContextId)
@@ -6595,8 +6627,11 @@ EEventLoadNodeExecutionResult FAsyncPackage2::Event_PostLoadExportBundle(FAsyncL
 	if (Package->LinkerLoadState.IsSet())
 	{
 		// No async postload for now
-		Package->AsyncPackageLoadingState = EAsyncPackageLoadingState2::DeferredPostLoad;
-		Package->GetExportBundleNode(ExportBundle_DeferredPostLoad, 0).ReleaseBarrier(&ThreadState);
+		FAsyncLoadingPostLoadGroup* DeferredPostLoadGroup = Package->DeferredPostLoadGroup;
+		check(DeferredPostLoadGroup);
+		check(DeferredPostLoadGroup->PackagesWithExportsToPostLoadCount > 0);
+		--DeferredPostLoadGroup->PackagesWithExportsToPostLoadCount;
+		Package->AsyncLoadingThread.ConditionalBeginDeferredPostLoad(ThreadState, DeferredPostLoadGroup);
 		return EEventLoadNodeExecutionResult::Complete;
 	}
 #endif
@@ -6719,16 +6754,15 @@ EEventLoadNodeExecutionResult FAsyncPackage2::Event_PostLoadExportBundle(FAsyncL
 			UE_ASYNC_PACKAGE_LOG(Verbose, Package->Desc, TEXT("AsyncThread: FullyLoaded"),
 				TEXT("Async loading of package is done, and UPackage is marked as fully loaded."));
 			// mimic old loader behavior for now, but this is more correctly also done in FinishUPackage
-			// called from ProcessLoadedPackagesFromGameThread just before complection callbacks
+			// called from ProcessLoadedPackagesFromGameThread just before completion callbacks
 			Package->LinkerRoot->MarkAsFullyLoaded();
 		}
 
-		check(Package->AsyncPackageLoadingState == EAsyncPackageLoadingState2::PostLoad);
-		Package->AsyncPackageLoadingState = EAsyncPackageLoadingState2::DeferredPostLoad;
-		for (int32 ExportBundleIndex = 0; ExportBundleIndex < Package->Data.TotalExportBundleCount; ++ExportBundleIndex)
-		{
-			Package->GetExportBundleNode(ExportBundle_DeferredPostLoad, ExportBundleIndex).ReleaseBarrier(&ThreadState);
-		}
+		FAsyncLoadingPostLoadGroup* DeferredPostLoadGroup = Package->DeferredPostLoadGroup;
+		check(DeferredPostLoadGroup);
+		check(DeferredPostLoadGroup->PackagesWithExportsToPostLoadCount > 0);
+		--DeferredPostLoadGroup->PackagesWithExportsToPostLoadCount;
+		Package->AsyncLoadingThread.ConditionalBeginDeferredPostLoad(ThreadState, DeferredPostLoadGroup);
 	}
 
 	return EEventLoadNodeExecutionResult::Complete;
@@ -8032,6 +8066,9 @@ FAsyncPackage2::~FAsyncPackage2()
 
 	FMemory::Free(Data.MemoryBuffer0);
 	FMemory::Free(Data.MemoryBuffer1);
+
+	check(PostLoadGroup == nullptr);
+	check(DeferredPostLoadGroup == nullptr);
 }
 
 void FAsyncPackage2::ReleaseRef()
