Change ObjectHandle private methods to take UObject* instead of UObject&.

diff --git a/Engine/Source/Runtime/CoreUObject/Private/UObject/Obj.cpp b/Engine/Source/Runtime/CoreUObject/Private/UObject/Obj.cpp
--- a/Engine/Source/Runtime/CoreUObject/Private/UObject/Obj.cpp
+++ b/Engine/Source/Runtime/CoreUObject/Private/UObject/Obj.cpp
@@ -345,7 +345,7 @@ bool UObject::Rename( const TCHAR* InName, UObject* NewOuter, ERenameFlags Flags
 			}
 		}
 #if UE_WITH_OBJECT_HANDLE_LATE_RESOLVE
-		UE::CoreUObject::Private::UpdateRenamedObject(*this, NewName, NewOuter);
+		UE::CoreUObject::Private::UpdateRenamedObject(this, NewName, NewOuter);
 #endif
 		LowLevelRename(NewName, NewOuter);
 	}
diff --git a/Engine/Source/Runtime/CoreUObject/Private/UObject/ObjectHandleIndex.cpp b/Engine/Source/Runtime/CoreUObject/Private/UObject/ObjectHandleIndex.cpp
--- a/Engine/Source/Runtime/CoreUObject/Private/UObject/ObjectHandleIndex.cpp
+++ b/Engine/Source/Runtime/CoreUObject/Private/UObject/ObjectHandleIndex.cpp
@@ -345,20 +345,20 @@ namespace UE::CoreUObject::Private
 		}
 	}
 
-	void FreeObjectHandle(const UObjectBase& Object)
+	void FreeObjectHandle(const UObjectBase* Object)
 	{
-		int32 ObjectIndex = GUObjectArray.ObjectToIndex(&Object);
+		int32 ObjectIndex = GUObjectArray.ObjectToIndex(Object);
 		FPackedObjectRef& PackedObjectRef = GObjectHandleIndex.ObjectIndexToPackedObjectRef[ObjectIndex];
 		PackedObjectRef.EncodedRef = 0;
 	}
 
-	void UpdateRenamedObject(const UObject& Object, FName NewName, UObject* NewOuter)
+	void UpdateRenamedObject(const UObject* Object, FName NewName, UObject* NewOuter)
 	{
 		if (!UE::LinkerLoad::IsImportLazyLoadEnabled())
 		{
 			return;
 		}
-		int32 ObjectIndex = GUObjectArray.ObjectToIndex(&Object);
+		int32 ObjectIndex = GUObjectArray.ObjectToIndex(Object);
 		FPackedObjectRef PackedObjectRef = GObjectHandleIndex.ObjectIndexToPackedObjectRef[ObjectIndex];
 		if (PackedObjectRef.EncodedRef == 0)
 		{
@@ -367,7 +367,7 @@ namespace UE::CoreUObject::Private
 
 		check(NewName != NAME_None);
 
-		FMinimalName MinimalName = NameToMinimalName(Object.GetFName());
+		FMinimalName MinimalName = NameToMinimalName(Object->GetFName());
 		FObjectId ObjectId;
 		FPackageId PackageId;
 		Unpack(PackedObjectRef, PackageId, ObjectId);
@@ -375,7 +375,7 @@ namespace UE::CoreUObject::Private
 		FRWScopeLock GlobalLockScope(GObjectHandleIndex.Lock, SLT_Write);
 		FObjectHandlePackageData& PackageData = GObjectHandleIndex.PackageData[PackageId.ToIndex()];
 		FRWScopeLock PackageLockScope(PackageData.Lock, SLT_ReadOnly);
-		if (Object.GetClass() == UPackage::StaticClass())
+		if (Object->GetClass() == UPackage::StaticClass())
 		{
 			//update the package name at existing index. existing object handles will be correct when unpacked
 			//add in the new name pointing the existing index. new object handles will resolve to the existing index
@@ -422,7 +422,7 @@ namespace UE::CoreUObject::Private
 			}
 			else if (CurrentObjectPath == OldObjectPath)
 			{
-				CurrentObjectPath.MakeWeakObjPtr(Object);
+				CurrentObjectPath.MakeWeakObjPtr(*Object);
 				continue;
 			}
 			else if (CurrentObjectPath.IsSimple())
@@ -437,7 +437,7 @@ namespace UE::CoreUObject::Private
 				CurrentObjectPath.Resolve(ResolvedNames);
 
 				//resolve all object along the path and convert the object paths to weak objects
-				const UObject* CurrentObject = Object.GetPackage();
+				const UObject* CurrentObject = Object->GetPackage();
 				for (int32 ObjectPathIndex = 0; ObjectPathIndex < ResolvedNames.Num() && CurrentObject; ++ObjectPathIndex)
 				{
 					CurrentObject = StaticFindObjectFastInternal(nullptr, CurrentObject, ResolvedNames[ObjectPathIndex]);
@@ -522,9 +522,9 @@ namespace UE::CoreUObject::Private
 		return PackedObjectRef;
 	}
 
-	FPackedObjectRef FindExistingPackedObjectRef(const UObject& Object)
+	FPackedObjectRef FindExistingPackedObjectRef(const UObject* Object)
 	{
-		int32 ObjectIndex = GUObjectArray.ObjectToIndex(&Object);
+		int32 ObjectIndex = GUObjectArray.ObjectToIndex(Object);
 		if (ObjectIndex == INDEX_NONE)
 		{
 			return FPackedObjectRef {0};
diff --git a/Engine/Source/Runtime/CoreUObject/Private/UObject/UObjectBase.cpp b/Engine/Source/Runtime/CoreUObject/Private/UObject/UObjectBase.cpp
--- a/Engine/Source/Runtime/CoreUObject/Private/UObject/UObjectBase.cpp
+++ b/Engine/Source/Runtime/CoreUObject/Private/UObject/UObjectBase.cpp
@@ -152,7 +152,7 @@ UObjectBase::~UObjectBase()
 		check(IsValidLowLevel());
 		check(GetFName() == NAME_None);
 #if UE_WITH_OBJECT_HANDLE_LATE_RESOLVE
-		UE::CoreUObject::Private::FreeObjectHandle(*this);
+		UE::CoreUObject::Private::FreeObjectHandle(this);
 #endif 
 		GUObjectArray.FreeUObjectIndex(this);
 	}
diff --git a/Engine/Source/Runtime/CoreUObject/Public/UObject/ObjectHandle.h b/Engine/Source/Runtime/CoreUObject/Public/UObject/ObjectHandle.h
--- a/Engine/Source/Runtime/CoreUObject/Public/UObject/ObjectHandle.h
+++ b/Engine/Source/Runtime/CoreUObject/Public/UObject/ObjectHandle.h
@@ -60,7 +60,7 @@ namespace UE::CoreUObject::Private
 	};
 
 	/* Returns the packed object ref for this object IF one exists otherwise returns a null PackedObjectRef */
-	COREUOBJECT_API FPackedObjectRef FindExistingPackedObjectRef(const UObject& Object);
+	COREUOBJECT_API FPackedObjectRef FindExistingPackedObjectRef(const UObject* Object);
 
 	/* Creates and ObjectRef from a packed object ref*/
 	COREUOBJECT_API FObjectRef MakeObjectRef(FPackedObjectRef Handle);
@@ -141,7 +141,7 @@ inline bool operator==(UE::CoreUObject::Private::FObjectHandlePrivate LHS, UE::C
 		}
 
 		//if packed ref empty then can't be equal as RHS is an unresolved pointer
-		FPackedObjectRef PackedLhs = FindExistingPackedObjectRef(*Obj);
+		FPackedObjectRef PackedLhs = FindExistingPackedObjectRef(Obj);
 		if (PackedLhs.EncodedRef == 0)
 		{
 			return false;
@@ -159,7 +159,7 @@ inline bool operator==(UE::CoreUObject::Private::FObjectHandlePrivate LHS, UE::C
 		}
 
 		//if packed ref empty then can't be equal as RHS is an unresolved pointer
-		FPackedObjectRef PackedRhs = FindExistingPackedObjectRef(*Obj);
+		FPackedObjectRef PackedRhs = FindExistingPackedObjectRef(Obj);
 		if (PackedRhs.EncodedRef == 0)
 		{
 			return false;
@@ -187,7 +187,7 @@ inline uint32 GetTypeHash(UE::CoreUObject::Private::FObjectHandlePrivate Handle)
 	{
 		const UObject* Obj = ReadObjectHandlePointerNoCheck(Handle);
 
-		FPackedObjectRef PackedObjectRef = FindExistingPackedObjectRef(*Obj);
+		FPackedObjectRef PackedObjectRef = FindExistingPackedObjectRef(Obj);
 		if (PackedObjectRef.EncodedRef == 0)
 		{
 			return GetTypeHash(Obj);
@@ -343,8 +343,8 @@ namespace UE::CoreUObject::Private
 #if UE_WITH_OBJECT_HANDLE_LATE_RESOLVE
 	//forward declarations
 	void InitObjectHandles(int32 Size);
-	void FreeObjectHandle(const UObjectBase& Object);
-	void UpdateRenamedObject(const UObject& Obj, FName NewName, UObject* NewOuter);
+	void FreeObjectHandle(const UObjectBase* Object);
+	void UpdateRenamedObject(const UObject* Obj, FName NewName, UObject* NewOuter);
 	UE::CoreUObject::Private::FPackedObjectRef MakePackedObjectRef(const UObject* Object);
 #endif
 }
