Remove some zenloader tracing recently added that caused perf and trace size regression

#rnx

diff --git a/Engine/Source/Runtime/CoreUObject/Private/Serialization/AsyncLoading2.cpp b/Engine/Source/Runtime/CoreUObject/Private/Serialization/AsyncLoading2.cpp
--- a/Engine/Source/Runtime/CoreUObject/Private/Serialization/AsyncLoading2.cpp
+++ b/Engine/Source/Runtime/CoreUObject/Private/Serialization/AsyncLoading2.cpp
@@ -6884,8 +6884,6 @@ FEventLoadNode2& FAsyncPackage2::GetExportBundleNode(EEventLoadNode2 Phase, uint
 
 void FAsyncLoadingThread2::UpdateSyncLoadContext(FAsyncLoadingThreadState2& ThreadState)
 {
-	TRACE_CPUPROFILER_EVENT_SCOPE(FAsyncLoadingThread2::UpdateSyncLoadContext);
-
 	if (ThreadState.bIsAsyncLoadingThread)
 	{
 		FAsyncLoadingSyncLoadContext* CreatedOnMainThread;
@@ -8341,8 +8339,6 @@ void FAsyncPackage2::FinishUPackage()
 #if WITH_EDITOR
 void FAsyncLoadingThread2::ConditionalProcessEditorCallbacks()
 {
-	TRACE_CPUPROFILER_EVENT_SCOPE(FAsyncLoadingThread2::ConditionalProcessEditorCallbacks);
-
 	check(IsInGameThread());
 	if (!GameThreadState->SyncLoadContextStack.IsEmpty())
 	{
@@ -8353,6 +8349,8 @@ void FAsyncLoadingThread2::ConditionalProcessEditorCallbacks()
 
 	while (!EditorCompletedUPackages.IsEmpty() || !EditorLoadedAssets.IsEmpty())
 	{
+		TRACE_CPUPROFILER_EVENT_SCOPE(FAsyncLoadingThread2::ConditionalProcessEditorCallbacks);
+
 		TArray<UObject*> LocalEditorLoadedAssets;
 		Swap(LocalEditorLoadedAssets, EditorLoadedAssets);
 		TArray<UPackage*> LocalEditorCompletedUPackages;
