SavePackage: Restore SavePackage support for OnlyEditorOnly, backout the backout of CL 25146866. The server crashes were caused by a different change.
#rnx

diff --git a/Engine/Source/Runtime/CoreUObject/Private/UObject/SavePackage2.cpp b/Engine/Source/Runtime/CoreUObject/Private/UObject/SavePackage2.cpp
--- a/Engine/Source/Runtime/CoreUObject/Private/UObject/SavePackage2.cpp
+++ b/Engine/Source/Runtime/CoreUObject/Private/UObject/SavePackage2.cpp
@@ -568,7 +568,7 @@ ESavePackageResult ValidateExports(FSaveContext& SaveContext)
 		{
 			UObject* Object = Export.Obj;
 			// Consider redirectors to world/levels as map packages too.
-			while (UObjectRedirector* Redirector = Cast<UObjectRedirector>(Object))
+			if (UObjectRedirector* Redirector = Cast<UObjectRedirector>(Object))
 			{
 				Object = Redirector->DestinationObject;
 			}
@@ -975,19 +975,22 @@ ESavePackageResult CreateLinker(FSaveContext& SaveContext)
 
 		SaveContext.GetLinker()->bRehydratePayloads = SaveContext.ShouldRehydratePayloads();
 		
-		// The package trailer is not supported for text based assets yet
-		if (!SaveContext.IsTextFormat() && !SaveContext.IsProceduralSave())
+		if (UE::FPackageTrailer::IsEnabled())
 		{
-			SaveContext.GetLinker()->PackageTrailerBuilder = MakeUnique<UE::FPackageTrailerBuilder>(SaveContext.GetPackage()->GetName());
-		}
-		else if ((SaveContext.GetSaveArgs().SaveFlags & SAVE_BulkDataByReference) != 0)
-		{
-			if (const FLinkerLoad* LinkerLoad = FLinkerLoad::FindExistingLinkerForPackage(SaveContext.GetPackage()))
+			// The package trailer is not supported for text based assets yet
+			if (!SaveContext.IsTextFormat() && !SaveContext.IsProceduralSave())
+			{		
+				SaveContext.GetLinker()->PackageTrailerBuilder = MakeUnique<UE::FPackageTrailerBuilder>(SaveContext.GetPackage()->GetName());
+			}
+			else if ((SaveContext.GetSaveArgs().SaveFlags & SAVE_BulkDataByReference) != 0)
 			{
-				const UE::FPackageTrailer* Trailer = LinkerLoad->GetPackageTrailer();
-				if (Trailer && Trailer->GetNumPayloads(UE::EPayloadStorageType::Any) > 0)
+				if (const FLinkerLoad* LinkerLoad = FLinkerLoad::FindExistingLinkerForPackage(SaveContext.GetPackage()))
 				{
-					SaveContext.GetLinker()->PackageTrailerBuilder = UE::FPackageTrailerBuilder::CreateReferenceToTrailer(*Trailer, SaveContext.GetPackage()->GetName());
+					const UE::FPackageTrailer* Trailer = LinkerLoad->GetPackageTrailer();
+					if (Trailer && Trailer->GetNumPayloads(UE::EPayloadStorageType::Any) > 0)
+					{
+						SaveContext.GetLinker()->PackageTrailerBuilder = UE::FPackageTrailerBuilder::CreateReferenceToTrailer(*Trailer, SaveContext.GetPackage()->GetName());
+					}
 				}
 			}
 		}
@@ -1021,6 +1024,10 @@ ESavePackageResult CreateLinker(FSaveContext& SaveContext)
 
 struct FNameEntryIdSortHelper
 {
+private:
+	/** the linker that we're sorting names for */
+	friend struct TDereferenceWrapper<FNameEntryId, FNameEntryIdSortHelper>;
+
 	/** Comparison function used by Sort */
 	FORCEINLINE bool operator()(const FName& A, const FName& B) const
 	{
@@ -1035,6 +1042,19 @@ struct FNameEntryIdSortHelper
 	}
 };
 
+struct FObjectResourceSortHelper
+{
+private:
+	friend struct TDereferenceWrapper<FObjectImport, FObjectResourceSortHelper>;
+	friend struct TDereferenceWrapper<FObjectExport, FObjectResourceSortHelper>;
+
+	/** Comparison function used by Sort */
+	FORCEINLINE bool operator()(const FObjectResource& A, const FObjectResource& B) const
+	{
+		return A.ObjectName.Compare(B.ObjectName) < 0;
+	}
+};
+
 ESavePackageResult BuildLinker(FSaveContext& SaveContext)
 {
 	// Setup Linker 
@@ -1136,9 +1156,10 @@ ESavePackageResult BuildLinker(FSaveContext& SaveContext)
 		Linker->Summary.NameCount = Linker->NameMap.Num();
 		Linker->Summary.NamesReferencedFromExportDataCount = NamesReferencedFromExportData.Num();
 
-		Algo::Sort(MakeArrayView(Linker->NameMap.GetData(), Linker->Summary.NamesReferencedFromExportDataCount), FNameEntryIdSortHelper());
-		Algo::Sort(MakeArrayView(Linker->NameMap.GetData() + Linker->Summary.NamesReferencedFromExportDataCount,
-			Linker->Summary.NameCount - Linker->Summary.NamesReferencedFromExportDataCount), FNameEntryIdSortHelper());
+		Sort(Linker->NameMap.GetData(), Linker->Summary.NamesReferencedFromExportDataCount, FNameEntryIdSortHelper());
+		Sort(Linker->NameMap.GetData() + Linker->Summary.NamesReferencedFromExportDataCount,
+			Linker->Summary.NameCount - Linker->Summary.NamesReferencedFromExportDataCount,
+			FNameEntryIdSortHelper());
 
 		if (!SaveContext.IsTextFormat())
 		{
@@ -1212,7 +1233,7 @@ ESavePackageResult BuildLinker(FSaveContext& SaveContext)
 				ObjectImport.ObjectName = ReplacedName;
 			}
 		}
-		//Algo::SortBy(Linker->ImportMap, &FObjectResource::ObjectName, FNameLexicalLess());
+		//Sort(&Linker->ImportMap[0], Linker->ImportMap.Num(), FObjectResourceSortHelper());
 
 		// @todo: To stay consistent with the old save and prevent binary diff between the algo, use the old import sort for now
 		// a future cvar could allow projects use the less expensive sort in their own time down the line
@@ -1237,7 +1258,7 @@ ESavePackageResult BuildLinker(FSaveContext& SaveContext)
 				Export.PackageFlags = Package->GetPackageFlags();
 			}
 		}
-		//Algo::SortBy(Linker->ExportMap, &FObjectResource::ObjectName, FNameLexicalLess());
+		//Sort(&Linker->ExportMap[0], Linker->ExportMap.Num(), FObjectResourceSortHelper());
 
 		// @todo: To stay consistent with the old save and prevent binary diff between the algo, use the old import sort for now
 		// a future cvar could allow projects use the less expensive sort in their own time down the line
@@ -1722,33 +1743,6 @@ int64 WriteObjectDataResources(TArray<FObjectDataResource>& DataResources, FStru
 	return Linker.Tell() - Linker.Summary.DataResourceOffset;
 }
 
-/**
- * Utility for safely setting the TotalHeaderSize member of FPackageFileSummary with a int64 value.
- * 
- * FPackageFileSummary uses int32 for a lot of offsets but a lot of our package writing code is
- * capable of handling files that exceed MAX_int32 so the final size is calculated as a int64.
- * If this value is truncated when storing in FPackageFileSummary then the package will not read
- * in correctly and most likely cause a crash. Rather than allowing the user to save bad data we 
- * can use this utility to catch the error and log it so that the user can take action.
- */
-ESavePackageResult SetSummaryTotalHeaderSize(FSaveContext& SaveContext, int64 TotalHeaderSize)
-{
-	FLinkerSave* Linker = SaveContext.GetLinker();
-
-	if (TotalHeaderSize <= MAX_int32)
-	{
-		Linker->Summary.TotalHeaderSize = (int32)TotalHeaderSize;
-		return ESavePackageResult::Success;
-	}
-	else
-	{
-		UE_LOG(LogSavePackage, Error, TEXT("Package header for '%s' is too large (%" UINT64_FMT " bytes), some package file summary offsets will be truncated when stored as a int32"),
-			*SaveContext.GetPackage()->GetName(), TotalHeaderSize);
-
-		return ESavePackageResult::Error;
-	}
-}
-
 ESavePackageResult WritePackageHeader(FStructuredArchive::FRecord& StructuredArchiveRoot, FSaveContext& SaveContext)
 {
 	FLinkerSave* Linker = SaveContext.GetLinker();
@@ -1892,13 +1886,16 @@ ESavePackageResult WritePackageHeader(FStructuredArchive::FRecord& StructuredArc
 		SCOPED_SAVETIMER(UPackage_Save_PreloadDependencies);
 		SavePreloadDependencies(StructuredArchiveRoot, SaveContext);
 	}
-	
+	Linker->Summary.TotalHeaderSize = (int32)Linker->Tell();
+
 	// Rather than check if an offset is truncated every time we assign one, we can just check the final TotalHeaderSize to see if it is truncated.
 	// Checking every time an offset is assigned would let us fail quicker but a) relies that new code follows the convention b) bloats the code a fair bit.
-	const ESavePackageResult Result = SetSummaryTotalHeaderSize(SaveContext, Linker->Tell());
-	if (Result != ESavePackageResult::Success)
+	if (Linker->Tell() > MAX_int32)
 	{
-		return Result;
+		UE_LOG(LogSavePackage, Error, TEXT("Package header for '%s' is too large (%" UINT64_FMT " bytes), some package file summary offsets will be truncated when stored as a int32"),
+			*SaveContext.GetPackage()->GetName(), Linker->Tell());
+
+		return ESavePackageResult::Error;
 	}
 
 	return ReturnSuccessOrCancel();
@@ -2758,13 +2755,7 @@ ESavePackageResult SaveHarvestedRealms(FSaveContext& SaveContext, ESaveRealm Har
 				const int64 DataResourceSize = WriteObjectDataResources(Linker.DataResourceMap, StructuredArchiveRoot, SaveContext);
 				check(DataResourceSize >= 0);
 
-				// Check to make sure that the package header is not too large
-				SaveContext.Result = SetSummaryTotalHeaderSize(SaveContext, (int64)Linker.Summary.TotalHeaderSize + DataResourceSize);
-				if (SaveContext.Result != ESavePackageResult::Success)
-				{
-					return SaveContext.Result;
-				}
-
+				Linker.Summary.TotalHeaderSize += DataResourceSize;
 				{
 					// Disables writing stack trace data when appending the exports data
 					FArchiveStackTraceDisabledScope _; 
