Disable buffer overrun warning false positive in order to support FastGenPGO.

#rnx

diff --git a/Engine/Source/Runtime/CoreUObject/Public/UObject/UnrealType.h b/Engine/Source/Runtime/CoreUObject/Public/UObject/UnrealType.h
--- a/Engine/Source/Runtime/CoreUObject/Public/UObject/UnrealType.h
+++ b/Engine/Source/Runtime/CoreUObject/Public/UObject/UnrealType.h
@@ -2698,6 +2698,8 @@ protected:
 	 */
 	UObject* ConstructDefaultObjectValueIfNecessary(UObject* ExistingValue) const;
 
+	// Disable false positive buffer overrun warning during pgoprofile linking step
+	PRAGMA_DISABLE_BUFFER_OVERRUN_WARNING
 	/* Helper functions for UObject property types that wrap the object pointer in a smart pointer */
 	template <typename T>
 	void GetWrappedUObjectPtrValues(UObject** OutObjects, const void* SrcAddress, EPropertyMemoryAccess SrcAccess, int32 ArrayIndex, int32 ArrayCount) const
@@ -2751,6 +2753,9 @@ protected:
 			}
 		}
 	}
+	// Enable back buffer overrun warning
+	PRAGMA_ENABLE_BUFFER_OVERRUN_WARNING
+
 	template <typename T>
 	void SetWrappedUObjectPtrValues(void* DestAddress, EPropertyMemoryAccess DestAccess, UObject** InValues, int32 ArrayIndex, int32 ArrayCount) const
 	{
