Convert assertion failure into ensure failure to give more information and to not halt the cook. Keep it as an ensure rather than an error so we get callstacks for diagnostics.
The first time this occurred the callstack was during receive replication on the CookDirector from AssetDatas created on the CookWorker.
#rnx

diff --git a/Engine/Source/Runtime/CoreUObject/Private/AssetRegistry/AssetData.cpp b/Engine/Source/Runtime/CoreUObject/Private/AssetRegistry/AssetData.cpp
--- a/Engine/Source/Runtime/CoreUObject/Private/AssetRegistry/AssetData.cpp
+++ b/Engine/Source/Runtime/CoreUObject/Private/AssetRegistry/AssetData.cpp
@@ -532,9 +532,24 @@ void FAssetData::SetTagsAndAssetBundles(FAssetDataTagMap&& Tags)
 {
 	using namespace UE::AssetData::Private;
 
-	for (TPair<FName, FString>& Tag : Tags)
+	for (FAssetDataTagMap::TIterator Iter(Tags); Iter; ++Iter)
 	{
-		check(!Tag.Key.IsNone() && !Tag.Value.IsEmpty());
+		if (Iter->Key.IsNone())
+		{
+			ensureMsgf(!Iter->Key.IsNone(),
+				TEXT("FAssetData::SetTagsAndAssetBundles called on %s with empty key name. Empty key names are invalid. The Tag will be removed."),
+				*this->GetFullName());
+			Iter.RemoveCurrent();
+			continue;
+		}
+		if (Iter->Value.IsEmpty())
+		{
+			ensureMsgf(!Iter->Value.IsEmpty(),
+				TEXT("FAssetData::SetTagsAndAssetBundles called on %s with empty value for tag %s. Empty values are invalid. The Tag will be removed."),
+				*this->GetFullName(), *Iter->Key.ToString());
+			Iter.RemoveCurrent();
+			continue;
+		}
 	}
 
 	FString AssetBundles;
