[BugFix] FScriptMapHelper::Iterator and FScriptSetHelper::Iterator post-increment operators were pre-incrementing

diff --git a/Engine/Source/Runtime/CoreUObject/Public/UObject/UnrealType.h b/Engine/Source/Runtime/CoreUObject/Public/UObject/UnrealType.h
--- a/Engine/Source/Runtime/CoreUObject/Public/UObject/UnrealType.h
+++ b/Engine/Source/Runtime/CoreUObject/Public/UObject/UnrealType.h
@@ -5130,7 +5130,7 @@ public:
 		}
 
 		FIterator& operator++() { Advance(); return *this; }
-		FIterator& operator++(int) { Advance(); return *this; }
+		FIterator operator++(int) { const FIterator Temp(*this); Advance(); return Temp; }
 		explicit operator bool() const { return Map.IsValidIndex(CurrentIndex); }
 		int32 operator*() const { return CurrentIndex; }
 
@@ -5736,7 +5736,7 @@ public:
 		}
 
 		FIterator& operator++() { Advance(); return *this; }
-		FIterator& operator++(int) { Advance(); return *this; }
+		FIterator operator++(int) { const FIterator Temp(*this); Advance(); return Temp; }
 		explicit operator bool() const { return Set.IsValidIndex(CurrentIndex); }
 		int32 operator*() const { return CurrentIndex; }
 
