Fix potential truncation issue in FArchiveStackTraceWriter.

#rnx

- LastDifferenceCallstackOffsetIndex was only ever being used as a int32 so changed it to that type.

diff --git a/Engine/Source/Runtime/CoreUObject/Private/Serialization/ArchiveStackTrace.cpp b/Engine/Source/Runtime/CoreUObject/Private/Serialization/ArchiveStackTrace.cpp
--- a/Engine/Source/Runtime/CoreUObject/Private/Serialization/ArchiveStackTrace.cpp
+++ b/Engine/Source/Runtime/CoreUObject/Private/Serialization/ArchiveStackTrace.cpp
@@ -1022,7 +1022,7 @@ void FArchiveStackTrace::CompareWith(const FPackageData& SourcePackage, const TC
 bool FArchiveStackTraceWriter::GenerateDiffMap(const FPackageData& SourcePackage, const FPackageData& DestPackage, const FArchiveCallstacks& Callstacks, int32 MaxDiffsToFind, FArchiveDiffMap& OutDiffMap)
 {
 	bool bIdentical = true;
-	int64 LastDifferenceCallstackOffsetIndex = -1;
+	int32 LastDifferenceCallstackOffsetIndex = -1;
 	FArchiveCallstacks::FCallstackData* DifferenceCallstackData = nullptr;
 
 	const int64 SourceSize = SourcePackage.Size - SourcePackage.StartOffset;
@@ -1038,7 +1038,7 @@ bool FArchiveStackTraceWriter::GenerateDiffMap(const FPackageData& SourcePackage
 			bIdentical = false;
 			if (OutDiffMap.Num() < MaxDiffsToFind)
 			{
-				const int32 DifferenceCallstackOffsetIndex = Callstacks.GetCallstackIndexAtOffset(DestAbsoluteOffset, FMath::Max<int64>(LastDifferenceCallstackOffsetIndex, 0));
+				const int32 DifferenceCallstackOffsetIndex = Callstacks.GetCallstackIndexAtOffset(DestAbsoluteOffset, FMath::Max<int32>(LastDifferenceCallstackOffsetIndex, 0));
 				if (DifferenceCallstackOffsetIndex >= 0 && DifferenceCallstackOffsetIndex != LastDifferenceCallstackOffsetIndex)
 				{
 					const FArchiveCallstacks::FCallstackAtOffset& CallstackAtOffset = Callstacks.GetCallstack(DifferenceCallstackOffsetIndex);
