Added guards to stop property size calculation overflowing.

diff --git a/Engine/Source/Runtime/CoreUObject/Private/UObject/Property.cpp b/Engine/Source/Runtime/CoreUObject/Private/UObject/Property.cpp
--- a/Engine/Source/Runtime/CoreUObject/Private/UObject/Property.cpp
+++ b/Engine/Source/Runtime/CoreUObject/Private/UObject/Property.cpp
@@ -1377,6 +1377,14 @@ EConvertFromTypeResult FProperty::ConvertFromType(const FPropertyTag& Tag, FStru
 	return EConvertFromTypeResult::UseSerializeItem;
 }
 
+namespace UE::CoreUObject::Private
+{
+	[[noreturn]] void OnInvalidPropertySize(uint32 InvalidPropertySize, const FProperty* Prop)
+	{
+		UE_LOG(LogProperty, Fatal, TEXT("Invalid property size %u when linking property %s of size %d"), InvalidPropertySize, *Prop->GetFullName(), Prop->GetSize());
+		for (;;);
+	}
+}
 
 int32 FProperty::SetupOffset()
 {
@@ -1390,7 +1398,13 @@ int32 FProperty::SetupOffset()
 	{
 		Offset_Internal = Align(0, GetMinAlignment());
 	}
-	return Offset_Internal + GetSize();
+
+	uint32 UnsignedTotal = (uint32)Offset_Internal + (uint32)GetSize();
+	if (UnsignedTotal >= (uint32)MAX_int32)
+	{
+		UE::CoreUObject::Private::OnInvalidPropertySize(UnsignedTotal, this);
+	}
+	return (int32)UnsignedTotal;
 }
 
 void FProperty::SetOffset_Internal(int32 NewOffset)
