[BugFix] FindNth* methods in FScriptMapHelper and FScriptSetHelper were indexing out of order

diff --git a/Engine/Source/Runtime/CoreUObject/Public/UObject/UnrealType.h b/Engine/Source/Runtime/CoreUObject/Public/UObject/UnrealType.h
--- a/Engine/Source/Runtime/CoreUObject/Public/UObject/UnrealType.h
+++ b/Engine/Source/Runtime/CoreUObject/Public/UObject/UnrealType.h
@@ -4679,7 +4679,7 @@ public:
 		const int32 Index = FindInternalIndex(N);
 		
 		checkSlow(IsValidIndex(Index));
-		return (Index != INDEX_NONE) ? GetPairPtr(FindInternalIndex(Index)) : nullptr;
+		return (Index != INDEX_NONE) ? GetPairPtr(Index) : nullptr;
 	}
 	
 	/**
@@ -4691,9 +4691,7 @@ public:
 	uint8* FindNthKeyPtr(int32 N)
 	{
 		const int32 Index = FindInternalIndex(N);
-		
-		checkSlow(IsValidIndex(Index));
-		return (Index != INDEX_NONE) ? GetKeyPtr(FindInternalIndex(Index)) : nullptr;
+		return (Index != INDEX_NONE) ? GetKeyPtr(Index) : nullptr;
 	}
 	
 	/**
@@ -4705,9 +4703,7 @@ public:
 	uint8* FindNthValuePtr(int32 N)
 	{
 		const int32 Index = FindInternalIndex(N);
-		
-		checkSlow(IsValidIndex(Index));
-		return (Index != INDEX_NONE) ? GetValuePtr(FindInternalIndex(Index)) : nullptr;
+		return (Index != INDEX_NONE) ? GetValuePtr(Index) : nullptr;
 	}
 	
 	/**
@@ -4719,9 +4715,7 @@ public:
 	const uint8* FindNthPairPtr(int32 N) const
 	{
 		const int32 Index = FindInternalIndex(N);
-		
-		checkSlow(IsValidIndex(Index));
-		return (Index != INDEX_NONE) ? GetPairPtr(FindInternalIndex(Index)) : nullptr;
+		return (Index != INDEX_NONE) ? GetPairPtr(Index) : nullptr;
 	}
 
 	/**
@@ -5427,9 +5421,7 @@ public:
 	uint8* FindNthElementPtr(int32 N)
 	{
 		const int32 Index = FindInternalIndex(N);
-		
-		checkSlow(IsValidIndex(Index));
-		return (Index != INDEX_NONE) ? GetElementPtr(FindInternalIndex(Index)) : nullptr;
+		return (Index != INDEX_NONE) ? GetElementPtr(Index) : nullptr;
 	}
 
 	/**
@@ -5441,7 +5433,7 @@ public:
 	const uint8* FindNthElementPtr(int32 N) const
 	{
 		const int32 Index = FindInternalIndex(N);
-		return (Index != INDEX_NONE) ? GetElementPtr(FindInternalIndex(Index)) : nullptr;
+		return (Index != INDEX_NONE) ? GetElementPtr(Index) : nullptr;
 	}
 
 	/**
