Do not treat FIoHash as a rvalue when updating a payload as it doens't actually do anything.

#rnx

diff --git a/Engine/Source/Runtime/CoreUObject/Private/Serialization/EditorBulkData.cpp b/Engine/Source/Runtime/CoreUObject/Private/Serialization/EditorBulkData.cpp
--- a/Engine/Source/Runtime/CoreUObject/Private/Serialization/EditorBulkData.cpp
+++ b/Engine/Source/Runtime/CoreUObject/Private/Serialization/EditorBulkData.cpp
@@ -1827,7 +1827,7 @@ void FEditorBulkData::SerializeToPackageTrailer(FLinkerSave& LinkerSave, FCompre
 	LinkerSave.PackageTrailerBuilder->AddPayload(PayloadContentId, MoveTemp(PayloadToSerialize), PayloadFilter, MoveTemp(OnPayloadWritten));
 }
 
-void FEditorBulkData::UpdatePayloadImpl(FSharedBuffer&& InPayload, FIoHash&& InPayloadID, UObject* Owner)
+void FEditorBulkData::UpdatePayloadImpl(FSharedBuffer&& InPayload, const FIoHash& InPayloadID, UObject* Owner)
 {
 	TRACE_CPUPROFILER_EVENT_SCOPE(FEditorBulkData::UpdatePayloadImpl);
 
@@ -1852,7 +1852,7 @@ void FEditorBulkData::UpdatePayloadImpl(FSharedBuffer&& InPayload, FIoHash&& InP
 	}
 
 	PayloadSize = (int64)Payload.GetSize();
-	PayloadContentId = MoveTemp(InPayloadID);
+	PayloadContentId = InPayloadID;
 
 	EnumRemoveFlags(Flags,	EFlags::IsVirtualized |
 							EFlags::ReferencesLegacyFile |
@@ -1955,7 +1955,7 @@ void FEditorBulkData::UpdatePayload(FSharedBuffer InPayload, UObject* Owner)
 {
 	TRACE_CPUPROFILER_EVENT_SCOPE(FEditorBulkData::UpdatePayload);
 	FIoHash NewPayloadId = HashPayload(InPayload);
-	UpdatePayloadImpl(MoveTemp(InPayload), MoveTemp(NewPayloadId), Owner);
+	UpdatePayloadImpl(MoveTemp(InPayload), NewPayloadId, Owner);
 }
 
 FEditorBulkData::FSharedBufferWithID::FSharedBufferWithID(FSharedBuffer InPayload)
@@ -1966,7 +1966,7 @@ FEditorBulkData::FSharedBufferWithID::FSharedBufferWithID(FSharedBuffer InPayloa
 
 void FEditorBulkData::UpdatePayload(FSharedBufferWithID InPayload, UObject* Owner)
 {
-	UpdatePayloadImpl(MoveTemp(InPayload.Payload), MoveTemp(InPayload.PayloadId), Owner);
+	UpdatePayloadImpl(MoveTemp(InPayload.Payload), InPayload.PayloadId, Owner);
 }
 
 void FEditorBulkData::SetCompressionOptions(ECompressionOptions Option)
diff --git a/Engine/Source/Runtime/CoreUObject/Public/Serialization/EditorBulkData.h b/Engine/Source/Runtime/CoreUObject/Public/Serialization/EditorBulkData.h
--- a/Engine/Source/Runtime/CoreUObject/Public/Serialization/EditorBulkData.h
+++ b/Engine/Source/Runtime/CoreUObject/Public/Serialization/EditorBulkData.h
@@ -390,7 +390,7 @@ private:
 	/** The new path that saves payloads to the FPackageTrailer which is then appended to the end of the package file */
 	void SerializeToPackageTrailer(FLinkerSave& LinkerSave, FCompressedBuffer PayloadToSerialize, EFlags UpdatedFlags, UObject* Owner);
 
-	void UpdatePayloadImpl(FSharedBuffer&& InPayload, FIoHash&& InPayloadID, UObject* Owner);
+	void UpdatePayloadImpl(FSharedBuffer&& InPayload, const FIoHash& InPayloadID, UObject* Owner);
 
 	FCompressedBuffer GetDataInternal() const;
 
