Remove redundant WITH_EDITOR GC code

#rnx

diff --git a/Engine/Source/Runtime/CoreUObject/Private/UObject/GarbageCollection.cpp b/Engine/Source/Runtime/CoreUObject/Private/UObject/GarbageCollection.cpp
--- a/Engine/Source/Runtime/CoreUObject/Private/UObject/GarbageCollection.cpp
+++ b/Engine/Source/Runtime/CoreUObject/Private/UObject/GarbageCollection.cpp
@@ -4942,20 +4942,7 @@ void UObject::CallAddReferencedObjects(FReferenceCollector& Collector)
 
 void UObject::AddReferencedObjects(UObject* This, FReferenceCollector& Collector)
 {
-#if WITH_EDITOR
-	// Required by the unified GC when running in the editor
-	if (GIsEditor)
-	{
-		UObject* LoadOuter = This->GetOuter();
-		UClass* Class = This->GetClass();
-		UPackage* Package = This->GetExternalPackageInternal();
-		Collector.AllowEliminatingReferences(false);
-		Collector.AddReferencedObject(LoadOuter, This);
-		Collector.AddReferencedObject(Package, This);
-		Collector.AllowEliminatingReferences(true);
-		Collector.AddReferencedObject(Class, This);
-	}
-#endif
+	// This function exists to compare against in GetAROFunc()
 }
 
 bool UObject::IsDestructionThreadSafe() const
@@ -5437,12 +5424,7 @@ static AROFunc GetAROFunc(UClass* Class)
 {
 	AROFunc ARO = Class->CppClassStaticFunctions.GetAddReferencedObjects();
 	check(ARO != nullptr);
-
-#if WITH_EDITOR
-	return ARO;
-#else
 	return ARO != &UObject::AddReferencedObjects ? ARO : nullptr;
-#endif
 }
 
 } // namespace UE::GC
