Fix deadlock by not holding the annotation mutex while calling into code that takes another lock

#rnx

diff --git a/Engine/Source/Runtime/CoreUObject/Public/UObject/UObjectAnnotation.h b/Engine/Source/Runtime/CoreUObject/Public/UObject/UObjectAnnotation.h
--- a/Engine/Source/Runtime/CoreUObject/Public/UObject/UObjectAnnotation.h
+++ b/Engine/Source/Runtime/CoreUObject/Public/UObject/UObjectAnnotation.h
@@ -79,7 +79,6 @@ private:
 	void AddAnnotationInternal(const UObjectBase* Object, T&& Annotation)
 	{
 		check(Object);
-		FScopeLock AnnotationMapLock(&AnnotationMapCritical);
 		AnnotationCacheKey = Object;
 		AnnotationCacheValue = Forward<T>(Annotation);
 		if (AnnotationCacheValue.IsDefault())
@@ -88,7 +87,14 @@ private:
 		}
 		else
 		{
-			if (AnnotationMap.Num() == 0)
+			bool bWasEmpty = false;
+			{
+				FScopeLock AnnotationMapLock(&AnnotationMapCritical);
+				bWasEmpty = (AnnotationMap.Num() == 0);
+				AnnotationMap.Add(AnnotationCacheKey, AnnotationCacheValue);
+			}
+
+			if (bWasEmpty)
 			{
 				// we are adding the first one, so if we are auto removing or verifying removal, register now
 #if (UE_BUILD_SHIPPING || UE_BUILD_TEST)
@@ -98,7 +104,6 @@ private:
 					GUObjectArray.AddUObjectDeleteListener(this);
 				}
 			}
-			AnnotationMap.Add(AnnotationCacheKey, AnnotationCacheValue);
 		}
 	}
 
@@ -127,13 +132,20 @@ public:
 	TAnnotation GetAndRemoveAnnotation(const UObjectBase *Object)
 	{		
 		check(Object);
-		FScopeLock AnnotationMapLock(&AnnotationMapCritical);
-		AnnotationCacheKey = Object;
-		AnnotationCacheValue = TAnnotation();
-		const bool bHadElements = (AnnotationMap.Num() > 0);
+		bool bHadElements = false;
+		bool bIsNowEmpty = false;
+
+		// Avoid holding the lock while we call GUObjectArray.RemoveUObjectDeleteListener as it could deadlock
 		TAnnotation Result;
-		AnnotationMap.RemoveAndCopyValue(AnnotationCacheKey, Result);
-		if (bHadElements && AnnotationMap.Num() == 0)
+		{
+			FScopeLock AnnotationMapLock(&AnnotationMapCritical);
+			AnnotationCacheKey = Object;
+			AnnotationCacheValue = TAnnotation();
+			bHadElements = (AnnotationMap.Num() > 0);
+			AnnotationMap.RemoveAndCopyValue(AnnotationCacheKey, Result);
+			bIsNowEmpty = (AnnotationMap.Num() == 0);
+		}
+		if (bHadElements && bIsNowEmpty)
 		{
 			// we are removing the last one, so if we are auto removing or verifying removal, unregister now
 #if (UE_BUILD_SHIPPING || UE_BUILD_TEST)
@@ -153,12 +165,18 @@ public:
 	void RemoveAnnotation(const UObjectBase *Object)
 	{
 		check(Object);
-		FScopeLock AnnotationMapLock(&AnnotationMapCritical);
-		AnnotationCacheKey = Object;
-		AnnotationCacheValue = TAnnotation();
-		const bool bHadElements = (AnnotationMap.Num() > 0);
-		AnnotationMap.Remove(AnnotationCacheKey);
-		if (bHadElements && AnnotationMap.Num() == 0)
+		bool bHadElements = false;
+		bool bIsNowEmpty = false;
+		// Avoid holding the lock while we call GUObjectArray.RemoveUObjectDeleteListener as it could deadlock
+		{
+			FScopeLock AnnotationMapLock(&AnnotationMapCritical);
+			AnnotationCacheKey = Object;
+			AnnotationCacheValue = TAnnotation();
+			bHadElements = (AnnotationMap.Num() > 0);
+			AnnotationMap.Remove(AnnotationCacheKey);
+			bIsNowEmpty = (AnnotationMap.Num() == 0);
+		}
+		if (bHadElements && bIsNowEmpty)
 		{
 			// we are removing the last one, so if we are auto removing or verifying removal, unregister now
 #if (UE_BUILD_SHIPPING || UE_BUILD_TEST)
@@ -175,11 +193,16 @@ public:
 	 */
 	void RemoveAllAnnotations()
 	{
-		FScopeLock AnnotationMapLock(&AnnotationMapCritical);
-		AnnotationCacheKey = NULL;
-		AnnotationCacheValue = TAnnotation();
-		const bool bHadElements = (AnnotationMap.Num() > 0);
-		AnnotationMap.Empty();
+		bool bHadElements = false;
+
+		// Avoid holding the lock while we call GUObjectArray.RemoveUObjectDeleteListener as it could deadlock
+		{
+			FScopeLock AnnotationMapLock(&AnnotationMapCritical);
+			AnnotationCacheKey = nullptr;
+			AnnotationCacheValue = TAnnotation();
+			bHadElements = (AnnotationMap.Num() > 0);
+			AnnotationMap.Empty();
+		}
 		if (bHadElements)
 		{
 			// we are removing the last one, so if we are auto removing or verifying removal, unregister now
