Delete UE_WITH_GC that was only used by C++ UHT, which is unsupported and about to be removed

#rnx

diff --git a/Engine/Source/Runtime/CoreUObject/Private/UObject/GarbageCollection.cpp b/Engine/Source/Runtime/CoreUObject/Private/UObject/GarbageCollection.cpp
--- a/Engine/Source/Runtime/CoreUObject/Private/UObject/GarbageCollection.cpp
+++ b/Engine/Source/Runtime/CoreUObject/Private/UObject/GarbageCollection.cpp
@@ -1124,8 +1124,6 @@ namespace UE::GC {
 
 using AROFunc = void (*)(UObject*, FReferenceCollector&);
 
-#if UE_WITH_GC
-
 //////////////////////////////////////////////////////////////////////////
 
 /** Pool for reusing contexts between CollectReferences calls */
@@ -3318,34 +3316,13 @@ void TReachabilityCollector<Options>::HandleObjectReferences(UObject** InObjects
 
 }
 
-#else // !UE_WITH_GC
-static constexpr int32 MaxWorkers = 1;
-static void CollectGarbageInternal(EObjectFlags, bool)											{ unimplemented(); }
-void PadObjectArray(TArray<UObject*>&)															{ unimplemented(); }
-FWorkerContext::FWorkerContext()																{ unimplemented(); }
-FWorkerContext::~FWorkerContext()																{ unimplemented(); }
-FWorkBlockifier::~FWorkBlockifier()																{ unimplemented(); }
-void FWorkBlockifier::FreeOwningBlock(UObject*const*)											{ unimplemented(); }
-void FWorkBlockifier::PushFullBlockSync()														{ unimplemented(); }
-FWorkBlock* FWorkBlockifier::PopFullBlockSync()													{ unimplemented(); return nullptr; }
-FWorkBlock* FWorkBlockifier::PopFullBlockAsync()												{ unimplemented(); return nullptr; }
-FWorkBlock* FWorkBlockifier::PopWipBlock()														{ unimplemented(); return nullptr; }
-void FSlowARO::CallSync(uint32, UObject*, FReferenceCollector&)									{ unimplemented(); }
-bool FSlowARO::TryQueueCall(uint32 SlowAROIndex, UObject* Object, FWorkerContext& Context)		{ unimplemented(); return false; }
-void RegisterSlowImplementation(AROFunc, EAROFlags)												{ return; }
-int32 FindSlowImplementation(AROFunc)															{ return INDEX_NONE; }
-ELoot StealWork(FWorkerContext& Context, FReferenceCollector& Collector, FWorkBlock*& OutBlock)	{ return ELoot::Nothing; }
-#endif // !UE_WITH_GC
-
 } // namespace UE::GC
 
 //////////////////////////////////////////////////////////////////////////
 
 void ShutdownGarbageCollection()
 {
-#if UE_WITH_GC
 	UE::GC::FContextPoolScope().Cleanup();
-#endif
 	delete GAsyncPurge;
 	GAsyncPurge = nullptr;
 #if ENABLE_GC_HISTORY
@@ -3416,8 +3393,6 @@ void FReferenceFinder::HandleObjectReference( UObject*& InObject, const UObject*
 
 //////////////////////////////////////////////////////////////////////////
 
-#if UE_WITH_GC
-
 namespace UE::GC
 {
 
@@ -3789,8 +3764,6 @@ PRAGMA_ENABLE_DEPRECATION_WARNINGS
 
 } // namespace UE::GC
 
-#endif // UE_WITH_GC
-
 // Allow parallel GC to be overridden to single threaded via console command.
 static int32 GAllowParallelGC = 1;
 
@@ -3867,9 +3840,6 @@ static bool IncrementalDestroyGarbage(bool bUseTimeLimit, double TimeLimit);
  */
 void IncrementalPurgeGarbage(bool bUseTimeLimit, double TimeLimit)
 {
-#if !UE_WITH_GC
-	return;
-#else
 	if (GExitPurge)
 	{
 		GObjPurgeIsRequired = true;
@@ -3952,10 +3922,8 @@ void IncrementalPurgeGarbage(bool bUseTimeLimit, double TimeLimit)
 		// when running incrementally using a time limit, add one last tick for the memory trim
 		bCompleted = bCompleted && !bUseTimeLimit;
 	}
-#endif // !UE_WITH_GC
 }
 
-#if UE_WITH_GC
 
 
 static bool GWarningTimeOutHasBeenDisplayedGC = false;
@@ -4284,7 +4252,6 @@ bool IncrementalDestroyGarbage(bool bUseTimeLimit, double TimeLimit)
 
 	return bCompleted;
 }
-#endif // UE_WITH_GC
 
 /**
  * Returns whether an incremental purge is still pending/ in progress.
@@ -4506,7 +4473,6 @@ static void UpdateGCHistory(TConstArrayView<TUniquePtr<FWorkerContext>> Contexts
 #endif // ENABLE_GC_HISTORY
 }
 
-#if UE_WITH_GC
 
 template<bool bPerformFullPurge>
 FORCENOINLINE void CollectGarbageImpl(EObjectFlags KeepFlags);
@@ -4762,7 +4728,6 @@ void CollectGarbageImpl(EObjectFlags KeepFlags)
 	GTimingInfo.LastGCTime = FPlatformTime::Seconds();
 	STAT_ADD_CUSTOMMESSAGE_NAME( STAT_NamedMarker, TEXT( "GarbageCollection - End" ) );
 }
-#endif // !UE_WITH_GC
 
 } // namespace UE::GC
 
@@ -5845,8 +5810,6 @@ TUniquePtr<FTokenStreamBuilder> FIntrinsicClassTokens::ConsumeBuilder(UClass* Cl
 
 //////////////////////////////////////////////////////////////////////////
 
-#if UE_WITH_GC
-
 void DumpMemoryStats(FOutputDevice& OutputDevice)
 {
 	FContextPoolScope Pool;
@@ -6112,8 +6075,6 @@ void ProcessAsync(void (*ProcessSync)(void*, FWorkerContext&), void* Processor,
 
 //////////////////////////////////////////////////////////////////////////
 
-#endif //UE_WITH_GC
-
 } // namespace UE::GC
  
 int32 GetNumCollectReferenceWorkers()
diff --git a/Engine/Source/Runtime/CoreUObject/Private/UObject/GarbageCollectionVerification.h b/Engine/Source/Runtime/CoreUObject/Private/UObject/GarbageCollectionVerification.h
--- a/Engine/Source/Runtime/CoreUObject/Private/UObject/GarbageCollectionVerification.h
+++ b/Engine/Source/Runtime/CoreUObject/Private/UObject/GarbageCollectionVerification.h
@@ -8,12 +8,11 @@
 
 #include "CoreMinimal.h"
 #include "Stats/Stats.h"
-#include "UObject/GarbageCollection.h"	// Needed for UE_WITH_GC definition
 #include "UObject/UObjectGlobals.h"
 #include "UObject/Class.h"
 
 // UE_BUILD_SHIPPING has GShouldVerifyGCAssumptions=false by default
-#define VERIFY_DISREGARD_GC_ASSUMPTIONS			!(UE_BUILD_SHIPPING || UE_BUILD_TEST || !UE_WITH_GC)
+#define VERIFY_DISREGARD_GC_ASSUMPTIONS			!(UE_BUILD_SHIPPING || UE_BUILD_TEST)
 
 #if VERIFY_DISREGARD_GC_ASSUMPTIONS
 
diff --git a/Engine/Source/Runtime/CoreUObject/Public/UObject/GarbageCollection.h b/Engine/Source/Runtime/CoreUObject/Public/UObject/GarbageCollection.h
--- a/Engine/Source/Runtime/CoreUObject/Public/UObject/GarbageCollection.h
+++ b/Engine/Source/Runtime/CoreUObject/Public/UObject/GarbageCollection.h
@@ -24,10 +24,6 @@
 class FProperty;
 class UObject;
 
-#if !defined(UE_WITH_GC)
-#	define UE_WITH_GC	1
-#endif
-
 /** Context sensitive keep flags for garbage collection */
 #define GARBAGE_COLLECTION_KEEPFLAGS	(GIsEditor ? RF_Standalone : RF_NoFlags)
 
