Fix potential truncation issue in TReferenceSearchBase::FCollector

#rnx

- Bumped BufferStartIndex to a int64 as it is used to store the offset into a TArray64.

diff --git a/Engine/Source/Runtime/CoreUObject/Private/UObject/ReferenceChainSearch.cpp b/Engine/Source/Runtime/CoreUObject/Private/UObject/ReferenceChainSearch.cpp
--- a/Engine/Source/Runtime/CoreUObject/Private/UObject/ReferenceChainSearch.cpp
+++ b/Engine/Source/Runtime/CoreUObject/Private/UObject/ReferenceChainSearch.cpp
@@ -333,7 +333,7 @@ namespace UE::ReferenceChainSearch
 					// Flush current edge list to graph buffer
 					FVertex SourceVertex = Policy.ObjectToVertex(PreviousReferencingObject);
 					TArray<FVertex> EdgeList = CurrentEdgeList.Array();
-					int32 BufferStartIndex = GraphBuffer.Num();
+					const int64 BufferStartIndex = GraphBuffer.Num();
 					// During construction store edge lists relative to null to be fixed up later
 					FVertex* StartAddress = nullptr;
 					GraphBuffer.Append(EdgeList);
