Add an overload to FEditorBulkData::UpdatePayload that takes a FCompressedBuffer directly

#rnx

- FCompressedBuffer will calculate the FIoHash of its data and FEditorBulkData can take the hash from there rather than calculating one itself. If you create one with no compression codec set you could use that to control when the hashing takes place before passing the data to FEditorBulkData.
- In the future we might change FEditorBulkData to hold its data in a FCompressedBuffer, allowing us to avoid decompressing the payload unless someone actually asks for it in uncompressed format.
- Added some unit tests

diff --git a/Engine/Source/Runtime/CoreUObject/Private/Serialization/EditorBulkData.cpp b/Engine/Source/Runtime/CoreUObject/Private/Serialization/EditorBulkData.cpp
--- a/Engine/Source/Runtime/CoreUObject/Private/Serialization/EditorBulkData.cpp
+++ b/Engine/Source/Runtime/CoreUObject/Private/Serialization/EditorBulkData.cpp
@@ -1953,11 +1953,22 @@ TFuture<FCompressedBuffer>FEditorBulkData::GetCompressedPayload() const
 
 void FEditorBulkData::UpdatePayload(FSharedBuffer InPayload, UObject* Owner)
 {
-	TRACE_CPUPROFILER_EVENT_SCOPE(FEditorBulkData::UpdatePayload);
+	TRACE_CPUPROFILER_EVENT_SCOPE(FEditorBulkData::UpdatePayload (FSharedBuffer));
 	FIoHash NewPayloadId = HashPayload(InPayload);
+
 	UpdatePayloadImpl(MoveTemp(InPayload), NewPayloadId, Owner);
 }
 
+void FEditorBulkData::UpdatePayload(FCompressedBuffer InPayload, UObject* Owner)
+{
+	TRACE_CPUPROFILER_EVENT_SCOPE(FEditorBulkData::UpdatePayload(FCompressedBuffer));
+
+	FIoHash NewPayloadId = InPayload.GetRawSize() > 0 ? InPayload.GetRawHash() : FIoHash();
+	FSharedBuffer NewPayload = InPayload.DecompressToComposite().ToShared();
+
+	UpdatePayloadImpl(MoveTemp(NewPayload), NewPayloadId, Owner);
+}
+
 FEditorBulkData::FSharedBufferWithID::FSharedBufferWithID(FSharedBuffer InPayload)
 	: Payload(MoveTemp(InPayload))
 	, PayloadId(HashPayload(Payload))
diff --git a/Engine/Source/Runtime/CoreUObject/Private/Tests/Serialization/EditorBulkDataTests.cpp b/Engine/Source/Runtime/CoreUObject/Private/Tests/Serialization/EditorBulkDataTests.cpp
--- a/Engine/Source/Runtime/CoreUObject/Private/Tests/Serialization/EditorBulkDataTests.cpp
+++ b/Engine/Source/Runtime/CoreUObject/Private/Tests/Serialization/EditorBulkDataTests.cpp
@@ -1,6 +1,7 @@
 // Copyright Epic Games, Inc. All Rights Reserved.
 
 #include "Algo/AllOf.h"
+#include "Compression/OodleDataCompression.h"
 #include "Memory/SharedBuffer.h"
 #include "Serialization/BulkDataRegistry.h"
 #include "Serialization/EditorBulkData.h"
@@ -19,7 +20,7 @@
 namespace UE::Serialization
 {
 	
-constexpr const uint32 TestFlags = EAutomationTestFlags::ApplicationContextMask | EAutomationTestFlags::EngineFilter;
+constexpr const uint32 TestFlags = EAutomationTestFlags::ApplicationContextMask | EAutomationTestFlags::SmokeFilter;
 
 bool IsBulkDataRegistryEnabled()
 {
@@ -43,6 +44,46 @@ TUniquePtr<uint8[]> CreateRandomData(int64 BufferSize)
 	return Buffer;
 }
 
+/** Creates a FSharedBuffer full of random data to make it easy to have something to test against. */
+FSharedBuffer CreateRandomPayload(int64 BufferSize)
+{
+	TUniquePtr<uint8[]> Data = CreateRandomData(BufferSize);
+
+	return FSharedBuffer::TakeOwnership(Data.Release(), BufferSize, [](void* Ptr, uint64) { delete[](uint8*)Ptr; });
+}
+
+/** Creates a FSharedBuffer with semi random data. */
+FSharedBuffer CreatePayload(int64 BufferSize, int64 Stride)
+{
+	TUniquePtr<uint8[]> Data = CreateRandomData(BufferSize);
+
+	uint8 Value = (uint8)(FMath::Rand() % 255);
+	int64 SpanCount = 0;
+
+	for (int64 Index = 0; Index < BufferSize; ++Index)
+	{
+		Data[Index] = Value;
+
+		if (++SpanCount >= Stride)
+		{
+			Value = (uint8)(FMath::Rand() % 255);
+			SpanCount = 0;
+		}
+	}
+
+	return FSharedBuffer::TakeOwnership(Data.Release(), BufferSize, [](void* Ptr, uint64) { delete[](uint8*)Ptr; });
+}
+
+bool CompareSharedBufferContents(const FSharedBuffer& LHS, const FSharedBuffer& RHS)
+{
+	if (LHS.GetSize() != RHS.GetSize())
+	{
+		return false;
+	}
+
+	return FMemory::Memcmp(LHS.GetData(), RHS.GetData(), LHS.GetSize()) == 0;
+}
+
 /**
  * This test creates a very basic FEditorBulkData object with in memory payload and validates that we are able to retrieve the 
  * payload via both TFuture and callback methods. It then creates copies of the object and makes sure that we can get the payload from
@@ -127,10 +168,10 @@ bool FEditorBulkDataTestEmpty::RunTest(const FString& Parameters)
 }
 
 /**
- * Test the various methods for updating the payload that a FEditorBulkData owns
+ * Test the various methods for updating the payload that a FEditorBulkData owns via FSharedBuffer
  */
-IMPLEMENT_SIMPLE_AUTOMATION_TEST(FEditorBulkDataTestUpdatePayload, TEXT("System.CoreUObject.Serialization.EditorBulkData.UpdatePayload"), TestFlags)
-bool FEditorBulkDataTestUpdatePayload::RunTest(const FString& Parameters)
+IMPLEMENT_SIMPLE_AUTOMATION_TEST(FEditorBulkDataTestUpdatePayloadSharedBuffer, TEXT("System.CoreUObject.Serialization.EditorBulkData.UpdatePayloadSharedBuffer"), TestFlags)
+bool FEditorBulkDataTestUpdatePayloadSharedBuffer::RunTest(const FString& Parameters)
 {
 	// Create a memory buffer of all zeros
 	const int64 BufferSize = 1024;
@@ -200,6 +241,51 @@ bool FEditorBulkDataTestUpdatePayload::RunTest(const FString& Parameters)
 	return true;
 }
 
+IMPLEMENT_SIMPLE_AUTOMATION_TEST(FEditorBulkDataTestUpdatePayloadCompressedBuffer, TEXT("System.CoreUObject.Serialization.EditorBulkData.UpdatePayloadCompressedBuffer"), TestFlags)
+bool FEditorBulkDataTestUpdatePayloadCompressedBuffer::RunTest(const FString& Parameters)
+{
+	// Create a memory buffer
+	const int64 BufferSize = 1024;
+	const int64 BufferStride = 32;
+
+//	FSharedBuffer InitialPayload = CreateRandomPayload(BufferSize);
+	FSharedBuffer InitialPayload = CreatePayload(BufferSize, BufferStride);
+	
+	{
+		FCompressedBuffer UncompressedPayload = FCompressedBuffer::Compress(InitialPayload, ECompressedBufferCompressor::NotSet, ECompressedBufferCompressionLevel::None);
+
+		FEditorBulkData BulkData;
+		BulkData.UpdatePayload(UncompressedPayload, nullptr);
+
+		FSharedBuffer BulkDataPayload = BulkData.GetPayload().Get();
+
+		TestEqual(TEXT("Hash of the payload in FCompressedBuffer and FEditorBulkData"), UncompressedPayload.GetRawHash(), BulkData.GetPayloadId());
+		TestTrue(TEXT("The bulkdata payload has the same data as the original payload"), CompareSharedBufferContents(BulkDataPayload, BulkDataPayload));
+
+		// Since the data was never compressed there is no reason that the data needed to be copied at any point
+		// so the returned payload should be a reference to the original FSharedBuffer
+		TestEqual(TEXT("The bulkdata and the original payload should have the same memory addresses"), BulkDataPayload.GetData(), InitialPayload.GetData());
+	}
+
+	{
+		FCompressedBuffer CompressedPayload = FCompressedBuffer::Compress(InitialPayload, ECompressedBufferCompressor::Kraken, ECompressedBufferCompressionLevel::Fast);
+
+		FEditorBulkData BulkData;
+		BulkData.UpdatePayload(CompressedPayload, nullptr);
+
+		FSharedBuffer BulkDataPayload = BulkData.GetPayload().Get();
+
+		TestEqual(TEXT("Hash of the payload in FCompressedBuffer and FEditorBulkData"), CompressedPayload.GetRawHash(), BulkData.GetPayloadId());
+		TestTrue(TEXT("The bulkdata payload has the same data as the original payload"), CompareSharedBufferContents(BulkDataPayload, BulkDataPayload));
+
+		// Since the data was compressed we will not have a reference to the original FSharedBuffer and should
+		// expect different memory addresses.
+		TestNotEqual(TEXT("The bulkdata and the original payload should have different memory addresses"), BulkDataPayload.GetData(), InitialPayload.GetData());
+	}
+
+	return true;
+}
+
 /**
  * This test will create a buffer, then serialize it to a FEditorBulkData object via FEditorBulkDataWriter.
  * Then we will serialize the FEditorBulkData object back to a second buffer and compare the results.
diff --git a/Engine/Source/Runtime/CoreUObject/Public/Serialization/EditorBulkData.h b/Engine/Source/Runtime/CoreUObject/Public/Serialization/EditorBulkData.h
--- a/Engine/Source/Runtime/CoreUObject/Public/Serialization/EditorBulkData.h
+++ b/Engine/Source/Runtime/CoreUObject/Public/Serialization/EditorBulkData.h
@@ -221,6 +221,16 @@ public:
 	 */
 	void UpdatePayload(FSharedBuffer InPayload, UObject* Owner = nullptr);
 
+	/**
+	 * Replaces the existing payload (if any) with a new one.
+	 * Note that FCompressedBuffers already know the FIoHash of the raw payload and so we do not
+	 * recalculate it, making this method much faster for larger payloads than the other overloads.
+	 * 
+	 * @param InPayload	The payload that this bulkdata object should reference. @see FCompressedBuffer
+	 * @param Owner The object that owns the bulkdata, or null to not associate with a UObject.
+	 */
+	void UpdatePayload(FCompressedBuffer InPayload, UObject* Owner = nullptr);
+	
 	/**
 	 * Utility struct used to compute the Payload ID before calling UpdatePayload
 	 */
