Set FEditorBulkData tests back to EAutomationTestFlags::EngineFilter after an accidental change

#rnx

diff --git a/Engine/Source/Runtime/CoreUObject/Private/Tests/Serialization/EditorBulkDataTests.cpp b/Engine/Source/Runtime/CoreUObject/Private/Tests/Serialization/EditorBulkDataTests.cpp
--- a/Engine/Source/Runtime/CoreUObject/Private/Tests/Serialization/EditorBulkDataTests.cpp
+++ b/Engine/Source/Runtime/CoreUObject/Private/Tests/Serialization/EditorBulkDataTests.cpp
@@ -20,7 +20,7 @@
 namespace UE::Serialization
 {
 	
-constexpr const uint32 TestFlags = EAutomationTestFlags::ApplicationContextMask | EAutomationTestFlags::SmokeFilter;
+constexpr const uint32 TestFlags = EAutomationTestFlags::ApplicationContextMask | EAutomationTestFlags::EngineFilter;
 
 bool IsBulkDataRegistryEnabled()
 {
