SavePackage: Blueprints: Fix the BlueprintGeneratedClass and its default object to be marked as usedingame, despite UBlueprint being an editor-only class.
#rnx

diff --git a/Engine/Source/Runtime/CoreUObject/Private/UObject/SavePackage/PackageHarvester.cpp b/Engine/Source/Runtime/CoreUObject/Private/UObject/SavePackage/PackageHarvester.cpp
--- a/Engine/Source/Runtime/CoreUObject/Private/UObject/SavePackage/PackageHarvester.cpp
+++ b/Engine/Source/Runtime/CoreUObject/Private/UObject/SavePackage/PackageHarvester.cpp
@@ -51,7 +51,9 @@ EObjectMark GenerateMarksForObject(const UObject* InObject, const ITargetPlatfor
 	}
 	else
 	// If NotForClient and NotForServer, it is implicitly editor only
-	if ((Marks & OBJECTMARK_NotForClient) && (Marks & OBJECTMARK_NotForServer))
+	// If doing an editor save, HasNonEditorOnlyReferences can override it to still not be marked as editoronly
+	if ((Marks & OBJECTMARK_NotForClient) && (Marks & OBJECTMARK_NotForServer) &&
+		(TargetPlatform != nullptr || !InObject->HasNonEditorOnlyReferences()))
 	{
 		Marks = (EObjectMark)(Marks | OBJECTMARK_EditorOnly);
 	}
diff --git a/Engine/Source/Runtime/Engine/Classes/Engine/Blueprint.h b/Engine/Source/Runtime/Engine/Classes/Engine/Blueprint.h
--- a/Engine/Source/Runtime/Engine/Classes/Engine/Blueprint.h
+++ b/Engine/Source/Runtime/Engine/Classes/Engine/Blueprint.h
@@ -1069,6 +1069,7 @@ public:
 	virtual bool NeedsLoadForClient() const override;
 	virtual bool NeedsLoadForServer() const override;
 	virtual bool NeedsLoadForEditorGame() const override;
+	virtual bool HasNonEditorOnlyReferences() const override;
 	//~ End UObject Interface
 
 #if WITH_EDITORONLY_DATA
diff --git a/Engine/Source/Runtime/Engine/Private/Blueprint.cpp b/Engine/Source/Runtime/Engine/Private/Blueprint.cpp
--- a/Engine/Source/Runtime/Engine/Private/Blueprint.cpp
+++ b/Engine/Source/Runtime/Engine/Private/Blueprint.cpp
@@ -1797,6 +1797,13 @@ bool UBlueprint::NeedsLoadForEditorGame() const
 	return true;
 }
 
+bool UBlueprint::HasNonEditorOnlyReferences() const
+{
+	// The this->BlueprintGeneratedClass is reference that we need to mark as UsedInGame,
+	// despite UBlueprint being editor-only due to NeedsLoadForClient,NeedsLoadForServer == false
+	return true;
+}
+
 void UBlueprint::TagSubobjects(EObjectFlags NewFlags)
 {
 	Super::TagSubobjects(NewFlags);
