Fix infinite loading when flushing a package with a deferred post load group

diff --git a/Engine/Source/Runtime/CoreUObject/Private/Serialization/AsyncLoading2.cpp b/Engine/Source/Runtime/CoreUObject/Private/Serialization/AsyncLoading2.cpp
--- a/Engine/Source/Runtime/CoreUObject/Private/Serialization/AsyncLoading2.cpp
+++ b/Engine/Source/Runtime/CoreUObject/Private/Serialization/AsyncLoading2.cpp
@@ -3967,10 +3967,11 @@ void FAsyncLoadingThread2::IncludePackageInSyncLoadContextRecursive(FAsyncLoadin
 	}
 
 	Package->SyncLoadContextId = ContextId;
-	if (Package->PostLoadGroup && Package->PostLoadGroup->SyncLoadContextId < ContextId)
+	FAsyncLoadingPostLoadGroup* PostLoadGroup = Package->PostLoadGroup ? Package->PostLoadGroup : Package->DeferredPostLoadGroup;
+	if (PostLoadGroup && PostLoadGroup->SyncLoadContextId < ContextId)
 	{
-		Package->PostLoadGroup->SyncLoadContextId = ContextId;
-		for (FAsyncPackage2* PackageInPostLoadGroup : Package->PostLoadGroup->Packages)
+		PostLoadGroup->SyncLoadContextId = ContextId;
+		for (FAsyncPackage2* PackageInPostLoadGroup : PostLoadGroup->Packages)
 		{
 			if (PackageInPostLoadGroup->SyncLoadContextId < ContextId)
 			{
@@ -8526,6 +8527,10 @@ void FAsyncLoadingThread2::FlushLoading(int32 RequestId)
 						CurrentlyExecutingPackage->PostLoadGroup->PackagesWithExportsToSerializeCount = 1;
 					}
 				}
+				else
+				{
+					ensureMsgf(CurrentlyExecutingPackage->DeferredPostLoadGroup == nullptr, TEXT("Flushing while executing `%s` async PostLoad which is illegal"), *CurrentlyExecutingPackage->Desc.UPackageName.ToString());
+				}
 			}
 		}
 
