changed `FBlueprintSupport::RegisterDeferredDependenciesInStruct` to use FObjectProperty api to avoid c casts which do not work for unresolved TObjectPtr
https://p4-swarm.epicgames.net/reviews/25174026

diff --git a/Engine/Source/Runtime/CoreUObject/Private/Blueprint/BlueprintSupport.cpp b/Engine/Source/Runtime/CoreUObject/Private/Blueprint/BlueprintSupport.cpp
--- a/Engine/Source/Runtime/CoreUObject/Private/Blueprint/BlueprintSupport.cpp
+++ b/Engine/Source/Runtime/CoreUObject/Private/Blueprint/BlueprintSupport.cpp
@@ -169,8 +169,8 @@ void FBlueprintSupport::RegisterDeferredDependenciesInStruct(const UStruct* Stru
 	{
 		const FObjectProperty* Property = It.Key();
 		void* PropertyValue = (void*)It.Value();
-		UObject* ObjectValue = *((UObject**)PropertyValue);
-		
+		TObjectPtr<UObject> ObjectValue = Property->GetObjectPtrPropertyValue(PropertyValue);
+
 		ULinkerPlaceholderExportObject* PlaceholderVal = Cast<ULinkerPlaceholderExportObject>(ObjectValue);
 		ULinkerPlaceholderClass* PlaceholderClass = Cast<ULinkerPlaceholderClass>(ObjectValue);
 
