SavePackage's TransientPropertyOverrides: Change LinkerSave to not write out UObject properties that are marked with transient, rather than writing them out with nullptr. This removes the need for a hack fix of the property name and makes the package output more closely match the output that would occur if the property had no values, which should be more robust against future changes.
#rnx

diff --git a/Engine/Source/Runtime/CoreUObject/Private/UObject/LinkerSave.cpp b/Engine/Source/Runtime/CoreUObject/Private/UObject/LinkerSave.cpp
--- a/Engine/Source/Runtime/CoreUObject/Private/UObject/LinkerSave.cpp
+++ b/Engine/Source/Runtime/CoreUObject/Private/UObject/LinkerSave.cpp
@@ -389,22 +389,9 @@ FArchive& FLinkerSave::operator<<( FName& InName )
 FArchive& FLinkerSave::operator<<( UObject*& Obj )
 {
 	FPackageIndex Save;
-
 	if (Obj)
 	{
-		bool bKeepObject = true;
-		if (TransientPropertyOverrides && !TransientPropertyOverrides->IsEmpty())
-		{
-			const TSet<FProperty*>* Props = TransientPropertyOverrides->Find(CurrentlySavingExportObject);
-			if (Props && Props->Contains(this->GetSerializedProperty()))
-			{
-				bKeepObject = false;
-			}
-		}
-		if (bKeepObject)
-		{
-			Save = MapObject(Obj);
-		}
+		Save = MapObject(Obj);
 	}
 	return *this << Save;
 }
@@ -446,6 +433,20 @@ FArchive& FLinkerSave::operator<<(FLazyObjectPtr& LazyObjectPtr)
 	ID = LazyObjectPtr.GetUniqueID();
 	return *this << ID;
 }
+
+bool FLinkerSave::ShouldSkipProperty(const FProperty* InProperty) const
+{
+	if (TransientPropertyOverrides && !TransientPropertyOverrides->IsEmpty())
+	{
+		const TSet<FProperty*>* Props = TransientPropertyOverrides->Find(CurrentlySavingExportObject);
+		if (Props && Props->Contains(InProperty))
+		{
+			return true;
+		}
+	}
+	return false;
+}
+
 void FLinkerSave::SetSerializeContext(FUObjectSerializeContext* InLoadContext)
 {
 	SaveContext = InLoadContext;
diff --git a/Engine/Source/Runtime/CoreUObject/Private/UObject/SavePackage/PackageHarvester.cpp b/Engine/Source/Runtime/CoreUObject/Private/UObject/SavePackage/PackageHarvester.cpp
--- a/Engine/Source/Runtime/CoreUObject/Private/UObject/SavePackage/PackageHarvester.cpp
+++ b/Engine/Source/Runtime/CoreUObject/Private/UObject/SavePackage/PackageHarvester.cpp
@@ -826,9 +826,6 @@ bool FPackageHarvester::ShouldSkipProperty(const FProperty* InProperty) const
 	const TSet<FProperty*>* Props = TransientPropertyOverrides.Find(CurrentExportDependencies.CurrentExport);
 	if (Props && Props->Contains(InProperty))
 	{
-		// Hack fix: harvest the property name so it can be found when the property is serialized by LinkerSave
-		// TODO: Skip the property during LinkerSave as well and remove this hack.
-		const_cast<FPackageHarvester*>(this)->HarvestExportDataName(InProperty->GetFName()); 
 		return true;
 	}
 	return false;
diff --git a/Engine/Source/Runtime/CoreUObject/Private/UObject/SavePackage/PackageHarvester.h b/Engine/Source/Runtime/CoreUObject/Private/UObject/SavePackage/PackageHarvester.h
--- a/Engine/Source/Runtime/CoreUObject/Private/UObject/SavePackage/PackageHarvester.h
+++ b/Engine/Source/Runtime/CoreUObject/Private/UObject/SavePackage/PackageHarvester.h
@@ -82,7 +82,7 @@ public:
 	virtual FArchive& operator<<(FSoftObjectPath& Value) override;
 	virtual FArchive& operator<<(FName& Name) override;
 
-	virtual bool ShouldSkipProperty(const FProperty* InProperty) const;
+	virtual bool ShouldSkipProperty(const FProperty* InProperty) const override;
 
 	/** Return whether the Object and all its Outers are native. */
 	static bool IsObjNative(UObject* InObj);
diff --git a/Engine/Source/Runtime/CoreUObject/Public/UObject/LinkerSave.h b/Engine/Source/Runtime/CoreUObject/Public/UObject/LinkerSave.h
--- a/Engine/Source/Runtime/CoreUObject/Public/UObject/LinkerSave.h
+++ b/Engine/Source/Runtime/CoreUObject/Public/UObject/LinkerSave.h
@@ -161,6 +161,7 @@ public:
 	FArchive& operator<<( UObject*& Obj );
 	FArchive& operator<<(FSoftObjectPath& SoftObjectPath);
 	FArchive& operator<<( FLazyObjectPtr& LazyObjectPtr );
+	virtual bool ShouldSkipProperty(const FProperty* InProperty) const override;
 	virtual void SetSerializeContext(FUObjectSerializeContext* InLoadContext) override;
 	FUObjectSerializeContext* GetSerializeContext() override;
 	virtual void UsingCustomVersion(const struct FGuid& Guid) override;
