Don't create an adaptor for native non-structured struct serializers if we can avoid it.
Saves 25% of the time (1s) in an object reference collector in a test level.

diff --git a/Engine/Source/Runtime/CoreUObject/Private/UObject/Class.cpp b/Engine/Source/Runtime/CoreUObject/Private/UObject/Class.cpp
--- a/Engine/Source/Runtime/CoreUObject/Private/UObject/Class.cpp
+++ b/Engine/Source/Runtime/CoreUObject/Private/UObject/Class.cpp
@@ -2785,18 +2785,28 @@ void UScriptStruct::SerializeItem(FStructuredArchive::FSlot Slot, void* Value, v
 		else
 		{
 #if WITH_TEXT_ARCHIVE_SUPPORT
-			FArchiveUObjectFromStructuredArchive Adapter(Slot);
-			FArchive& Ar = Adapter.GetArchive();
-			bItemSerialized = TheCppStructOps->Serialize(Ar, Value);
-			if (bItemSerialized && !Slot.IsFilled())
+			if (Slot.GetUnderlyingArchive().IsTextFormat())
 			{
-				// The struct said that serialization succeeded but it didn't actually write anything.
-				Slot.EnterRecord();
+				FArchiveUObjectFromStructuredArchive Adapter(Slot);
+				FArchive& Ar = Adapter.GetArchive();
+				bItemSerialized = TheCppStructOps->Serialize(Ar, Value);
+				if (bItemSerialized && !Slot.IsFilled())
+				{
+					// The struct said that serialization succeeded but it didn't actually write anything.
+					Slot.EnterRecord();
+				}
+				Adapter.Close();
 			}
-			Adapter.Close();
-#else
-			bItemSerialized = TheCppStructOps->Serialize(Slot.GetUnderlyingArchive(), Value);
+			else
 #endif
+			{
+				bItemSerialized = TheCppStructOps->Serialize(Slot.GetUnderlyingArchive(), Value);
+				if (bItemSerialized && !Slot.IsFilled())
+				{
+					// The struct said that serialization succeeded but it didn't actually write anything.
+					Slot.EnterRecord();
+				}
+			}
 		}		
 	}
 
