Require FGCObject debug name

#rnx

diff --git a/Engine/Config/BaseEngine.ini b/Engine/Config/BaseEngine.ini
--- a/Engine/Config/BaseEngine.ini
+++ b/Engine/Config/BaseEngine.ini
@@ -1456,7 +1456,6 @@ gc.AssetClustreringEnabled=False
 gc.ActorClusteringEnabled=False
 gc.UseDisregardForGCOnDedicatedServers=False
 gc.MultithreadedDestructionEnabled=True
-gc.VerifyGCObjectNames=True
 gc.VerifyUObjectsAreNotFGCObjects=False
 gc.PendingKillEnabled=True
 
diff --git a/Engine/Source/Runtime/CoreUObject/Private/Misc/GCObjectReferencer.cpp b/Engine/Source/Runtime/CoreUObject/Private/Misc/GCObjectReferencer.cpp
--- a/Engine/Source/Runtime/CoreUObject/Private/Misc/GCObjectReferencer.cpp
+++ b/Engine/Source/Runtime/CoreUObject/Private/Misc/GCObjectReferencer.cpp
@@ -13,14 +13,6 @@
 #include "UObject/GarbageCollectionVerification.h"
 #include "HAL/IConsoleManager.h"
 
-static int32 GVerifyGCObjectNames = 1;
-static FAutoConsoleVariableRef CVerifyGCObjectNames(
-	TEXT("gc.VerifyGCObjectNames"),
-	GVerifyGCObjectNames,
-	TEXT("If true, the engine will verify if all FGCObject-derived classes define GetReferencerName() function overrides"),
-	ECVF_Default
-);
-
 // Global GC state flags
 extern bool GObjIncrementalPurgeIsInProgress;
 extern bool GObjUnhashUnreachableIsInProgress;
@@ -42,16 +34,9 @@ struct UGCObjectReferencer::FImpl
 
 
 UGCObjectReferencer::UGCObjectReferencer(const FObjectInitializer& ObjectInitializer)
-	: Super(ObjectInitializer)
-	, Impl(new FImpl)
-{
-#if VERIFY_DISREGARD_GC_ASSUMPTIONS
-	if (!HasAnyFlags(RF_ClassDefaultObject))
-	{
-		FCoreUObjectDelegates::GetPreGarbageCollectDelegate().AddUObject(this, &UGCObjectReferencer::VerifyGCObjectNames);
-	}
-#endif // VERIFY_DISREGARD_GC_ASSUMPTIONS
-}
+: Super(ObjectInitializer)
+, Impl(new FImpl)
+{}
 
 UGCObjectReferencer::UGCObjectReferencer(FVTableHelper& Helper)
 : Super(Helper)
@@ -97,9 +82,6 @@ void UGCObjectReferencer::AddObject(FGCObject* Object)
 	// Make sure there are no duplicates. Should be impossible...
 	checkSlow(!ReferencedObjects.Contains(Object));
 	ReferencedObjects.Add(Object);
-#if VERIFY_DISREGARD_GC_ASSUMPTIONS
-	bReferencedObjectsChangedSinceLastNameVerify = true;
-#endif // VERIFY_DISREGARD_GC_ASSUMPTIONS
 
 #if WITH_EDITORONLY_DATA
 	OnGCObjectAdded.Broadcast(Object);
@@ -115,9 +97,6 @@ void UGCObjectReferencer::RemoveObject(FGCObject* Object)
 	TArray<FGCObject*>& ReferencedObjects = Object->bCanMakeInitialReferences ? Impl->InitialReferencedObjects : Impl->RemainingReferencedObjects; 
 	int32 NumRemoved = ReferencedObjects.RemoveSingleSwap(Object);
 	check(NumRemoved == 1);
-#if VERIFY_DISREGARD_GC_ASSUMPTIONS
-	bReferencedObjectsChangedSinceLastNameVerify = true;
-#endif // VERIFY_DISREGARD_GC_ASSUMPTIONS
 }
 
 bool UGCObjectReferencer::GetReferencerName(UObject* Object, FString& OutName, bool bOnlyIfAddingReferenced) const
@@ -232,76 +211,6 @@ void UGCObjectReferencer::AddInitialReferences(TArray<UObject**>& Out)
 	}
 }
 
-/** 
-  * Helper class to make sure the ensure that verifies GetReferencerName is fired from within a callstack that can be
-  * associated with the offending FGCObject-derived class.
-  */
-class FVerifyReferencerNameCollector : public FReferenceCollector
-{
-	/** Current FGCObject being processed */
-	FGCObject* CurrentGCObject = nullptr;
-	/** Since we only need to fire the ensure once per FGCObject instance this will prevent us from spamming with ensures */
-	bool bCurrentObjectVerified = false;
-
-public:
-
-	void SetCurrentGCObject(FGCObject* InGCObject)
-	{
-		CurrentGCObject = InGCObject;
-		bCurrentObjectVerified = false;
-	}
-
-	void VerifyReferencerName()
-	{
-		if (!bCurrentObjectVerified && CurrentGCObject)
-		{
-			if (!ensureAlwaysMsgf(CurrentGCObject->GetReferencerName() !=  UE::GC::UnknownGCObjectName,
-				TEXT("Please make sure all FGCObject derived classes have a unique name by overriding FGCObject::GetReferencerName() function. FGCObject::GetReferencerName() will become pure virtual in the next engine release. See callstack for details.")))
-			{
-				bCurrentObjectVerified = true;
-			}
-		}
-	}
-
-	// Begin FGCObject interface
-	virtual void HandleObjectReference(UObject*& Object, const UObject* ReferencingObject, const FProperty* ReferencingProperty) override
-	{
-		VerifyReferencerName();
-	}
-	virtual void HandleObjectReferences(UObject** InObjects, const int32 ObjectNum, const UObject* ReferencingObject, const FProperty* InReferencingProperty) override
-	{
-		VerifyReferencerName();
-	}
-	virtual bool IsIgnoringArchetypeRef() const override
-	{
-		return false;
-	}
-	virtual bool IsIgnoringTransient() const override
-	{
-		return false;
-	}
-	// End FGCObject interface
-};
-
-void UGCObjectReferencer::VerifyGCObjectNames()
-{
-	if (bReferencedObjectsChangedSinceLastNameVerify && GShouldVerifyGCAssumptions && GVerifyGCObjectNames)
-	{
-		FVerifyReferencerNameCollector VerifyReferencerNameCollector;
-
-		for (TArrayView<FGCObject*> ReferencedObjects : { MakeArrayView(Impl->RemainingReferencedObjects),
-														  MakeArrayView(Impl->InitialReferencedObjects) })
-		{
-			for (FGCObject* GCObject : ReferencedObjects)
-			{
-				check(GCObject);
-				VerifyReferencerNameCollector.SetCurrentGCObject(GCObject);
-				GCObject->AddReferencedObjects(VerifyReferencerNameCollector);
-			}
-		}
-	}
-}
-
 IMPLEMENT_CORE_INTRINSIC_CLASS(UGCObjectReferencer, UObject, 
 	{
 		Class->CppClassStaticFunctions = UOBJECT_CPPCLASS_STATICFUNCTIONS_FORCLASS(UGCObjectReferencer);
diff --git a/Engine/Source/Runtime/CoreUObject/Public/UObject/GCObject.h b/Engine/Source/Runtime/CoreUObject/Public/UObject/GCObject.h
--- a/Engine/Source/Runtime/CoreUObject/Public/UObject/GCObject.h
+++ b/Engine/Source/Runtime/CoreUObject/Public/UObject/GCObject.h
@@ -27,8 +27,6 @@ class UGCObjectReferencer : public UObject
 	/** Current FGCObject* that references are being added from  */
 	FGCObject* CurrentlySerializingObject = nullptr;
 
-	/** True if the list of referenced objects has changed since we last ran VerifyGCObjectNames */
-	bool bReferencedObjectsChangedSinceLastNameVerify = false;
 	friend struct FReplaceReferenceHelper;
 
 public:
@@ -77,11 +75,6 @@ public:
 	 */
 	virtual void FinishDestroy() override;
 
-	/**
-	 * Verifies that all FGCObjects have specified names
-	 */
-	COREUOBJECT_API void VerifyGCObjectNames();
-
 	/**
 	 * Returns the currently serializing object
 	 */
@@ -186,10 +179,7 @@ public:
 	virtual void AddReferencedObjects( FReferenceCollector& Collector ) = 0;
 
 	/** Overload this method to report a name for your referencer */
-	virtual FString GetReferencerName() const
-	{
-		return UnknownGCObjectName;
-	}
+	virtual FString GetReferencerName() const = 0;
 
 	/** Overload this method to report how the specified object is referenced, if necessary */
 	virtual bool GetReferencerPropertyName(UObject* Object, FString& OutPropertyName) const
diff --git a/Engine/Source/Runtime/Engine/Classes/Engine/CoreSettings.h b/Engine/Source/Runtime/Engine/Classes/Engine/CoreSettings.h
--- a/Engine/Source/Runtime/Engine/Classes/Engine/CoreSettings.h
+++ b/Engine/Source/Runtime/Engine/Classes/Engine/CoreSettings.h
@@ -217,11 +217,6 @@ protected:
 		ToolTip = "If false, DisregardForGC will be disabled for dedicated servers."))
 	uint32 UseDisregardForGCOnDedicatedServers : 1;
 
-	UPROPERTY(EditAnywhere, config, Category = Debug, meta = (
-		ConsoleVariable = "gc.VerifyGCObjectNames", DisplayName = "Verify FGCObject names",
-		ToolTip = "If true, the engine will verify if all FGCObject-derived classes define GetReferencerName() function overrides."))
-	uint32 VerifyGCObjectNames : 1;
-
 	UPROPERTY(EditAnywhere, config, Category = Debug, meta = (
 		ConsoleVariable = "gc.VerifyUObjectsAreNotFGCObjects", DisplayName = "Verify UObjects Are Not FGCObjects",
 		ToolTip = "If true, the engine will throw a warning when it detects a UObject-derived class which also derives from FGCObject or any of its members is derived from FGCObject."))
diff --git a/Engine/Source/Runtime/Engine/Private/CoreSettings.cpp b/Engine/Source/Runtime/Engine/Private/CoreSettings.cpp
--- a/Engine/Source/Runtime/Engine/Private/CoreSettings.cpp
+++ b/Engine/Source/Runtime/Engine/Private/CoreSettings.cpp
@@ -210,7 +210,6 @@ UGarbageCollectionSettings::UGarbageCollectionSettings()
 	AssetClusteringEnabled = true;
 	ActorClusteringEnabled = true;
 	UseDisregardForGCOnDedicatedServers = false;
-	VerifyGCObjectNames = true;
 	VerifyUObjectsAreNotFGCObjects = true;
 	PendingKillEnabled = false;
 	DumpObjectCountsToLogWhenMaxObjectLimitExceeded = false;
