SavePackage: Suppress spurious static analysis warning "'TargetPlatform' pointer was utilized before it was verified against nullptr."
#rnx

diff --git a/Engine/Source/Runtime/CoreUObject/Private/UObject/SavePackage/PackageHarvester.cpp b/Engine/Source/Runtime/CoreUObject/Private/UObject/SavePackage/PackageHarvester.cpp
--- a/Engine/Source/Runtime/CoreUObject/Private/UObject/SavePackage/PackageHarvester.cpp
+++ b/Engine/Source/Runtime/CoreUObject/Private/UObject/SavePackage/PackageHarvester.cpp
@@ -39,7 +39,8 @@ EObjectMark GenerateMarksForObject(const UObject* InObject, const ITargetPlatfor
 		// no need to query the target platform if the object is editoronly and the targetplatform doesn't allow editor objects 
 		bCheckTargetPlatform = !bIsEditorOnlyObject || bTargetAllowsEditorObjects;
 	}
-	if (bCheckTargetPlatform && (!InObject->NeedsLoadForTargetPlatform(TargetPlatform) || !TargetPlatform->AllowObject(InObject)))
+	if (bCheckTargetPlatform && TargetPlatform && 
+		(!InObject->NeedsLoadForTargetPlatform(TargetPlatform) || !TargetPlatform->AllowObject(InObject)))
 	{
 		Marks = (EObjectMark)(Marks | OBJECTMARK_NotForTargetPlatform);
 	}
