Fixed mismatch string formatting argument in SavePackage2.

diff --git a/Engine/Source/Runtime/CoreUObject/Private/UObject/SavePackage2.cpp b/Engine/Source/Runtime/CoreUObject/Private/UObject/SavePackage2.cpp
--- a/Engine/Source/Runtime/CoreUObject/Private/UObject/SavePackage2.cpp
+++ b/Engine/Source/Runtime/CoreUObject/Private/UObject/SavePackage2.cpp
@@ -231,7 +231,7 @@ ESavePackageResult ValidatePackage(FSaveContext& SaveContext)
 			}
 			else
 			{
-				UE_LOG(LogSavePackage, Warning, TEXT("%s"), *FString::Printf(TEXT("Filename is too long (%d characters); this may interfere with cooking for consoles. Unreal filenames should be no longer than %s characters. Filename value: %s"), BaseFilename.Len(), MaxFilenameLength, *BaseFilename));
+				UE_LOG(LogSavePackage, Warning, TEXT("%s"), *FString::Printf(TEXT("Filename is too long (%d characters); this may interfere with cooking for consoles. Unreal filenames should be no longer than %d characters. Filename value: %s"), BaseFilename.Len(), MaxFilenameLength, *BaseFilename));
 			}
 		}
 	}
