change to FObjectPtrProperty::Identical to avoid loading lazy loaded objects
https://p4-swarm.epicgames.net/reviews/25288349

diff --git a/Engine/Source/Runtime/CoreUObject/Private/UObject/PropertyObjectPtr.cpp b/Engine/Source/Runtime/CoreUObject/Private/UObject/PropertyObjectPtr.cpp
--- a/Engine/Source/Runtime/CoreUObject/Private/UObject/PropertyObjectPtr.cpp
+++ b/Engine/Source/Runtime/CoreUObject/Private/UObject/PropertyObjectPtr.cpp
@@ -152,12 +152,11 @@ bool FObjectPtrProperty::Identical(const void* A, const void* B, uint32 PortFlag
 			UObject* DSO = ClassA ? ClassA->GetDefaultSubobjectByName(ObjectA.GetFName()) : nullptr;
 			bPerformDeepComparison = DSO != nullptr;
 		}
-		if (bPerformDeepComparison)
+		if (bPerformDeepComparison && ObjectA.GetClass() == ObjectB.GetClass() && ObjectA.GetFName() == ObjectB.GetFName())
 		{
 			return FObjectPropertyBase::StaticIdentical(ObjectA.Get(), ObjectB.Get(), PortFlags);
 		}
 	}
-
 	return false;
 }
 
diff --git a/Engine/Source/Runtime/CoreUObject/Tests/ObjectPropertyTest.cpp b/Engine/Source/Runtime/CoreUObject/Tests/ObjectPropertyTest.cpp
--- a/Engine/Source/Runtime/CoreUObject/Tests/ObjectPropertyTest.cpp
+++ b/Engine/Source/Runtime/CoreUObject/Tests/ObjectPropertyTest.cpp
@@ -382,4 +382,62 @@ TEST_CASE("UE::CoreUObject::FObjectProperty::ParseObjectPropertyValue")
 }
 #endif
 
+
+TEST_CASE("UE::FObjectPtrProperty::Identical")
+{
+	int ResolveCount = 0;
+#if UE_WITH_OBJECT_HANDLE_TRACKING
+	auto CallbackHandle = UE::CoreUObject::AddObjectHandleReferenceResolvedCallback([&ResolveCount](const FObjectRef&, UPackage*, UObject*)
+		{
+			++ResolveCount;
+		});
+	ON_SCOPE_EXIT
+	{
+		UE::CoreUObject::RemoveObjectHandleReferenceResolvedCallback(CallbackHandle);
+	};
+#endif
+	UClass* Class = UObjectPtrTestClassWithRef::StaticClass();
+	FObjectProperty* Property = CastField<FObjectProperty>(Class->FindPropertyByName(TEXT("ObjectPtr")));
+	REQUIRE(Property != nullptr);
+
+
+	UPackage* TestPackage = NewObject<UPackage>(nullptr, TEXT("Test/TestPackageName"), RF_Transient);
+	TestPackage->AddToRoot();
+	UPackage* TestPackage2 = NewObject<UPackage>(nullptr, TEXT("Test/TestPackageName2"), RF_Transient);
+	TestPackage2->AddToRoot();
+	ON_SCOPE_EXIT
+	{
+		TestPackage->RemoveFromRoot();
+		TestPackage2->RemoveFromRoot();
+	};
+	TObjectPtr<UObject> ObjWithRef = NewObject<UObjectPtrTestClassWithRef>(TestPackage, TEXT("UObjectWithClassProperty"));
+	TObjectPtr<UObject> Obj1 = NewObject<UObjectPtrTestClass>(TestPackage, TEXT("UObjectPtrTestClass"));
+	TObjectPtr<UObject> Obj2 = NewObject<UObjectPtrTestClass>(TestPackage2, TEXT("UObjectPtrTestClass"));
+
+#if UE_WITH_OBJECT_HANDLE_TRACKING
+
+	FObjectHandle Handle1 = MakeUnresolvedHandle(ObjWithRef.Get());
+	FObjectHandle Handle2 = MakeUnresolvedHandle(Obj1.Get());
+	FObjectHandle Handle3 = MakeUnresolvedHandle(Obj2.Get());
+	TObjectPtr<UObjectPtrTestClass> ObjectPtr = *reinterpret_cast<TObjectPtr<UObjectPtrTestClass>*>(&Handle1);
+
+	ObjWithRef = *reinterpret_cast<TObjectPtr<UObject>*>(&Handle1);
+	Obj1 = *reinterpret_cast<TObjectPtr<UObject>*>(&Handle2);
+	Obj2 = *reinterpret_cast<TObjectPtr<UObject>*>(&Handle3);
+#endif
+
+	CHECK(Property->Identical(&Obj1, &Obj1, 0u));
+	CHECK(!Property->Identical(&Obj1, nullptr, 0u));
+	CHECK(!Property->Identical(nullptr, &Obj1, 0u));
+	CHECK(!Property->Identical(&ObjWithRef, &Obj2, 0u));
+	CHECK(!Property->Identical(&ObjWithRef, &Obj2, 0u));
+	CHECK(!Property->Identical(&Obj1, &ObjWithRef, PPF_DeepComparison));
+
+	CHECK(ResolveCount == 0);
+	CHECK(Property->Identical(&Obj1, &Obj2, PPF_DeepComparison));
+#if UE_WITH_OBJECT_HANDLE_TRACKING
+	CHECK(ResolveCount == 2);
+#endif
+	
+}
 #endif
