Making sure running out of available UObject slots crashes regardless of checks being enabled or not otherwise it will siletnly result in stomps or OOM crashes

diff --git a/Engine/Source/Runtime/CoreUObject/Public/UObject/UObjectArray.h b/Engine/Source/Runtime/CoreUObject/Public/UObject/UObjectArray.h
--- a/Engine/Source/Runtime/CoreUObject/Public/UObject/UObjectArray.h
+++ b/Engine/Source/Runtime/CoreUObject/Public/UObject/UObjectArray.h
@@ -228,12 +228,10 @@ namespace UE::UObjectArrayPrivate
 
 	FORCEINLINE void CheckUObjectLimitReached(const int32 NumUObjects, const int32 MaxUObjects, const int32 NewUObjectCount)
 	{
-#if DO_CHECK
 		if ((NumUObjects + NewUObjectCount) > MaxUObjects)
 		{
 			FailMaxUObjectCountExceeded(MaxUObjects, NewUObjectCount);
 		}
-#endif // DO_CHECK
 	}
 };
 
