Fix a bug where FArchiveCrc32 would not hash all of the data provided to it if the data was sufficiently large enough.

#rnx

- This should only change the hash when FArchiveCrc32::Serialize is called on buffers larger than MAX_int32 in size, which given the current use cases of FArchiveCrc32 is probably not occuring anywhere anyway.
- Found this issue when fixing truncation warnings.
- Serialize accepts data lengths of type int64 but FCrc::MemCrc32 would silently truncate the data to int32.
-- This doesn't result in the data being hashed up to length MAX_int32 as the iteration inside of MemCrc32 (at least in some implementations) uses a signed index, passing in larger values will result in this index being negative and our loop checks for that index being reduced to 0.
- To fix this we should call FCrc::MemCrc32 with the data length clamped to MAX_int32 and then make repeated calls until the entire buffer hash been hashed.

diff --git a/Engine/Source/Runtime/CoreUObject/Private/Serialization/ArchiveCrc32.cpp b/Engine/Source/Runtime/CoreUObject/Private/Serialization/ArchiveCrc32.cpp
--- a/Engine/Source/Runtime/CoreUObject/Private/Serialization/ArchiveCrc32.cpp
+++ b/Engine/Source/Runtime/CoreUObject/Private/Serialization/ArchiveCrc32.cpp
@@ -14,7 +14,16 @@ FArchiveCrc32::FArchiveCrc32(uint32 InCRC, UObject* InRootObject)
 
 void FArchiveCrc32::Serialize(void* Data, int64 Num)
 {
-	CRC = FCrc::MemCrc32(Data, Num, CRC);
+	uint8* BytePointer = static_cast<uint8*>(Data);
+	while (Num > 0)
+	{
+		const int32 BytesToHash = static_cast<int32>(FMath::Min(Num, int64(MAX_int32)));
+
+		CRC = FCrc::MemCrc32(BytePointer, BytesToHash, CRC);
+
+		Num -= BytesToHash;
+		BytePointer += BytesToHash;
+	}
 }
 
 FArchive& FArchiveCrc32::operator<<(class FName& Name)
