Fix potential truncation issue in AsyncLoading2

#rnx

- If the value being stored in ExportBundleEntriesCount exceeds MAX_int32 then we will end up failing the check below and printing an error anyway so we might as well just cast at this point.

diff --git a/Engine/Source/Runtime/CoreUObject/Private/Serialization/AsyncLoading2.cpp b/Engine/Source/Runtime/CoreUObject/Private/Serialization/AsyncLoading2.cpp
--- a/Engine/Source/Runtime/CoreUObject/Private/Serialization/AsyncLoading2.cpp
+++ b/Engine/Source/Runtime/CoreUObject/Private/Serialization/AsyncLoading2.cpp
@@ -5107,7 +5107,7 @@ static void ReadAsyncPackageHeader(FAsyncPackageSerializationState& Serializatio
 	HeaderData.ExportCount = HeaderData.ExportMap.Num();
 
 	const uint64 ExportBundleEntriesSize = PackageSummary->DependencyBundleHeadersOffset - PackageSummary->ExportBundleEntriesOffset;
-	const int32 ExportBundleEntriesCount = ExportBundleEntriesSize / sizeof(FExportBundleEntry);
+	const int32 ExportBundleEntriesCount = static_cast<int32>(ExportBundleEntriesSize / sizeof(FExportBundleEntry));
 
 	if (ExportBundleEntriesCount != HeaderData.ExportCount * FExportBundleEntry::ExportCommandType_Count)
 	{
