UClass uses UE::FMutex instead of FRWLock if USE_UE_LOCK_FOR_UCLASS_FUNCTION_HANDLING is defined (useful on platform that has a limited system objects count)

#rnx

diff --git a/Engine/Source/Runtime/CoreUObject/Private/UObject/Class.cpp b/Engine/Source/Runtime/CoreUObject/Private/UObject/Class.cpp
--- a/Engine/Source/Runtime/CoreUObject/Private/UObject/Class.cpp
+++ b/Engine/Source/Runtime/CoreUObject/Private/UObject/Class.cpp
@@ -4879,8 +4879,8 @@ void UClass::Serialize( FArchive& Ar )
 	// serialize the function map
 	//@TODO: UCREMOVAL: Should we just regenerate the FuncMap post load, instead of serializing it?
 	{
-		FWriteScopeLock ScopeLock(FuncMapLock);
-	Ar << FuncMap;
+	 	FUClassFuncScopeWriteLock ScopeLock(FuncMapLock);
+		Ar << FuncMap;
 	}
 
 	// Class flags first.
@@ -5300,8 +5300,8 @@ void UClass::PurgeClass(bool bRecompilingOnLoad)
 	DeleteUnresolvedScriptProperties();
 
 	{
-		FWriteScopeLock ScopeLock(FuncMapLock);
-	FuncMap.Empty();
+		FUClassFuncScopeWriteLock ScopeLock(FuncMapLock);
+		FuncMap.Empty();
 	}
 	ClearFunctionMapsCaches();
 	PropertyLink = nullptr;
@@ -5869,7 +5869,7 @@ void UClass::CreateLinkAndAddChildFunctionsToMap(const FClassFunctionLinkInfo* F
 
 void UClass::ClearFunctionMapsCaches()
 {
-	FWriteScopeLock ScopeLock(SuperFuncMapLock);
+	FUClassFuncScopeWriteLock ScopeLock(SuperFuncMapLock);
 	SuperFuncMap.Empty();
 }
 
@@ -5882,7 +5882,7 @@ UFunction* UClass::FindFunctionByName(FName InName, EIncludeSuperFlag::Type Incl
 	UE_AUTORTFM_OPEN(
 	{
 		{
-			FReadScopeLock ScopeLock(FuncMapLock);
+			FUClassFuncScopeReadLock ScopeLock(FuncMapLock);
 			Result = FuncMap.FindRef(InName);
 		}
 
@@ -5893,7 +5893,7 @@ UFunction* UClass::FindFunctionByName(FName InName, EIncludeSuperFlag::Type Incl
 			{
 				bool bFoundInSuperFuncMap = false;
 				{
-					FReadScopeLock ScopeLock(SuperFuncMapLock);
+					FUClassFuncScopeReadLock ScopeLock(SuperFuncMapLock);
 					if (UFunction** SuperResult = SuperFuncMap.Find(InName))
 					{
 						Result = *SuperResult;
@@ -5917,7 +5917,7 @@ UFunction* UClass::FindFunctionByName(FName InName, EIncludeSuperFlag::Type Incl
 						Result = SuperClass->FindFunctionByName(InName);
 					}
 
-					FWriteScopeLock ScopeLock(SuperFuncMapLock);
+					FUClassFuncScopeWriteLock ScopeLock(SuperFuncMapLock);
 					SuperFuncMap.Add(InName, Result);
 				}
 			}
diff --git a/Engine/Source/Runtime/CoreUObject/Public/UObject/Class.h b/Engine/Source/Runtime/CoreUObject/Public/UObject/Class.h
--- a/Engine/Source/Runtime/CoreUObject/Public/UObject/Class.h
+++ b/Engine/Source/Runtime/CoreUObject/Public/UObject/Class.h
@@ -15,6 +15,8 @@
 #include "Containers/StringFwd.h"
 #include "Containers/UnrealString.h"
 #include "CoreTypes.h"
+#include "Async/Mutex.h"
+#include "Async/UniqueLock.h"
 #include "HAL/CriticalSection.h"
 #include "HAL/PlatformCrt.h"
 #include "HAL/UnrealMemory.h"
@@ -111,6 +113,21 @@ template <typename FuncType> class TFunctionRef;
 COREUOBJECT_API DECLARE_LOG_CATEGORY_EXTERN(LogClass, Log, All);
 COREUOBJECT_API DECLARE_LOG_CATEGORY_EXTERN(LogScriptSerialization, Log, All);
 
+#ifndef USE_UE_LOCK_FOR_UCLASS_FUNCTION_HANDLING
+	#define USE_UE_LOCK_FOR_UCLASS_FUNCTION_HANDLING 0
+#endif
+
+#if USE_UE_LOCK_FOR_UCLASS_FUNCTION_HANDLING
+	typedef UE::TUniqueLock<UE::FMutex> FUClassFuncScopeReadLock;
+	typedef UE::TUniqueLock<UE::FMutex> FUClassFuncScopeWriteLock;
+	typedef UE::FMutex FUClassFuncLock;
+#else
+	typedef FReadScopeLock FUClassFuncScopeReadLock;
+	typedef FWriteScopeLock FUClassFuncScopeWriteLock;
+	typedef FRWLock FUClassFuncLock;
+#endif
+
+
 /*-----------------------------------------------------------------------------
 	FRepRecord.
 -----------------------------------------------------------------------------*/
@@ -2824,7 +2841,7 @@ public:
 #if WITH_EDITOR
 	void GenerateFunctionList(TArray<FName>& OutArray) const 
 	{ 
-		FReadScopeLock ScopeLock(FuncMapLock);
+		FUClassFuncScopeReadLock ScopeLock(FuncMapLock)
 		FuncMap.GenerateKeyArray(OutArray); 
 	}
 #endif // WITH_EDITOR
@@ -2844,13 +2861,13 @@ private:
 	TMap<FName, UFunction*> FuncMap;
 
 	/** Scope lock to avoid the FuncMap being read and written to simultaneously on multiple threads. */
-	mutable FRWLock FuncMapLock;
+	mutable FUClassFuncLock FuncMapLock;
 
 	/** A cache of all functions by name that exist in a parent (superclass or interface) context */
 	mutable TMap<FName, UFunction*> SuperFuncMap;
 
 	/** Scope lock to avoid the SuperFuncMap being read and written to simultaneously on multiple threads. */
-	mutable FRWLock SuperFuncMapLock;
+	mutable FUClassFuncLock SuperFuncMapLock;
 
 public:
 	/**
@@ -2936,7 +2953,7 @@ public:
 	/** Add a function to the function map */
 	void AddFunctionToFunctionMap(UFunction* Function, FName FuncName)
 	{
-		FWriteScopeLock ScopeLock(FuncMapLock);
+		FUClassFuncScopeWriteLock ScopeLock(FuncMapLock);
 		FuncMap.Add(FuncName, Function);
 	}
 
@@ -2945,7 +2962,7 @@ public:
 	/** Remove a function from the function map */
 	void RemoveFunctionFromFunctionMap(UFunction* Function)
 	{
-		FWriteScopeLock ScopeLock(FuncMapLock);
+		FUClassFuncScopeWriteLock ScopeLock(FuncMapLock);
 		FuncMap.Remove(Function->GetFName());
 	}
 
