Fixed CIS

diff --git a/Engine/Source/Runtime/CoreUObject/Public/UObject/Class.h b/Engine/Source/Runtime/CoreUObject/Public/UObject/Class.h
--- a/Engine/Source/Runtime/CoreUObject/Public/UObject/Class.h
+++ b/Engine/Source/Runtime/CoreUObject/Public/UObject/Class.h
@@ -2841,7 +2841,7 @@ public:
 #if WITH_EDITOR
 	void GenerateFunctionList(TArray<FName>& OutArray) const 
 	{ 
-		FUClassFuncScopeReadLock ScopeLock(FuncMapLock)
+		FUClassFuncScopeReadLock ScopeLock(FuncMapLock);
 		FuncMap.GenerateKeyArray(OutArray); 
 	}
 #endif // WITH_EDITOR
