Soaking ensure to see if we need import map delinking

Resetlinkers Looping doesn't appear to be required.
This CL adds tracing to raise confidence that we can remove the offending loops.

diff --git a/Engine/Source/Runtime/CoreUObject/Private/UObject/LinkerManager.cpp b/Engine/Source/Runtime/CoreUObject/Private/UObject/LinkerManager.cpp
--- a/Engine/Source/Runtime/CoreUObject/Private/UObject/LinkerManager.cpp
+++ b/Engine/Source/Runtime/CoreUObject/Private/UObject/LinkerManager.cpp
@@ -132,10 +132,22 @@ void FLinkerManager::ResetLoaders(UObject* InPkg)
 					{
 						for (auto& Import : Linker->ImportMap)
 						{
-							if (Import.SourceLinker == LinkerToReset)
+							if (Import.SourceLinker)
 							{
-								Import.SourceLinker = nullptr;
-								Import.SourceIndex = INDEX_NONE;
+								// This code is a N^2 loop, searching the ImportMap's of all Loaders, looking for references
+								// that matched the loader we are resetting, and nulling them.
+								// But it looks as though all Loaders have their ImportrMaps null'd in DissociateImportsAndForcedExports.
+								// So this ended up being a time consuming loop that did nothing.
+								// To gain some confidence that this code is not needed, I've added these ensureMsgf's.
+								// If these don't go off, then this code, and the back links may be removed.
+								// 
+								// Soaking to see if this code is actually needed.
+								ensureMsgf(false, TEXT("ResetLoaders has a non null SourceLinker! Linker %p, Import.SourceLinker %p, LinkerToReset = %p"), Linker, Import.SourceLinker, LinkerToReset);
+								if (Import.SourceLinker == LinkerToReset)
+								{
+									Import.SourceLinker = nullptr;
+									Import.SourceIndex = INDEX_NONE;
+								}
 							}
 						}
 					}
@@ -188,6 +200,8 @@ void FLinkerManager::ResetLoaders(TArrayView<UObject*> InPackages)
 
 void FLinkerManager::ResetLoaders(const TSet<FLinkerLoad*>& InLinkerLoads)
 {
+	TRACE_CPUPROFILER_EVENT_SCOPE(FLinkerManager::ResetLoaders_Set);
+
 	// Remove import references
 	{
 #if THREADSAFE_UOBJECTS
@@ -199,11 +213,22 @@ void FLinkerManager::ResetLoaders(const TSet<FLinkerLoad*>& InLinkerLoads)
 			if (!InLinkerLoads.Contains(Linker))
 			{
 				for (auto& Import : Linker->ImportMap)
-				{
-					if (InLinkerLoads.Contains(Import.SourceLinker))
+				{			
+					if (Import.SourceLinker)
 					{
-						Import.SourceLinker = NULL;
-						Import.SourceIndex = INDEX_NONE;
+						// It looks as though all Loaders have their ImportrMaps null'd in DissociateImportsAndForcedExports.
+						// So this ends up being a time consuming loop that does nothing.
+						// To gain some confidence that this code is not needed, I've added these ensureMsgf's.
+						// If these don't go off, then this code, and the back links may be removed.
+						// 
+						// Soaking to see if this code is actually needed.
+						ensureMsgf(false, TEXT("ResetLoaders has a non null SourceLinker! Linker %p, Import.SourceLinker %p"), Linker, Import.SourceLinker);
+
+						if (InLinkerLoads.Contains(Import.SourceLinker))
+						{
+							Import.SourceLinker = NULL;
+							Import.SourceIndex = INDEX_NONE;
+						}
 					}
 				}
 			}
