Optimize PIE soft object path fixups.
- FSoftObjectPath properties report as containing weak object references for serializers.
- Mark PIE ID fixup archive as ArIsObjectReferenceCollector.

diff --git a/Engine/Source/Programs/LowLevelTests/Tests/ArchiveReplaceObjectRefTests.cpp b/Engine/Source/Programs/LowLevelTests/Tests/ArchiveReplaceObjectRefTests.cpp
new file mode 100644
--- /dev/null
+++ b/Engine/Source/Programs/LowLevelTests/Tests/ArchiveReplaceObjectRefTests.cpp
@@ -0,0 +1,74 @@
+// Copyright Epic Games, Inc. All Rights Reserved.
+
+#include "CoreTypes.h"
+#include "Serialization/ArchiveReplaceObjectRef.h"
+#include "TestHarness.h"
+#include "TestObject.h"
+#include "UObject/Package.h"
+
+#if WITH_LOW_LEVEL_TESTS
+
+namespace UE
+{
+    TEST_CASE("CoreUObject::Serialization::FArchiveReplaceObjectRef", "[CoreUObject][Serialization]")
+    {
+        FName PackageName = MakeUniqueObjectName(nullptr, UPackage::StaticClass(), FName("/Memory/ArchiveReplaceObjectRefTest"));
+        UPackage* Package = NewObject<UPackage>(nullptr, UPackage::StaticClass(), PackageName);
+        UTestObject* Source = NewObject<UTestObject>(Package, "Source");
+        UTestObject* From1 = NewObject<UTestObject>(Package, "From1");
+        UTestObject* To1 = NewObject<UTestObject>(Package, "To1");
+        UTestObject* From2 = NewObject<UTestObject>(Package, "From2");
+        UTestObject* To2 = NewObject<UTestObject>(Package, "To2");
+
+        TMap<UObject*, UObject*> ReplacementMap;
+        ReplacementMap.Add(From1, To1);
+        ReplacementMap.Add(From2, To2);
+
+        SECTION("StrongReferences", "TObjectPtr fields should be replaced")
+        {
+            Source->StrongObjectReference = From1;
+            Source->EmbeddedStruct.StrongObjectReference = From2; 
+            FTestStruct& S = Source->ArrayStructs.AddZeroed_GetRef();
+            S.StrongObjectReference = From1;
+            FArchiveReplaceObjectRef Ar(Source, ReplacementMap);            
+            CHECK(Source->StrongObjectReference == To1);
+            CHECK(Source->EmbeddedStruct.StrongObjectReference == To2);
+            REQUIRE(Source->ArrayStructs.Num() == 1);
+            CHECK(Source->ArrayStructs[0].StrongObjectReference == To1);
+        }
+        
+        SECTION("WeakReferences", "TWeakObjectPtr fields should be replaced")
+        {
+            Source->WeakObjectReference = From1;
+            Source->EmbeddedStruct.WeakObjectReference = From2;
+            FTestStruct& S = Source->ArrayStructs.AddZeroed_GetRef();
+            S.WeakObjectReference = From1;
+            FArchiveReplaceObjectRef Ar(Source, ReplacementMap);            
+            CHECK(Source->WeakObjectReference == To1);
+            CHECK(Source->EmbeddedStruct.WeakObjectReference == To2);
+            REQUIRE(Source->ArrayStructs.Num() == 1);
+            CHECK(Source->ArrayStructs[0].WeakObjectReference == To1);
+        }
+        
+        SECTION("SoftReferences", "TSoftObjectPtr and FSoftObjectPath fields should be replaced")
+        {
+            Source->SoftObjectReference = From1;
+            Source->SoftObjectPath = From2;
+            Source->EmbeddedStruct.SoftObjectReference = From2;
+            Source->EmbeddedStruct.SoftObjectPath = From1;
+            FTestStruct& S = Source->ArrayStructs.AddZeroed_GetRef();
+            S.SoftObjectReference = From1;
+            S.SoftObjectPath = From2;
+            FArchiveReplaceObjectRef Ar(Source, ReplacementMap);            
+            CHECK(Source->SoftObjectReference.GetUniqueID() == FSoftObjectPath(To1));
+            CHECK(Source->SoftObjectPath == FSoftObjectPath(To2));
+            CHECK(Source->EmbeddedStruct.SoftObjectReference.GetUniqueID() == FSoftObjectPath(To2));
+            CHECK(Source->EmbeddedStruct.SoftObjectPath == FSoftObjectPath(To1));
+            REQUIRE(Source->ArrayStructs.Num() == 1);
+            CHECK(Source->ArrayStructs[0].SoftObjectReference.GetUniqueID() == FSoftObjectPath(To1));
+            CHECK(Source->ArrayStructs[0].SoftObjectPath == FSoftObjectPath(To2));
+        }
+    }
+}
+
+#endif
\ No newline at end of file
diff --git a/Engine/Source/Runtime/CoreUObject/Private/UObject/Class.cpp b/Engine/Source/Runtime/CoreUObject/Private/UObject/Class.cpp
--- a/Engine/Source/Runtime/CoreUObject/Private/UObject/Class.cpp
+++ b/Engine/Source/Runtime/CoreUObject/Private/UObject/Class.cpp
@@ -35,9 +35,11 @@
 #include "UObject/Interface.h"
 #include "UObject/LinkerPlaceholderClass.h"
 #include "UObject/LinkerPlaceholderFunction.h"
+#include "UObject/StructOnScope.h"
 #include "UObject/StructScriptLoader.h"
 #include "UObject/PropertyHelper.h"
 #include "UObject/CoreRedirects.h"
+#include "UObject/ObjectMacros.h"
 #include "Internationalization/PolyglotTextData.h"
 #include "Serialization/ArchiveScriptReferenceCollector.h"
 #include "Serialization/ArchiveUObjectFromStructuredArchive.h"
@@ -59,7 +61,9 @@
 #include "AssetRegistry/AssetData.h"
 #include "HAL/PlatformStackWalk.h"
 #include "String/Find.h"
+#include "String/ParseTokens.h"
 #include "AutoRTFM/AutoRTFM.h"
+#include "Serialization/TestUndeclaredScriptStructObjectReferences.h"
 
 // This flag enables some expensive class tree validation that is meant to catch mutations of 
 // the class tree outside of SetSuperStruct. It has been disabled because loading blueprints 
@@ -940,7 +944,9 @@ void UStruct::Link(FArchive& Ar, bool bRelinkExistingProperties)
 	{
 		FProperty* Property = *It;
 
-		if (Property->ContainsObjectReference(EncounteredStructProps) || Property->ContainsWeakObjectReference())
+		// Ref link contains any properties which contain object references including types with user-defined serializers which don't explicitly specify whether they 
+		// contain object references
+		if (Property->ContainsObjectReference(EncounteredStructProps, EPropertyObjectReferenceType::Any))
 		{
 			*RefLinkPtr = Property;
 			RefLinkPtr = &(*RefLinkPtr)->NextRef;
@@ -1082,7 +1088,9 @@ void UStruct::SerializeBin( FStructuredArchive::FSlot Slot, void* Data ) const
 
 	FStructuredArchive::FStream PropertyStream = Slot.EnterStream();
 	
-	if (UnderlyingArchive.IsObjectReferenceCollector() && !UnderlyingArchive.IsModifyingWeakAndStrongReferences())
+	// RefLink contains Strong, Weak and Soft object references including FSoftObjectPath
+	// If objects wish to serialize other properties they should not set ArIsObjectReferenceCollector
+	if (UnderlyingArchive.IsObjectReferenceCollector())
 	{
 		// The FProperty instance might start in the middle of a cache line	
 		static constexpr uint32 ExtraPrefetchBytes = PLATFORM_CACHE_LINE_SIZE - /* min alignment */ 16;
@@ -3928,8 +3936,172 @@ bool FAutomationTestAttemptToFindShortTypeNamesInMetaData::RunTest(const FString
 
 #endif // WITH_EDITORONLY_DATA
 
+// bExactCheck - Check for places where structs serialize a different set of object references to what they declare (Conservative is always an error)
+// otherwise - Check for places where structs serialize object references they don't declare (Conservative allows serializing any reference types)
+static bool FindUndeclaredObjectReferencesInStructSerializers(FAutomationTestBase& Test, bool bExactCheck)
+{
+	struct FTestArchive : public FArchiveUObject
+	{
+		FTestArchive()
+		{
+			// Not persistent, loading, saving, etc
+			ArIsObjectReferenceCollector = true;
+		}
+		
+		EPropertyObjectReferenceType References = EPropertyObjectReferenceType::None;
+
+		virtual FArchive& operator<<(FLazyObjectPtr& Value) override 
+		{
+			References |= EPropertyObjectReferenceType::Weak;
+			return *this;
+		}
+		virtual FArchive& operator<<(FObjectPtr& Value) override 
+		{ 
+			References |= EPropertyObjectReferenceType::Strong;
+			return *this;
+		}
+		virtual FArchive& operator<<(FSoftObjectPtr& Value) override 
+		{
+			References |= EPropertyObjectReferenceType::Weak | EPropertyObjectReferenceType::Soft;
+			return *this;
+		}
+		virtual FArchive& operator<<(FSoftObjectPath& Value) override 
+		{
+			References |= EPropertyObjectReferenceType::Soft;
+			return *this;
+		}
+		virtual FArchive& operator<<(FWeakObjectPtr& Value) override 
+		{ 
+			References |= EPropertyObjectReferenceType::Weak;
+			return *this;
+		}
+	};
+
+	bool bFoundProblems = false;
+
+	for (TObjectIterator<UScriptStruct> ScriptIt; ScriptIt; ++ScriptIt)
+	{
+		UScriptStruct* Struct = *ScriptIt;
+		if ((Struct->StructFlags & STRUCT_SerializeNative) == 0)
+		{
+			continue;
+		}
+
+		auto LogUndeclaredObjectReference = [&bFoundProblems, &Test, Struct](ELogVerbosity::Type Verbosity, const TCHAR* RefType)
+		{
+			Test.AddError(FString::Printf(TEXT("Struct %s%s serializes an object reference of type EPropertyObjectReferenceType::%s but does not declare it"),
+				Struct->GetPrefixCPP(),
+				*Struct->GetName(),
+				RefType
+				));
+			bFoundProblems = true;
+		};
+		auto LogOverdeclaredObjectReference = [&bFoundProblems, &Test, Struct](ELogVerbosity::Type Verbosity, const TCHAR* RefType)
+		{
+			Test.AddError(FString::Printf(TEXT("Struct %s%s declares an object reference of type EPropertyObjectReferenceType::%s but does not serialize it"),
+				Struct->GetPrefixCPP(),
+				*Struct->GetName(),
+				RefType
+				));
+			bFoundProblems = true;
+		};
+		EPropertyObjectReferenceType DeclaredReferences = Struct->GetCppStructOps()->GetCapabilities().HasSerializerObjectReferences;			
+		EPropertyObjectReferenceType PropertyReferences = EPropertyObjectReferenceType::None;
+		for (EPropertyObjectReferenceType Type : { EPropertyObjectReferenceType::Strong, EPropertyObjectReferenceType::Weak, EPropertyObjectReferenceType::Soft })
+		{
+			FProperty* Property = Struct->PropertyLink;
+			TArray<const FStructProperty*> EncounteredStructProps; 
+			while (Property && !EnumHasAllFlags(PropertyReferences, Type))
+			{
+				if (Property->ContainsObjectReference(EncounteredStructProps, Type))
+				{
+					PropertyReferences |= Type;
+				}
+				Property = Property->PropertyLinkNext;
+			}
+		}
+
+		FStructOnScope TempStruct(Struct);
+		
+		FTestArchive Ar;
+		Struct->SerializeItem(Ar, TempStruct.GetStructMemory(), nullptr);
+		
+		if(!bExactCheck && !EnumHasAllFlags(DeclaredReferences, EPropertyObjectReferenceType::Conservative))
+		{
+			for (EPropertyObjectReferenceType Type : { EPropertyObjectReferenceType::Strong, EPropertyObjectReferenceType::Weak, EPropertyObjectReferenceType::Soft })
+			{
+				if (EnumHasAllFlags(Ar.References, Type) && !EnumHasAllFlags(DeclaredReferences | PropertyReferences, Type))
+				{
+					LogUndeclaredObjectReference(ELogVerbosity::Warning, LexToString(Type));
+				}
+			}
+		}
+		if(bExactCheck && EnumHasAllFlags(DeclaredReferences, EPropertyObjectReferenceType::Conservative))
+		{
+			FString SerializedTypes;
+			for (uint32 Flag = 1, Max = (uint32)EPropertyObjectReferenceType::MAX; Flag != Max; Flag <<= 1)
+			{
+				EPropertyObjectReferenceType TypedFlag = (EPropertyObjectReferenceType)Flag;
+				if (EnumHasAllFlags(Ar.References, TypedFlag))
+				{
+					if (SerializedTypes.IsEmpty())
+					{
+						SerializedTypes = LexToString(TypedFlag);
+					}
+					else 
+					{
+						SerializedTypes.Appendf(TEXT(" | %s"),LexToString(TypedFlag));
+					}
+				}
+			}	
+			if (SerializedTypes.IsEmpty())
+			{
+				SerializedTypes = LexToString(EPropertyObjectReferenceType::None);
+			}
+			Test.AddError(FString::Printf(TEXT("Struct %s%s declares its object references as EPropertyObjectReferenceType::Conservative which is never exact. Actual types serialized: %s"),
+				Struct->GetPrefixCPP(),
+				*Struct->GetName(),
+				*SerializedTypes
+			));
+			bFoundProblems = true;
+		}
+		else if (bExactCheck)
+		{
+			for (EPropertyObjectReferenceType Type : { EPropertyObjectReferenceType::Strong, EPropertyObjectReferenceType::Weak, EPropertyObjectReferenceType::Soft })
+			{
+				if (EnumHasAllFlags(DeclaredReferences, Type) && !EnumHasAllFlags(Ar.References, Type))
+				{
+					LogOverdeclaredObjectReference(ELogVerbosity::Warning, LexToString(Type));
+				}
+			}
+		}
+	}
+	return !bFoundProblems;
+}
+
+// Test for finding structs which misreport their object properties for native serialization
+IMPLEMENT_SIMPLE_AUTOMATION_TEST(FAutomationTestFindUndeclaredObjectReferencesInStructSerializers, "UObject.Class FindUndeclaredObjectReferencesInStructSerializers", EAutomationTestFlags::EditorContext | EAutomationTestFlags::ClientContext | EAutomationTestFlags::ServerContext | EAutomationTestFlags::SmokeFilter)
+bool FAutomationTestFindUndeclaredObjectReferencesInStructSerializers::RunTest(const FString& Parameters)
+{
+	return FindUndeclaredObjectReferencesInStructSerializers(*this, false);
+}
+
+// Test for finding structs which unoptimally their object properties for native serialization - not a smoke test
+
+	IMPLEMENT_SIMPLE_AUTOMATION_TEST(FAutomationTestFindInexactObjectReferencesInStructSerializers, "UObject.Class FindInexactObjectReferencesInStructSerializers", EAutomationTestFlags::EditorContext | EAutomationTestFlags::ClientContext | EAutomationTestFlags::ServerContext | EAutomationTestFlags::EngineFilter )
+bool FAutomationTestFindInexactObjectReferencesInStructSerializers::RunTest(const FString& Parameters)
+{
+	return FindUndeclaredObjectReferencesInStructSerializers(*this, true);
+}
+
 #endif // !(UE_BUILD_TEST || UE_BUILD_SHIPPING)
 
+bool FTestUndeclaredScriptStructObjectReferencesTest::Serialize(FArchive& Ar)
+{
+	Ar << StrongObjectPointer << SoftObjectPointer << SoftObjectPath << WeakObjectPointer;
+	return true;
+}
+
 IMPLEMENT_CORE_INTRINSIC_CLASS(UScriptStruct, UStruct,
 	{
 	}
@@ -6897,6 +7069,12 @@ UScriptStruct* TBaseStructure<FTestUninitializedScriptStructMembersTest>::Get()
 	return ScriptStruct;
 }
 
+UScriptStruct* TBaseStructure<FTestUndeclaredScriptStructObjectReferencesTest>::Get()
+{
+	static UScriptStruct* ScriptStruct = StaticGetBaseStructureInternal(TEXT("TestUndeclaredScriptStructObjectReferencesTest"));
+	return ScriptStruct;
+}
+
 UScriptStruct* TBaseStructure<FTopLevelAssetPath>::Get()
 {
 	static UScriptStruct* ScriptStruct = StaticGetBaseStructureInternal(TEXT("TopLevelAssetPath"));
diff --git a/Engine/Source/Runtime/CoreUObject/Private/UObject/GarbageCollection.cpp b/Engine/Source/Runtime/CoreUObject/Private/UObject/GarbageCollection.cpp
--- a/Engine/Source/Runtime/CoreUObject/Private/UObject/GarbageCollection.cpp
+++ b/Engine/Source/Runtime/CoreUObject/Private/UObject/GarbageCollection.cpp
@@ -4980,33 +4980,41 @@ bool FStructProperty::ContainsObjectReference(TArray<const FStructProperty*>& En
 	{
 		return false;
 	}
-	else
+
+	if (!Struct)
 	{
-		if (!Struct)
-		{
-			UE_LOG(LogGarbage, Warning, TEXT("Broken FStructProperty does not have a UStruct: %s"), *GetFullName() );
-		}
-		else if (Struct->StructFlags & STRUCT_AddStructReferencedObjects)
+		UE_LOG(LogGarbage, Warning, TEXT("Broken FStructProperty does not have a UStruct: %s"), *GetFullName() );
+		return false;
+	}
+
+	if (EnumHasAnyFlags(InReferenceType, EPropertyObjectReferenceType::Strong) && (Struct->StructFlags & STRUCT_AddStructReferencedObjects))
+	{
+		return true;
+	}
+
+	if (Struct->StructFlags & STRUCT_SerializeNative)
+	{
+		UScriptStruct::ICppStructOps* Ops = Struct->GetCppStructOps();
+		if (Ops && Ops->HasSerializerObjectReferences(InReferenceType))
 		{
 			return true;
 		}
-		else
+	}
+
+	EncounteredStructProps.Add(this);
+	bool bValue = false;
+	FProperty* Property = Struct->PropertyLink;
+	while (Property)
+	{
+		if (Property->ContainsObjectReference(EncounteredStructProps, InReferenceType))
 		{
-			EncounteredStructProps.Add(this);
-			FProperty* Property = Struct->PropertyLink;
-			while( Property )
-			{
-				if (Property->ContainsObjectReference(EncounteredStructProps, InReferenceType))
-				{
-					EncounteredStructProps.RemoveSingleSwap(this, false /*bAllowShrinking*/);
-					return true;
-				}
-				Property = Property->PropertyLinkNext;
-			}
-			EncounteredStructProps.RemoveSingleSwap(this, false /*bAllowShrinking*/);
+			bValue = true;
+			break;
 		}
-		return false;
+		Property = Property->PropertyLinkNext;
 	}
+	EncounteredStructProps.RemoveSingleSwap(this, false /*bAllowShrinking*/);
+	return bValue;
 }
 
 bool FFieldPathProperty::ContainsObjectReference(TArray<const FStructProperty*>& EncounteredStructProps, EPropertyObjectReferenceType InReferenceType /*= EPropertyObjectReferenceType::Strong*/) const
diff --git a/Engine/Source/Runtime/CoreUObject/Private/UObject/Property.cpp b/Engine/Source/Runtime/CoreUObject/Private/UObject/Property.cpp
--- a/Engine/Source/Runtime/CoreUObject/Private/UObject/Property.cpp
+++ b/Engine/Source/Runtime/CoreUObject/Private/UObject/Property.cpp
@@ -27,6 +27,7 @@
 #include "Math/InterpCurvePoint.h"
 #include "UObject/Package.h"
 #include "UObject/ReleaseObjectVersion.h"
+#include "Serialization/TestUndeclaredScriptStructObjectReferences.h"
 
 DEFINE_LOG_CATEGORY(LogProperty);
 
@@ -44,6 +45,7 @@ struct TVector3StructOpsTypeTraits : public TStructOpsTypeTraitsBase2<T>
 		WithStructuredSerializer = true,
 		WithStructuredSerializeFromMismatchedTag = true,
 	};
+	static constexpr EPropertyObjectReferenceType WithSerializerObjectReferences = EPropertyObjectReferenceType::None;
 };
 template<> struct TStructOpsTypeTraits<FVector3f> : public TVector3StructOpsTypeTraits<FVector3f> {};
 template<> struct TStructOpsTypeTraits<FVector3d> : public TVector3StructOpsTypeTraits<FVector3d> {};
@@ -62,6 +64,7 @@ struct TIntPointStructOpsTypeTraits : public TStructOpsTypeTraitsBase2<T>
 		WithSerializer = true,
 		WithSerializeFromMismatchedTag = true,
 	};
+	static constexpr EPropertyObjectReferenceType WithSerializerObjectReferences = EPropertyObjectReferenceType::None;
 };
 template<> struct TStructOpsTypeTraits<FInt32Point> : public TIntPointStructOpsTypeTraits<FInt32Point> {};
 template<> struct TStructOpsTypeTraits<FInt64Point> : public TIntPointStructOpsTypeTraits<FInt64Point> {};
@@ -85,6 +88,7 @@ struct TIntVectorStructOpsTypeTraits : public TStructOpsTypeTraitsBase2<T>
 		WithSerializer = true,
 		WithSerializeFromMismatchedTag = true,
 	};
+	static constexpr EPropertyObjectReferenceType WithSerializerObjectReferences = EPropertyObjectReferenceType::None;
 };
 template<> struct TStructOpsTypeTraits<FInt32Vector2> : public TIntVectorStructOpsTypeTraits<FInt32Vector2> {};
 template<> struct TStructOpsTypeTraits<FInt64Vector2> : public TIntVectorStructOpsTypeTraits<FInt64Vector2> {};
@@ -132,6 +136,7 @@ struct TVector2StructOpsTypeTraits : public TStructOpsTypeTraitsBase2<T>
 		WithSerializer = true,
 		WithSerializeFromMismatchedTag = true,
 	};
+	static constexpr EPropertyObjectReferenceType WithSerializerObjectReferences = EPropertyObjectReferenceType::None;
 };
 template<> struct TStructOpsTypeTraits<FVector2f> : public TVector2StructOpsTypeTraits<FVector2f> {};
 template<> struct TStructOpsTypeTraits<FVector2d> : public TVector2StructOpsTypeTraits<FVector2d> {};
@@ -150,6 +155,7 @@ struct TVector4StructOpsTypeTraits : public TStructOpsTypeTraitsBase2<T>
 		WithSerializer = true,
 		WithSerializeFromMismatchedTag = true,
 	};
+	static constexpr EPropertyObjectReferenceType WithSerializerObjectReferences = EPropertyObjectReferenceType::None;
 };
 template<> struct TStructOpsTypeTraits<FVector4f> : public TVector4StructOpsTypeTraits<FVector4f> {};
 template<> struct TStructOpsTypeTraits<FVector4d> : public TVector4StructOpsTypeTraits<FVector4d> {};
@@ -170,6 +176,7 @@ struct TPlaneStructOpsTypeTraits : public TStructOpsTypeTraitsBase2<T>
 		WithSerializer = true,
 		WithSerializeFromMismatchedTag = true,
 	};
+	static constexpr EPropertyObjectReferenceType WithSerializerObjectReferences = EPropertyObjectReferenceType::None;
 };
 template<> struct TStructOpsTypeTraits<FPlane4f> : public TPlaneStructOpsTypeTraits<FPlane4f> {};
 template<> struct TStructOpsTypeTraits<FPlane4d> : public TPlaneStructOpsTypeTraits<FPlane4d> {};
@@ -190,6 +197,7 @@ struct TRotatorStructOpsTypeTraits : public TStructOpsTypeTraitsBase2<T>
 		WithSerializer = true,
 		WithSerializeFromMismatchedTag = true,
 	};
+	static constexpr EPropertyObjectReferenceType WithSerializerObjectReferences = EPropertyObjectReferenceType::None;
 };
 
 template<> struct TStructOpsTypeTraits<FRotator3f> : public TRotatorStructOpsTypeTraits<FRotator3f> {};
@@ -209,6 +217,7 @@ struct TBox3StructOpsTypeTraits : public TStructOpsTypeTraitsBase2<T>
 		WithSerializer = true,
 		WithSerializeFromMismatchedTag = true,
 	};
+	static constexpr EPropertyObjectReferenceType WithSerializerObjectReferences = EPropertyObjectReferenceType::None;
 };
 template<> struct TStructOpsTypeTraits<FBox3f> : public TBox3StructOpsTypeTraits<FBox3f> {};
 template<> struct TStructOpsTypeTraits<FBox3d> : public TBox3StructOpsTypeTraits<FBox3d> {};
@@ -244,6 +253,7 @@ struct TMatrixStructOpsTypeTraits : public TStructOpsTypeTraitsBase2<T>
 		WithSerializer = true,
 		WithSerializeFromMismatchedTag = true,
 	};
+	static constexpr EPropertyObjectReferenceType WithSerializerObjectReferences = EPropertyObjectReferenceType::None;
 };
 
 template<> struct TStructOpsTypeTraits<FMatrix44f> : public TMatrixStructOpsTypeTraits<FMatrix44f> {};
@@ -285,6 +295,7 @@ struct TStructOpsTypeTraits<FLinearColor> : public TStructOpsTypeTraitsBase2<FLi
 		WithZeroConstructor = true,
 		WithStructuredSerializer = true,
 	};
+	static constexpr EPropertyObjectReferenceType WithSerializerObjectReferences = EPropertyObjectReferenceType::None;
 };
 UE_IMPLEMENT_STRUCT("/Script/CoreUObject", LinearColor);
 
@@ -298,6 +309,7 @@ struct TStructOpsTypeTraits<FColor> : public TStructOpsTypeTraitsBase2<FColor>
 		WithZeroConstructor = true,
 		WithSerializer = true,
 	};
+	static constexpr EPropertyObjectReferenceType WithSerializerObjectReferences = EPropertyObjectReferenceType::None;
 };
 UE_IMPLEMENT_STRUCT("/Script/CoreUObject", Color);
 
@@ -315,6 +327,7 @@ struct TQuatStructOpsTypeTraits : public TStructOpsTypeTraitsBase2<T>
 		WithSerializer = true,
 		WithSerializeFromMismatchedTag = true,
 	};
+	static constexpr EPropertyObjectReferenceType WithSerializerObjectReferences = EPropertyObjectReferenceType::None;
 };
 template<> struct TStructOpsTypeTraits<FQuat4f> : public TQuatStructOpsTypeTraits<FQuat4f> {};
 template<> struct TStructOpsTypeTraits<FQuat4d> : public TQuatStructOpsTypeTraits<FQuat4d> {};
@@ -332,6 +345,7 @@ struct TStructOpsTypeTraits<FTwoVectors> : public TStructOpsTypeTraitsBase2<FTwo
 		WithSerializer = true,
 		WithNoDestructor = true,
 	};
+	static constexpr EPropertyObjectReferenceType WithSerializerObjectReferences = EPropertyObjectReferenceType::None;
 };
 UE_IMPLEMENT_STRUCT("/Script/CoreUObject", TwoVectors);
 
@@ -443,6 +457,7 @@ struct TStructOpsTypeTraits<FGuid> : public TStructOpsTypeTraitsBase2<FGuid>
 		WithSerializer = true,
 		WithStructuredSerializer = true,
 	};
+	static constexpr EPropertyObjectReferenceType WithSerializerObjectReferences = EPropertyObjectReferenceType::None;
 };
 UE_IMPLEMENT_STRUCT("/Script/CoreUObject", Guid);
 
@@ -486,6 +501,7 @@ struct TStructOpsTypeTraits<FDateTime> : public TStructOpsTypeTraitsBase2<FDateT
 		WithZeroConstructor = true,
 		WithIdenticalViaEquality = true,
 	};
+	static constexpr EPropertyObjectReferenceType WithSerializerObjectReferences = EPropertyObjectReferenceType::None;
 };
 UE_IMPLEMENT_STRUCT("/Script/CoreUObject", DateTime);
 
@@ -503,6 +519,7 @@ struct TStructOpsTypeTraits<FTimespan> : public TStructOpsTypeTraitsBase2<FTimes
 		WithZeroConstructor = true,
 		WithIdenticalViaEquality = true,
 	};
+	static constexpr EPropertyObjectReferenceType WithSerializerObjectReferences = EPropertyObjectReferenceType::None;
 };
 UE_IMPLEMENT_STRUCT("/Script/CoreUObject", Timespan);
 
@@ -514,6 +531,7 @@ struct TStructOpsTypeTraits<FFrameNumber> : public TStructOpsTypeTraitsBase2<FFr
 		WithSerializer = true,
 		WithIdenticalViaEquality = true
 	};
+	static constexpr EPropertyObjectReferenceType WithSerializerObjectReferences = EPropertyObjectReferenceType::None;
 };
 UE_IMPLEMENT_STRUCT("/Script/CoreUObject", FrameNumber);
 
@@ -530,6 +548,7 @@ struct TStructOpsTypeTraits<FSoftObjectPath> : public TStructOpsTypeTraitsBase2<
 		WithImportTextItem = true,
 		WithStructuredSerializeFromMismatchedTag = true,
 	};
+	static constexpr EPropertyObjectReferenceType WithSerializerObjectReferences = EPropertyObjectReferenceType::Soft;
 };
 UE_IMPLEMENT_STRUCT("/Script/CoreUObject", SoftObjectPath);
 
@@ -546,6 +565,7 @@ struct TStructOpsTypeTraits<FSoftClassPath> : public TStructOpsTypeTraitsBase2<F
 		WithImportTextItem = true,
 		WithStructuredSerializeFromMismatchedTag = true,
 	};
+	static constexpr EPropertyObjectReferenceType WithSerializerObjectReferences = EPropertyObjectReferenceType::Soft;
 };
 UE_IMPLEMENT_STRUCT("/Script/CoreUObject", SoftClassPath);
 
@@ -579,6 +599,18 @@ struct TStructOpsTypeTraits<FPrimaryAssetId> : public TStructOpsTypeTraitsBase2<
 };
 UE_IMPLEMENT_STRUCT("/Script/CoreUObject", PrimaryAssetId);
 
+template<>
+struct TStructOpsTypeTraits<FTestUndeclaredScriptStructObjectReferencesTest> : public TStructOpsTypeTraitsBase2<FTestUndeclaredScriptStructObjectReferencesTest>
+{
+	enum 
+	{
+		WithSerializer = true,
+	};
+	static constexpr EPropertyObjectReferenceType WithSerializerObjectReferences = EPropertyObjectReferenceType::Strong | EPropertyObjectReferenceType::Weak | EPropertyObjectReferenceType::Soft;
+};
+UE_IMPLEMENT_STRUCT("/Script/CoreUObject", TestUndeclaredScriptStructObjectReferencesTest);
+
+
 template<>
 struct TStructOpsTypeTraits<FFallbackStruct> : public TStructOpsTypeTraitsBase2<FFallbackStruct>
 {
@@ -589,6 +621,24 @@ UE_IMPLEMENT_STRUCT("/Script/CoreUObject", FallbackStruct);
 	Helpers.
 -----------------------------------------------------------------------------*/
 
+const TCHAR* LexToString(EPropertyObjectReferenceType Type)
+{
+	switch(Type)
+	{
+	case EPropertyObjectReferenceType::None:
+		return TEXT("None");
+	case EPropertyObjectReferenceType::Strong:
+		return TEXT("Strong");
+	case EPropertyObjectReferenceType::Weak:
+		return TEXT("Weak");
+	case EPropertyObjectReferenceType::Soft:
+		return TEXT("Soft");
+	case EPropertyObjectReferenceType::Conservative:
+		return TEXT("Conservative");
+	}	
+	return TEXT("Unknown");
+}
+
 constexpr FAsciiSet AlphaNumericChars = "abcdefghijklmnopqrstuvwxyzABCDEFGHIJKLMNOPQRSTUVWXYZ0123456789";
 
 FORCEINLINE constexpr bool IsValidTokenStart(TCHAR FirstChar, bool bDottedNames)
diff --git a/Engine/Source/Runtime/CoreUObject/Private/UObject/SoftObjectPath.cpp b/Engine/Source/Runtime/CoreUObject/Private/UObject/SoftObjectPath.cpp
--- a/Engine/Source/Runtime/CoreUObject/Private/UObject/SoftObjectPath.cpp
+++ b/Engine/Source/Runtime/CoreUObject/Private/UObject/SoftObjectPath.cpp
@@ -839,7 +839,20 @@ bool FSoftObjectPathThreadContext::GetSerializationOptions(FName& OutPackageName
 FThreadSafeCounter FSoftObjectPath::CurrentTag(1);
 TSet<FName> FSoftObjectPath::PIEPackageNames;
 
+#if WITH_LOW_LEVEL_TESTS
 
+#include "TestHarness.h"
+
+std::ostream& operator<<(std::ostream& Stream, const FSoftObjectPath& Value)
+{
+	TStringBuilder<FName::StringBufferSize*2> Builder;
+	Builder << Value;
+	return Stream << Builder.ToView();
+}
+
+#endif
+
+#if WITH_DEV_AUTOMATION_TESTS
 
 IMPLEMENT_SIMPLE_AUTOMATION_TEST(FSoftObjectPathImportTextTests, "System.CoreUObject.SoftObjectPath.ImportText", EAutomationTestFlags::ApplicationContextMask | EAutomationTestFlags::EngineFilter);
 bool FSoftObjectPathImportTextTests::RunTest(const FString& Parameters)
@@ -995,3 +1008,18 @@ bool FSoftObjectPathTrySetPathTests::RunTest(const FString& Parameters)
 
 	return true;
 }
+
+IMPLEMENT_SIMPLE_AUTOMATION_TEST(FSoftObjectPathFixupForPIETests, "System.CoreUObject.SoftObjectPath.FixupForPIE", EAutomationTestFlags::ApplicationContextMask | EAutomationTestFlags::EngineFilter);
+bool FSoftObjectPathFixupForPIETests::RunTest(const FString& Parameters)
+{
+	const TCHAR* TestOriginalPath = TEXT("/Game/Maps/Arena.Arena:PersistentLevel.Target");	
+	const int32 PieInstanceID = 7;
+	const FString ExpectedFinalPath = FString::Printf(TEXT("/Game/Maps/%s_%d_Arena.Arena:PersistentLevel.Target"), PLAYWORLD_PACKAGE_PREFIX, PieInstanceID);
+	
+	FSoftObjectPath SoftPath(TestOriginalPath);
+	SoftPath.FixupForPIE(PieInstanceID);	
+	TestEqual(TEXT("Fixed up path should be PIE package with correct id"), SoftPath.ToString(), ExpectedFinalPath);
+	return true;
+}
+
+#endif
\ No newline at end of file
diff --git a/Engine/Source/Runtime/CoreUObject/Public/Serialization/TestUndeclaredScriptStructObjectReferences.h b/Engine/Source/Runtime/CoreUObject/Public/Serialization/TestUndeclaredScriptStructObjectReferences.h
new file mode 100644
--- /dev/null
+++ b/Engine/Source/Runtime/CoreUObject/Public/Serialization/TestUndeclaredScriptStructObjectReferences.h
@@ -0,0 +1,25 @@
+// Copyright Epic Games, Inc. All Rights Reserved.
+#pragma once 
+
+#include "UObject/ObjectPtr.h"
+#include "UObject/SoftObjectPath.h"
+#include "UObject/SoftObjectPtr.h"
+#include "UObject/WeakObjectPtr.h"
+
+class CORE_API FArchive;
+ 
+// Helper struct to test if struct serializer object reference declaration tests work properly
+struct FTestUndeclaredScriptStructObjectReferencesTest
+{
+	TObjectPtr<UObject> StrongObjectPointer;
+	TSoftObjectPtr<UObject> SoftObjectPointer;
+	FSoftObjectPath SoftObjectPath;
+	TWeakObjectPtr<UObject> WeakObjectPointer;
+	
+	bool Serialize(FArchive& Ar);
+};
+
+template<> struct TBaseStructure<FTestUndeclaredScriptStructObjectReferencesTest>
+{
+	COREUOBJECT_API static UScriptStruct* Get();
+};
\ No newline at end of file
diff --git a/Engine/Source/Runtime/CoreUObject/Public/UObject/Class.h b/Engine/Source/Runtime/CoreUObject/Public/UObject/Class.h
--- a/Engine/Source/Runtime/CoreUObject/Public/UObject/Class.h
+++ b/Engine/Source/Runtime/CoreUObject/Public/UObject/Class.h
@@ -109,6 +109,7 @@ struct FPropertyTag;
 struct FRandomStream;
 struct FUObjectSerializeContext;
 template <typename FuncType> class TFunctionRef;
+enum class EPropertyObjectReferenceType : uint32;
 
 COREUOBJECT_API DECLARE_LOG_CATEGORY_EXTERN(LogClass, Log, All);
 COREUOBJECT_API DECLARE_LOG_CATEGORY_EXTERN(LogScriptSerialization, Log, All);
@@ -864,6 +865,8 @@ struct TStructOpsTypeTraitsBase2
 		WithPureVirtual                = false,                         // struct has PURE_VIRTUAL functions and cannot be constructed when CHECK_PUREVIRTUALS is true
 		WithCanEditChange			   = false,							// struct has an editor-only CanEditChange function that can conditionally make child properties read-only in the details panel (same idea as UObject::CanEditChange)
 	};
+
+	static constexpr EPropertyObjectReferenceType WithSerializerObjectReferences = EPropertyObjectReferenceType::Conservative; // struct's Serialize method(s) may serialize object references of these types - default Conservative means unknown and object reference collector archives should serialize this struct 
 };
 
 template<class CPPSTRUCT>
@@ -903,6 +906,7 @@ public:
 		struct FCapabilities
 		{
 			EPropertyFlags ComputedPropertyFlags;
+			EPropertyObjectReferenceType HasSerializerObjectReferences;
 			bool HasNoopConstructor : 1;
 			bool HasZeroConstructor : 1;
 			bool HasDestructor : 1;
@@ -1107,6 +1111,11 @@ public:
 		{
 			return GetCapabilities().HasAddStructReferencedObjects;
 		}
+		/** returns true if the native serialize functions may serialize object references of the given type **/
+		bool HasSerializerObjectReferences(EPropertyObjectReferenceType Type) const
+		{
+			return EnumHasAnyFlags(GetCapabilities().HasSerializerObjectReferences, Type);
+		}
 		/** 
 		 * return a pointer to a function that can add referenced objects
 		 * @return true if the copy was imported, otherwise it will fall back to FStructProperty::ImportText
@@ -1187,6 +1196,7 @@ public:
 				| (TIsTriviallyDestructible<CPPSTRUCT>::Value ? CPF_NoDestructor : CPF_None)
 				| (TIsZeroConstructType<CPPSTRUCT>::Value ? CPF_ZeroConstructor : CPF_None)
 				| (TModels_V<CGetTypeHashable, CPPSTRUCT> ? CPF_HasGetValueTypeHash : CPF_None),
+				TTraits::WithSerializerObjectReferences,
 				TTraits::WithNoInitConstructor,
 				TTraits::WithZeroConstructor,
 				!(TTraits::WithNoDestructor || TIsPODType<CPPSTRUCT>::Value),
diff --git a/Engine/Source/Runtime/CoreUObject/Public/UObject/NoExportTypes.h b/Engine/Source/Runtime/CoreUObject/Public/UObject/NoExportTypes.h
--- a/Engine/Source/Runtime/CoreUObject/Public/UObject/NoExportTypes.h
+++ b/Engine/Source/Runtime/CoreUObject/Public/UObject/NoExportTypes.h
@@ -25,6 +25,7 @@
 #include "UObject/TopLevelAssetPath.h"
 #include "UObject/SoftObjectPath.h"
 #include "UObject/PropertyAccessUtil.h"
+#include "Serialization/TestUndeclaredScriptStructObjectReferences.h"
 
 #include "Math/InterpCurvePoint.h"
 #include "Math/UnitConversion.h"
@@ -2716,6 +2717,22 @@ struct FTestUninitializedScriptStructMembersTest
 	float UnusedValue;
 };
 
+USTRUCT(noexport, IsAlwaysAccessible, HasDefaults)
+struct FTestUndeclaredScriptStructObjectReferencesTest
+{
+	UPROPERTY(Transient)
+	TObjectPtr<UObject> StrongObjectPointer;
+
+	UPROPERTY(Transient)
+	TSoftObjectPtr<UObject> SoftObjectPointer;
+
+	UPROPERTY(Transient)
+	FSoftObjectPath SoftObjectPath;
+	
+	UPROPERTY(Transient)
+	TWeakObjectPtr<UObject> WeakObjectPointer;
+};
+
 /**
  * Direct base class for all UE objects
  * @note The full C++ class is located here: Engine\Source\Runtime\CoreUObject\Public\UObject\Object.h
diff --git a/Engine/Source/Runtime/CoreUObject/Public/UObject/ObjectMacros.h b/Engine/Source/Runtime/CoreUObject/Public/UObject/ObjectMacros.h
--- a/Engine/Source/Runtime/CoreUObject/Public/UObject/ObjectMacros.h
+++ b/Engine/Source/Runtime/CoreUObject/Public/UObject/ObjectMacros.h
@@ -499,6 +499,21 @@ enum class EMapPropertyFlags : uint8
 
 ENUM_CLASS_FLAGS(EMapPropertyFlags)
 
+enum class EPropertyObjectReferenceType : uint32
+{
+	None = 0,
+	Strong = 1 << 0, // Hard reference to a UObject, keeps the object from being garbage collected
+	Weak = 1 << 1,   // Weak reference to a UObject, does not keep the object from being garbage collected, does not become valid again after object is unloaded and reloaded.
+	Soft = 1 << 2,	 // Soft path/identity-based reference to a UObject, does not keep the object from being garbage collected, does become valid again after object is unloaded and reloaded.
+	Conservative = 1 << 3,	 // Not a real reference type, used to mark native struct serializers which may serialize unknown reference types and to conservatively populate RefLink with such struct properties.
+	
+	MAX = Conservative << 1, // Marker for iterating over all flags
+	Any = Strong | Weak | Soft | Conservative,
+};
+ENUM_CLASS_FLAGS(EPropertyObjectReferenceType);
+
+COREUOBJECT_API const TCHAR* LexToString(EPropertyObjectReferenceType Type);
+
 /**
  * Flags describing an object instance
  */
diff --git a/Engine/Source/Runtime/CoreUObject/Public/UObject/SoftObjectPath.h b/Engine/Source/Runtime/CoreUObject/Public/UObject/SoftObjectPath.h
--- a/Engine/Source/Runtime/CoreUObject/Public/UObject/SoftObjectPath.h
+++ b/Engine/Source/Runtime/CoreUObject/Public/UObject/SoftObjectPath.h
@@ -615,6 +615,11 @@ PRAGMA_ENABLE_DEPRECATION_WARNINGS
 	}
 }
 
+#if WITH_LOW_LEVEL_TESTS
+#include <ostream>
+COREUOBJECT_API std::ostream& operator<<(std::ostream& Stream, const FSoftObjectPath& Value);
+#endif
+
 #if UE_ENABLE_INCLUDE_ORDER_DEPRECATED_IN_5_2
 #include "CoreMinimal.h"
 #endif
diff --git a/Engine/Source/Runtime/CoreUObject/Public/UObject/SoftObjectPtr.h b/Engine/Source/Runtime/CoreUObject/Public/UObject/SoftObjectPtr.h
--- a/Engine/Source/Runtime/CoreUObject/Public/UObject/SoftObjectPtr.h
+++ b/Engine/Source/Runtime/CoreUObject/Public/UObject/SoftObjectPtr.h
@@ -11,6 +11,16 @@
 #include "UObject/PersistentObjectPtr.h"
 #include "UObject/SoftObjectPath.h"
 
+/**
+ * TIsSoftObjectPointerType
+ * Trait for recognizing 'soft' (path-based) object pointer types
+ */
+template<typename T> 
+struct TIsSoftObjectPointerType
+{ 
+	enum { Value = false };
+};
+
 /**
  * FSoftObjectPtr is a type of weak pointer to a UObject, that also keeps track of the path to the object on disk.
  * It will change back and forth between being Valid and Pending as the referenced object loads or unloads.
@@ -96,6 +106,7 @@ public:
 
 template <> struct TIsPODType<FSoftObjectPtr> { enum { Value = TIsPODType<TPersistentObjectPtr<FSoftObjectPath> >::Value }; };
 template <> struct TIsWeakPointerType<FSoftObjectPtr> { enum { Value = TIsWeakPointerType<TPersistentObjectPtr<FSoftObjectPath> >::Value }; };
+template <> struct TIsSoftObjectPointerType<FSoftObjectPtr> { enum { Value = true }; };
 
 /**
  * TSoftObjectPtr is templatized wrapper of the generic FSoftObjectPtr, it can be used in UProperties
@@ -411,6 +422,7 @@ FArchive& operator<<(FArchive& Ar, TSoftObjectPtr<T>& Ptr)
 
 template<class T> struct TIsPODType<TSoftObjectPtr<T> > { enum { Value = TIsPODType<FSoftObjectPtr>::Value }; };
 template<class T> struct TIsWeakPointerType<TSoftObjectPtr<T> > { enum { Value = TIsWeakPointerType<FSoftObjectPtr>::Value }; };
+template<class T> struct TIsSoftObjectPointerType<TSoftObjectPtr<T>> { enum { Value = TIsSoftObjectPointerType<FSoftObjectPtr>::Value }; };
 
 template <typename T>
 struct TCallTraits<TSoftObjectPtr<T>> : public TCallTraitsBase<TSoftObjectPtr<T>>
diff --git a/Engine/Source/Runtime/CoreUObject/Public/UObject/UnrealType.h b/Engine/Source/Runtime/CoreUObject/Public/UObject/UnrealType.h
--- a/Engine/Source/Runtime/CoreUObject/Public/UObject/UnrealType.h
+++ b/Engine/Source/Runtime/CoreUObject/Public/UObject/UnrealType.h
@@ -122,14 +122,6 @@ enum class EConvertFromTypeResult
 	Converted
 };
 
-enum class EPropertyObjectReferenceType : uint32
-{
-	None = 0,
-	Strong = 1 << 0,
-	Weak = 1 << 1
-};
-ENUM_CLASS_FLAGS(EPropertyObjectReferenceType);
-
 enum class EPropertyMemoryAccess : uint8
 {
 	// Direct memory access - the associated pointer points to the memory at the reflected item.
@@ -2873,8 +2865,9 @@ PRAGMA_ENABLE_DEPRECATION_WARNINGS
 	// FProperty interface.
 	virtual bool ContainsObjectReference(TArray<const FStructProperty*>& EncounteredStructProps, EPropertyObjectReferenceType InReferenceType = EPropertyObjectReferenceType::Strong) const override
 	{
-		return (!!(InReferenceType & EPropertyObjectReferenceType::Strong) && !TIsWeakPointerType<InTCppType>::Value) ||
-			(!!(InReferenceType & EPropertyObjectReferenceType::Weak) && TIsWeakPointerType<InTCppType>::Value);
+		return (EnumHasAnyFlags(InReferenceType, EPropertyObjectReferenceType::Strong) && !TIsWeakPointerType<InTCppType>::Value)
+			|| (EnumHasAnyFlags(InReferenceType, EPropertyObjectReferenceType::Weak) && TIsWeakPointerType<InTCppType>::Value)
+			|| (EnumHasAnyFlags(InReferenceType, EPropertyObjectReferenceType::Soft) && TIsSoftObjectPointerType<InTCppType>::Value);
 	}
 	// End of FProperty interface
 
diff --git a/Engine/Source/Runtime/Engine/Private/Level.cpp b/Engine/Source/Runtime/Engine/Private/Level.cpp
--- a/Engine/Source/Runtime/Engine/Private/Level.cpp
+++ b/Engine/Source/Runtime/Engine/Private/Level.cpp
@@ -4204,6 +4204,7 @@ void ULevel::BeginCacheForCookedPlatformData(const ITargetPlatform *TargetPlatfo
 
 void ULevel::FixupForPIE(int32 InPIEInstanceID, TFunctionRef<void(int32, FSoftObjectPath&)> InCustomFixupFunction)
 {
+	TRACE_CPUPROFILER_EVENT_SCOPE(ULevel::FixupForPIE);
 	FPIEFixupSerializer FixupSerializer(this, InPIEInstanceID, InCustomFixupFunction);
 	Serialize(FixupSerializer);
 }
diff --git a/Engine/Source/Runtime/Engine/Private/PieFixupSerializer.cpp b/Engine/Source/Runtime/Engine/Private/PieFixupSerializer.cpp
--- a/Engine/Source/Runtime/Engine/Private/PieFixupSerializer.cpp
+++ b/Engine/Source/Runtime/Engine/Private/PieFixupSerializer.cpp
@@ -3,6 +3,7 @@
 #include "PieFixupSerializer.h"
 
 #include "Components/InstancedStaticMeshComponent.h"
+#include "Misc/AutomationTest.h"
 #include "UObject/Package.h"
 #include "UObject/UnrealType.h"
 
@@ -19,6 +20,8 @@ FPIEFixupSerializer::FPIEFixupSerializer(UObject* InRoot, int32 InPIEInstanceID)
 	, PIEInstanceID(InPIEInstanceID)
 {
 	this->ArShouldSkipBulkData = true;
+	this->ArIsObjectReferenceCollector = true;
+	this->ArIsModifyingWeakAndStrongReferences = true;
 }
 
 FPIEFixupSerializer::FPIEFixupSerializer(UObject* InRoot, int32 InPIEInstanceID, TFunctionRef<void(int32, FSoftObjectPath&)> InSoftObjectPathFixupFunction)
@@ -27,6 +30,8 @@ FPIEFixupSerializer::FPIEFixupSerializer(UObject* InRoot, int32 InPIEInstanceID,
 	, PIEInstanceID(InPIEInstanceID)
 {
 	this->ArShouldSkipBulkData = true;
+	this->ArIsObjectReferenceCollector = true;
+	this->ArIsModifyingWeakAndStrongReferences = true;
 }
 
 bool FPIEFixupSerializer::ShouldSkipProperty(const FProperty* InProperty) const
@@ -76,3 +81,46 @@ FArchive& FPIEFixupSerializer::operator<<(FSoftObjectPtr& Value)
 	*this << Value.GetUniqueID();
 	return *this;
 }
+
+
+#if WITH_DEV_AUTOMATION_TESTS
+
+#include "Tests/PieFixupTestObjects.h"
+
+IMPLEMENT_SIMPLE_AUTOMATION_TEST(FPIEFixupSerializerTest, "System.Engine.PIE.FixupSoftReferences", EAutomationTestFlags::ApplicationContextMask | EAutomationTestFlags::EngineFilter);
+bool FPIEFixupSerializerTest::RunTest(const FString& Parameters)
+{
+	FName PackageName = MakeUniqueObjectName(nullptr, UPackage::StaticClass(), FName("/Memory/PieFixupTest"));
+	UPackage* Package = NewObject<UPackage>(nullptr, UPackage::StaticClass(), PackageName);
+	UPieFixupTestObject* Obj = NewObject<UPieFixupTestObject>(Package);
+	Obj->Path = TEXT("/Game/Maps/Arena.Arena:PersistentLevel.SpawnPoint.Root");
+	Obj->TypedPtr = TSoftObjectPtr<AActor>(FSoftObjectPath(TEXT("/Game/Maps/Arena.Arena:PersistentLevel.Target")));
+	
+	Obj->Struct.Path = Obj->Path;
+	Obj->Struct.TypedPtr = Obj->TypedPtr;
+	
+	FPieFixupStructWithSoftObjectPath& InArray = Obj->Array.AddZeroed_GetRef();
+	InArray.Path = Obj->Path;
+	InArray.TypedPtr = Obj->TypedPtr;
+	
+	int32 PieInstanceID  = 3;
+	FPIEFixupSerializer Fixup(Package, PieInstanceID );
+	Fixup << Obj;
+	
+	TestEqual(TEXT("Path is fixed up"), Obj->Path.ToString(),
+		FString::Printf(TEXT("/Game/Maps/%s_%d_Arena.Arena:PersistentLevel.SpawnPoint.Root"), PLAYWORLD_PACKAGE_PREFIX, PieInstanceID));
+	TestEqual(TEXT("Ptr is fixed up"), Obj->TypedPtr.GetUniqueID().ToString(),
+		FString::Printf(TEXT("/Game/Maps/%s_%d_Arena.Arena:PersistentLevel.Target"), PLAYWORLD_PACKAGE_PREFIX, PieInstanceID));
+	TestEqual(TEXT("Path in struct is fixed up"), Obj->Struct.Path.ToString(),
+		FString::Printf(TEXT("/Game/Maps/%s_%d_Arena.Arena:PersistentLevel.SpawnPoint.Root"), PLAYWORLD_PACKAGE_PREFIX, PieInstanceID));
+	TestEqual(TEXT("Ptr in struct is fixed up"), Obj->Struct.TypedPtr.GetUniqueID().ToString(),
+		FString::Printf(TEXT("/Game/Maps/%s_%d_Arena.Arena:PersistentLevel.Target"), PLAYWORLD_PACKAGE_PREFIX, PieInstanceID));
+	TestEqual(TEXT("Path in array of structs is fixed up "), InArray.Path.ToString(),
+		FString::Printf(TEXT("/Game/Maps/%s_%d_Arena.Arena:PersistentLevel.SpawnPoint.Root"), PLAYWORLD_PACKAGE_PREFIX, PieInstanceID));
+	TestEqual(TEXT("Ptr in array of structs is fixed up "), InArray.TypedPtr.GetUniqueID().ToString(),
+		FString::Printf(TEXT("/Game/Maps/%s_%d_Arena.Arena:PersistentLevel.Target"), PLAYWORLD_PACKAGE_PREFIX, PieInstanceID));
+
+	return true;
+}
+
+#endif
\ No newline at end of file
diff --git a/Engine/Source/Runtime/Engine/Private/Tests/PieFixupTestObjects.h b/Engine/Source/Runtime/Engine/Private/Tests/PieFixupTestObjects.h
new file mode 100644
--- /dev/null
+++ b/Engine/Source/Runtime/Engine/Private/Tests/PieFixupTestObjects.h
@@ -0,0 +1,38 @@
+// Copyright Epic Games, Inc. All Rights Reserved.
+
+#pragma once
+
+#include "UObject/Object.h"
+#include "UObject/SoftObjectPtr.h"
+#include "PieFixupTestObjects.generated.h"
+
+USTRUCT()
+struct FPieFixupStructWithSoftObjectPath
+{
+    GENERATED_BODY()
+
+    UPROPERTY()
+    FSoftObjectPath Path;
+
+    UPROPERTY()
+    TSoftObjectPtr<AActor> TypedPtr;
+};
+
+UCLASS()
+class UPieFixupTestObject : public UObject
+{
+	GENERATED_BODY()
+
+public:
+    UPROPERTY()
+    FSoftObjectPath Path;
+
+    UPROPERTY()
+    TSoftObjectPtr<AActor> TypedPtr;
+    
+    UPROPERTY()
+    FPieFixupStructWithSoftObjectPath Struct;
+
+    UPROPERTY()
+    TArray<FPieFixupStructWithSoftObjectPath> Array;
+};
diff --git a/Engine/Source/Runtime/Engine/Private/World.cpp b/Engine/Source/Runtime/Engine/Private/World.cpp
--- a/Engine/Source/Runtime/Engine/Private/World.cpp
+++ b/Engine/Source/Runtime/Engine/Private/World.cpp
@@ -3853,6 +3853,7 @@ UWorld* UWorld::DuplicateWorldForPIE(const FString& PackageName, UWorld* OwningW
 	UWorld* PIELevelWorld = GetDuplicatedWorldForPIE(EditorLevelWorld, PIELevelPackage, PIEInstanceID);
 
 	{
+		TRACE_CPUPROFILER_EVENT_SCOPE(FixupForPIE);
 		// The owning world may contain lazy pointers to actors in the sub-level we just duplicated so make sure they are fixed up with the PIE GUIDs
 		FPIEFixupSerializer FixupSerializer(OwningWorld, PIEInstanceID);
 		FixupSerializer << OwningWorld;
