Ensure that FDynamicallyTypedValue::IsInline does not evaluate ConstExprCeilLogTwo(alignof(UPTRINT)) at runtime by binding the result to a constexpr local.

diff --git a/Engine/Source/Runtime/CoreUObject/Public/UObject/DynamicallyTypedValue.h b/Engine/Source/Runtime/CoreUObject/Public/UObject/DynamicallyTypedValue.h
--- a/Engine/Source/Runtime/CoreUObject/Public/UObject/DynamicallyTypedValue.h
+++ b/Engine/Source/Runtime/CoreUObject/Public/UObject/DynamicallyTypedValue.h
@@ -156,8 +156,9 @@ namespace UE
 		// Whether the value's data is stored in InlineData or in the memory pointed to by HeapData.
 		bool IsInline() const
 		{
+			constexpr uint8 InlineDataMinAlignmentLogTwo = FMath::ConstExprCeilLogTwo(alignof(UPTRINT));
 			return Type->GetNumBytes() <= sizeof(UPTRINT)
-				&& Type->GetMinAlignmentLogTwo() <= FMath::ConstExprCeilLogTwo(alignof(UPTRINT));
+				&& Type->GetMinAlignmentLogTwo() <= InlineDataMinAlignmentLogTwo;
 		}
 
 		// Allocates heap memory for the value if it uses it.
