Fix PropagatePreEditChange potentially crashing when working with properties contained within structs of subobjects.

#rnx

diff --git a/Engine/Source/Runtime/CoreUObject/Private/UObject/Obj.cpp b/Engine/Source/Runtime/CoreUObject/Private/UObject/Obj.cpp
--- a/Engine/Source/Runtime/CoreUObject/Private/UObject/Obj.cpp
+++ b/Engine/Source/Runtime/CoreUObject/Private/UObject/Obj.cpp
@@ -571,8 +571,8 @@ void UObject::PropagatePreEditChange( TArray<UObject*>& AffectedObjects, FEditPr
 		}
 	}
 
-	check(PropertyAboutToChange.GetActiveNode() != nullptr);
-	const FProperty* ChangedProperty = PropertyAboutToChange.GetActiveNode()->GetValue();
+	check(PropertyAboutToChange.GetActiveMemberNode() != nullptr);
+	const FProperty* ChangedProperty = PropertyAboutToChange.GetActiveMemberNode()->GetValue();
 	FPropagatedEditChangeAnnotation Annotation = PropagatedEditChangeAnnotation.GetAnnotation(this);
 
 	for ( int32 i = 0; i < Instances.Num(); i++ )
