SavePackage: Rename private namespace SavePackageUtilities to UE::SavePackageUtilities to match the existing public namespace, thereby reducing complexity.
#rnx

diff --git a/Engine/Source/Runtime/CoreUObject/Private/UObject/SavePackage/PackageHarvester.cpp b/Engine/Source/Runtime/CoreUObject/Private/UObject/SavePackage/PackageHarvester.cpp
--- a/Engine/Source/Runtime/CoreUObject/Private/UObject/SavePackage/PackageHarvester.cpp
+++ b/Engine/Source/Runtime/CoreUObject/Private/UObject/SavePackage/PackageHarvester.cpp
@@ -46,7 +46,7 @@ EObjectMark GenerateMarksForObject(const UObject* InObject, const ITargetPlatfor
 	}
 #endif
 	
-	if (SavePackageUtilities::IsStrippedEditorOnlyObject(InObject, true, false))
+	if (UE::SavePackageUtilities::IsStrippedEditorOnlyObject(InObject, true, false))
 	{
 		Marks = (EObjectMark)(Marks | OBJECTMARK_EditorOnly);
 	}
@@ -549,7 +549,8 @@ void FPackageHarvester::TryHarvestExportInternal(UObject* InObject)
 
 #if WITH_EDITORONLY_DATA
 	// Remove the Game realm if the object is editoronly
-	bool bIsEditorOnlyObject = SavePackageUtilities::IsStrippedEditorOnlyObject(InObject, true /* bCheckRecursive */, false /* bCheckMarks */)
+	bool bIsEditorOnlyObject = UE::SavePackageUtilities::IsStrippedEditorOnlyObject(InObject,
+		true /* bCheckRecursive */, false /* bCheckMarks */)
 		&& !InObject->HasNonEditorOnlyReferences();
 	FHarvestScope EditorOnlyScope = EnterConditionalEditorOnlyScope(bIsEditorOnlyObject);
 #endif
@@ -633,7 +634,7 @@ void FPackageHarvester::ProcessImport(UObject* InObject)
 				// Gets all subobjects defined in a class, including the CDO, CDO components and blueprint-created components
 				TArray<UObject*> ObjectTemplates;
 				ObjectTemplates.Add(CDO);
-				SavePackageUtilities::GetCDOSubobjects(CDO, ObjectTemplates);
+				UE::SavePackageUtilities::GetCDOSubobjects(CDO, ObjectTemplates);
 				for (UObject* ObjTemplate : ObjectTemplates)
 				{
 					// Recurse into templates
diff --git a/Engine/Source/Runtime/CoreUObject/Private/UObject/SavePackage/SaveContext.cpp b/Engine/Source/Runtime/CoreUObject/Private/UObject/SavePackage/SaveContext.cpp
--- a/Engine/Source/Runtime/CoreUObject/Private/UObject/SavePackage/SaveContext.cpp
+++ b/Engine/Source/Runtime/CoreUObject/Private/UObject/SavePackage/SaveContext.cpp
@@ -195,7 +195,7 @@ EObjectMark FSaveContext::GetExcludedObjectMarksForGameRealm(const ITargetPlatfo
 {
 	if (TargetPlatform)
 	{
-		return SavePackageUtilities::GetExcludedObjectMarksForTargetPlatform(TargetPlatform);
+		return UE::SavePackageUtilities::GetExcludedObjectMarksForTargetPlatform(TargetPlatform);
 	}
 	else
 	{
diff --git a/Engine/Source/Runtime/CoreUObject/Private/UObject/SavePackage/SavePackageUtilities.cpp b/Engine/Source/Runtime/CoreUObject/Private/UObject/SavePackage/SavePackageUtilities.cpp
--- a/Engine/Source/Runtime/CoreUObject/Private/UObject/SavePackage/SavePackageUtilities.cpp
+++ b/Engine/Source/Runtime/CoreUObject/Private/UObject/SavePackage/SavePackageUtilities.cpp
@@ -182,14 +182,15 @@ namespace UE
 	}
 }
 
- FSavePackageSettings& FSavePackageSettings::GetDefaultSettings()
+FSavePackageSettings& FSavePackageSettings::GetDefaultSettings()
 {
 	static FSavePackageSettings Default;
 	return Default;
 }
 
-namespace SavePackageUtilities
+namespace UE::SavePackageUtilities
 {
+
 const FName NAME_World("World");
 const FName NAME_Level("Level");
 const FName NAME_PrestreamPackage("PrestreamPackage");
@@ -618,11 +619,6 @@ bool IsStrippedEditorOnlyObject(const UObject* InObject, bool bCheckRecursive, b
 	return true;
 #endif
 }
-	
-} // end namespace SavePackageUtilities
-
-namespace UE::SavePackageUtilities
-{
 
 bool IsUpdatingLoadedPath(bool bIsCooking, const FPackagePath& TargetPackagePath, uint32 SaveFlags)
 {
@@ -1736,10 +1732,6 @@ bool CanSkipEditorReferencedPackagesWhenCooking()
 	return bResult;
 }
 
-}
-
-namespace SavePackageUtilities
-{
 /**
  * Static: Saves thumbnail data for the specified package outer and linker
  *
@@ -1992,7 +1984,7 @@ void SaveWorldLevelInfo(UPackage* InOuter, FLinkerSave* Linker, FStructuredArchi
 	}
 }
 
-} // end namespace SavePackageUtilities
+} // end namespace UE::SavePackageUtilities
 
 void UPackage::WaitForAsyncFileWrites()
 {
@@ -2041,161 +2033,160 @@ bool UPackage::IsEmptyPackage(UPackage* Package, const UObject* LastReferencer)
 }
 
 
-namespace UE
+namespace UE::AssetRegistry
 {
-namespace AssetRegistry
+
+// See the corresponding ReadPackageDataMain and ReadPackageDataDependencies defined in PackageReader.cpp in AssetRegistry module
+void WritePackageData(FStructuredArchiveRecord& ParentRecord, bool bIsCooking, const UPackage* Package,
+	FLinkerSave* Linker, const TSet<UObject*>& ImportsUsedInGame, const TSet<FName>& SoftPackagesUsedInGame,
+	const ITargetPlatform* TargetPlatform, TArray<FAssetData>* OutAssetDatas)
 {
-	// See the corresponding ReadPackageDataMain and ReadPackageDataDependencies defined in PackageReader.cpp in AssetRegistry module
-	void WritePackageData(FStructuredArchiveRecord& ParentRecord, bool bIsCooking, const UPackage* Package,
-		FLinkerSave* Linker, const TSet<UObject*>& ImportsUsedInGame, const TSet<FName>& SoftPackagesUsedInGame,
-		const ITargetPlatform* TargetPlatform, TArray<FAssetData>* OutAssetDatas)
-	{
-		// To avoid large patch sizes, we have frozen cooked package format at the format before VER_UE4_ASSETREGISTRY_DEPENDENCYFLAGS
-		bool bPreDependencyFormat = bIsCooking;
+	// To avoid large patch sizes, we have frozen cooked package format at the format before VER_UE4_ASSETREGISTRY_DEPENDENCYFLAGS
+	bool bPreDependencyFormat = bIsCooking;
 
-		// WritePackageData is currently only called if not bTextFormat; we rely on that to save offsets
-		FArchive& BinaryArchive = ParentRecord.GetUnderlyingArchive();
-		check(!BinaryArchive.IsTextFormat());
+	// WritePackageData is currently only called if not bTextFormat; we rely on that to save offsets
+	FArchive& BinaryArchive = ParentRecord.GetUnderlyingArchive();
+	check(!BinaryArchive.IsTextFormat());
 
-		// Store the asset registry offset in the file and enter a record for the asset registry data
-		Linker->Summary.AssetRegistryDataOffset = (int32)BinaryArchive.Tell();
-		FStructuredArchiveRecord AssetRegistryRecord = ParentRecord.EnterField(TEXT("AssetRegistry")).EnterRecord();
+	// Store the asset registry offset in the file and enter a record for the asset registry data
+	Linker->Summary.AssetRegistryDataOffset = (int32)BinaryArchive.Tell();
+	FStructuredArchiveRecord AssetRegistryRecord = ParentRecord.EnterField(TEXT("AssetRegistry")).EnterRecord();
 
-		// Offset to Dependencies
-		int64 OffsetToAssetRegistryDependencyDataOffset = INDEX_NONE;
-		if (!bPreDependencyFormat)
-		{
-			// Write placeholder data for the offset to the separately-serialized AssetRegistryDependencyData
-			OffsetToAssetRegistryDependencyDataOffset = BinaryArchive.Tell();
-			int64 AssetRegistryDependencyDataOffset = 0;
-			AssetRegistryRecord << SA_VALUE(TEXT("AssetRegistryDependencyDataOffset"), AssetRegistryDependencyDataOffset);
-			check(BinaryArchive.Tell() == OffsetToAssetRegistryDependencyDataOffset + sizeof(AssetRegistryDependencyDataOffset));
-		}
+	// Offset to Dependencies
+	int64 OffsetToAssetRegistryDependencyDataOffset = INDEX_NONE;
+	if (!bPreDependencyFormat)
+	{
+		// Write placeholder data for the offset to the separately-serialized AssetRegistryDependencyData
+		OffsetToAssetRegistryDependencyDataOffset = BinaryArchive.Tell();
+		int64 AssetRegistryDependencyDataOffset = 0;
+		AssetRegistryRecord << SA_VALUE(TEXT("AssetRegistryDependencyDataOffset"), AssetRegistryDependencyDataOffset);
+		check(BinaryArchive.Tell() == OffsetToAssetRegistryDependencyDataOffset + sizeof(AssetRegistryDependencyDataOffset));
+	}
 
-		// Collect the tag map
-		TArray<UObject*> AssetObjects;
-		if (!(Linker->Summary.GetPackageFlags() & PKG_FilterEditorOnly))
+	// Collect the tag map
+	TArray<UObject*> AssetObjects;
+	if (!(Linker->Summary.GetPackageFlags() & PKG_FilterEditorOnly))
+	{
+		// Find any exports which are not in the tag map
+		for (int32 i = 0; i < Linker->ExportMap.Num(); i++)
 		{
-			// Find any exports which are not in the tag map
-			for (int32 i = 0; i < Linker->ExportMap.Num(); i++)
+			FObjectExport& Export = Linker->ExportMap[i];
+			if (Export.Object && Export.Object->IsAsset())
 			{
-				FObjectExport& Export = Linker->ExportMap[i];
-				if (Export.Object && Export.Object->IsAsset())
-				{
-					AssetObjects.Add(Export.Object);
-				}
+				AssetObjects.Add(Export.Object);
 			}
 		}
-		int32 ObjectCount = AssetObjects.Num();
-		FStructuredArchive::FArray AssetArray = AssetRegistryRecord.EnterArray(TEXT("TagMap"), ObjectCount);
-		if (OutAssetDatas)
-		{
-			OutAssetDatas->Reset();
-		}
-		FString PackageName = Package->GetName();
+	}
+	int32 ObjectCount = AssetObjects.Num();
+	FStructuredArchive::FArray AssetArray = AssetRegistryRecord.EnterArray(TEXT("TagMap"), ObjectCount);
+	if (OutAssetDatas)
+	{
+		OutAssetDatas->Reset();
+	}
+	FString PackageName = Package->GetName();
 
-		for (int32 ObjectIdx = 0; ObjectIdx < AssetObjects.Num(); ++ObjectIdx)
-		{
-			const UObject* Object = AssetObjects[ObjectIdx];
+	for (int32 ObjectIdx = 0; ObjectIdx < AssetObjects.Num(); ++ObjectIdx)
+	{
+		const UObject* Object = AssetObjects[ObjectIdx];
 
-			// Exclude the package name in the object path, we just need to know the path relative to the package we are saving
-			FString ObjectPath = Object->GetPathName(Package);
-			FString ObjectClassName = Object->GetClass()->GetPathName();
+		// Exclude the package name in the object path, we just need to know the path relative to the package we are saving
+		FString ObjectPath = Object->GetPathName(Package);
+		FString ObjectClassName = Object->GetClass()->GetPathName();
 
-			TArray<UObject::FAssetRegistryTag> SourceTags;
-			Object->GetAssetRegistryTags(SourceTags);
+		TArray<UObject::FAssetRegistryTag> SourceTags;
+		Object->GetAssetRegistryTags(SourceTags);
 #if WITH_EDITOR
-			Object->GetExtendedAssetRegistryTagsForSave(TargetPlatform, SourceTags);
+		Object->GetExtendedAssetRegistryTagsForSave(TargetPlatform, SourceTags);
 #endif // WITH_EDITOR
 
-			TArray<UObject::FAssetRegistryTag> Tags;
-			for (UObject::FAssetRegistryTag& SourceTag : SourceTags)
+		TArray<UObject::FAssetRegistryTag> Tags;
+		for (UObject::FAssetRegistryTag& SourceTag : SourceTags)
+		{
+			UObject::FAssetRegistryTag* Existing = Tags.FindByPredicate([SourceTag](const UObject::FAssetRegistryTag& InTag) { return InTag.Name == SourceTag.Name; });
+			if (Existing)
 			{
-				UObject::FAssetRegistryTag* Existing = Tags.FindByPredicate([SourceTag](const UObject::FAssetRegistryTag& InTag) { return InTag.Name == SourceTag.Name; });
-				if (Existing)
-				{
-					Existing->Value = SourceTag.Value;
-				}
-				else
-				{
-					Tags.Add(SourceTag);
-				}
+				Existing->Value = SourceTag.Value;
+			}
+			else
+			{
+				Tags.Add(SourceTag);
 			}
+		}
 
-			int32 TagCount = Tags.Num();
+		int32 TagCount = Tags.Num();
 
-			FStructuredArchive::FRecord AssetRecord = AssetArray.EnterElement().EnterRecord();
-			AssetRecord << SA_VALUE(TEXT("Path"), ObjectPath) << SA_VALUE(TEXT("Class"), ObjectClassName);
+		FStructuredArchive::FRecord AssetRecord = AssetArray.EnterElement().EnterRecord();
+		AssetRecord << SA_VALUE(TEXT("Path"), ObjectPath) << SA_VALUE(TEXT("Class"), ObjectClassName);
 
-			FStructuredArchive::FMap TagMap = AssetRecord.EnterField(TEXT("Tags")).EnterMap(TagCount);
+		FStructuredArchive::FMap TagMap = AssetRecord.EnterField(TEXT("Tags")).EnterMap(TagCount);
 
-			for (TArray<UObject::FAssetRegistryTag>::TConstIterator TagIter(Tags); TagIter; ++TagIter)
-			{
-				FString Key = TagIter->Name.ToString();
-				FString Value = TagIter->Value;
+		for (TArray<UObject::FAssetRegistryTag>::TConstIterator TagIter(Tags); TagIter; ++TagIter)
+		{
+			FString Key = TagIter->Name.ToString();
+			FString Value = TagIter->Value;
 
-				TagMap.EnterElement(Key) << Value;
-			}
+			TagMap.EnterElement(Key) << Value;
+		}
 
-			if (OutAssetDatas)
+		if (OutAssetDatas)
+		{
+			FAssetDataTagMap TagsAndValues;
+			for (UObject::FAssetRegistryTag& Tag : Tags)
 			{
-				FAssetDataTagMap TagsAndValues;
-				for (UObject::FAssetRegistryTag& Tag : Tags)
+				if (!Tag.Name.IsNone() && !Tag.Value.IsEmpty())
 				{
-					if (!Tag.Name.IsNone() && !Tag.Value.IsEmpty())
-					{
-						TagsAndValues.Add(Tag.Name, MoveTemp(Tag.Value));
-					}
+					TagsAndValues.Add(Tag.Name, MoveTemp(Tag.Value));
 				}
-				// if we do not have a full object path already, build it
-				const bool bFullObjectPath = ObjectPath.StartsWith(TEXT("/"), ESearchCase::CaseSensitive);
-				if (!bFullObjectPath)
+			}
+			// if we do not have a full object path already, build it
+			const bool bFullObjectPath = ObjectPath.StartsWith(TEXT("/"), ESearchCase::CaseSensitive);
+			if (!bFullObjectPath)
+			{
+				// if we do not have a full object path, ensure that we have a top level object for the package and not a sub object
+				if (!ensureMsgf(!ObjectPath.Contains(TEXT("."), ESearchCase::CaseSensitive),
+					TEXT("Cannot make FAssetData for sub object %s in package %s!"), *ObjectPath, *PackageName))
 				{
-					// if we do not have a full object path, ensure that we have a top level object for the package and not a sub object
-					if (!ensureMsgf(!ObjectPath.Contains(TEXT("."), ESearchCase::CaseSensitive),
-						TEXT("Cannot make FAssetData for sub object %s in package %s!"), *ObjectPath, *PackageName))
-					{
-						continue;
-					}
-					ObjectPath = PackageName + TEXT(".") + ObjectPath;
+					continue;
 				}
-
-				OutAssetDatas->Emplace(PackageName, ObjectPath, FTopLevelAssetPath(ObjectClassName),
-					MoveTemp(TagsAndValues), Package->GetChunkIDs(), Package->GetPackageFlags());
+				ObjectPath = PackageName + TEXT(".") + ObjectPath;
 			}
-		}
-		if (bPreDependencyFormat)
-		{
-			// The legacy format did not write the other sections, or the offsets to those other sections
-			return;
-		}
 
-		// Overwrite the placeholder offset for the AssetRegistryDependencyData and enter a record for the asset registry dependency data
-		{
-			int64 AssetRegistryDependencyDataOffset = Linker->Tell();
-			BinaryArchive.Seek(OffsetToAssetRegistryDependencyDataOffset);
-			BinaryArchive << AssetRegistryDependencyDataOffset;
-			BinaryArchive.Seek(AssetRegistryDependencyDataOffset);
+			OutAssetDatas->Emplace(PackageName, ObjectPath, FTopLevelAssetPath(ObjectClassName),
+				MoveTemp(TagsAndValues), Package->GetChunkIDs(), Package->GetPackageFlags());
 		}
-		FStructuredArchiveRecord DependencyDataRecord = ParentRecord.EnterField(TEXT("AssetRegistryDependencyData")).EnterRecord();
+	}
+	if (bPreDependencyFormat)
+	{
+		// The legacy format did not write the other sections, or the offsets to those other sections
+		return;
+	}
 
-		// Convert the IsUsedInGame sets into a bitarray with a value per import/softpackagereference
-		TBitArray<> ImportUsedInGameBits;
-		TBitArray<> SoftPackageUsedInGameBits;
-		ImportUsedInGameBits.Reserve(Linker->ImportMap.Num());
-		for (int32 ImportIndex = 0; ImportIndex < Linker->ImportMap.Num(); ++ImportIndex)
-		{
-			ImportUsedInGameBits.Add(ImportsUsedInGame.Contains(Linker->ImportMap[ImportIndex].XObject));
-		}
-		SoftPackageUsedInGameBits.Reserve(Linker->SoftPackageReferenceList.Num());
-		for (int32 SoftPackageIndex = 0; SoftPackageIndex < Linker->SoftPackageReferenceList.Num(); ++SoftPackageIndex)
-		{
-			SoftPackageUsedInGameBits.Add(SoftPackagesUsedInGame.Contains(Linker->SoftPackageReferenceList[SoftPackageIndex]));
-		}
+	// Overwrite the placeholder offset for the AssetRegistryDependencyData and enter a record for the asset registry dependency data
+	{
+		int64 AssetRegistryDependencyDataOffset = Linker->Tell();
+		BinaryArchive.Seek(OffsetToAssetRegistryDependencyDataOffset);
+		BinaryArchive << AssetRegistryDependencyDataOffset;
+		BinaryArchive.Seek(AssetRegistryDependencyDataOffset);
+	}
+	FStructuredArchiveRecord DependencyDataRecord = ParentRecord.EnterField(TEXT("AssetRegistryDependencyData")).EnterRecord();
 
-		// Serialize the Dependency section
-		DependencyDataRecord << SA_VALUE(TEXT("ImportUsedInGame"), ImportUsedInGameBits);
-		DependencyDataRecord << SA_VALUE(TEXT("SoftPackageUsedInGame"), SoftPackageUsedInGameBits);
+	// Convert the IsUsedInGame sets into a bitarray with a value per import/softpackagereference
+	TBitArray<> ImportUsedInGameBits;
+	TBitArray<> SoftPackageUsedInGameBits;
+	ImportUsedInGameBits.Reserve(Linker->ImportMap.Num());
+	for (int32 ImportIndex = 0; ImportIndex < Linker->ImportMap.Num(); ++ImportIndex)
+	{
+		ImportUsedInGameBits.Add(ImportsUsedInGame.Contains(Linker->ImportMap[ImportIndex].XObject));
 	}
+	SoftPackageUsedInGameBits.Reserve(Linker->SoftPackageReferenceList.Num());
+	for (int32 SoftPackageIndex = 0; SoftPackageIndex < Linker->SoftPackageReferenceList.Num(); ++SoftPackageIndex)
+	{
+		SoftPackageUsedInGameBits.Add(SoftPackagesUsedInGame.Contains(Linker->SoftPackageReferenceList[SoftPackageIndex]));
+	}
+
+	// Serialize the Dependency section
+	DependencyDataRecord << SA_VALUE(TEXT("ImportUsedInGame"), ImportUsedInGameBits);
+	DependencyDataRecord << SA_VALUE(TEXT("SoftPackageUsedInGame"), SoftPackageUsedInGameBits);
 }
+
 }
diff --git a/Engine/Source/Runtime/CoreUObject/Private/UObject/SavePackage/SavePackageUtilities.h b/Engine/Source/Runtime/CoreUObject/Private/UObject/SavePackage/SavePackageUtilities.h
--- a/Engine/Source/Runtime/CoreUObject/Private/UObject/SavePackage/SavePackageUtilities.h
+++ b/Engine/Source/Runtime/CoreUObject/Private/UObject/SavePackage/SavePackageUtilities.h
@@ -317,51 +317,53 @@ private:
 };
 
 // Utility functions used by both UPackage::Save and/or UPackage::Save2
-namespace SavePackageUtilities
+namespace UE::SavePackageUtilities
 {
-	extern const FName NAME_World;
-	extern const FName NAME_Level;
-	extern const FName NAME_PrestreamPackage;
 
-	void SaveThumbnails(UPackage* InOuter, FLinkerSave* Linker, FStructuredArchive::FSlot Slot);
+extern const FName NAME_World;
+extern const FName NAME_Level;
+extern const FName NAME_PrestreamPackage;
 
-	/**
-	 * Used to append additional data to the end of the package file by invoking callbacks stored in the linker.
-	 * They may be saved to the end of the file, or to a separate archive passed into the PackageWriter.
+void SaveThumbnails(UPackage* InOuter, FLinkerSave* Linker, FStructuredArchive::FSlot Slot);
+
+/**
+	* Used to append additional data to the end of the package file by invoking callbacks stored in the linker.
+	* They may be saved to the end of the file, or to a separate archive passed into the PackageWriter.
 	 
-	 * @param Linker The linker containing the exports. Provides the list of AdditionalData, and the data may write to it as their target archive.
-	 * @param InOutStartOffset In value is the offset in the Linker's archive where the datas will be put. If SavePackageContext settings direct
-	 *        the datas to write in a separate archive that will be combined after the linker, the value is the offset after the Linker archive's 
-	 *        totalsize and after any previous post-Linker archive data such as BulkDatas.
-	 *        Output value is incremented by the number of bytes written the Linker or the separate archive at the end of the linker.
-	 * @param SavePackageContext If non-null and configured to require it, data is passed to this PackageWriter on this context rather than appended to the Linker archive.
-	  */
-	ESavePackageResult AppendAdditionalData(FLinkerSave& Linker, int64& InOutDataStartOffset, FSavePackageContext* SavePackageContext);
+	* @param Linker The linker containing the exports. Provides the list of AdditionalData, and the data may write to it as their target archive.
+	* @param InOutStartOffset In value is the offset in the Linker's archive where the datas will be put. If SavePackageContext settings direct
+	*        the datas to write in a separate archive that will be combined after the linker, the value is the offset after the Linker archive's 
+	*        totalsize and after any previous post-Linker archive data such as BulkDatas.
+	*        Output value is incremented by the number of bytes written the Linker or the separate archive at the end of the linker.
+	* @param SavePackageContext If non-null and configured to require it, data is passed to this PackageWriter on this context rather than appended to the Linker archive.
+	*/
+ESavePackageResult AppendAdditionalData(FLinkerSave& Linker, int64& InOutDataStartOffset, FSavePackageContext* SavePackageContext);
 	
-	/** Used to create the sidecar file (.upayload) from payloads that have been added to the linker */
-	ESavePackageResult CreatePayloadSidecarFile(FLinkerSave& Linker, const FPackagePath& PackagePath, const bool bSaveToMemory,
-		FSavePackageOutputFileArray& AdditionalPackageFiles, FSavePackageContext* SavePackageContext);
+/** Used to create the sidecar file (.upayload) from payloads that have been added to the linker */
+ESavePackageResult CreatePayloadSidecarFile(FLinkerSave& Linker, const FPackagePath& PackagePath, const bool bSaveToMemory,
+	FSavePackageOutputFileArray& AdditionalPackageFiles, FSavePackageContext* SavePackageContext);
 	
-	void SaveWorldLevelInfo(UPackage* InOuter, FLinkerSave* Linker, FStructuredArchive::FRecord Record);
-	EObjectMark GetExcludedObjectMarksForTargetPlatform(const class ITargetPlatform* TargetPlatform);
-	void FindMostLikelyCulprit(const TArray<UObject*>& BadObjects, UObject*& MostLikelyCulprit, FString& OutReferencer, FSaveContext* InOptionalSaveContext = nullptr);
+void SaveWorldLevelInfo(UPackage* InOuter, FLinkerSave* Linker, FStructuredArchive::FRecord Record);
+EObjectMark GetExcludedObjectMarksForTargetPlatform(const class ITargetPlatform* TargetPlatform);
+void FindMostLikelyCulprit(const TArray<UObject*>& BadObjects, UObject*& MostLikelyCulprit, FString& OutReferencer, FSaveContext* InOptionalSaveContext = nullptr);
 	
-	/** 
-	  * Search 'OutputFiles' for output files that were saved to the temp directory and move those files
-	  * to their final location. Output files that were not saved to the temp directory will be ignored.
-	  * 
-	  * If errors are encountered then the original state of the package will be restored and should continue to work.
-	  */
-	ESavePackageResult FinalizeTempOutputFiles(const FPackagePath& PackagePath, const FSavePackageOutputFileArray& OutputFiles, const FDateTime& FinalTimeStamp);
-
-	void WriteToFile(const FString& Filename, const uint8* InDataPtr, int64 InDataSize);
-	void AsyncWriteFile(FLargeMemoryPtr Data, const int64 DataSize, const TCHAR* Filename, EAsyncWriteOptions Options, TArrayView<const FFileRegion> InFileRegions);
-	void AsyncWriteFile(EAsyncWriteOptions Options, FSavePackageOutputFile& File);
-
-	void GetCDOSubobjects(UObject* CDO, TArray<UObject*>& Subobjects);
-
-	/** Returns result of IsEditorOnlyObject if Engine:[Core.System]:CanStripEditorOnlyExportsAndImports (ini) is set to true */
-	bool IsStrippedEditorOnlyObject(const UObject* InObject, bool bCheckRecursive = true, bool bCheckMarks = true);
+/** 
+	* Search 'OutputFiles' for output files that were saved to the temp directory and move those files
+	* to their final location. Output files that were not saved to the temp directory will be ignored.
+	* 
+	* If errors are encountered then the original state of the package will be restored and should continue to work.
+	*/
+ESavePackageResult FinalizeTempOutputFiles(const FPackagePath& PackagePath, const FSavePackageOutputFileArray& OutputFiles, const FDateTime& FinalTimeStamp);
+
+void WriteToFile(const FString& Filename, const uint8* InDataPtr, int64 InDataSize);
+void AsyncWriteFile(FLargeMemoryPtr Data, const int64 DataSize, const TCHAR* Filename, EAsyncWriteOptions Options, TArrayView<const FFileRegion> InFileRegions);
+void AsyncWriteFile(EAsyncWriteOptions Options, FSavePackageOutputFile& File);
+
+void GetCDOSubobjects(UObject* CDO, TArray<UObject*>& Subobjects);
+
+/** Returns result of IsEditorOnlyObject if Engine:[Core.System]:CanStripEditorOnlyExportsAndImports (ini) is set to true */
+bool IsStrippedEditorOnlyObject(const UObject* InObject, bool bCheckRecursive = true, bool bCheckMarks = true);
+
 }
 
 #if ENABLE_COOK_STATS
diff --git a/Engine/Source/Runtime/CoreUObject/Private/UObject/SavePackage2.cpp b/Engine/Source/Runtime/CoreUObject/Private/UObject/SavePackage2.cpp
--- a/Engine/Source/Runtime/CoreUObject/Private/UObject/SavePackage2.cpp
+++ b/Engine/Source/Runtime/CoreUObject/Private/UObject/SavePackage2.cpp
@@ -432,7 +432,7 @@ ESavePackageResult HarvestPackage(FSaveContext& SaveContext)
 		// Harvest the PrestreamPackage class name if needed
 		if (PrestreamPackages.Num() > 0)
 		{
-			Harvester.HarvestPackageHeaderName(SavePackageUtilities::NAME_PrestreamPackage);
+			Harvester.HarvestPackageHeaderName(UE::SavePackageUtilities::NAME_PrestreamPackage);
 		}
 
 		// if we have a WorldTileInfo, we need to harvest its dependencies as well, i.e. Custom Version
@@ -575,7 +575,7 @@ ESavePackageResult ValidateExports(FSaveContext& SaveContext)
 			if (Object)
 			{
 				FName ClassName = Object->GetClass()->GetFName();
-				bContainsMap |= (ClassName == SavePackageUtilities::NAME_World || ClassName == SavePackageUtilities::NAME_Level);
+				bContainsMap |= (ClassName == UE::SavePackageUtilities::NAME_World || ClassName == UE::SavePackageUtilities::NAME_Level);
 			}
 		}
 		if (!bContainsMap)
@@ -657,7 +657,7 @@ ESavePackageResult ValidateIllegalReferences(FSaveContext& SaveContext, TArray<U
 		UObject* MostLikelyCulprit = nullptr;
 		FString CulpritString = TEXT("Unknown");
 		FString Referencer;
-		SavePackageUtilities::FindMostLikelyCulprit(ObjectsInOtherMaps, MostLikelyCulprit, Referencer, &SaveContext);
+		UE::SavePackageUtilities::FindMostLikelyCulprit(ObjectsInOtherMaps, MostLikelyCulprit, Referencer, &SaveContext);
 		if (MostLikelyCulprit != nullptr)
 		{
 			CulpritString = FString::Printf(TEXT("%s (%s)"), *MostLikelyCulprit->GetFullName(), *Referencer);
@@ -680,7 +680,7 @@ ESavePackageResult ValidateIllegalReferences(FSaveContext& SaveContext, TArray<U
 		UObject* MostLikelyCulprit = nullptr;
 		FString CulpritString = TEXT("Unknown");
 		FString Referencer;
-		SavePackageUtilities::FindMostLikelyCulprit(PrivateObjects, MostLikelyCulprit, Referencer, &SaveContext);
+		UE::SavePackageUtilities::FindMostLikelyCulprit(PrivateObjects, MostLikelyCulprit, Referencer, &SaveContext);
 		CulpritString = FString::Printf(TEXT("%s (%s)"),
 			(MostLikelyCulprit != nullptr) ? *MostLikelyCulprit->GetFullName() : TEXT("(unknown culprit)"),
 			*Referencer);
@@ -1204,7 +1204,7 @@ ESavePackageResult BuildLinker(FSaveContext& SaveContext)
 			// If the package import is a prestream package, mark it as such by hacking its class name
 			if (SaveContext.IsPrestreamPackage(Cast<UPackage>(Import)))
 			{
-				ObjectImport.ClassName = SavePackageUtilities::NAME_PrestreamPackage;
+				ObjectImport.ClassName = UE::SavePackageUtilities::NAME_PrestreamPackage;
 			}
 
 			if (ReplacedName != NAME_None)
@@ -1872,7 +1872,7 @@ ESavePackageResult WritePackageHeader(FStructuredArchive::FRecord& StructuredArc
 	// Save thumbnails
 	{
 		SCOPED_SAVETIMER(UPackage_Save_SaveThumbnails);
-		SavePackageUtilities::SaveThumbnails(SaveContext.GetPackage(), Linker, StructuredArchiveRoot.EnterField(TEXT("Thumbnails")));
+		UE::SavePackageUtilities::SaveThumbnails(SaveContext.GetPackage(), Linker, StructuredArchiveRoot.EnterField(TEXT("Thumbnails")));
 	}
 	{
 		// Save asset registry data so the editor can search for information about assets in this package
@@ -1884,7 +1884,7 @@ ESavePackageResult WritePackageHeader(FStructuredArchive::FRecord& StructuredArc
 	// Save level information used by World browser
 	{
 		SCOPED_SAVETIMER(UPackage_Save_WorldLevelData);
-		SavePackageUtilities::SaveWorldLevelInfo(SaveContext.GetPackage(), Linker, StructuredArchiveRoot);
+		UE::SavePackageUtilities::SaveWorldLevelInfo(SaveContext.GetPackage(), Linker, StructuredArchiveRoot);
 	}
 
 	// Write Preload Dependencies
@@ -1937,12 +1937,12 @@ ESavePackageResult WritePackageTextHeader(FStructuredArchive::FRecord& Structure
 	// Save thumbnails
 	{
 		SCOPED_SAVETIMER(UPackage_Save_SaveThumbnails);
-		SavePackageUtilities::SaveThumbnails(SaveContext.GetPackage(), Linker, StructuredArchiveRoot.EnterField(TEXT("Thumbnails")));
+		UE::SavePackageUtilities::SaveThumbnails(SaveContext.GetPackage(), Linker, StructuredArchiveRoot.EnterField(TEXT("Thumbnails")));
 	}
 	// Save level information used by World browser
 	{
 		SCOPED_SAVETIMER(UPackage_Save_WorldLevelData);
-		SavePackageUtilities::SaveWorldLevelInfo(SaveContext.GetPackage(), Linker, StructuredArchiveRoot);
+		UE::SavePackageUtilities::SaveWorldLevelInfo(SaveContext.GetPackage(), Linker, StructuredArchiveRoot);
 	}
 
 	return ReturnSuccessOrCancel();
@@ -2465,7 +2465,7 @@ ESavePackageResult FinalizeFile(FStructuredArchive::FRecord& StructuredArchiveRo
 
 			for (FSavePackageOutputFile& Entry : SaveContext.AdditionalPackageFiles)
 			{
-				SavePackageUtilities::AsyncWriteFile(WriteOptions, Entry);
+				UE::SavePackageUtilities::AsyncWriteFile(WriteOptions, Entry);
 			}
 		}
 		SaveContext.CloseLinkerArchives();
@@ -2495,7 +2495,7 @@ ESavePackageResult FinalizeFile(FStructuredArchive::FRecord& StructuredArchiveRo
 		// Add the .uasset file to the list of output files (TODO: Fix the 0 size, it isn't used after this point but needs to be cleaned up)
 		SaveContext.AdditionalPackageFiles.Emplace(SaveContext.GetFilename(), SaveContext.GetTempFilename().GetValue(), 0);
 
-		ESavePackageResult FinalizeResult = SavePackageUtilities::FinalizeTempOutputFiles(SaveContext.GetTargetPackagePath(), SaveContext.AdditionalPackageFiles,
+		ESavePackageResult FinalizeResult = UE::SavePackageUtilities::FinalizeTempOutputFiles(SaveContext.GetTargetPackagePath(), SaveContext.AdditionalPackageFiles,
 			SaveContext.GetFinalTimestamp());
 		
 		SaveContext.SetTempFilename(TOptional<FString>());
@@ -2998,14 +2998,14 @@ ESavePackageResult WriteAdditionalFiles(FSaveContext& SaveContext, FScopedSlowTa
 	SlowTask.EnterProgressFrame();
 
 	// Add any pending data blobs to the end of the file by invoking the callbacks
-	ESavePackageResult Result = SavePackageUtilities::AppendAdditionalData(*SaveContext.GetLinker(), VirtualExportsFileOffset, SaveContext.GetSavePackageContext());
+	ESavePackageResult Result = UE::SavePackageUtilities::AppendAdditionalData(*SaveContext.GetLinker(), VirtualExportsFileOffset, SaveContext.GetSavePackageContext());
 	if (Result != ESavePackageResult::Success)
 	{
 		return Result;
 	}
 
 	// Create the payload side car file (if needed)
-	Result = SavePackageUtilities::CreatePayloadSidecarFile(*SaveContext.GetLinker(), SaveContext.GetTargetPackagePath(), SaveContext.IsSaveToMemory(),
+	Result = UE::SavePackageUtilities::CreatePayloadSidecarFile(*SaveContext.GetLinker(), SaveContext.GetTargetPackagePath(), SaveContext.IsSaveToMemory(),
 		SaveContext.AdditionalPackageFiles, SaveContext.GetSavePackageContext());
 	if (Result != ESavePackageResult::Success)
 	{
