Intializing Object Handles with the actual global UObject array capacity instead of the max UObject count value from ini since the actual array capacity may be larger (it's rounded up to the nearest multiple of UObject chunk size)

diff --git a/Engine/Source/Runtime/CoreUObject/Private/UObject/UObjectBase.cpp b/Engine/Source/Runtime/CoreUObject/Private/UObject/UObjectBase.cpp
--- a/Engine/Source/Runtime/CoreUObject/Private/UObject/UObjectBase.cpp
+++ b/Engine/Source/Runtime/CoreUObject/Private/UObject/UObjectBase.cpp
@@ -1049,7 +1049,7 @@ void UObjectBaseInit()
 	GUObjectAllocator.AllocatePermanentObjectPool(SizeOfPermanentObjectPool);
 	GUObjectArray.AllocateObjectPool(MaxUObjects, MaxObjectsNotConsideredByGC, bPreAllocateUObjectArray);
 #if UE_WITH_OBJECT_HANDLE_LATE_RESOLVE
-	UE::CoreUObject::Private::InitObjectHandles(MaxUObjects);
+	UE::CoreUObject::Private::InitObjectHandles(GUObjectArray.GetObjectArrayCapacity());
 #endif
 
 	void InitNoPendingKill();
