Only run FindInexactObjectReferencesInStructSerializers when prompted by user.

diff --git a/Engine/Source/Runtime/CoreUObject/Private/UObject/Class.cpp b/Engine/Source/Runtime/CoreUObject/Private/UObject/Class.cpp
--- a/Engine/Source/Runtime/CoreUObject/Private/UObject/Class.cpp
+++ b/Engine/Source/Runtime/CoreUObject/Private/UObject/Class.cpp
@@ -4088,7 +4088,7 @@ bool FAutomationTestFindUndeclaredObjectReferencesInStructSerializers::RunTest(c
 
 // Test for finding structs which unoptimally their object properties for native serialization - not a smoke test
 
-	IMPLEMENT_SIMPLE_AUTOMATION_TEST(FAutomationTestFindInexactObjectReferencesInStructSerializers, "UObject.Class FindInexactObjectReferencesInStructSerializers", EAutomationTestFlags::EditorContext | EAutomationTestFlags::ClientContext | EAutomationTestFlags::ServerContext | EAutomationTestFlags::EngineFilter )
+	IMPLEMENT_SIMPLE_AUTOMATION_TEST(FAutomationTestFindInexactObjectReferencesInStructSerializers, "UObject.Class FindInexactObjectReferencesInStructSerializers", EAutomationTestFlags::EditorContext | EAutomationTestFlags::ClientContext | EAutomationTestFlags::ServerContext | EAutomationTestFlags::EngineFilter | EAutomationTestFlags::RequiresUser)
 bool FAutomationTestFindInexactObjectReferencesInStructSerializers::RunTest(const FString& Parameters)
 {
 	return FindUndeclaredObjectReferencesInStructSerializers(*this, true);
