New GC schema that replaces token stream

Flat, simpler and more compact representation that doesn't bake in depth-first visitation order of maps/sets and arrays-of-structs.

 ~0.5MB memory saved and 3-4% walltime reachability speedup for an internal project, despite making two things slower:
1. The new unbounded queue for arrays-of-structs lack work-stealing support
2. Struct schemas are allocated separately for simplicity, which is likely slower than inlining them with class schemas.
Both work-stealing and inlining struct schemas might be implemented as future optimizations.

diff --git a/Engine/Source/Editor/KismetCompiler/Private/KismetCompiler.cpp b/Engine/Source/Editor/KismetCompiler/Private/KismetCompiler.cpp
--- a/Engine/Source/Editor/KismetCompiler/Private/KismetCompiler.cpp
+++ b/Engine/Source/Editor/KismetCompiler/Private/KismetCompiler.cpp
@@ -3028,7 +3028,7 @@ void FKismetCompilerContext::FinishCompilingClass(UClass* Class)
 	if (ParentClass != NULL)
 	{
 		// Propagate the new parent's inheritable class flags
-		Class->ReferenceTokens.Reset();
+		Class->ReferenceSchema.Reset();
 		Class->ClassFlags &= ~CLASS_RecompilerClear;
 		Class->ClassFlags |= (ParentClass->ClassFlags & CLASS_ScriptInherit);//@TODO: ChangeParentClass had this, but I don't think I want it: | UClass::StaticClassFlags;  // will end up with CLASS_Intrinsic
 		Class->ClassCastFlags |= ParentClass->ClassCastFlags;
diff --git a/Engine/Source/Runtime/CoreUObject/Private/UObject/Class.cpp b/Engine/Source/Runtime/CoreUObject/Private/UObject/Class.cpp
--- a/Engine/Source/Runtime/CoreUObject/Private/UObject/Class.cpp
+++ b/Engine/Source/Runtime/CoreUObject/Private/UObject/Class.cpp
@@ -583,8 +583,7 @@ void UField::SetAssociatedFField(FField* InField)
 
 IMPLEMENT_CORE_INTRINSIC_CLASS(UField, UObject,
 	{
-		UE::GC::FTokenStreamBuilder& Builder = UE::GC::FIntrinsicClassTokens::AllocateBuilder(Class);
-		Builder.EmitObjectReference(STRUCT_OFFSET(UField, Next), TEXT("Next"));
+		UE::GC::DeclareIntrinsicMembers(Class, { UE_GC_MEMBER(UField, Next) });
 	}
 );
 
@@ -2207,13 +2206,9 @@ IMPLEMENT_CORE_INTRINSIC_CLASS(UStruct, UField,
 	{
 		Class->CppClassStaticFunctions = UOBJECT_CPPCLASS_STATICFUNCTIONS_FORCLASS(UStruct);
 
-		UE::GC::FTokenStreamBuilder& Builder = UE::GC::FIntrinsicClassTokens::AllocateBuilder(Class);
-		Builder.EmitObjectReference(STRUCT_OFFSET(UStruct, SuperStruct), TEXT("SuperStruct"));
-		Builder.EmitObjectReference(STRUCT_OFFSET(UStruct, Children), TEXT("Children"));
-
 		// Note: None of the *Link members need to be emitted, as they only contain properties
 		// that are in the Children chain or SuperStruct->Children chains.
-		Builder.EmitObjectArrayReference(STRUCT_OFFSET(UStruct, ScriptAndPropertyObjectReferences), TEXT("ScriptAndPropertyObjectReferences"));
+		UE::GC::DeclareIntrinsicMembers(Class, { UE_GC_MEMBER(UStruct, SuperStruct), UE_GC_MEMBER(UStruct, Children), UE_GC_MEMBER(UStruct, ScriptAndPropertyObjectReferences) });
 	}
 );
 
@@ -6425,29 +6420,21 @@ bool UClass::IsClassGroupName(const TCHAR* InGroupName) const
 
 
 #if WITH_EDITORONLY_DATA
-IMPLEMENT_CORE_INTRINSIC_CLASS(UClass, UStruct,
-	{
-		Class->CppClassStaticFunctions = UOBJECT_CPPCLASS_STATICFUNCTIONS_FORCLASS(UClass);
-
-		UE::GC::FTokenStreamBuilder& Builder = UE::GC::FIntrinsicClassTokens::AllocateBuilder(Class);
-		Builder.EmitObjectReference(STRUCT_OFFSET(UClass, ClassDefaultObject), TEXT("ClassDefaultObject"));
-		Builder.EmitObjectReference(STRUCT_OFFSET(UClass, ClassWithin), TEXT("ClassWithin"));
-		Builder.EmitObjectReference(STRUCT_OFFSET(UClass, ClassGeneratedBy), TEXT("ClassGeneratedBy"));
-		Builder.EmitObjectArrayReference(STRUCT_OFFSET(UClass, NetFields), TEXT("NetFields"));
-	}
-);
+#define UCLASS_EDITOR_GC_MEMBERS UE_GC_MEMBER(UClass, ClassGeneratedBy)
 #else
+#define UCLASS_EDITOR_GC_MEMBERS
+#endif
+
 IMPLEMENT_CORE_INTRINSIC_CLASS(UClass, UStruct,
 	{
 		Class->CppClassStaticFunctions = UOBJECT_CPPCLASS_STATICFUNCTIONS_FORCLASS(UClass);
 
-		UE::GC::FTokenStreamBuilder& Builder = UE::GC::FIntrinsicClassTokens::AllocateBuilder(Class);
-		Builder.EmitObjectReference(STRUCT_OFFSET(UClass, ClassDefaultObject), TEXT("ClassDefaultObject"));
-		Builder.EmitObjectReference(STRUCT_OFFSET(UClass, ClassWithin), TEXT("ClassWithin"));
-		Builder.EmitObjectArrayReference(STRUCT_OFFSET(UClass, NetFields), TEXT("NetFields"));
+		UE::GC::DeclareIntrinsicMembers(Class, {	UE_GC_MEMBER(UClass, ClassDefaultObject),
+													UE_GC_MEMBER(UClass, ClassWithin), 
+													UE_GC_MEMBER(UClass, NetFields),
+													UCLASS_EDITOR_GC_MEMBERS} );
 	}
 );
-#endif
 
 void GetPrivateStaticClassBody(
 	const TCHAR* PackageName,
diff --git a/Engine/Source/Runtime/CoreUObject/Private/UObject/GarbageCollection.cpp b/Engine/Source/Runtime/CoreUObject/Private/UObject/GarbageCollection.cpp
--- a/Engine/Source/Runtime/CoreUObject/Private/UObject/GarbageCollection.cpp
+++ b/Engine/Source/Runtime/CoreUObject/Private/UObject/GarbageCollection.cpp
@@ -52,9 +52,6 @@ DEFINE_LOG_CATEGORY(LogGarbage);
 
 #define PERF_DETAILED_PER_CLASS_GC_STATS				(LOOKING_FOR_PERF_ISSUES || 0) 
 
-static_assert(sizeof(FGCSkipInfo) == sizeof(uint32),		"Must have size of a GC token");
-static_assert(sizeof(FGCReferenceInfo) == sizeof(uint32),	"Must have size of a GC token");
-
 /** Object count during last mark phase																				*/
 FThreadSafeCounter		GObjectCountDuringLastMarkPhase;
 /** Whether UObject hash tables are locked by GC */
@@ -294,310 +291,6 @@ static FAutoConsoleVariableRef CGarbageReferenceTrackingEnabled(
 namespace UE::GC
 {
 
-namespace TokenDebugNames
-{
-#if ENABLE_GC_TOKEN_DEBUG_INFO
-#define DEFINE_TOKEN_DEBUG_ID(ID) static const FLazyName ID(#ID "Token")
-#else
-#define DEFINE_TOKEN_DEBUG_ID(ID) static constexpr EName ID = EName::None;
-#endif
-	DEFINE_TOKEN_DEBUG_ID(SkipInfo);
-	DEFINE_TOKEN_DEBUG_ID(Stride);
-	DEFINE_TOKEN_DEBUG_ID(Count);
-	DEFINE_TOKEN_DEBUG_ID(Pointer);
-	DEFINE_TOKEN_DEBUG_ID(EndOfPointer);
-	DEFINE_TOKEN_DEBUG_ID(ARO);
-	DEFINE_TOKEN_DEBUG_ID(EndOfStream);
-}
-
-#if !UE_BUILD_SHIPPING
-	static const TCHAR* ToStatName(EGCReferenceType Type)
-	{
-		switch (Type)
-		{
-		case GCRT_Object:								return TEXT("Reference");
-		case GCRT_ArrayObject:							return TEXT("Array<Reference>");
-		case GCRT_ArrayStruct:							return TEXT("Array<Struct>");
-		case GCRT_FixedArray:							return TEXT("FixedArray");
-		case GCRT_AddStructReferencedObjects:			return TEXT("StructARO");
-		case GCRT_AddReferencedObjects:					return TEXT("ClassARO");
-		case GCRT_AddTMapReferencedObjects:				return TEXT("Map");
-		case GCRT_AddTSetReferencedObjects:				return TEXT("Set");
-		case GCRT_AddFieldPathReferencedObject:			return TEXT("FieldPath");
-		case GCRT_ArrayAddFieldPathReferencedObject:	return TEXT("Array<FieldPath>");
-		case GCRT_EndOfStream:							return TEXT("EndOfStream");
-		case GCRT_ArrayObjectFreezable:					return TEXT("FrozenArray<Reference>");
-		case GCRT_ArrayStructFreezable:					return TEXT("FrozenArray<Struct>");
-		case GCRT_DynamicallyTypedValue:				return TEXT("DynamicallyTypedValue");
-		case GCRT_SlowAddReferencedObjects:				return TEXT("SlowARO");
-		default:										return TEXT("!!ERROR!!");
-		}
-	}
-
-	class FTokenStreamStats
-	{
-	public:
-		uint32 Counters[GCRT_Count] = {};
-		uint32 NumTokens = 0;
-		uint32 NumPointers = 0;
-		uint32 MinOffset = 0;
-		uint32 MaxOffset = 0;
-		uint32 MaxJump = 0;
-
-		FTokenStreamStats()  {}
-		explicit FTokenStreamStats(const FTokenStreamOwner& Tokens)
-			: NumTokens(Tokens.Strong.NumTokens())
-		{
-			FIterator It{Tokens.Strong, Tokens};
-			CountTokensUntil(It, NumTokens);
-		}
-
-		FTokenStreamStats& operator+=(const FTokenStreamStats& In)
-		{
-			for (uint32 Idx = 0; Idx < GCRT_Count; ++Idx)
-			{
-				Counters[Idx] += In.Counters[Idx];
-			}
-			NumTokens += In.NumTokens;
-			NumPointers += In.NumPointers;
-			MinOffset = MinOffset == 0 ? In.MinOffset : FMath::Min(MinOffset, In.MinOffset);
-			MaxOffset = FMath::Max(MaxOffset, In.MaxOffset);
-			MaxJump = FMath::Max(MaxJump, In.MaxJump);
-
-			return *this;
-		}
-	
-	private:
-		struct FIterator
-		{
-			const FTokenStreamView Tokens;
-			const FTokenStreamOwner& Owner; // Only for debug name access
-			uint32 Idx = 0;
-		
-			FGCReferenceInfo ConsumeReferenceInfo()
-			{
-				return Tokens.AccessReferenceInfo(Idx++);
-			}
-			
-			void Skip(FName DebugName)
-			{
-				check(GetDebugName() == DebugName);
-				++Idx;
-			}
-
-			uint32 ConsumeSkipInfo()
-			{
-				check(GetDebugName() == TokenDebugNames::SkipInfo);
-				return Tokens.ReadSkipInfo(/* in-out */ Idx).SkipIndex;
-			}
-
-			void* ConsumePointer()
-			{
-				check(GetDebugName() == TokenDebugNames::Pointer);
-				void* Ptr = Tokens.ReadPointer(/* in-out */ Idx);
-				Skip(TokenDebugNames::EndOfPointer);
-
-				return Ptr;
-			}
-
-#if ENABLE_GC_TOKEN_DEBUG_INFO
-			FName GetDebugName() const { return Owner.GetTokenInfo(FTokenId(Idx)).Name;}
-#else
-			static constexpr EName GetDebugName() { return EName::None; }
-#endif
-		};
-
-		void CountTokensUntil(FIterator& It, uint32 EndIdx)
-		{
-			static constexpr EGCReferenceType StrongFieldTokens[] = { 
-				GCRT_Object, GCRT_ArrayObject, GCRT_ArrayStruct, GCRT_AddTMapReferencedObjects, GCRT_AddTSetReferencedObjects, GCRT_FixedArray,
-				GCRT_AddFieldPathReferencedObject, GCRT_ArrayAddFieldPathReferencedObject, GCRT_ArrayObjectFreezable, GCRT_ArrayStructFreezable,
-				GCRT_Optional, GCRT_DynamicallyTypedValue };
-
-			const bool bObjectTokens = EndIdx == It.Tokens.NumTokens();
-			uint32 LastOffset = 0;
-			while (It.Idx < EndIdx)
-			{
-				FGCReferenceInfo Token = It.ConsumeReferenceInfo();
-				++Counters[Token.Type];
-
-				if (bObjectTokens && Algo::Find(StrongFieldTokens, Token.Type))
-				{
-					check(Token.Offset);
-					MinOffset = MinOffset == 0 ? Token.Offset : FMath::Min(MinOffset, Token.Offset);
-					MaxOffset = FMath::Max(Token.Offset, MaxOffset);
-					uint32 Jump = (uint32)FMath::Abs((int32)(Token.Offset - LastOffset));
-					LastOffset = Token.Offset;
-					MaxJump = FMath::Max(MaxJump, Jump);
-				}
-
-				// Traverse structural tokens
-				switch (Token.Type)
-				{
-					case GCRT_Optional:
-					case GCRT_ArrayStructFreezable:
-					case GCRT_ArrayStruct:
-						It.Skip(TokenDebugNames::Stride);
-						CountTokensUntil(It, It.ConsumeSkipInfo());
-						break;
-					case GCRT_FixedArray:
-						It.Skip(TokenDebugNames::Stride);
-						It.Skip(TokenDebugNames::Count);
-						break;
-					case GCRT_AddTSetReferencedObjects:
-					case GCRT_AddTMapReferencedObjects:
-					{
-						FProperty* Property = (FProperty*)It.ConsumePointer();
-						check(Property->IsA<FMapProperty>() || Property->IsA<FSetProperty>());
-						++NumPointers;
-						CountTokensUntil(It, It.ConsumeSkipInfo());
-						break;	
-					}
-					case GCRT_EndOfPointer:
-						checkf(false, TEXT("ConsumePointer() should consume all GCRT_EndOfPointer tokens"));
-						break;
-					case GCRT_AddStructReferencedObjects:
-						It.ConsumePointer();
-						++NumPointers;
-						break;
-					case GCRT_AddReferencedObjects:
-						It.ConsumePointer();
-						++NumPointers;
-						[[fallthrough]];
-					case GCRT_SlowAddReferencedObjects:
-						It.Skip(TokenDebugNames::EndOfStream);
-						[[fallthrough]];
-					case GCRT_EndOfStream:
-						check(It.Idx == EndIdx);
-						break;
-					default: // Non-structural leaf field
-						break;
-				}
-			}
-		} 
-	};
-
-	using FLog2Distribution = uint32[32];
-	enum class EDumpUnit {Bytes, Number};
-
-	static void DumpDistribution(FOutputDevice& Out, const TCHAR* Title, const FLog2Distribution& Bins, uint32 Total, EDumpUnit Unit)
-	{
-		uint32 End = 32;
-		for (; End > 0 && Bins[End - 1] == 0; --End);
-
-		uint32 Begin = 0;
-		for (; Begin < End && Bins[Begin] == 0; ++Begin);
-
-		Out.Logf(TEXT(" --- %-24s ---"), Title);
-
-		for (uint32 Idx = Begin, Sum = 0; Idx < End; ++Idx)
-		{
-			uint32 Num = Bins[Idx];
-			Sum += Num;
-			uint32 Max = Idx == 0 ? 0 : 1u << (Idx - 1);
-			float Percentage = (100.f * Num) / Total;
-			float Cumulative = (100.f * Sum) / Total;
-			if (Unit == EDumpUnit::Number)
-			{
-				Out.Logf(TEXT(" %4d %5.1f%% %5.1f%%  %d"), Max, Percentage, Cumulative, Num);
-			}
-			else
-			{
-				uint32 bK = Idx > 10;
-				Out.Logf(TEXT(" %3d%c %5.1f%% %5.1f%%  %d"), Max >> bK*10, "BK"[bK], Percentage, Cumulative, Num);
-			}
-		}
-	}
-
-	static void DumpTokenStats(FOutputDevice& Out, const FTokenStreamStats& Stats)
-	{
-		static constexpr TCHAR FormatStr[] = TEXT(" %18s %5.1f%% %5.1f%%  %d");
-
-		uint32 Sum = 0;
-		for (uint32 Idx = 0; Idx < GCRT_Count; ++Idx)
-		{
-			if (uint32 Num = Stats.Counters[Idx])
-			{
-				Sum += Num;
-				float Percentage = (100.f * Num) / Stats.NumTokens;
-				float Cumulative = (100.f * Sum) / Stats.NumTokens;
-				Out.Logf(FormatStr, ToStatName((EGCReferenceType)Idx), Percentage, Cumulative, Num);
-			}
-		}
-
-		uint32 Num = Stats.NumPointers * 3; // EndOfPointer + 2 32-bit tokens for 64-bit pointer
-		Sum += Num;
-		float Percentage = (100.f * Num) / Stats.NumTokens;
-		float Cumulative = (100.f * Sum) / Stats.NumTokens;
-		Out.Logf(FormatStr, TEXT("Pointer tokens"), Percentage, Cumulative, Num);
-		Out.Logf(FormatStr, TEXT("Skip/Stride/Count"), 100 - Cumulative, 100.f, Stats.NumTokens - Sum);
-	}
-	
-	static void DumpSchemaStats(FOutputDevice& Out)
-	{
-		TMap<void*, FTokenStreamStats> AllStats;
-		uint32 NumClasses = 0;
-		uint32 NumTokens = 0;
-		for (TObjectIterator<UClass> It; It; ++It)
-		{
-			FTokenStreamOwner& Tokens = It->ReferenceTokens;
-			FTokenStreamStats& Stats = AllStats.FindOrAdd(Tokens.Strong.GetTokenData());
-			if (Stats.NumTokens == 0)
-			{
-				Stats = FTokenStreamStats(Tokens);
-				NumTokens += Stats.NumTokens;
-			}
-			++NumClasses;
-		}
-		
-		FTokenStreamStats Sum;
-		FLog2Distribution SizeBins = {};
-		FLog2Distribution MinOffsetBins = {};
-		FLog2Distribution OffsetRangeBins = {};
-		FLog2Distribution MaxJumpBins = {};
-		FLog2Distribution NumTokenBins = {};
-		FLog2Distribution TokenBins[GCRT_Count] = {};
-		uint32 NumObjects = 0;
-		for (TObjectIterator<UObject> It; It; ++It)
-		{
-			const uint32 Size = It->GetClass()->GetPropertiesSize();
-			const FTokenStreamStats& Stats = AllStats.FindChecked(It->GetClass()->ReferenceTokens.Strong.GetTokenData());
-			Sum += Stats;
-			++SizeBins[FMath::CeilLogTwo(Size)];
-			++MinOffsetBins[FMath::CeilLogTwo(Stats.MinOffset)];
-			++OffsetRangeBins[FMath::CeilLogTwo(Stats.MaxOffset - Stats.MinOffset)];
-			++MaxJumpBins[FMath::CeilLogTwo(Stats.MaxJump)];
-			++NumTokenBins[FMath::CeilLogTwo(Stats.NumTokens)];
-			++NumObjects;
-
-			for (uint32 Idx = 0; Idx < GCRT_Count; ++Idx)
-			{
-				 ++TokenBins[Idx][FMath::CeilLogTwo(Stats.Counters[Idx])];
-			}
-		}
-
-		Out.Logf(TEXT("-------------------- GC strong token statistics --------------------"));
-		Out.Logf(TEXT(" %d token streams shared by %d classes contain %d tokens"), AllStats.Num(), NumClasses, NumTokens);
-		Out.Logf(TEXT(" Sweeping %d objects will process %d tokens"), NumObjects, Sum.NumTokens);
-		Out.Logf(TEXT("--------------------------- Distributions --------------------------"));
-		DumpDistribution(Out, TEXT("Object size"), SizeBins, NumObjects, EDumpUnit::Bytes);
-		DumpDistribution(Out, TEXT("Min offset"),	MinOffsetBins, NumObjects, EDumpUnit::Bytes);
-		DumpDistribution(Out, TEXT("Offset range"),	OffsetRangeBins, NumObjects, EDumpUnit::Bytes);
-		DumpDistribution(Out, TEXT("Max jump"),	MaxJumpBins, NumObjects, EDumpUnit::Bytes);
-		DumpDistribution(Out, TEXT("Tokens"), NumTokenBins, NumObjects, EDumpUnit::Number);
-		for (EGCReferenceType TokenType : {	GCRT_Object, GCRT_ArrayObject, GCRT_ArrayStruct, GCRT_AddTMapReferencedObjects, GCRT_AddTSetReferencedObjects, GCRT_FixedArray })
-		{
-			DumpDistribution(Out, ToStatName(TokenType), TokenBins[TokenType], NumObjects, EDumpUnit::Number);
-		}
-		Out.Logf(TEXT("------------------------ Tokens to process -------------------------"));
-		DumpTokenStats(Out, Sum);
-		Out.Logf(TEXT("--------------------------------------------------------------------"));
-	}
-
-	static FAutoConsoleCommandWithOutputDevice GDumpSchemaStats(TEXT("gc.DumpSchemaStats"), TEXT("Dump GC token stream statistics"), FConsoleCommandWithOutputDeviceDelegate::CreateStatic(&DumpSchemaStats));
-
-#endif // !UE_BUILD_SHIPPING
-
 	struct FStats
 	{
 #if PERF_DETAILED_PER_CLASS_GC_STATS
@@ -1117,21 +810,14 @@ bool IsInGarbageCollectorThread()
 	return GAsyncPurge ? GAsyncPurge->IsInAsyncPurgeThread() : IsInGameThread();
 }
 
-namespace UE::GC {
-
-	
-
-
-using AROFunc = void (*)(UObject*, FReferenceCollector&);
-
 //////////////////////////////////////////////////////////////////////////
 
+namespace UE::GC {
+
 /** Pool for reusing contexts between CollectReferences calls */
 class FContextPool
 {
-	friend class FContextPoolScope;
-	FContextPool();
-	
+	friend class FContextPoolScope;	
 	TArray<TUniquePtr<FWorkerContext>> Reusable;
 	int32 NumAllocated = 0;
 };
@@ -1141,6 +827,8 @@ class FContextPool
 // Bit of defensive coding during a large refactor, can be removed in the long run
 class FContextPoolScope
 {
+	UE_NONCOPYABLE(FContextPoolScope);
+
 	FContextPool& Pool;
 	bool bNested;
 	static FContextPool& Get();
@@ -1151,7 +839,6 @@ class FContextPoolScope
 	}
 public:
 	FContextPoolScope() : Pool(Get()) { CheckGameThread(); }
-	~FContextPoolScope() { CheckGameThread(); }
 
 	FWorkerContext* AllocateFromPool();
 	void ReturnToPool(FWorkerContext* Context);
@@ -1205,14 +892,6 @@ void FContextPoolScope::Cleanup()
 	Pool.Reusable.Empty();
 }
 
-static void DumpMemoryStats(FOutputDevice& OutputDevice);
-
-FContextPool::FContextPool()
-{
-	static FAutoConsoleCommandWithOutputDevice Cmd(TEXT("gc.DumpMemoryStats"), TEXT("Dump GC memory stats"),
-		FConsoleCommandWithOutputDeviceDelegate::CreateStatic(&DumpMemoryStats));
-}
-
 //////////////////////////////////////////////////////////////////////////
 	
 FORCEINLINE constexpr bool IsWithClusters(EGCOptions Options) {	return !!(Options & EGCOptions::WithClusters); }
@@ -1411,7 +1090,7 @@ struct FResolvedMutableReference : FImmutableReference
 	UObject** Mutable;
 };
 
-FORCEINLINE_DEBUGGABLE static bool ValidateReference(UObject* Object, FPermanentObjectPoolExtents PermanentPool, const UObject* ReferencingObject, FTokenId TokenId)
+FORCEINLINE_DEBUGGABLE static bool ValidateReference(UObject* Object, FPermanentObjectPoolExtents PermanentPool, const UObject* ReferencingObject, FMemberId MemberId)
 {
 	bool bOk = (!PermanentPool.Contains(Object)) & (!!Object) & IsObjectHandleResolved(reinterpret_cast<FObjectHandle&>(Object)); //-V792
 
@@ -1425,12 +1104,12 @@ FORCEINLINE_DEBUGGABLE static bool ValidateReference(UObject* Object, FPermanent
 			!Object->IsValidLowLevelFast())
 		{
 			UClass* ReferencingClass = ReferencingObject ? ReferencingObject->GetClass() : nullptr;
-			FString TokenName = ReferencingClass ? ReferencingClass->ReferenceTokens.GetTokenInfo(TokenId).Name.ToString() : TEXT("N/A");
+			FString MemberName = ReferencingClass ? GetMemberDebugInfo(ReferencingClass->ReferenceSchema.Get(), MemberId).Name.ToString() : TEXT("N/A");
 
-			UE_LOG(LogGarbage, Fatal, TEXT("Invalid object in GC: 0x%016llx, ReferencingObject: %s, TokenId %s (%d)"),
+			UE_LOG(LogGarbage, Fatal, TEXT("Invalid object in GC: 0x%016llx, ReferencingObject: %s, MemberId %s (%d)"),
 				(int64)(PTRINT)Object,
 				ReferencingObject ? *ReferencingObject->GetFullName() : TEXT("N/A"),
-				*TokenName, TokenId.AsPrintableIndex());
+				*MemberName, MemberId.AsPrintableIndex());
 		}
 	}
 #endif // ENABLE_GC_OBJECT_CHECKS
@@ -1456,155 +1135,341 @@ struct FReferenceMetadata
 
 //////////////////////////////////////////////////////////////////////////
 
-class FReferenceBatcherBase
-{
-protected:
-
-	template<class ReferenceType, uint32 Capacity, uint32 PrefetchCapacity = 0>
-	struct TBatchQueue
-	{
-		uint32 Num = 0;
-		ReferenceType Entries[Capacity + PrefetchCapacity];
-
-		FORCEINLINE static constexpr uint32 Max() { return Capacity; }
-		FORCEINLINE bool IsFull() const { return Num == Capacity; }
-		FORCEINLINE void Push(ReferenceType Entry) { Entries[Num++] = Entry; }
-		FORCEINLINE uint32 Slack() const { return Capacity - Num; }
-		FORCEINLINE ReferenceType& operator[](uint32 Idx) { return Entries[Idx]; }
-
-		TBatchQueue() { FMemory::Memzero(Entries); } // Zero out so prefetching ahead fetches same (near) null addresses
-		~TBatchQueue() { checkf(Num == 0, TEXT("Failed to flush")); }
-	};
+// Work-stealing algorithms are O(N^2), everyone steals from everyone.
+// Might want to improve that before going much wider.
+static constexpr int32 MaxWorkers = 16;
 
-	static constexpr uint32 ArrayBatchSize = 32;
-	static constexpr uint32 UnvalidatedBatchSize = 32;
-	static constexpr uint32 ValidatedBatchSize = 1024;
-	static constexpr uint32 ValidatedPrefetchAhead = 64;
+int32 GetNumSlowAROs();
 
-	struct FValidatedBitmask
+// Allocates and caches 4K pages used for temporary allocations during GC
+//
+// Currently doesn't free pages until shutdown
+class FPageAllocator
+{
+	struct alignas(PLATFORM_CACHE_LINE_SIZE) FWorkerCache
 	{
-		FORCEINLINE_DEBUGGABLE void Set(uint32 Idx, bool Value) { Words[Idx/64] |= (uint64(Value) << (Idx%64)); }
-		FORCEINLINE_DEBUGGABLE uint64 Get(uint32 Idx) { return (uint64(1) << (Idx%64)) & Words[Idx/64]; }
-		FORCEINLINE_DEBUGGABLE uint32 CountBits()
+		UE_NONCOPYABLE(FWorkerCache);
+		FWorkerCache() = default;
+		~FWorkerCache()
 		{
-			uint32 Out = 0;
-			for (uint64 Word : Words)
+			for (void* Page : MakeArrayView(Pages, Num))
 			{
-				Out += FPlatformMath::CountBits(Word);
+				FreePage(Page);
 			}
-			return Out;
 		}
 
-		FORCEINLINE_DEBUGGABLE static FValidatedBitmask Or(const FValidatedBitmask& A, const FValidatedBitmask& B)
+		bool Push(void* Page)
 		{
-			FValidatedBitmask Out;
-			for (uint32 WordIdx = 0; WordIdx < NumWords; ++WordIdx)
+			check(Page);
+			if (Num < Capacity)
 			{
-				Out.Words[WordIdx] = A.Words[WordIdx] | B.Words[WordIdx];
+				Pages[Num++] = Page;
+				return true;	
 			}
-			return Out;
+			return false;
 		}
 
-		FORCEINLINE_DEBUGGABLE static FValidatedBitmask And(const FValidatedBitmask& A, const FValidatedBitmask& B)
+		void* Pop()
 		{
-			FValidatedBitmask Out;
-			for (uint32 WordIdx = 0; WordIdx < NumWords; ++WordIdx)
-			{
-				Out.Words[WordIdx] = A.Words[WordIdx] & B.Words[WordIdx];
-			}
-			return Out;
+			return Num > 0 ? Pages[--Num] : nullptr;
 		}
 
+		static constexpr uint32 Capacity = 512 / MaxWorkers;
+		uint32 Num = 0;
+		void* Pages[Capacity] = {};
+	};	
 
-		static const uint32 NumWords = 1 + (UnvalidatedBatchSize - 1) / 64; //-V1064
-		uint64 Words[NumWords] = {};
-	};
-};
-
-// TReferenceBatcher helpers
-FORCEINLINE UObject*						GetObject(FImmutableReference In) { return In.Object; }
-FORCEINLINE UObject*						GetObject(FMutableReference In) { return *In.Object; }
-FORCEINLINE void							PrefetchObjectPointer(FImmutableReference) {}
-FORCEINLINE void							PrefetchObjectPointer(FMutableReference In) { FPlatformMisc::Prefetch(In.Object); }
-FORCEINLINE FImmutableReference				ToImmutableReference(FMutableReference In) { return {*In.Object}; }
-FORCEINLINE FImmutableReference				ToResolvedReference(FImmutableReference In) { return In; }
-FORCEINLINE FResolvedMutableReference		ToResolvedReference(FMutableReference In) { return { ToImmutableReference(In), In.Object }; }
-template<class RefType> RefType				ToReference(FMutableReference In);
-template<> FORCEINLINE FImmutableReference	ToReference(FMutableReference In) { return ToImmutableReference(In); }
-template<> FORCEINLINE FMutableReference	ToReference(FMutableReference In) { return In; }
-template<class RefType> RefType				MakeReference(UObject*& Object);
-template<> FORCEINLINE FImmutableReference	MakeReference(UObject*& Object) { return {Object}; }
-template<> FORCEINLINE FMutableReference	MakeReference(UObject*& Object) { return {&Object}; }
-
-/**
- * Queues up memory indirection to avoid synchronous cache misses by prefetching data just in time.
- *
- * Uses bounded queues to avoid branching, control temp memory consumption, allow compiler to generate vectorized code
- * and make life easier for the branch predictor.
- * 
- * Array queue feed into "unvalidated" references queue, which feed into validated references queue, which gets processed when full.
- * 
- * Unvalidated means references that might be nullptr or part of the permanent object pool, neither of which require processing.
- */
-template <class UnvalidatedReferenceType, class ValidatedReferenceType, class ProcessorType>
-class TReferenceBatcher : public FReferenceBatcherBase
-{
-	FWorkerContext& Context;
-	const FPermanentObjectPoolExtents PermanentPool;
-
-	alignas (PLATFORM_CACHE_LINE_SIZE)	TBatchQueue<FReferenceArray, ArrayBatchSize>									UnvalidatedArrays;
-	alignas (PLATFORM_CACHE_LINE_SIZE)	TBatchQueue<UnvalidatedReferenceType, UnvalidatedBatchSize>						UnvalidatedReferences;
-	alignas (PLATFORM_CACHE_LINE_SIZE)	TBatchQueue<ValidatedReferenceType, ValidatedBatchSize, ValidatedPrefetchAhead>	ValidatedReferences;
-
-public:
-	TReferenceBatcher(FWorkerContext& InContext) : Context(InContext) {}
-	
-	FORCEINLINE_DEBUGGABLE void PushArray(FReferenceArray UnvalidatedArray)
+	struct FSharedCache
 	{
-		FPlatformMisc::Prefetch(UnvalidatedArray.Array);
-		UnvalidatedArrays.Push({ UnvalidatedArray });
-		if (UnvalidatedArrays.IsFull())
+		UE_NONCOPYABLE(FSharedCache);
+		FSharedCache() = default;
+		~FSharedCache()
 		{
-			DrainArraysFull();
+			for (void* Page : Pages)
+			{
+				FreePage(Page);
+			}
 		}
-	}
 
-	FORCEINLINE_DEBUGGABLE void PushReference(UnvalidatedReferenceType UnvalidatedReference)
-	{
-		UnvalidatedReferences.Push({ UnvalidatedReference });
-		if (UnvalidatedReferences.IsFull())
+		void Push(void* Page)
 		{
-			DrainUnvalidatedFull();
+			FScopeLock Scope(&Lock);
+			Pages.Add(Page);
 		}
-	}
-	
-	FORCEINLINE_DEBUGGABLE void PushReferences(TArrayView<UObject*> Range)
-	{
-		if (Range.Num())
+
+		void PushSurplus(FWorkerCache& Worker)
 		{
-			// Fill up UnvalidatedReferences slack and drain it to avoid inserting an element at a time and testing if full
-			for (uint32 Slack = UnvalidatedReferences.Slack(); (uint32)Range.Num() >= Slack; Slack = UnvalidatedBatchSize) //-V1021
+			// Each worker context also needs a page for FWorkBlockifier and FStructBlockifer
+			const uint32 NumWorkerInitialPages = 2 + GetNumSlowAROs();
+
+			if (Worker.Num > NumWorkerInitialPages)
 			{
-				QueueUnvalidated(Range.Left(Slack));
-				Range.RightChopInline(Slack);
-				check(UnvalidatedReferences.IsFull());
-				DrainUnvalidatedFull();
+				FScopeLock Scope(&Lock);
+				Pages.Append(Worker.Pages + NumWorkerInitialPages, Worker.Num - NumWorkerInitialPages);
+				Worker.Num = NumWorkerInitialPages;
 			}
-
-			QueueUnvalidated(Range);
-			check(!UnvalidatedReferences.IsFull());
 		}
-	}
 
-	FORCEINLINE_DEBUGGABLE void FlushQueues()
-	{
-		// Slower path with dynamic Num that the compiler won't unroll or vectorize, unlike DrainXyzFull methods
+		void* Pop()
+		{
+			FScopeLock Scope(&Lock);
+			return Pages.IsEmpty() ? nullptr : Pages.Pop(/* shrink */ false);
+		}
+
+		FCriticalSection Lock;
+		TArray<void*> Pages;
+	};	
+
+	FSharedCache SharedCache;
+	FWorkerCache WorkerCaches[MaxWorkers];
+
+public:
+	static constexpr uint64 PageSize = 4096;
+
+	void* AllocatePage(int32 WorkerIdx)
+	{
+		check(WorkerIdx >=0 && WorkerIdx < MaxWorkers);
+		if (void* WorkerPage = WorkerCaches[WorkerIdx].Pop())
+		{
+			check(IsValidPage(WorkerPage));
+			return WorkerPage;
+		}
+		else if (void* SharedPage = SharedCache.Pop())
+		{
+			check(IsValidPage(SharedPage));
+			return SharedPage;
+		}
+
+		void* NewPage = FMemory::Malloc(PageSize, PageSize);
+		return NewPage;
+	}
+	
+	void ReturnWorkerPage(int32 WorkerIdx, void* Page)
+	{
+		check(IsValidPage(Page));
+		if (!WorkerCaches[WorkerIdx].Push(Page))
+		{
+			FreePage(Page);
+		}
+	}
+
+	void ReturnSharedPage(void* Page)
+	{
+		check(IsValidPage(Page));
+		SharedCache.Push(Page);
+	}
+
+	static void FreePage(void* Page)
+	{
+		check(IsValidPage(Page));
+		FMemory::Free(Page);
+	}
+
+	uint64 CountBytes() const
+	{
+		uint64 NumPages = SharedCache.Pages.Num();
+		for (const FWorkerCache& WorkerCache : WorkerCaches)
+		{
+			NumPages += WorkerCache.Num;
+		}
+		return NumPages * PageSize;
+	}
+
+	static bool IsValidPage(void* Page)
+	{
+		return Page && UPTRINT(Page) % PageSize == 0;
+	}
+};
+
+static FPageAllocator GScratchPages;
+
+////////////// TReferenceBatcher and FStructBatcher helpers //////////////
+
+class FBatcherBase
+{
+protected:
+	template<class EntryType, uint32 Capacity, uint32 PrefetchCapacity = 0>
+	struct TBatchQueue
+	{
+		uint32 Num = 0;
+		EntryType Entries[Capacity + PrefetchCapacity];
+
+		FORCEINLINE static constexpr uint32 Max() { return Capacity; }
+		FORCEINLINE bool IsFull() const { return Num == Capacity; }
+		FORCEINLINE void Push(EntryType Entry) { Entries[Num++] = Entry; }
+		FORCEINLINE uint32 Slack() const { return Capacity - Num; }
+		FORCEINLINE EntryType& operator[](uint32 Idx) { return Entries[Idx]; }
+
+		TBatchQueue() { FMemory::Memzero(Entries); } // Zero out so prefetching ahead fetches same (near) null addresses
+		~TBatchQueue() { checkf(Num == 0, TEXT("Failed to flush")); }
+	};
+
+	static constexpr uint32 ArrayBatchSize = 32;
+	static constexpr uint32 UnvalidatedBatchSize = 32;
+	static constexpr uint32 ValidatedBatchSize = 1024;
+	static constexpr uint32 ValidatedPrefetchAhead = 64;
+
+	struct FValidatedBitmask
+	{
+		FORCEINLINE_DEBUGGABLE void Set(uint32 Idx, bool Value) { Words[Idx/64] |= (uint64(Value) << (Idx%64)); }
+		FORCEINLINE_DEBUGGABLE uint64 Get(uint32 Idx) { return (uint64(1) << (Idx%64)) & Words[Idx/64]; }
+		FORCEINLINE_DEBUGGABLE uint32 CountBits()
+		{
+			uint32 Out = 0;
+			for (uint64 Word : Words)
+			{
+				Out += FPlatformMath::CountBits(Word);
+			}
+			return Out;
+		}
+
+		FORCEINLINE_DEBUGGABLE static FValidatedBitmask Or(const FValidatedBitmask& A, const FValidatedBitmask& B)
+		{
+			FValidatedBitmask Out;
+			for (uint32 WordIdx = 0; WordIdx < NumWords; ++WordIdx)
+			{
+				Out.Words[WordIdx] = A.Words[WordIdx] | B.Words[WordIdx];
+			}
+			return Out;
+		}
+
+		FORCEINLINE_DEBUGGABLE static FValidatedBitmask And(const FValidatedBitmask& A, const FValidatedBitmask& B)
+		{
+			FValidatedBitmask Out;
+			for (uint32 WordIdx = 0; WordIdx < NumWords; ++WordIdx)
+			{
+				Out.Words[WordIdx] = A.Words[WordIdx] & B.Words[WordIdx];
+			}
+			return Out;
+		}
+
+
+		static const uint32 NumWords = 1 + (UnvalidatedBatchSize - 1) / 64; //-V1064
+		uint64 Words[NumWords] = {};
+	};
+};
+
+FORCEINLINE UObject*						GetObject(FImmutableReference In) { return In.Object; }
+FORCEINLINE UObject*						GetObject(FMutableReference In) { return *In.Object; }
+FORCEINLINE void							PrefetchObjectPointer(FImmutableReference) {}
+FORCEINLINE void							PrefetchObjectPointer(FMutableReference In) { FPlatformMisc::Prefetch(In.Object); }
+FORCEINLINE FImmutableReference				ToImmutableReference(FMutableReference In) { return {*In.Object}; }
+FORCEINLINE FImmutableReference				ToResolvedReference(FImmutableReference In) { return In; }
+FORCEINLINE FResolvedMutableReference		ToResolvedReference(FMutableReference In) { return { ToImmutableReference(In), In.Object }; }
+template<class RefType> RefType				ToReference(FMutableReference In);
+template<> FORCEINLINE FImmutableReference	ToReference(FMutableReference In) { return ToImmutableReference(In); }
+template<> FORCEINLINE FMutableReference	ToReference(FMutableReference In) { return In; }
+template<class RefType> RefType				MakeReference(UObject*& Object);
+template<> FORCEINLINE FImmutableReference	MakeReference(UObject*& Object) { return {Object}; }
+template<> FORCEINLINE FMutableReference	MakeReference(UObject*& Object) { return {&Object}; }
+
+using FStructArray = Private::FStructArray;
+using FStridedReferenceArray = Private::FStridedReferenceArray;
+using FStridedReferenceView = Private::FStridedReferenceView;
+
+FORCEINLINE_DEBUGGABLE static TArrayView<UObject*> Split(/* in-out */ TArrayView<UObject*>& Body, int32 Idx)
+{
+	check(Idx <= Body.Num());
+	TArrayView<UObject*> Head(Body.GetData(), Idx);
+	Body = TArrayView<UObject*>(Body.GetData() + Idx, Body.Num() - Idx);
+	return Head;
+}
+
+FORCEINLINE_DEBUGGABLE static FStridedReferenceView Split(/* in-out */ FStridedReferenceView& Body, int32 Idx)
+{
+	check(Idx <= Body.Num);
+	check(Body.Stride > 0);
+	FStridedReferenceView Head = Body;
+	Head.Num = Idx;
+	Body.Data += Idx * Body.Stride;
+	Body.Num -= Idx;
+	return Head;
+}
+
+/////////////////////////// TReferenceBatcher ///////////////////////////
+
+/**
+ * Queues up memory indirection to avoid synchronous cache misses by prefetching data just in time.
+ *
+ * Uses bounded queues to avoid branching, control temp memory consumption, allow compiler to generate vectorized code
+ * and make life easier for the branch predictor.
+ * 
+ * Array queue feed into "unvalidated" references queue, which feed into validated references queue, which gets processed when full.
+ * 
+ * Unvalidated means references that might be nullptr or part of the permanent object pool, neither of which require processing.
+ */
+template <class UnvalidatedReferenceType, class ValidatedReferenceType, class ProcessorType>
+class TReferenceBatcher : public FBatcherBase
+{
+	FWorkerContext& Context;
+	const FPermanentObjectPoolExtents PermanentPool;
+
+	alignas (PLATFORM_CACHE_LINE_SIZE)	TBatchQueue<FReferenceArray, ArrayBatchSize>									UnvalidatedArrays;			// Drains to UnvalidatedReferences
+	alignas (PLATFORM_CACHE_LINE_SIZE)	TBatchQueue<FStridedReferenceArray, ArrayBatchSize>								UnvalidatedStridedArrays;	// Drains to UnvalidatedReferences
+	alignas (PLATFORM_CACHE_LINE_SIZE)	TBatchQueue<UnvalidatedReferenceType, UnvalidatedBatchSize>						UnvalidatedReferences;		// Drains to ValidatedReferences
+	alignas (PLATFORM_CACHE_LINE_SIZE)	TBatchQueue<ValidatedReferenceType, ValidatedBatchSize, ValidatedPrefetchAhead>	ValidatedReferences;		// Drains to ProcessorType::HandleBatchedReference
+
+public:
+	TReferenceBatcher(FWorkerContext& InContext) : Context(InContext) {}
+	
+	FORCEINLINE_DEBUGGABLE void PushArray(TArray<UObject*>& Array)
+	{
+		FPlatformMisc::Prefetch(&Array);
+		UnvalidatedArrays.Push({&Array});
+		if (UnvalidatedArrays.IsFull())
+		{
+			DrainArraysFull();
+		}
+	}
+
+	FORCEINLINE_DEBUGGABLE void PushArray(FStridedReferenceArray StridedArray)
+	{
+		FPlatformMisc::Prefetch(StridedArray.Array);
+		UnvalidatedStridedArrays.Push(StridedArray);
+		if (UnvalidatedStridedArrays.IsFull())
+		{
+			DrainStridedArraysFull();
+		}
+	}
+
+	template<class ViewType>
+	FORCEINLINE_DEBUGGABLE void PushReferenceView(ViewType View)
+	{
+		if (GetNum(View))
+		{
+			// Fill up UnvalidatedReferences slack and drain it to avoid inserting an element at a time and testing if full
+			for (uint32 Slack = UnvalidatedReferences.Slack(); (uint32)GetNum(View) >= Slack; Slack = UnvalidatedBatchSize) //-V1021
+			{
+				QueueUnvalidated(Split(/* in-out */ View, Slack));
+				check(UnvalidatedReferences.IsFull());
+				DrainUnvalidatedFull();
+			}
+
+			QueueUnvalidated(View);
+			check(!UnvalidatedReferences.IsFull());
+		}
+	}
+
+	FORCEINLINE_DEBUGGABLE void PushReference(UnvalidatedReferenceType UnvalidatedReference)
+	{
+		UnvalidatedReferences.Push(UnvalidatedReference);
+		if (UnvalidatedReferences.IsFull())
+		{
+			DrainUnvalidatedFull();
+		}
+	}
+
+	FORCEINLINE_DEBUGGABLE void FlushQueues()
+	{
+		// Slower path with dynamic Num that the compiler won't unroll or vectorize, unlike DrainXyzFull methods
 		DrainArrays(UnvalidatedArrays.Num);
+		DrainStridedArrays(UnvalidatedStridedArrays.Num);
 		DrainUnvalidated(UnvalidatedReferences.Num);
 		DrainValidated(ValidatedReferences.Num);
 	}
 
 private:
+
 	// Process queued arrays of references and feed contents to unvalidated queue
 	FORCEINLINE_DEBUGGABLE void DrainArrays(uint32 Num)
 	{
@@ -1617,29 +1482,49 @@ private:
 
 		for (FReferenceArray Entry : MakeArrayView(UnvalidatedArrays.Entries, Num))
 		{
-			PushReferences(*Entry.Array);
+			PushReferenceView(MakeArrayView(*Entry.Array));
 		}
 
 		UnvalidatedArrays.Num = 0;
 	}
 
-	// Helps generate vectorized code with static Num
-	FORCENOINLINE void DrainArraysFull()
+	// Process queued strided arrays of references and feed contents to unvalidated queue
+	FORCEINLINE_DEBUGGABLE void DrainStridedArrays(uint32 Num)
 	{
-		DrainArrays(UnvalidatedArrays.Max());
+		check(Num <= ArrayBatchSize);
+
+		FStridedReferenceView Views[ArrayBatchSize];
+		for (uint32 Idx = 0; Idx < Num; ++Idx)
+		{
+			Views[Idx] = ToView(UnvalidatedStridedArrays.Entries[Idx]);
+		}
+
+		for (FStridedReferenceView View : MakeArrayView(Views, Num))
+		{
+			FPlatformMisc::Prefetch(View.Data);
+		}
+		
+		for (FStridedReferenceView View : MakeArrayView(Views, Num))
+		{
+			PushReferenceView(View);
+		}
+
+		UnvalidatedStridedArrays.Num = 0;
 	}
 
-	FORCEINLINE_DEBUGGABLE void QueueUnvalidated(TArrayView<UObject*> Range)
+
+	template<class ViewType>
+	FORCEINLINE_DEBUGGABLE void QueueUnvalidated(ViewType View)
 	{
-		check((uint32)Range.Num() <= UnvalidatedReferences.Slack());
+		check((uint32)GetNum(View) <= UnvalidatedReferences.Slack());
 
-		UnvalidatedReferenceType* Unvalidated = UnvalidatedReferences.Entries + UnvalidatedReferences.Num;
-		for (int32 Idx = 0; Idx < Range.Num(); ++Idx)
+		UnvalidatedReferenceType* UnvalidatedIt = UnvalidatedReferences.Entries + UnvalidatedReferences.Num;
+		for (UObject*& Object : View)
 		{
-			Unvalidated[Idx] = { MakeReference<UnvalidatedReferenceType>(Range[Idx]) };
+			*UnvalidatedIt++ = MakeReference<UnvalidatedReferenceType>(Object);
 		}
 
-		UnvalidatedReferences.Num += Range.Num();
+		UnvalidatedReferences.Num += GetNum(View);
 	}
 
 	// Validate all unvalidated references and feed referencing in need of processing to validated queue
@@ -1680,12 +1565,6 @@ private:
 		UnvalidatedReferences.Num = 0;
 	}
 
-	// Helps generate vectorized code with static Num
-	FORCENOINLINE void DrainUnvalidatedFull()
-	{
-		DrainUnvalidated(UnvalidatedReferences.Max());
-	}
-
 	FORCEINLINE_DEBUGGABLE void QueueValidReferences(uint32 NumToAppend, FValidatedBitmask Validations, uint32& InOutIdx)
 	{
 		checkSlow(NumToAppend <= Validations.CountBits());
@@ -1779,247 +1658,367 @@ private:
 		ValidatedReferences.Num = 0;
 	}
 
-	// Helps generate code with static Num
-	FORCENOINLINE void DrainValidatedFull()
-	{
-		DrainValidated(ValidatedReferences.Max());
-	}
+	// Helps generate vectorized code with static Num
+
+	FORCENOINLINE void DrainArraysFull()			{ DrainArrays(UnvalidatedArrays.Max()); }
+	FORCENOINLINE void DrainStridedArraysFull()		{ DrainStridedArrays(UnvalidatedStridedArrays.Max()); }
+	FORCENOINLINE void DrainUnvalidatedFull()		{ DrainUnvalidated(UnvalidatedReferences.Max()); }
+	FORCENOINLINE void DrainValidatedFull()			{ DrainValidated(ValidatedReferences.Max()); }
 };
 
-//////////////////////////////////////////////////////////////////////////
 
-FORCEINLINE_DEBUGGABLE static void PadObjects(UObject* LastObject, TArrayView<UObject*> Padding)
+///////////////////////// FStructBatcher + helpers ///////////////////////
+
+// Fixed block of array-of-structs waiting to be processed
+struct FStructArrayBlock
 {
-	check(LastObject->IsValidLowLevel());
+	static constexpr uint32 Lookahead = 2;
+	static constexpr uint32 Capacity = (FPageAllocator::PageSize - sizeof(FStructArrayBlock*)) / sizeof(FStructArray) - Lookahead;
 
-	for (UObject*& ObjectPastEnd : Padding)
+	FStructArray Data[Capacity + Lookahead];
+	FStructArrayBlock* NextFull;
+
+	FStructArray* GetPadEnd() { return Data + Capacity; }
+};
+
+// Unbounded queue of validated / non-empty AoS
+class FStructBlockifier
+{
+	UE_NONCOPYABLE(FStructBlockifier);
+public:
+	explicit FStructBlockifier(int32 WorkerIdx)
+	: WorkerIndex(WorkerIdx)
 	{
-		ObjectPastEnd = LastObject;
+		AllocateWipBlock(nullptr);
 	}
-}
 
-void PadObjectArray(TArray<UObject*>& Objects)
-{
-	if (int32 Num = Objects.Num())
+	~FStructBlockifier()
 	{
-		Objects.Reserve(Num + ObjectLookahead);
-		PadObjects(Objects.Last(), MakeArrayView(Objects.GetData() + Num, ObjectLookahead));
+		check(!CanPop());
+		FreeBlock(Wip);
+	}
+	
+	void Push(FStructArray AoS)
+	{
+		*WipIt = AoS;
+		if (++WipIt == Wip->Data+ FStructArrayBlock::Capacity)
+		{
+			AllocateWipBlock(Wip);
+		}
 	}
-}
 
-FORCEINLINE_DEBUGGABLE void PadBlock(FWorkBlock& Block)
-{
-	PadObjects(Block.GetObjects().Last(), Block.GetPadding());
-}
+	FStructArray* PushUninitialized(int32 Num)
+	{
+		check(Num > 0 && Num <= GetSlack());
 
-//////////////////////////////////////////////////////////////////////////
+		FStructArray* Out = WipIt;
 
-// Work-stealing algorithms are O(N^2), everyone steals from everyone.
-// Might want to improve that before going much wider.
-static constexpr int32 MaxWorkers = 16;
+		WipIt += Num;
+		if (WipIt == Wip->Data + FStructArrayBlock::Capacity)
+		{
+			AllocateWipBlock(Wip);
+		}
+		return Out;
+	}
 
-class FWorkerIndexAllocator
-{
-	std::atomic<uint64> Used{0};
-	static_assert(MaxWorkers <= 64, "Currently supports single uint64 word");
-public:
-	int32 AllocateChecked()
+	int32 GetSlack() const
 	{
-		while (true)
-		{
-			uint64 UsedNow = Used.load(std::memory_order_relaxed);
-			uint64 FreeIndex = FPlatformMath::CountTrailingZeros64(~UsedNow);
-			checkf(FreeIndex < MaxWorkers, TEXT("Exceeded max active GC worker contexts"));
+		return IntCastChecked<int32>(Wip->Data + FStructArrayBlock::Capacity - WipIt);
+	}
 
-			uint64 Mask = uint64(1) << FreeIndex;
-			if ((Used.fetch_or(Mask) & Mask) == 0)
-			{
-				return static_cast<int32>(FreeIndex);
-			}
+	bool CanPop() const
+	{
+		return (PartialNum() > 0) | !!(Wip->NextFull);
+	}
+
+	FStructArrayBlock* PopBlock(int32& OutNum)
+	{
+		if (FStructArrayBlock* Full = Wip->NextFull)
+		{
+			FPlatformMisc::Prefetch(Full->NextFull);
+			FPlatformMisc::Prefetch(Full->Data);
+			PadBlock(Full, FStructArrayBlock::Capacity);
+			Wip->NextFull = Full->NextFull;
+			OutNum = FStructArrayBlock::Capacity;
+			return Full;
 		}
+		
+		OutNum = PartialNum();
+		check(OutNum > 0);
+		PadBlock(Wip, OutNum);
+		FStructArrayBlock* Out = Wip;
+		AllocateWipBlock(nullptr);
+		return Out;
 	}
 
-	void FreeChecked(int32 Index)
+	void FreeBlock(FStructArrayBlock* Block)
 	{
-		check(Index >= 0 && Index < MaxWorkers);
-		uint64 Mask = uint64(1) << Index;
-		uint64 Old = Used.fetch_and(~Mask);
-		checkf(Old & Mask, TEXT("Index already freed"));
+		GScratchPages.ReturnWorkerPage(WorkerIndex, Block);
 	}
-};
-static FWorkerIndexAllocator GWorkerIndices;
 
-FWorkerContext::FWorkerContext()
-{
-	AllocateWorkerIndex();
-	ObjectsToSerialize.Init();
-}
+private:
+	FStructArray* WipIt; // Wip->Data cursor
+	FStructArrayBlock* Wip;
+	const int32 WorkerIndex;
 	
-FWorkerContext::~FWorkerContext()
+	int32 PartialNum() const
+	{
+		check(WipIt >= Wip->Data && WipIt < Wip->Data + FStructArrayBlock::Capacity);
+		return static_cast<int32>(WipIt - Wip->Data);
+	}
+
+	static void PadBlock(FStructArrayBlock* Block, int32 PadIdx)
+	{
+		check(PadIdx <= FStructArrayBlock::Capacity);
+		FMemory::Memzero(Block->Data + PadIdx, sizeof(FStructArray) * FStructArrayBlock::Lookahead);
+	}
+
+	void AllocateWipBlock(FStructArrayBlock* NextFull)
+	{	
+		static_assert(sizeof(FStructArrayBlock) <= FPageAllocator::PageSize, "Block must fit in page");
+	
+		Wip = reinterpret_cast<FStructArrayBlock*>(GScratchPages.AllocatePage(WorkerIndex));
+		//Wip = new (GScratchPages.AllocatePage(WorkerIndex)) FStructArrayBlock;
+		Wip->NextFull = NextFull;
+		WipIt = Wip->Data;
+	}
+};
+
+struct FSparseStructArray
 {
-	FreeWorkerIndex();
-}
+	FSchemaView Schema{NoInit};
+	FScriptSparseArray* Array;
+	int32 Num;
+	uint32 Stride;
+};
 
-void FWorkerContext::AllocateWorkerIndex()
+struct FStructArrayWithoutStride
 {
-	check(GetWorkerIndex() == INDEX_NONE);
-	ObjectsToSerialize.SetWorkerIndex(GWorkerIndices.AllocateChecked());
-}
+	FSchemaView Schema{NoInit};
+	uint8* Data;
+	int32 Num;
+};
 
-void FWorkerContext::FreeWorkerIndex()
+static FStructArray CopyAndLoadStride(FStructArrayWithoutStride In)
 {
-	if (GetWorkerIndex() != INDEX_NONE)
-	{
-		GWorkerIndices.FreeChecked(GetWorkerIndex());
-		ObjectsToSerialize.SetWorkerIndex(INDEX_NONE);
-	}
+	return {In.Schema, In.Data, In.Num, In.Schema.GetStructStride() };
 }
 
-//////////////////////////////////////////////////////////////////////////
+// Queues up arrays-of-structs and maps/sets as sparse array-of-structs for later processing by ProcessStructs
+class FStructBatcher : public FBatcherBase
+{
+	alignas (PLATFORM_CACHE_LINE_SIZE)	TBatchQueue<FSparseStructArray, ArrayBatchSize>			UnvalidatedSparseStructArrays;	// Drains to ValidatedStructArrays
+	alignas (PLATFORM_CACHE_LINE_SIZE)	TBatchQueue<FStructArrayWithoutStride, ArrayBatchSize>	UnvalidatedStructArrays;		// Drains to ValidatedStructArrays
+	alignas (PLATFORM_CACHE_LINE_SIZE)	FStructBlockifier										ValidatedStructArrays;			// Drained by ProcessStructs, which feed back into TReferenceBatcher
 
-int32 GetNumSlowAROs();
+public:
+	explicit FStructBatcher(int32 WorkerIdx) : ValidatedStructArrays(WorkerIdx) {}
 
-// Allocates and caches 4K pages used for temporary allocations during GC
-//
-// Currently doesn't free pages until shutdown
-class FPageAllocator
-{
-	struct alignas(PLATFORM_CACHE_LINE_SIZE) FWorkerCache
+	FORCEINLINE_DEBUGGABLE void PushSparseStructArray(FSchemaView Schema, FScriptSparseArray& Array)
 	{
-		UE_NONCOPYABLE(FWorkerCache);
-		FWorkerCache() = default;
-		~FWorkerCache()
-		{
-			for (void* Page : MakeArrayView(Pages, Num))
-			{
-				FreePage(Page);
-			}
-		}
+		check(!Schema.IsEmpty());
 
-		bool Push(void* Page)
+		UnvalidatedSparseStructArrays.Push({ Schema, &Array, Array.Num(), Schema.GetStructStride() });
+		if (UnvalidatedSparseStructArrays.IsFull())
 		{
-			check(Page);
-			if (Num < Capacity)
-			{
-				Pages[Num++] = Page;
-				return true;	
-			}
-			return false;
+			DrainSparseStructArraysFull();
 		}
+	}
+	
+	FORCEINLINE_DEBUGGABLE void PushStructArray(FSchemaView Schema, uint8* Data, int32 Num)
+	{
+		check(!Schema.IsEmpty());
 
-		void* Pop()
+		UnvalidatedStructArrays.Push({Schema, Data, Num});
+		if (UnvalidatedStructArrays.IsFull())
 		{
-			return Num > 0 ? Pages[--Num] : nullptr;
+			DrainStructArraysFull();
 		}
+	}
 
-		static constexpr uint32 Capacity = 512 / MaxWorkers;
-		uint32 Num = 0;
-		void* Pages[Capacity] = {};
-	};	
+	FORCEINLINE_DEBUGGABLE void FlushBoundedQueues()
+	{
+		DrainSparseStructArrays(UnvalidatedSparseStructArrays.Num);
+		DrainStructArrays(UnvalidatedStructArrays.Num);
+	}
 
-	struct FSharedCache
+	FStructBlockifier& GetUnboundedQueue() { return ValidatedStructArrays;	}
+private:
+	FORCEINLINE_DEBUGGABLE void DrainSparseStructArrays(uint32 Num)
 	{
-		UE_NONCOPYABLE(FSharedCache);
-		FSharedCache() = default;
-		~FSharedCache()
-		{
-			for (void* Page : Pages)
-			{
-				FreePage(Page);
-			}
-		}
+		check(Num <= ArrayBatchSize);
 
-		void Push(void* Page)
+		if (Num == 0)
 		{
-			FScopeLock Scope(&Lock);
-			Pages.Add(Page);
+			return;
 		}
 
-		void PushSurplus(FWorkerCache& Worker)
+		// This can likely be optimized further, e.g. prefetching, intermediate queues, FValidatedBitmask for IsCompact() etc
+		for (FSparseStructArray Entry : MakeArrayView(UnvalidatedSparseStructArrays.Entries, Num))
 		{
-			// Each worker context also owns a page via FWorkBlockifier::Wip
-			const uint32 NumWorkerInitialPages = 1 + GetNumSlowAROs();
-
-			if (Worker.Num > NumWorkerInitialPages)
+			if (Entry.Num)
 			{
-				FScopeLock Scope(&Lock);
-				Pages.Append(Worker.Pages + NumWorkerInitialPages, Worker.Num - NumWorkerInitialPages);
-				Worker.Num = NumWorkerInitialPages;
-			}
-		}
-
-		void* Pop()
-		{
-			FScopeLock Scope(&Lock);
-			return Pages.IsEmpty() ? nullptr : Pages.Pop(/* shrink */ false);
-		}
-
-		FCriticalSection Lock;
-		TArray<void*> Pages;
-	};	
+				FScriptSparseArray& Array = *Entry.Array;
+				FStructArray Slice = { Entry.Schema, Private::GetSparseData(Array), Entry.Num, Entry.Stride };
+				if (Array.IsCompact())
+				{
+					ValidatedStructArrays.Push(Slice);
+				}
+				else
+				{
+					Slice.Num = 0;
+					for (uint32 Idx = 0, NumIndices = Array.GetMaxIndex(); Idx < NumIndices; ++Idx)
+					{
+						if (Array.IsAllocated(Idx))
+						{
+							++Slice.Num;
+						}
+						else
+						{
+							if (Slice.Num > 0)
+							{
+								ValidatedStructArrays.Push(Slice);
+								Slice.Data += Slice.Stride * Slice.Num;
+								Slice.Num = 0;
+							}
 
-	FSharedCache SharedCache;
-	FWorkerCache WorkerCaches[MaxWorkers];
+							Slice.Data += Slice.Stride;
+						}
+					}
 
-public:
-	static constexpr uint64 PageSize = 4096;
+					if (Slice.Num > 0)
+					{
+						ValidatedStructArrays.Push(Slice);
+					}
+				}
+			}
+		}
 
-	void* AllocatePage(int32 WorkerIdx)
+		UnvalidatedSparseStructArrays.Num = 0;
+	}
+	
+	FORCEINLINE_DEBUGGABLE void DrainStructArrays(uint32 Num)
 	{
-		check(WorkerIdx >=0 && WorkerIdx < MaxWorkers);
-		if (void* WorkerPage = WorkerCaches[WorkerIdx].Pop())
+		check(Num <= UnvalidatedStructArrays.Max());
+		
+		if (Num == 0)
 		{
-			check(IsValidPage(WorkerPage));
-			return WorkerPage;
+			return;
 		}
-		else if (void* SharedPage = SharedCache.Pop())
+		
+		uint32 ValidIndices[UnvalidatedStructArrays.Max()];
+		uint32* ValidIt = ValidIndices;
+		for (uint32 Idx = 0; Idx < Num; ++Idx)
 		{
-			check(IsValidPage(SharedPage));
-			return SharedPage;
+			bool bIsValid = UnvalidatedStructArrays.Entries[Idx].Num > 0;
+			*ValidIt = Idx;
+			ValidIt += bIsValid;
 		}
 
-		void* NewPage = FMemory::Malloc(PageSize, PageSize);
-		return NewPage;
-	}
-	
-	void ReturnWorkerPage(int32 WorkerIdx, void* Page)
-	{
-		check(IsValidPage(Page));
-		if (!WorkerCaches[WorkerIdx].Push(Page))
+		for (uint32* IdxIt = ValidIndices, *IdxEnd = ValidIt; IdxIt != IdxEnd; )
 		{
-			FreePage(Page);
+			uint32* NextIdxEnd = FMath::Min(IdxIt + ValidatedStructArrays.GetSlack(), IdxEnd);
+			FStructArray* OutIt = ValidatedStructArrays.PushUninitialized(static_cast<int32>(NextIdxEnd - IdxIt));
+			do
+			{
+				*OutIt++ = CopyAndLoadStride(UnvalidatedStructArrays.Entries[*IdxIt++]);
+			}
+			while (IdxIt != NextIdxEnd);
 		}
+
+		UnvalidatedStructArrays.Num = 0;
 	}
 
-	void ReturnSharedPage(void* Page)
+	// Helps generate vectorized code with static Num
+
+	FORCENOINLINE void DrainStructArraysFull()			{ DrainStructArrays(UnvalidatedStructArrays.Max());	}
+	FORCENOINLINE void DrainSparseStructArraysFull()	{ DrainSparseStructArrays(UnvalidatedSparseStructArrays.Max());	}
+};
+
+//////////////////////////////////////////////////////////////////////////
+
+FORCEINLINE_DEBUGGABLE static void PadObjects(UObject* LastObject, TArrayView<UObject*> Padding)
+{
+	check(LastObject->IsValidLowLevel());
+
+	for (UObject*& ObjectPastEnd : Padding)
 	{
-		check(IsValidPage(Page));
-		SharedCache.Push(Page);
+		ObjectPastEnd = LastObject;
 	}
+}
 
-	static void FreePage(void* Page)
+void PadObjectArray(TArray<UObject*>& Objects)
+{
+	if (int32 Num = Objects.Num())
 	{
-		check(IsValidPage(Page));
-		FMemory::Free(Page);
+		Objects.Reserve(Num + ObjectLookahead);
+		PadObjects(Objects.Last(), MakeArrayView(Objects.GetData() + Num, ObjectLookahead));
 	}
+}
 
-	uint64 CountBytes() const
+FORCEINLINE_DEBUGGABLE void PadBlock(FWorkBlock& Block)
+{
+	PadObjects(Block.GetObjects().Last(), Block.GetPadding());
+}
+
+//////////////////////////////////////////////////////////////////////////
+
+class FWorkerIndexAllocator
+{
+	std::atomic<uint64> Used{0};
+	static_assert(MaxWorkers <= 64, "Currently supports single uint64 word");
+public:
+	int32 AllocateChecked()
 	{
-		uint64 NumPages = SharedCache.Pages.Num();
-		for (const FWorkerCache& WorkerCache : WorkerCaches)
+		while (true)
 		{
-			NumPages += WorkerCache.Num;
+			uint64 UsedNow = Used.load(std::memory_order_relaxed);
+			uint64 FreeIndex = FPlatformMath::CountTrailingZeros64(~UsedNow);
+			checkf(FreeIndex < MaxWorkers, TEXT("Exceeded max active GC worker contexts"));
+
+			uint64 Mask = uint64(1) << FreeIndex;
+			if ((Used.fetch_or(Mask) & Mask) == 0)
+			{
+				return static_cast<int32>(FreeIndex);
+			}
 		}
-		return NumPages * PageSize;
 	}
 
-	static bool IsValidPage(void* Page)
+	void FreeChecked(int32 Index)
 	{
-		return Page && UPTRINT(Page) % PageSize == 0;
+		check(Index >= 0 && Index < MaxWorkers);
+		uint64 Mask = uint64(1) << Index;
+		uint64 Old = Used.fetch_and(~Mask);
+		checkf(Old & Mask, TEXT("Index already freed"));
 	}
 };
+static FWorkerIndexAllocator GWorkerIndices;
 
-static FPageAllocator GScratchPages;
+FWorkerContext::FWorkerContext()
+{
+	AllocateWorkerIndex();
+	ObjectsToSerialize.Init();
+}
+	
+FWorkerContext::~FWorkerContext()
+{
+	FreeWorkerIndex();
+}
+
+void FWorkerContext::AllocateWorkerIndex()
+{
+	check(GetWorkerIndex() == INDEX_NONE);
+	ObjectsToSerialize.SetWorkerIndex(GWorkerIndices.AllocateChecked());
+}
+
+void FWorkerContext::FreeWorkerIndex()
+{
+	if (GetWorkerIndex() != INDEX_NONE)
+	{
+		GWorkerIndices.FreeChecked(GetWorkerIndex());
+		ObjectsToSerialize.SetWorkerIndex(INDEX_NONE);
+	}
+}
 
 // Bounded SPMC work-stealing queue
 //
@@ -2490,7 +2489,7 @@ template<class Quantity> FORCEINLINE TArrayView<UObject*> FAROQueue::Steal()	{ r
 class FSlowAROManager
 {
 public:
-	void RegisterImplementation(AROFunc ARO, EAROFlags Flags)
+	void RegisterImplementation(ObjectAROFn ARO, EAROFlags Flags)
 	{
 		check(IsInGameThread());
 		check(AROs.Find(ARO) == INDEX_NONE);
@@ -2507,7 +2506,7 @@ public:
 		ExtraSlowAROs |= decltype(ExtraSlowAROs)(EnumHasAllFlags(Flags, EAROFlags::ExtraSlow)) << Idx;
 	}
 
-	int32 FindImplementation(AROFunc ARO) const
+	int32 FindImplementation(ObjectAROFn ARO) const
 	{
 		return AROs.Find(ARO);
 	}
@@ -2541,7 +2540,7 @@ public:
 		const int32 WorkerIdx = Context.GetWorkerIndex();
 		for (uint32 AROIdx : UnbalancedAROIndices)
 		{
-			AROFunc ARO = AROs[AROIdx];
+			ObjectAROFn ARO = AROs[AROIdx];
 			FAROQueue& Queue = GetQueue(AROIdx, WorkerIdx);
 			for (TArrayView<UObject*> Objects = Queue.Pop<FAROQueue::FMany>(); Objects.Num(); Objects = Queue.Pop<FAROQueue::FMany>())
 			{
@@ -2593,7 +2592,7 @@ private:
 	int32 NumWorkers = 0;
 	TArray<FAROQueue> AllQueues;
 	TArray<int32, TFixedAllocator<2>> UnbalancedAROIndices;
-	TArray<AROFunc, TFixedAllocator<Capacity>> AROs;
+	TArray<ObjectAROFn, TFixedAllocator<Capacity>> AROs;
 
 	FAROQueue& GetQueue(int32 AROIdx, int32 WorkerIdx)
 	{
@@ -2604,7 +2603,7 @@ private:
 
 	// @return if processing should stop since NumCalls reached CallLimit
 	template<class Quantity>
-	static bool ProcessQueues(FWorkerContext& Context, TArrayView<FAROQueue> WorkerQueues, AROFunc ARO, FReferenceCollector& Collector, uint32& NumCalls)
+	static bool ProcessQueues(FWorkerContext& Context, TArrayView<FAROQueue> WorkerQueues, ObjectAROFn ARO, FReferenceCollector& Collector, uint32& NumCalls)
 	{
 		static constexpr uint32 CallLimit = FAROQueue::NumPop<Quantity>;
 		const int32 WorkerIdx = Context.GetWorkerIndex();
@@ -2616,7 +2615,7 @@ private:
 
 	// @return if processing should stop since NumCalls reached or exceeded CallLimit
 	template<TArrayView<UObject*>(FAROQueue::*PopFunc)(), uint32 CallLimit>
-	FORCEINLINE_DEBUGGABLE static bool ProcessQueuesUsing(FWorkerContext& Context, TArrayView<FAROQueue> Queues, AROFunc ARO, FReferenceCollector& Collector, uint32& NumCalls)
+	FORCEINLINE_DEBUGGABLE static bool ProcessQueuesUsing(FWorkerContext& Context, TArrayView<FAROQueue> Queues, ObjectAROFn ARO, FReferenceCollector& Collector, uint32& NumCalls)
 	{
 		for (FAROQueue& Queue : Queues)
 		{
@@ -2707,12 +2706,12 @@ bool FSlowARO::ProcessAllCalls(FWorkerContext& Context, FReferenceCollector& Col
 	return GSlowARO.GetPostInit().ProcessAllQueues(Context, Collector);
 }
 
-void RegisterSlowImplementation(AROFunc ARO, EAROFlags Flags)
+void RegisterSlowImplementation(ObjectAROFn ARO, EAROFlags Flags)
 {
 	GSlowARO.Get().RegisterImplementation(ARO, Flags);
 }
 
-int32 FindSlowImplementation(AROFunc ARO)
+int32 FindSlowImplementation(ObjectAROFn ARO)
 {
 	return GSlowARO.GetPostInit().FindImplementation(ARO);
 }
@@ -2728,10 +2727,10 @@ int32 GetNumSlowAROs()
 FORCEINLINE_DEBUGGABLE void KillReference(UObject*& Object) { Object = nullptr; }
 
 template <EGCOptions Options>
-constexpr FORCEINLINE EKillable MayKill(EGCTokenType TokenType, bool bAllowKill)
+constexpr FORCEINLINE EKillable MayKill(EOrigin Origin, bool bAllowKill)
 {
 	// To avoid content changes, allow reference elimination inside of Blueprints
-	return (bAllowKill & (IsPendingKill(Options) || TokenType == EGCTokenType::NonNative)) ? EKillable::Yes : EKillable::No;
+	return (bAllowKill & (IsPendingKill(Options) || Origin == EOrigin::Blueprint)) ? EKillable::Yes : EKillable::No;
 }
 
 // Return whether flag was cleared. Only thread-safe for concurrent clear, not concurrent set+clear. Don't use during mark phase.
@@ -2777,12 +2776,12 @@ public:
 	static constexpr EInternalObjectFlags KillFlag = IsWithPendingKill() ? EInternalObjectFlags::PendingKill : EInternalObjectFlags::Garbage;
 	PRAGMA_ENABLE_DEPRECATION_WARNINGS
 
-	constexpr static FORCEINLINE EKillable MayKill(EGCTokenType TokenType, bool bAllowKill) { return UE::GC::MayKill<Options>(TokenType, bAllowKill); }	
+	constexpr static FORCEINLINE EKillable MayKill(EOrigin Origin, bool bAllowKill) { return UE::GC::MayKill<Options>(Origin, bAllowKill); }	
 
-	static FORCEINLINE void ProcessReferenceDirectly(FWorkerContext& Context, FPermanentObjectPoolExtents PermanentPool, const UObject* ReferencingObject, UObject*& Object, FTokenId TokenId, EKillable Killable)
+	static FORCEINLINE void ProcessReferenceDirectly(FWorkerContext& Context, FPermanentObjectPoolExtents PermanentPool, const UObject* ReferencingObject, UObject*& Object, FMemberId MemberId, EKillable Killable)
 	{
-		return Killable == EKillable::No	? ProcessReferenceDirectly<EKillable::No >(Context, PermanentPool, ReferencingObject, Object, TokenId)
-											: ProcessReferenceDirectly<EKillable::Yes>(Context, PermanentPool, ReferencingObject, Object, TokenId);
+		return Killable == EKillable::No	? ProcessReferenceDirectly<EKillable::No >(Context, PermanentPool, ReferencingObject, Object, MemberId)
+											: ProcessReferenceDirectly<EKillable::Yes>(Context, PermanentPool, ReferencingObject, Object, MemberId);
 	}
 
 	static FORCEINLINE void DetectGarbageReference(FWorkerContext& Context, FReferenceMetadata Metadata)
@@ -2798,10 +2797,10 @@ public:
 	 * @param bAllowReferenceElimination	Whether to allow NULL'ing the reference if RF_PendingKill is set
 	*/
 	template<EKillable Killable>
-	static FORCEINLINE_DEBUGGABLE void ProcessReferenceDirectly(FWorkerContext& Context, FPermanentObjectPoolExtents PermanentPool, const UObject* ReferencingObject, UObject*& Object, FTokenId TokenId)
+	static FORCEINLINE_DEBUGGABLE void ProcessReferenceDirectly(FWorkerContext& Context, FPermanentObjectPoolExtents PermanentPool, const UObject* ReferencingObject, UObject*& Object, FMemberId MemberId)
 	{
 		UE::GC::Stats.IncreaseObjectRefStats(Object);
-		if (ValidateReference(Object, PermanentPool, ReferencingObject, TokenId))
+		if (ValidateReference(Object, PermanentPool, ReferencingObject, MemberId))
 		{
 			const int32 ObjectIndex = GUObjectArray.ObjectToIndex(Object);
 			FImmutableReference Reference = {Object};	
@@ -2908,50 +2907,63 @@ public:
 template <class ProcessorType>
 struct TBatchDispatcher
 {
-	UE_NONCOPYABLE(TBatchDispatcher);
+	static constexpr bool bBatching = true;
+	static constexpr bool bParallel = IsParallel(ProcessorType::Options);
 
 	FWorkerContext& Context;
 	const FPermanentObjectPoolExtents PermanentPool;
+	FReferenceCollector& Collector;
 	TReferenceBatcher<FMutableReference, FResolvedMutableReference, ProcessorType> KillableBatcher;
 	TReferenceBatcher<FImmutableReference, FImmutableReference, ProcessorType> ImmutableBatcher;
+	FStructBatcher StructBatcher;
 
-	explicit TBatchDispatcher(FWorkerContext& InContext)
+	UE_NONCOPYABLE(TBatchDispatcher);
+	explicit TBatchDispatcher(FWorkerContext& InContext, FReferenceCollector& InCollector)
 	: Context(InContext)
+	, Collector(InCollector)
 	, KillableBatcher(InContext)
 	, ImmutableBatcher(InContext)
+	, StructBatcher(InContext.GetWorkerIndex())
 	{}
 
 	// ProcessObjectArray API
 
-	FORCENOINLINE void HandleReferenceDirectly(const UObject* ReferencingObject, UObject*& Object, FTokenId TokenId, EKillable Killable)
+	FORCENOINLINE void HandleReferenceDirectly(const UObject* ReferencingObject, UObject*& Object, FMemberId MemberId, EKillable Killable)
 	{
-		ProcessorType::ProcessReferenceDirectly(Context, PermanentPool, ReferencingObject, Object, TokenId, Killable);
+		ProcessorType::ProcessReferenceDirectly(Context, PermanentPool, ReferencingObject, Object, MemberId, Killable);
 		Context.Stats.AddReferences(1);
 	}
 
-	FORCEINLINE void HandleReferenceDirectly(const UObject* ReferencingObject, UObject*& Object, FTokenId TokenId, EGCTokenType TokenType, bool bAllowReferenceElimination)
+	FORCEINLINE void HandleReferenceDirectly(const UObject* ReferencingObject, UObject*& Object, FMemberId MemberId, EOrigin Origin, bool bAllowReferenceElimination)
 	{
-		HandleReferenceDirectly(ReferencingObject, Object, TokenId, ProcessorType::MayKill(TokenType, bAllowReferenceElimination));
+		HandleReferenceDirectly(ReferencingObject, Object, MemberId, ProcessorType::MayKill(Origin, bAllowReferenceElimination));
 	}
 
-	FORCEINLINE void HandleKillableReference(UObject*& Object, FTokenId TokenId, EGCTokenType TokenType)
+	FORCEINLINE_DEBUGGABLE void HandleKillableReference(UObject*& Object, FMemberId MemberId, EOrigin Origin)
 	{
-		QueueReference(Context.GetReferencingObject(), Object, TokenId, ProcessorType::MayKill(TokenType, true));
+		QueueReference(Context.GetReferencingObject(), Object, MemberId, ProcessorType::MayKill(Origin, true));
 	}
 
-	FORCEINLINE void HandleImmutableReference(UObject* Object, FTokenId TokenId, EGCTokenType TokenType)
+	FORCEINLINE void HandleImmutableReference(UObject* Object, FMemberId MemberId, EOrigin Origin)
 	{
 		ImmutableBatcher.PushReference(FImmutableReference{Object});
 	}
 
-	FORCEINLINE void HandleKillableArray(TArray<UObject*>& Array, FTokenId TokenId, EGCTokenType TokenType)
+	template<class ArrayType>
+	FORCEINLINE void HandleKillableArray(ArrayType&& Array, FMemberId MemberId, EOrigin Origin)
 	{
-		QueueArray(Context.GetReferencingObject(), Array, TokenId, ProcessorType::MayKill(TokenType, true));
+		QueueArray(Context.GetReferencingObject(), (ArrayType&&)Array, MemberId, ProcessorType::MayKill(Origin, true));
 	}
 
-	FORCEINLINE void HandleKillableReferences(TArrayView<UObject*> Objects, FTokenId TokenId, EGCTokenType TokenType)
+	FORCEINLINE void HandleKillableReferences(TArrayView<UObject*> Objects, FMemberId MemberId, EOrigin Origin)
 	{
-		QueueReferences(Context.GetReferencingObject(), Objects, TokenId, ProcessorType::MayKill(TokenType, true));
+		QueueReferences(Context.GetReferencingObject(), Objects, MemberId, ProcessorType::MayKill(Origin, true));
+	}
+
+	FORCEINLINE void HandleWeakReference(FWeakObjectPtr& WeakPtr, const UObject* ReferencingObject, FMemberId MemberId, EMemberType)
+	{
+		UObject* WeakObject = WeakPtr.Get(true);
+		HandleReferenceDirectly(ReferencingObject, WeakObject, MemberId, EKillable::No);
 	}
 
 	FORCENOINLINE void FlushQueuedReferences()
@@ -2960,9 +2972,32 @@ struct TBatchDispatcher
 		ImmutableBatcher.FlushQueues();
 	}
 
+	static constexpr bool bSupportsStructQueueing = true;
+
+	FORCENOINLINE bool FlushToStructBlocks()
+	{
+		StructBatcher.FlushBoundedQueues();
+		return StructBatcher.GetUnboundedQueue().CanPop();
+	}
+
+	FORCEINLINE_DEBUGGABLE FStructBlockifier& GetStructBlocks()
+	{
+		return StructBatcher.GetUnboundedQueue();
+	}
+
+	FORCEINLINE_DEBUGGABLE void QueueStructArray(FSchemaView Schema, uint8* Data, int32 Num)
+	{
+		StructBatcher.PushStructArray(Schema, Data, Num);
+	}
+
+	FORCEINLINE_DEBUGGABLE void QueueSparseStructArray(FSchemaView Schema, FScriptSparseArray& Array)
+	{
+		StructBatcher.PushSparseStructArray(Schema, Array);
+	}
+
 	// Collector API
 
-	FORCEINLINE_DEBUGGABLE void QueueReference(const UObject* ReferencingObject,  UObject*& Object, FTokenId TokenId, EKillable Killable)
+	FORCEINLINE_DEBUGGABLE void QueueReference(const UObject* ReferencingObject,  UObject*& Object, FMemberId MemberId, EKillable Killable)
 	{
 		if (Killable == EKillable::Yes)
 		{
@@ -2975,44 +3010,45 @@ struct TBatchDispatcher
 		}
 	}
 
-	FORCEINLINE_DEBUGGABLE void QueueReferences(const UObject* ReferencingObject, TArrayView<UObject*> References, FTokenId TokenId, EKillable Killable)
+	FORCEINLINE_DEBUGGABLE void QueueReferences(const UObject* ReferencingObject, TArrayView<UObject*> References, FMemberId MemberId, EKillable Killable)
 	{
 		if (Killable == EKillable::Yes)
 		{
-			KillableBatcher.PushReferences(References);
+			KillableBatcher.PushReferenceView(References);
 		}
 		else
 		{
-			ImmutableBatcher.PushReferences(References);
+			ImmutableBatcher.PushReferenceView(References);
 		}
 	}
 	
-	FORCEINLINE_DEBUGGABLE void QueueArray(const UObject* ReferencingObject, TArray<UObject*>& Array, FTokenId TokenId, EKillable Killable)
+	template<class ArrayType>
+	FORCEINLINE_DEBUGGABLE void QueueArray(const UObject* ReferencingObject, ArrayType&& Array, FMemberId MemberId, EKillable Killable)
 	{
 		if (Killable == EKillable::Yes)
 		{
-			KillableBatcher.PushArray(FReferenceArray{&Array});
+			KillableBatcher.PushArray((ArrayType&&)Array);
 		}
 		else
 		{
-			ImmutableBatcher.PushArray(FReferenceArray{&Array});
+			ImmutableBatcher.PushArray((ArrayType&&)Array);
 		}
 	}
 
-	FORCEINLINE_DEBUGGABLE void QueueSet(const UObject* ReferencingObject, TSet<UObject*>& Objects, FTokenId TokenId, EKillable Killable)
+	FORCEINLINE_DEBUGGABLE void QueueSet(const UObject* ReferencingObject, TSet<UObject*>& Objects, FMemberId MemberId, EKillable Killable)
 	{
 		if (Killable == EKillable::Yes)
 		{
 			for (UObject*& Object : Objects)
 			{
-				QueueReference(ReferencingObject, Object, TokenId, EKillable::Yes);
+				QueueReference(ReferencingObject, Object, MemberId, EKillable::Yes);
 			}
 		}
 		else
 		{
 			for (UObject*& Object : Objects)
 			{
-				QueueReference(ReferencingObject, Object, TokenId, EKillable::No);
+				QueueReference(ReferencingObject, Object, MemberId, EKillable::No);
 			}
 		}
 	}
@@ -3025,9 +3061,9 @@ class TReachabilityCollectorBase : public FReferenceCollector
 {
 protected:
 	bool bAllowEliminatingReferences = true;
-	EGCTokenType CurrentTokenType = EGCTokenType::Native;
+	EOrigin CurrentOrigin = EOrigin::Other;
 
-	FORCEINLINE EKillable MayKill() const { return UE::GC::MayKill<Options>(CurrentTokenType, bAllowEliminatingReferences); }
+	FORCEINLINE EKillable MayKill() const { return UE::GC::MayKill<Options>(CurrentOrigin, bAllowEliminatingReferences); }
 
 public:
 	virtual bool IsIgnoringArchetypeRef() const override final { return false;}
@@ -3037,7 +3073,7 @@ public:
 
 	virtual void SetIsProcessingNativeReferences(bool bIsNative) override final
 	{
-		CurrentTokenType = bIsNative ? EGCTokenType::Native : EGCTokenType::NonNative;
+		CurrentOrigin = bIsNative ? EOrigin::Other : EOrigin::Blueprint;
 	}
 
 	virtual bool NeedsPropertyReferencer() const override final { return ENABLE_GC_OBJECT_CHECKS; }
@@ -3056,25 +3092,24 @@ class TReachabilityCollector final : public TReachabilityCollectorBase<Options>
 	friend DispatcherType& GetDispatcher(TReachabilityCollector<Options>& Collector, ProcessorType&, FWorkerContext&) { return Collector.Dispatcher; }
 
 public:
-	TReachabilityCollector(ProcessorType&, FWorkerContext& Context) : Dispatcher(Context) {}
+	TReachabilityCollector(ProcessorType&, FWorkerContext& Context) : Dispatcher(Context, *this) {}
 		
 	virtual void HandleObjectReference(UObject*& InObject, const UObject* InReferencingObject, const FProperty* InReferencingProperty) override;
 	virtual void HandleObjectReferences(UObject** InObjects, const int32 ObjectNum, const UObject* InReferencingObject, const FProperty* InReferencingProperty) override;
 	
 	virtual void AddStableReference(UObject** Object) override
 	{
-		Dispatcher.QueueReference(Dispatcher.Context.GetReferencingObject(), *Object, ETokenlessId::Collector, MayKill());
+		Dispatcher.QueueReference(Dispatcher.Context.GetReferencingObject(), *Object, EMemberlessId::Collector, MayKill());
 	}
 	
 	virtual void AddStableReferenceArray(TArray<UObject*>* Objects) override
 	{
-		Dispatcher.QueueArray(Dispatcher.Context.GetReferencingObject(), *Objects, ETokenlessId::Collector, MayKill());
+		Dispatcher.QueueArray(Dispatcher.Context.GetReferencingObject(), *Objects, EMemberlessId::Collector, MayKill());
 	}
 
 	virtual void AddStableReferenceSet(TSet<UObject*>* Objects) override
 	{
-
-		Dispatcher.QueueSet(Dispatcher.Context.GetReferencingObject(), *Objects, ETokenlessId::Collector, MayKill());
+		Dispatcher.QueueSet(Dispatcher.Context.GetReferencingObject(), *Objects, EMemberlessId::Collector, MayKill());
 	}
 	
 	virtual bool MarkWeakObjectReferenceForClearing(UObject** WeakReference) override
@@ -3119,15 +3154,15 @@ public:
 	bool TracksHistory() const { return bTrackHistory; }
 	bool TracksGarbage() const { return bTrackGarbage; }
 
-	FORCENOINLINE void HandleTokenStreamObjectReference(FWorkerContext& Context, const UObject* ReferencingObject, UObject*& Object, FTokenId TokenId, EGCTokenType TokenType, bool bAllowReferenceElimination)
+	FORCENOINLINE void HandleTokenStreamObjectReference(FWorkerContext& Context, const UObject* ReferencingObject, UObject*& Object, FMemberId MemberId, EOrigin Origin, bool bAllowReferenceElimination)
 	{
 		UE::GC::Stats.IncreaseObjectRefStats(Object);
-		if (ValidateReference(Object, PermanentPool, ReferencingObject, TokenId))
+		if (ValidateReference(Object, PermanentPool, ReferencingObject, MemberId))
 		{
 			FReferenceMetadata Metadata(GUObjectArray.ObjectToIndex(Object));
 			if (bAllowReferenceElimination & Metadata.Has(TReachabilityProcessor<Options>::KillFlag)) //-V792
 			{
-				if (MayKill<Options>(TokenType, bAllowReferenceElimination) == EKillable::Yes)
+				if (MayKill<Options>(Origin, bAllowReferenceElimination) == EKillable::Yes)
 				{
 					check(IsPendingKill(Options) || ReferencingObject);
 					checkSlow(Metadata.ObjectItem->GetOwnerIndex() <= 0);
@@ -3136,8 +3171,8 @@ public:
 				}
 				else if (!IsPendingKill(Options) && bTrackGarbage)
 				{
-					check(TokenType == EGCTokenType::Native);
-					HandleGarbageReference(Context, ReferencingObject, Object, TokenId);
+					check(Origin == EOrigin::Other);
+					HandleGarbageReference(Context, ReferencingObject, Object, MemberId);
 				}
 			}
 			
@@ -3145,7 +3180,7 @@ public:
 
 			if (bReachedFirst && bTrackHistory)
 			{
-				HandleHistoryReference(Context, ReferencingObject, Object, TokenId);
+				HandleHistoryReference(Context, ReferencingObject, Object, MemberId);
 			}
 		}
 	}
@@ -3155,14 +3190,14 @@ private:
 	const bool bTrackGarbage;
 	const bool bTrackHistory;
 
-	FORCENOINLINE static void HandleGarbageReference(FWorkerContext& Context, const UObject* ReferencingObject, UObject*& Object, FTokenId TokenId)
+	FORCENOINLINE static void HandleGarbageReference(FWorkerContext& Context, const UObject* ReferencingObject, UObject*& Object, FMemberId MemberId)
 	{
 		const UObject* Referencer = ReferencingObject ? ReferencingObject : Context.GetReferencingObject();
 		if (IsValid(Referencer))
 		{
 			if (Referencer != FGCObject::GGCObjectReferencer)
 			{
-				FName PropertyName = Referencer->GetClass()->ReferenceTokens.GetTokenInfo(TokenId).Name;
+				FName PropertyName = GetMemberDebugInfo(Referencer->GetClass()->ReferenceSchema.Get(), MemberId).Name;
 				Context.GarbageReferences.Add(FGarbageReferenceInfo(Referencer, Object, PropertyName));
 			}
 			else
@@ -3173,13 +3208,13 @@ private:
 		}
 	}
 
-	FORCENOINLINE static void HandleHistoryReference(FWorkerContext& Context, const UObject* ReferencingObject, UObject*& Object, FTokenId TokenId)
+	FORCENOINLINE static void HandleHistoryReference(FWorkerContext& Context, const UObject* ReferencingObject, UObject*& Object, FMemberId MemberId)
 	{
 		FGCDirectReference Ref(Object);
 		const UObject* Referencer = ReferencingObject ? ReferencingObject : Context.GetReferencingObject();	
 		if (Referencer != FGCObject::GGCObjectReferencer)
 		{
-			Ref.ReferencerName = Referencer->GetClass()->ReferenceTokens.GetTokenInfo(TokenId).Name;
+			Ref.ReferencerName = GetMemberDebugInfo(Referencer->GetClass()->ReferenceSchema.Get(), MemberId).Name;
 		}
 		else if (FGCObject::GGCObjectReferencer == Referencer)
 		{
@@ -3209,7 +3244,7 @@ class TDebugReachabilityCollector final : public TReachabilityCollectorBase<Opti
 {
 	using Super = TReachabilityCollectorBase<Options>;
 	using Super::bAllowEliminatingReferences;
-	using Super::CurrentTokenType;
+	using Super::CurrentOrigin;
 	using ProcessorType = TDebugReachabilityProcessor<Options>;
 	ProcessorType& Processor;
 	FWorkerContext& Context;
@@ -3229,7 +3264,7 @@ public:
 		}
 #endif // ENABLE_GC_OBJECT_CHECKS
 
-		Processor.HandleTokenStreamObjectReference(Context, ReferencingObject, Object, ETokenlessId::Collector, CurrentTokenType, bAllowEliminatingReferences);
+		Processor.HandleTokenStreamObjectReference(Context, ReferencingObject, Object, EMemberlessId::Collector, CurrentOrigin, bAllowEliminatingReferences);
 		Context.Stats.AddReferences(1);
 	}
 
@@ -3241,8 +3276,6 @@ public:
 		}
 	}
 
-
-
 	virtual void AddStableReference(UObject** Object) override
 	{
 		HandleObjectReference(*Object, Context.GetReferencingObject(), nullptr);
@@ -3290,7 +3323,7 @@ template <EGCOptions Options>
 void TReachabilityCollector<Options>::HandleObjectReference(UObject*& Object, const UObject* ReferencingObject, const FProperty* ReferencingProperty)
 {
 	CheckReference(Object, ReferencingObject, ReferencingProperty);
-	Dispatcher.HandleReferenceDirectly(ReferencingObject, Object, ETokenlessId::Collector, MayKill());
+	Dispatcher.HandleReferenceDirectly(ReferencingObject, Object, EMemberlessId::Collector, MayKill());
 }
 
 template <EGCOptions Options>
@@ -3311,9 +3344,56 @@ void TReachabilityCollector<Options>::HandleObjectReferences(UObject** InObjects
 
 	for (UObject*& Object : Objects)
 	{
-		Dispatcher.HandleReferenceDirectly(ReferencingObject, Object, ETokenlessId::Collector, Killable);
+		Dispatcher.HandleReferenceDirectly(ReferencingObject, Object, EMemberlessId::Collector, Killable);
 	}
+}
+
+//////////////////////////////////////////////////////////////////////////
 
+template <typename ProcessorType, typename CollectorType>
+void TFastReferenceCollector<ProcessorType, CollectorType>::ProcessStructs(DispatcherType& Dispatcher)
+{
+	if constexpr (DispatcherType::bBatching)
+	{
+		check(Dispatcher.GetStructBlocks().CanPop());
+		FStructBlockifier& Blocks = Dispatcher.GetStructBlocks();
+		Dispatcher.Context.ReferencingObject = nullptr;
+
+		do
+		{
+			int32 Num;
+			FStructArrayBlock* Block = Blocks.PopBlock(/* out */ Num);
+			ON_SCOPE_EXIT{ Blocks.FreeBlock(Block); };
+			FStructArray* ArrayIt = Block->Data;
+			FStructArray* ArrayEnd = ArrayIt + Num;
+		
+			for (int32 PrefetchIdx = 0; PrefetchIdx < FStructArrayBlock::Lookahead; ++PrefetchIdx)
+			{
+				FPlatformMisc::Prefetch(ArrayIt[PrefetchIdx].Data);
+			}
+			FPlatformMisc::Prefetch(ArrayIt[0].Schema.GetWords());
+
+			do
+			{
+				FPlatformMisc::Prefetch(ArrayIt[FStructArrayBlock::Lookahead].Data);
+				FPlatformMisc::Prefetch(ArrayIt[1].Schema.GetWords());
+
+				const FStructArray AoS = *ArrayIt;
+				check(AoS.Num > 0);
+				check(AoS.Stride % 8 == 0);
+				uint8* StructIt = AoS.Data;
+				uint8* StructEnd = AoS.Data + AoS.Num * AoS.Stride;
+				do
+				{
+					Private::VisitMembers(Dispatcher, AoS.Schema, StructIt); 
+					StructIt += AoS.Stride;
+				}
+				while (StructIt != StructEnd);
+			}
+			while (++ArrayIt != ArrayEnd);
+		}
+		while (Blocks.CanPop() || Dispatcher.FlushToStructBlocks());
+	}
 }
 
 } // namespace UE::GC
@@ -3396,17 +3476,6 @@ void FReferenceFinder::HandleObjectReference( UObject*& InObject, const UObject*
 namespace UE::GC
 {
 
-
-/**
- * Implementation of parallel realtime garbage collector using recursive subdivision
- *
- * The approach is to create an array of uint32 tokens for each class that describe object references. This is done for 
- * script exposed classes by traversing the properties and additionally via manual function calls to emit tokens for
- * native only classes in the construction singleton IMPLEMENT_INTRINSIC_CLASS. 
- * A third alternative is a AddReferencedObjects callback per object which 
- * is used to deal with object references from types that aren't supported by the reflectable type system.
- * interface doesn't make sense to implement for.
- */
 class FRealtimeGC : public FGarbageCollectionTracer
 {
 	typedef void(FRealtimeGC::*MarkObjectsFn)(EObjectFlags);
@@ -3924,10 +3993,7 @@ void IncrementalPurgeGarbage(bool bUseTimeLimit, double TimeLimit)
 	}
 }
 
-
-
 static bool GWarningTimeOutHasBeenDisplayedGC = false;
-
 static bool GEnableTimeoutOnPendingDestroyedObjectGC = true;
 #if UE_BUILD_SHIPPING || USING_ADDRESS_SANITISER
 static FAutoConsoleVariableRef CVarGCEnableTimeoutOnPendingDestroyedObjectInShipping(
@@ -4473,7 +4539,6 @@ static void UpdateGCHistory(TConstArrayView<TUniquePtr<FWorkerContext>> Contexts
 #endif // ENABLE_GC_HISTORY
 }
 
-
 template<bool bPerformFullPurge>
 FORCENOINLINE void CollectGarbageImpl(EObjectFlags KeepFlags);
 
@@ -4642,7 +4707,6 @@ void CollectGarbageImpl(EObjectFlags KeepFlags)
 					GC.PerformReachabilityAnalysis(KeepFlags, Options);
 					UE_LOG(LogGarbage, Log, TEXT("%.2f ms for GC rerun to track garbage references (gc.GarbageReferenceTrackingEnabled=%d)"), (FPlatformTime::Seconds() - StartTime) * 1000, GGarbageReferenceTrackingEnabled);
 				}
-
 			}
 
 
@@ -5034,204 +5098,96 @@ bool FMulticastDelegateProperty::ContainsObjectReference(TArray<const FStructPro
 	return !!(InReferenceType & EPropertyObjectReferenceType::Weak);
 }
 
-namespace UE::GC
-{
-/**
- * Scope helper structure to emit tokens for fixed arrays in the case of ArrayDim (passed in count) being > 1.
- */
-struct FFixedArrayTokenHelper
-{
-	/**
-	 * Constructor, emitting necessary tokens for fixed arrays if count > 1 and also keeping track of count so 
-	 * destructor can do the same.
-	 *
-	 * @param InReferenceTokenStream	Token stream to emit tokens to
-	 * @param InOffset					offset into object/ struct
-	 * @param InCount					array count
-	 * @param InStride					array type stride (e.g. sizeof(struct) or sizeof(UObject*))
-	 * @param InProperty                the property this array represents
-	 */
-	FFixedArrayTokenHelper(FTokenStreamBuilder& TokenStream,  int32 InOffset, int32 InCount, int32 InStride, const FProperty& InProperty, FGCStackSizeHelper& InStackSizeHelper)
-		: ReferenceTokenStream(TokenStream)
-		, StackSizeHelper(InStackSizeHelper)
-		, Count(InCount)
-	{
-		if( InCount > 1 )
-		{
-			ReferenceTokenStream.EmitObjectReference(InOffset, *StackSizeHelper.GetPropertyPath(), GCRT_FixedArray);
-
-			ReferenceTokenStream.EmitStride(InStride);
-			ReferenceTokenStream.EmitCount(InCount);
-			// GCRT_FixedArray pushes a new stack frame in TFastReferenceCollector
-			StackSizeHelper.Push(&InProperty);
-		}
-	}
-
-	/** Destructor, emitting return if ArrayDim > 1 */
-	~FFixedArrayTokenHelper()
-	{
-		if( Count > 1 )
-		{
-			ReferenceTokenStream.EmitReturn();
-			StackSizeHelper.Pop();
-		}
-	}
-
-private:
-	/** Reference token stream used to emit to */
-	FTokenStreamBuilder&	ReferenceTokenStream;
-	/** GC token stream stack size helper */
-	FGCStackSizeHelper&			StackSizeHelper;
-	/** Size of fixed array */
-	int32							Count;
-};
-
-} // namespace UE::GC
-
-
-/**
- * Emits tokens used by realtime garbage collection code to passed in ReferenceTokenStream. The offset emitted is relative
- * to the passed in BaseOffset which is used by e.g. arrays of structs.
- */
-void FProperty::EmitReferenceInfo(UE::GC::FTokenStreamBuilder& TokenStream, int32 BaseOffset, TArray<const FStructProperty*>& EncounteredStructProps, FGCStackSizeHelper& StackSizeHelper)
+void FProperty::EmitReferenceInfo(UE::GC::FSchemaBuilder& Schema, int32 BaseOffset, TArray<const FStructProperty*>& EncounteredStructProps, UE::GC::FPropertyStack& DebugPath)
 {
 }
 
-/**
- * Emits tokens used by realtime garbage collection code to passed in OwnerClass' ReferenceTokenStream. The offset emitted is relative
- * to the passed in BaseOffset which is used by e.g. arrays of structs.
- */
-void FObjectProperty::EmitReferenceInfo(UE::GC::FTokenStreamBuilder& TokenStream, int32 BaseOffset, TArray<const FStructProperty*>& EncounteredStructProps, FGCStackSizeHelper& StackSizeHelper)
+void FObjectProperty::EmitReferenceInfo(UE::GC::FSchemaBuilder& Schema, int32 BaseOffset, TArray<const FStructProperty*>& EncounteredStructProps, UE::GC::FPropertyStack& DebugPath)
 {
-	UE::GC::FFixedArrayTokenHelper FixedArrayHelper(TokenStream, BaseOffset + GetOffset_ForGC(), ArrayDim, sizeof(UObject*), *this, StackSizeHelper);
-	TokenStream.EmitObjectReference(BaseOffset + GetOffset_ForGC(), *StackSizeHelper.GetPropertyPath(), GCRT_Object);
+	for (int32 Idx = 0, Num = ArrayDim; Idx < Num; ++Idx)
+	{
+		Schema.Add(UE::GC::DeclareMember(DebugPath, BaseOffset + GetOffset_ForGC() + Idx * sizeof(FObjectPtr), UE::GC::EMemberType::Reference));
+	}
 }
 
-/**
- * Emits tokens used by realtime garbage collection code to passed in OwnerClass' ReferenceTokenStream. The offset emitted is relative
- * to the passed in BaseOffset which is used by e.g. arrays of structs.
- */
-void FArrayProperty::EmitReferenceInfo(UE::GC::FTokenStreamBuilder& TokenStream, int32 BaseOffset, TArray<const FStructProperty*>& EncounteredStructProps, FGCStackSizeHelper& StackSizeHelper)
+void FArrayProperty::EmitReferenceInfo(UE::GC::FSchemaBuilder& Schema, int32 BaseOffset, TArray<const FStructProperty*>& EncounteredStructProps, UE::GC::FPropertyStack& DebugPath)
 {
+	using namespace UE::GC;
 	if (Inner->ContainsObjectReference(EncounteredStructProps, EPropertyObjectReferenceType::Strong))
 	{
 		bool bUsesFreezableAllocator = EnumHasAnyFlags(ArrayFlags, EArrayPropertyFlags::UsesMemoryImageAllocator);
+		EMemberType Type = bUsesFreezableAllocator ?  EMemberType::FreezableStructArray : EMemberType::StructArray;
+		FSchemaBuilder InnerSchema(Inner->ElementSize);
 
 		// Structs and nested arrays share the same implementation on the Garbage Collector side
 		// as arrays of structs already push the array memory into the GC stack and process its tokens
 		// which is exactly what is required for nested arrays to work
 		if( Inner->IsA(FObjectProperty::StaticClass()) || Inner->IsA(FObjectPtrProperty::StaticClass()) )
 		{
-			TokenStream.EmitObjectReference(BaseOffset + GetOffset_ForGC(), *StackSizeHelper.GetPropertyPath(), bUsesFreezableAllocator ? GCRT_ArrayObjectFreezable : GCRT_ArrayObject);
+			Type = bUsesFreezableAllocator ?  EMemberType::FreezableReferenceArray :  EMemberType::ReferenceArray;
 		}
 		else if( Inner->IsA(FInterfaceProperty::StaticClass()) )
 		{
-			TokenStream.EmitObjectReference(BaseOffset + GetOffset_ForGC(), *StackSizeHelper.GetPropertyPath(), bUsesFreezableAllocator ? GCRT_ArrayStructFreezable : GCRT_ArrayStruct);
-
-			// GCRT_ArrayStruct and GCRT_ArrayStructFreezable push a new stack frame in TFastReferenceCollector
-			FGCStackSizeHelperScope StackSizeScope(StackSizeHelper, Inner);
-
-			TokenStream.EmitStride(Inner->ElementSize);
-			const uint32 SkipIndexIndex = TokenStream.EmitSkipIndexPlaceholder();
-
-			TokenStream.EmitObjectReference(0, *StackSizeHelper.GetPropertyPath(), GCRT_Object);
-
-			const uint32 SkipIndex = TokenStream.EmitReturn();
-			TokenStream.UpdateSkipIndexPlaceholder(SkipIndexIndex, SkipIndex);
+			FPropertyStackScope PropertyScope(DebugPath, Inner);
+			InnerSchema.Add(UE::GC::DeclareMember(DebugPath, /* offsetof(FScriptInterface, ObjectPointer) */ 0, EMemberType::Reference));
 		}
 		else if (Inner->IsA(FFieldPathProperty::StaticClass()))
 		{
-			TokenStream.EmitObjectReference(BaseOffset + GetOffset_ForGC(), *StackSizeHelper.GetPropertyPath(), GCRT_ArrayAddFieldPathReferencedObject);
+			Type = EMemberType::FieldPathArray;
 		}
 		else
 		{
-			// In the general case, emit a GCRT_ArrayStruct that can handle any inner property type.
-			TokenStream.EmitObjectReference(BaseOffset + GetOffset_ForGC(), *StackSizeHelper.GetPropertyPath(), bUsesFreezableAllocator ? GCRT_ArrayStructFreezable : GCRT_ArrayStruct);
-
-			// GCRT_ArrayStruct and GCRT_ArrayStructFreezable push a new stack frame in TFastReferenceCollector
-			FGCStackSizeHelperScope StackSizeScope(StackSizeHelper, Inner);
-
-			TokenStream.EmitStride(Inner->ElementSize);
-			const uint32 SkipIndexIndex = TokenStream.EmitSkipIndexPlaceholder();
-			Inner->EmitReferenceInfo(TokenStream, 0, EncounteredStructProps, StackSizeHelper);
-			const uint32 SkipIndex = TokenStream.EmitReturn();
-			TokenStream.UpdateSkipIndexPlaceholder(SkipIndexIndex, SkipIndex);
+			FPropertyStackScope PropertyScope(DebugPath, Inner);
+			Inner->EmitReferenceInfo(InnerSchema, 0, EncounteredStructProps, DebugPath);
 		}
+
+		Schema.Add(UE::GC::DeclareMember(DebugPath, BaseOffset + GetOffset_ForGC(), Type, InnerSchema.Build()));
 	}
 }
 
-
-/**
- * Emits tokens used by realtime garbage collection code to passed in OwnerClass' ReferenceTokenStream. The offset emitted is relative
- * to the passed in BaseOffset which is used by e.g. arrays of structs.
- */
-void FMapProperty::EmitReferenceInfo(UE::GC::FTokenStreamBuilder& TokenStream, int32 BaseOffset, TArray<const FStructProperty*>& EncounteredStructProps, FGCStackSizeHelper& StackSizeHelper)
+void FMapProperty::EmitReferenceInfo(UE::GC::FSchemaBuilder& Schema, int32 BaseOffset, TArray<const FStructProperty*>& EncounteredStructProps, UE::GC::FPropertyStack& DebugPath)
 {
+	using namespace UE::GC;
 	if (ContainsObjectReference(EncounteredStructProps, EPropertyObjectReferenceType::Strong))
 	{
-		// TMap reference tokens are processed by GC in a similar way to an array of structs
-		TokenStream.EmitObjectReference(BaseOffset + GetOffset_ForGC(), *StackSizeHelper.GetPropertyPath(), GCRT_AddTMapReferencedObjects);
-		TokenStream.EmitPointer((const void*)this);
-		const uint32 SkipIndexIndex = TokenStream.EmitSkipIndexPlaceholder();
-
-		if (KeyProp->ContainsObjectReference(EncounteredStructProps, EPropertyObjectReferenceType::Strong))
-		{
-			// GCRT_AddTMapReferencedObjects pushes a new stack frame in TFastReferenceCollector
-			FGCStackSizeHelperScope StackSizeScope(StackSizeHelper, KeyProp);
-			KeyProp->EmitReferenceInfo(TokenStream, 0, EncounteredStructProps, StackSizeHelper);
-		}
-		if (ValueProp->ContainsObjectReference(EncounteredStructProps, EPropertyObjectReferenceType::Strong))
+		FSchemaBuilder PairSchema(GetPairStride());
+		for (FProperty* Prop : {KeyProp, ValueProp})
 		{
-			// GCRT_AddTMapReferencedObjects pushes a new stack frame in TFastReferenceCollector
-			FGCStackSizeHelperScope StackSizeScope(StackSizeHelper, ValueProp);
-			ValueProp->EmitReferenceInfo(TokenStream, 0, EncounteredStructProps, StackSizeHelper);
+			if (Prop->ContainsObjectReference(EncounteredStructProps, EPropertyObjectReferenceType::Strong))
+			{
+				FPropertyStackScope PropertyScope(DebugPath, Prop);
+				Prop->EmitReferenceInfo(PairSchema, 0, EncounteredStructProps, DebugPath);
+			}
 		}
 
-		const uint32 SkipIndex = TokenStream.EmitReturn();
-		TokenStream.UpdateSkipIndexPlaceholder(SkipIndexIndex, SkipIndex);
+		Schema.Add(UE::GC::DeclareMember(DebugPath, BaseOffset + GetOffset_ForGC(), EMemberType::SparseStructArray, PairSchema.Build()));
 	}
 }
 
-/**
-* Emits tokens used by realtime garbage collection code to passed in OwnerClass' ReferenceTokenStream. The offset emitted is relative
-* to the passed in BaseOffset which is used by e.g. arrays of structs.
-*/
-void FSetProperty::EmitReferenceInfo(UE::GC::FTokenStreamBuilder& TokenStream, int32 BaseOffset, TArray<const FStructProperty*>& EncounteredStructProps, FGCStackSizeHelper& StackSizeHelper)
+void FSetProperty::EmitReferenceInfo(UE::GC::FSchemaBuilder& Schema, int32 BaseOffset, TArray<const FStructProperty*>& EncounteredStructProps, UE::GC::FPropertyStack& DebugPath)
 {
+	using namespace UE::GC;
 	if (ContainsObjectReference(EncounteredStructProps, EPropertyObjectReferenceType::Strong))
 	{
-		// TSet reference tokens are processed by GC in a similar way to an array of structs
-		TokenStream.EmitObjectReference(BaseOffset + GetOffset_ForGC(), *StackSizeHelper.GetPropertyPath(), GCRT_AddTSetReferencedObjects);
-		TokenStream.EmitPointer((const void*)this);
-
-		// GCRT_AddTSetReferencedObjects pushes a new stack frame in TFastReferenceCollector
-		FGCStackSizeHelperScope StackSizeScope(StackSizeHelper, ElementProp);
-
-		const uint32 SkipIndexIndex = TokenStream.EmitSkipIndexPlaceholder();
-		ElementProp->EmitReferenceInfo(TokenStream, 0, EncounteredStructProps, StackSizeHelper);
-		const uint32 SkipIndex = TokenStream.EmitReturn();
-		TokenStream.UpdateSkipIndexPlaceholder(SkipIndexIndex, SkipIndex);
+		FSchemaBuilder ElementSchema(GetStride());
+		FPropertyStackScope PropertyScope(DebugPath, ElementProp);
+		ElementProp->EmitReferenceInfo(ElementSchema, 0, EncounteredStructProps, DebugPath);
+		Schema.Add(UE::GC::DeclareMember(DebugPath, BaseOffset + GetOffset_ForGC(), EMemberType::SparseStructArray, ElementSchema.Build()));
 	}
 }
 
-
-/**
- * Emits tokens used by realtime garbage collection code to passed in ReferenceTokenStream. The offset emitted is relative
- * to the passed in BaseOffset which is used by e.g. arrays of structs.
- */
-void FStructProperty::EmitReferenceInfo(UE::GC::FTokenStreamBuilder& TokenStream, int32 BaseOffset, TArray<const FStructProperty*>& EncounteredStructProps, FGCStackSizeHelper& StackSizeHelper)
+void FStructProperty::EmitReferenceInfo(UE::GC::FSchemaBuilder& Schema, int32 BaseOffset, TArray<const FStructProperty*>& EncounteredStructProps, UE::GC::FPropertyStack& DebugPath)
 {
-	check(Struct);
+	using namespace UE::GC;
+	
+	const int32 Offset = BaseOffset + GetOffset_ForGC();
 	if (Struct->StructFlags & STRUCT_AddStructReferencedObjects)
 	{
-		UScriptStruct::ICppStructOps* CppStructOps = Struct->GetCppStructOps();
-		check(CppStructOps); // else should not have STRUCT_AddStructReferencedObjects
-		UE::GC::FFixedArrayTokenHelper FixedArrayHelper(TokenStream, BaseOffset + GetOffset_ForGC(), ArrayDim, ElementSize, *this, StackSizeHelper);
-
-		TokenStream.EmitObjectReference(BaseOffset + GetOffset_ForGC(), *StackSizeHelper.GetPropertyPath(), GCRT_AddStructReferencedObjects);
-
-		void *FunctionPtr = (void*)CppStructOps->AddStructReferencedObjects();
-		TokenStream.EmitPointer(FunctionPtr);
+		StructAROFn StructARO = Struct->GetCppStructOps()->AddStructReferencedObjects();	
+		for (int32 Idx = 0, Num = ArrayDim; Idx < Num; ++Idx)
+		{	
+			Schema.Add(UE::GC::DeclareMember(DebugPath, Offset + Idx * ElementSize, EMemberType::MemberARO, StructARO) );
+		}
 	}
 
 	// Check if the struct has any properties that reference UObjects
@@ -5246,186 +5202,41 @@ void FStructProperty::EmitReferenceInfo(UE::GC::FTokenStreamBuilder& TokenStream
 		}
 		EncounteredStructProps.RemoveSingleSwap(this, false /* bAllowShrinking */);
 	}
-	// If the struct has UObject properties (and only if) emit tokens for them
-	if (bHasPropertiesWithObjectReferences)
-	{
-		UE::GC::FFixedArrayTokenHelper FixedArrayHelper(TokenStream, BaseOffset + GetOffset_ForGC(), ArrayDim, ElementSize, *this, StackSizeHelper);
-
-		FProperty* Property = Struct->PropertyLink;
-		while (Property)
-		{
-			FGCStackSizeHelperScope StackSizeScope(StackSizeHelper, Property);
-			Property->EmitReferenceInfo(TokenStream, BaseOffset + GetOffset_ForGC(), EncounteredStructProps, StackSizeHelper);
-			Property = Property->PropertyLinkNext;
-		}
-	}
-}
-
-/**
- * Emits tokens used by realtime garbage collection code to passed in ReferenceTokenStream. The offset emitted is relative
- * to the passed in BaseOffset which is used by e.g. arrays of structs.
- */
-void FInterfaceProperty::EmitReferenceInfo(UE::GC::FTokenStreamBuilder& TokenStream, int32 BaseOffset, TArray<const FStructProperty*>& EncounteredStructProps, FGCStackSizeHelper& StackSizeHelper)
-{
-	UE::GC::FFixedArrayTokenHelper FixedArrayHelper(TokenStream, BaseOffset + GetOffset_ForGC(), ArrayDim, sizeof(FScriptInterface), *this, StackSizeHelper);
-
-	TokenStream.EmitObjectReference(BaseOffset + GetOffset_ForGC(), *StackSizeHelper.GetPropertyPath(), GCRT_Object);
-}
-
-void FFieldPathProperty::EmitReferenceInfo(UE::GC::FTokenStreamBuilder& TokenStream, int32 BaseOffset, TArray<const FStructProperty*>& EncounteredStructProps, FGCStackSizeHelper& StackSizeHelper)
-{
-	static_assert(sizeof(FFieldPath) == sizeof(TFieldPath<FProperty>), "TFieldPath should have the same size as the underlying FFieldPath");
-	UE::GC::FFixedArrayTokenHelper FixedArrayHelper(TokenStream, BaseOffset + GetOffset_ForGC(), ArrayDim, sizeof(FFieldPath), *this, StackSizeHelper);
-	TokenStream.EmitObjectReference(BaseOffset + GetOffset_ForGC(), *StackSizeHelper.GetPropertyPath(), GCRT_AddFieldPathReferencedObject);
-}
-
-namespace UE::GC
-{
-
-static int32 GNumTokens;
-static int32 GNumTokenStreams;
-static void UpdateGlobalTokenCount(int32 Delta)
-{
-#if !UE_BUILD_SHIPPING
-	GNumTokens += Delta;
-	GNumTokenStreams += FMath::Sign(Delta);
-#endif
-}
-
-struct FTokenStreamHeader
-{
-	std::atomic<int32> RefCount{1};
-};
-
-static FTokenStreamHeader* GetStreamHeader(FTokenStreamView View)
-{
-	check(View.GetTokenData());
-	return reinterpret_cast<FTokenStreamHeader*>(View.GetTokenData()) - 1;
-}
 
-static void DropStreamReference(FTokenStreamView View)
-{
-	if (View.GetTokenData())
+	// Emit schema members
+	if (bHasPropertiesWithObjectReferences)
 	{
-		if (GetStreamHeader(View)->RefCount.fetch_sub(1) == 1)
+		for (int32 Idx = 0, Num = ArrayDim; Idx < Num; ++Idx)
 		{
-			UpdateGlobalTokenCount(-int32(View.NumTokens()));
-			FMemory::Free(GetStreamHeader(View));
+			for (FProperty* Property = Struct->PropertyLink; Property; Property = Property->PropertyLinkNext)
+			{
+				Property->EmitReferenceInfo(Schema, Offset + Idx * ElementSize, EncounteredStructProps, DebugPath);
+			}
 		}
 	}
 }
 
-static void AddStreamReference(FTokenStreamView View)
+void FInterfaceProperty::EmitReferenceInfo(UE::GC::FSchemaBuilder& Schema, int32 BaseOffset, TArray<const FStructProperty*>& EncounteredStructProps, UE::GC::FPropertyStack& DebugPath)
 {
-	if (View.GetTokenData())
+	const int32 Offset = BaseOffset + GetOffset_ForGC();
+	for (int32 Idx = 0, Num = ArrayDim; Idx < Num; ++Idx)
 	{
-		GetStreamHeader(View)->RefCount.fetch_add(1);
+		Schema.Add(UE::GC::DeclareMember(DebugPath, Offset + Idx * sizeof(FScriptInterface), UE::GC::EMemberType::Reference) );
 	}
 }
 
-static void UpdateStreamReference(FTokenStreamView& Out, FTokenStreamView New)
+void FFieldPathProperty::EmitReferenceInfo(UE::GC::FSchemaBuilder& Schema, int32 BaseOffset, TArray<const FStructProperty*>& EncounteredStructProps, UE::GC::FPropertyStack& DebugPath)
 {
-	if (Out.GetTokenData() != New.GetTokenData())
+	static_assert(sizeof(FFieldPath) == sizeof(TFieldPath<FProperty>), "TFieldPath should have the same size as the underlying FFieldPath");
+	const int32 Offset = BaseOffset + GetOffset_ForGC();
+	for (int32 Idx = 0, Num = ArrayDim; Idx < Num; ++Idx)
 	{
-		DropStreamReference(Out);
-		AddStreamReference(New);
-		Out = New;
+		Schema.Add(UE::GC::DeclareMember(DebugPath, Offset + Idx * sizeof(FFieldPath), UE::GC::EMemberType::FieldPath));
 	}
 }
 
-static FGCReferenceInfo GEndOfStreamToken{GCRT_EndOfStream, 0};
-
-uint32 UE::GC::FTokenStreamBuilder::MaxStackSize;
-
-
-void FTokenStreamBuilder::EmitObjectReference(int32 Offset, FName DebugName, EGCReferenceType Kind)
-{
-	FGCReferenceInfo ObjectReference(Kind, Offset);
-	int32 Index = EmitReferenceInfo(ObjectReference, DebugName);
-}
-
-void FTokenStreamBuilder::EmitObjectArrayReference(int32 Offset, FName DebugName)
-{
-	check(DebugOwner.HasAnyClassFlags(CLASS_Intrinsic));
-	EmitObjectReference(Offset, DebugName, GCRT_ArrayObject);
-}
-
-uint32 FTokenStreamBuilder::EmitStructArrayBegin(int32 Offset, FName DebugName, int32 Stride)
-{
-	check(DebugOwner.HasAnyClassFlags(CLASS_Intrinsic));
-	EmitObjectReference(Offset, DebugName, GCRT_ArrayStruct);
-	EmitStride(Stride);
-	return EmitSkipIndexPlaceholder();
-}
-
-/**
- * Realtime garbage collection helper function used to indicate the end of an array of structs. The
- * index following the current one will be written to the passed in SkipIndexIndex in order to be
- * able to skip tokens for empty dynamic arrays.
- *
- * @param SkipIndexIndex
- */
-void FTokenStreamBuilder::EmitStructArrayEnd(uint32 SkipIndexIndex)
-{
-	check(DebugOwner.HasAnyClassFlags(CLASS_Intrinsic));
-	uint32 SkipIndex = EmitReturn();
-	UpdateSkipIndexPlaceholder( SkipIndexIndex, SkipIndex );
-}
-
-void FTokenStreamBuilder::EmitFixedArrayBegin(int32 Offset, FName DebugName, int32 Stride, int32 Count)
-{
-	check(DebugOwner.HasAnyClassFlags(CLASS_Intrinsic));
-	EmitObjectReference(Offset, DebugName, GCRT_FixedArray);
-	EmitStride(Stride);
-	EmitCount(Count);
-}
-
-void FTokenStreamBuilder::EmitFixedArrayEnd()
-{
-	check(DebugOwner.HasAnyClassFlags(CLASS_Intrinsic));
-	EmitReturn();
-}
-
-// FTokenStreamBuilderIterator helper
-struct FNamedReferenceInfo
-{
-	FGCReferenceInfo ReferenceInfo;
-	FName DebugName;
-};
-
-static AROFunc GetAROFunc(UClass* Class)
-{
-	AROFunc ARO = Class->CppClassStaticFunctions.GetAddReferencedObjects();
-	check(ARO != nullptr);
-	return ARO != &UObject::AddReferencedObjects ? ARO : nullptr;
-}
-
-} // namespace UE::GC
-
-struct FScopeLockIfNotNative
-{
-	FCriticalSection& ScopeCritical;
-	const bool bNotNative;
-	FScopeLockIfNotNative(FCriticalSection& InScopeCritical, bool bIsNotNative)
-		: ScopeCritical(InScopeCritical)
-		, bNotNative(bIsNotNative)
-	{
-		if (bNotNative)
-		{
-			ScopeCritical.Lock();
-		}
-	}
-	~FScopeLockIfNotNative()
-	{
-		if (bNotNative)
-		{
-			ScopeCritical.Unlock();
-		}
-	}
-};
-
-/** Token stream can assemble code can sometimes be called from two threads throuh a web of async loading calls. */
-static FCriticalSection GReferenceTokenStreamCritical;
+/** Both game thread and async loading can assemble schemas */
+static FCriticalSection GAssembleSchemaLock;
 
 void UClass::AssembleReferenceTokenStream(bool bForce)
 {
@@ -5439,406 +5250,129 @@ void UClass::AssembleReferenceTokenStream(bool bForce)
 	}
 	else
 	{
-		FScopeLock NonNativeLock(&GReferenceTokenStreamCritical);
+		FScopeLock NonNativeLock(&GAssembleSchemaLock);
 		AssembleReferenceTokenStreamInternal(bForce);
 	}
 }
 
+static UE::GC::ObjectAROFn GetARO(UClass* Class)
+{
+	UE::GC::ObjectAROFn ARO = Class->CppClassStaticFunctions.GetAddReferencedObjects();
+	check(ARO != nullptr);
+	return ARO != &UObject::AddReferencedObjects ? ARO : nullptr;
+}
+
 void UClass::AssembleReferenceTokenStreamInternal(bool bForce)
 {
 	using namespace UE::GC;
 
 	if (!HasAnyClassFlags(CLASS_TokenStreamAssembled) || bForce)
 	{
-		FTokenStreamBuilder StackBuilder(*this);
-		TUniquePtr<FTokenStreamBuilder> IntrinsicBuilder;
-		if (ClassFlags & CLASS_Intrinsic)
+		if (bForce)
 		{
-			IntrinsicBuilder = FIntrinsicClassTokens::ConsumeBuilder(this);
+			ClassFlags &= ~CLASS_TokenStreamAssembled;
 		}
-		FTokenStreamBuilder& Builder = IntrinsicBuilder ? *IntrinsicBuilder : StackBuilder;
 
-		if (bForce)
+		FSchemaBuilder Schema(/* don't store sizeof(class) to enable super class schema reuse */ 0);
+		FSchemaView SuperSchema;
+		if (UClass* SuperClass = GetSuperClass())
 		{
-			ClassFlags &= ~CLASS_TokenStreamAssembled;
+			SuperClass->AssembleReferenceTokenStreamInternal();	
+			SuperSchema = SuperClass->ReferenceSchema.Get();
+			Schema.Append(SuperSchema);
 		}
+		const int32 NumSuperMembers = Schema.NumMembers();
 
 		{
-			FGCStackSizeHelper StackSizeHelper;
+			FPropertyStack DebugPath;
 			TArray<const FStructProperty*> EncounteredStructProps;
 
 			// Iterate over properties defined in this class
 			for (TFieldIterator<FProperty> It(this, EFieldIteratorFlags::ExcludeSuper); It; ++It)
 			{
 				FProperty* Property = *It;
-				FGCStackSizeHelperScope StackSizeScope(StackSizeHelper, Property);
-				Property->EmitReferenceInfo(Builder, 0, EncounteredStructProps, StackSizeHelper);
+				FPropertyStackScope PropertyScope(DebugPath, Property);
+				Property->EmitReferenceInfo(Schema, 0, EncounteredStructProps, DebugPath);
 			}
-
-			Builder.SetStackSize(StackSizeHelper.GetMaxStackSize());
 		}
-				
-		FTokenStreamView SuperTokens;
-		if (UClass* SuperClass = GetSuperClass())
+		
+		if (ClassFlags & CLASS_Intrinsic)
 		{
-			// Make sure super class has valid token stream.
-			SuperClass->AssembleReferenceTokenStreamInternal();
-			SuperTokens = FTokenStreamBuilder::DropFinalTokens(SuperClass->ReferenceTokens.Strong, GetAROFunc(SuperClass));
+			Schema.Append(UE::GC::GetIntrinsicSchema(this));
 		}
-
+		
 		// Make sure all Blueprint properties are marked as non-native
-		Builder.SetTokenType(GetClass()->HasAnyClassFlags(CLASS_NeedsDeferredDependencyLoading) ? EGCTokenType::NonNative : EGCTokenType::Native);
-		Builder.EmitFinalTokens(GetAROFunc(this));
+		// @todo Currently native super class properties of BP base classes are incorrectly marked as Blueprint
+		// @todo Investigate if CLASS_CompiledFromBlueprint is better to avoid reference eliminating non-native non-blueprint properties
+		EOrigin Origin = GetClass()->HasAnyClassFlags(CLASS_NeedsDeferredDependencyLoading) ? EOrigin::Blueprint : EOrigin::Other;
 
-		// Allocate flat token stream with current and super class tokens
-		FTokenStreamBuilder::Merge(/* out */ ReferenceTokens.Strong, Builder, SuperTokens);
+		bool bReuseSuper = Schema.NumMembers() == NumSuperMembers && NumSuperMembers > 0 && GetARO(this) == GetARO(GetSuperClass());
+		FSchemaView View(bReuseSuper ? SuperSchema : Schema.Build(GetARO(this)), Origin);
+		ReferenceSchema.Set(View);
 		
-		checkf(!HasAnyClassFlags(CLASS_TokenStreamAssembled), TEXT("Token stream already assembled for class '%s'"), *GetPathName()); // recursion here is probably bad
+		checkf(!HasAnyClassFlags(CLASS_TokenStreamAssembled), TEXT("GC schema already assembled for class '%s'"), *GetPathName()); // recursion here is probably bad
 		ClassFlags |= CLASS_TokenStreamAssembled;
 	}
 }
 
-namespace UE::GC
-{
-
 //////////////////////////////////////////////////////////////////////////
-
-void FTokenStreamBuilder::EmitFinalTokens(AROFunc ARO)
-{
-	if (ARO)
-	{
-		int32 SlowIndex = FindSlowImplementation(ARO);
-		if (SlowIndex == INDEX_NONE)
-		{
-			EmitReferenceInfo(FGCReferenceInfo(GCRT_AddReferencedObjects, 0), TokenDebugNames::ARO);
-			EmitPointer(reinterpret_cast<const void*>(ARO));
-		}
-		else
-		{
-			EmitReferenceInfo(FGCReferenceInfo(GCRT_SlowAddReferencedObjects, static_cast<uint32>(SlowIndex)), TokenDebugNames::ARO);
-		}
-	}
-
-	// Emit end of stream token.
-	EmitReferenceInfo(GEndOfStreamToken, TokenDebugNames::EndOfStream);
-}
-
- FTokenStreamView FTokenStreamBuilder::DropFinalTokens(FTokenStreamView Stream, AROFunc DropARO)
- {
-	check(Stream.DroppedNum == 0)
-	if (Stream.Num >= 5 &&
-		FGCReferenceInfo(Stream.Tokens[Stream.Num - 5]).Type == GCRT_AddReferencedObjects &&
-		FGCReferenceInfo(Stream.Tokens[Stream.Num - 2]).Type == GCRT_EndOfPointer &&
-		FGCReferenceInfo(Stream.Tokens[Stream.Num - 1]).Type == GCRT_EndOfStream)
-	{
-		checkf(DropARO, TEXT("Super class token stream assembled with ClassAddReferencedObjects set and now it has been reset."
-							" Super class token stream is outdated and should have been regenerated. "));
-		Stream.Num -= 5;
-		Stream.DroppedNum = 5;
-	}
-	else if (Stream.Num >= 2 &&
-		FGCReferenceInfo(Stream.Tokens[Stream.Num - 2]).Type == GCRT_SlowAddReferencedObjects &&
-		FGCReferenceInfo(Stream.Tokens[Stream.Num - 1]).Type == GCRT_EndOfStream)
-	{
-		checkf(DropARO, TEXT("Super class token stream assembled with ClassAddReferencedObjects set and now it has been reset."
-							" Super class token stream is outdated and should have been regenerated. "));
-		Stream.Num -= 2;
-		Stream.DroppedNum = 2;
-	}
-	else if (Stream.Num >= 1)
-	{
-		check(FGCReferenceInfo(Stream.Tokens[Stream.Num - 1]).Type == GCRT_EndOfStream);
-		Stream.Num -= 1;
-		Stream.DroppedNum = 1;
-	}
-
-	return Stream;
-}
-
-void FTokenStreamBuilder::Merge(FTokenStreamView& Out, const FTokenStreamBuilder& Class, FTokenStreamView Super) 
-{
-	check(Class.Tokens.Num() && FGCReferenceInfo(Class.Tokens.Last()).Type == GCRT_EndOfStream);
-
-	uint32 Num = Class.Num() + Super.Num;
-	if (Num == 1)
-	{
-		// Return empty stream if there's only an EndOfStream terminator	
-		check(Super.IsEmpty() && Class.Tokens[0] == GEndOfStreamToken);
-		UpdateStreamReference(Out, FTokenStreamView());
-	}
-	else if (Num == Out.Num &&
-		!FMemory::Memcmp(Out.Tokens,				Super.Tokens,			sizeof(uint32) * Super.Num) &&
-		!FMemory::Memcmp(Out.Tokens + Super.Num,	Class.Tokens.GetData(),	sizeof(uint32) * Class.Num()))
-	{
-		// Nothing changed, no need to update
-		check(Out.IsNonNative == (Class.TokenType == EGCTokenType::NonNative));
-	}
-	else if (Class.Tokens.Num() == Super.DroppedNum && 
-		!FMemory::Memcmp(Super.Tokens + Super.Num,	Class.Tokens.GetData(), sizeof(uint32) * Class.Num()))
-	{
-		// Reuse Super tokens since Class only contain Super's dropped final tokens
-		FTokenStreamView New = Super;
-		New.Num += New.DroppedNum;
-		New.DroppedNum = 0;
-		New.IsNonNative = Class.TokenType == EGCTokenType::NonNative;
-		UpdateStreamReference(Out, New);
-	}
-	else
-	{
-		// Allocate new flat stream
-		LLM_SCOPE(ELLMTag::GC);
-	
-		static_assert(alignof(FTokenStreamHeader) == alignof(uint32));
-		static_assert(alignof(FName) == alignof(uint32));
-		SIZE_T NumBytes = sizeof(FTokenStreamHeader) + Num * sizeof(uint32) + (ENABLE_GC_TOKEN_DEBUG_INFO ? Num * sizeof(FName) : 0);
-		FTokenStreamHeader* Header = new (FMemory::Malloc(NumBytes)) FTokenStreamHeader;
-		uint32* Tokens = reinterpret_cast<uint32*>(Header + 1);
-
-		if (Super.Num)
-		{
-			FMemory::Memcpy(Tokens,             Super.Tokens,           sizeof(uint32) * Super.Num  );
-		}
-		if (Class.Num())
-		{
-			FMemory::Memcpy(Tokens + Super.Num, Class.Tokens.GetData(), sizeof(uint32) * Class.Num());
-		}
-
-	#if ENABLE_GC_TOKEN_DEBUG_INFO
-		check(Class.DebugNames.Num() == Class.Tokens.Num());
-		FName* DebugNames = reinterpret_cast<FName*>(Tokens + Num);
-		if (Super.Num)   { FMemory::Memcpy(DebugNames,             Super.GetDebugNames(),      sizeof(FName) * Super.Num); }
-		if (Class.Num()) { FMemory::Memcpy(DebugNames + Super.Num, Class.DebugNames.GetData(), sizeof(FName) * Class.Num()); }
-	#endif
-
-		UpdateGlobalTokenCount(int32(Num));
-
-		// Update the global maximum stack size after assembling the token stream of each class
-		uint32 StackSize = FMath::Max(Class.StackSize, Super.StackSize);
-		MaxStackSize = FMath::Max(MaxStackSize, StackSize);
-		check(StackSize < (1u << FTokenStreamView::NumStackSizeBits));
-
-		FTokenStreamView New;
-		New.Tokens = Tokens;
-		New.Num = Num;
-		New.StackSize = StackSize;
-		New.IsNonNative = Class.TokenType == EGCTokenType::NonNative;
-
-		DropStreamReference(Out);
-		Out = New;
-	}
-}
-
-static uint32 PickReturnCount(uint32 OuterReturnCount, uint32 InnerReturnCount)
-{
-	if (InnerReturnCount > 1)
-	{
-		check(OuterReturnCount == 0);
-		return InnerReturnCount - 1;
-	}
-
-	return OuterReturnCount;
-}
-
-int32 FTokenStreamBuilder::EmitReferenceInfo(FGCReferenceInfo ReferenceInfo, FName DebugName)
-{
-	int32 TokenIndex = Tokens.Add(ReferenceInfo);
-#if ENABLE_GC_TOKEN_DEBUG_INFO
-	check(DebugNames.Num() == TokenIndex);
-	check(!DebugName.IsNone());
-	DebugNames.Add(DebugName);
-#endif
-	return TokenIndex;
-}
-
-uint32 FTokenStreamBuilder::EmitSkipIndexPlaceholder()
-{
-	uint32 TokenIndex = Tokens.Add(SkipIndexPlaceholderMagic);
-#if ENABLE_GC_TOKEN_DEBUG_INFO
-	check(DebugNames.Num() == TokenIndex);
-	DebugNames.Add(TokenDebugNames::SkipInfo);
-#endif
-	return TokenIndex;
-}
-
-void FTokenStreamBuilder::UpdateSkipIndexPlaceholder( uint32 SkipIndexIndex, uint32 SkipIndex )
-{
-	check( SkipIndex > 0 && SkipIndex <= (uint32)Tokens.Num() );
-	FGCReferenceInfo ReferenceInfo(Tokens[SkipIndex-1]);
-	check( ReferenceInfo.Type != GCRT_None );
-	check( Tokens[SkipIndexIndex] == SkipIndexPlaceholderMagic );
-	check( SkipIndexIndex < SkipIndex );
-	check( ReferenceInfo.ReturnCount >= 1 );
-	FGCSkipInfo SkipInfo;
-	SkipInfo.SkipIndex			= SkipIndex - SkipIndexIndex;
-	// We need to subtract 1 as ReturnCount includes return from this array.
-	SkipInfo.InnerReturnCount	= ReferenceInfo.ReturnCount - 1; 
-	Tokens[SkipIndexIndex]		= SkipInfo;
-}
-
-int32 FTokenStreamBuilder::EmitCount( uint32 Count )
-{
-	int32 TokenIndex = Tokens.Add( Count );
-#if ENABLE_GC_TOKEN_DEBUG_INFO
-	check(DebugNames.Num() == TokenIndex);
-	DebugNames.Add(TokenDebugNames::Count);
-#endif
-	return TokenIndex;
-}
-
-int32 FTokenStreamBuilder::EmitPointer( void const* Ptr )
-{
-	static constexpr uint32 NumTokensPerPointer = sizeof(void*) / sizeof(uint32); //-V514
-
-	const int32 StoreIndex = Tokens.Num();
-	Tokens.AddUninitialized(NumTokensPerPointer);
-	StorePointer(&Tokens[StoreIndex], Ptr);
-
-#if ENABLE_GC_TOKEN_DEBUG_INFO
-	check(DebugNames.Num() == StoreIndex);
-	for (int32 PointerTokenIndex = 0; PointerTokenIndex < NumTokensPerPointer; ++PointerTokenIndex)
-	{
-		DebugNames.Add(TokenDebugNames::Pointer);
-	}
-#endif
-
-	// Now inser the end of pointer marker, this will mostly be used for storing ReturnCount value
-	// if the pointer was stored at the end of struct array stream.
-	EmitReferenceInfo(FGCReferenceInfo(GCRT_EndOfPointer, 0), TokenDebugNames::EndOfPointer);
-
-	return StoreIndex;
-}
-
-int32 FTokenStreamBuilder::EmitStride( uint32 Stride )
-{
-	int32 TokenIndex = Tokens.Add( Stride );
-#if ENABLE_GC_TOKEN_DEBUG_INFO
-	check(DebugNames.Num() == TokenIndex);
-	DebugNames.Add(TokenDebugNames::Stride);
-#endif
-	return TokenIndex;
-}
-
-uint32 FTokenStreamBuilder::EmitReturn()
+namespace UE::GC
 {
-	FGCReferenceInfo ReferenceInfo(Tokens.Last());
-	check(ReferenceInfo.Type != GCRT_None);
-	ReferenceInfo.ReturnCount++;
-	Tokens.Last() = ReferenceInfo;
-	return Tokens.Num();
-}
-
-//////////////////////////////////////////////////////////////////////////
 
-static FName ToName(ETokenlessId Id)
+static FName ToName(EMemberlessId Id)
 {
-	static_assert(ETokenlessId::Collector == static_cast<ETokenlessId>(1u));
-	static_assert(ETokenlessId::Max == ETokenlessId::InitialReference);
+	static_assert(EMemberlessId::Collector == static_cast<EMemberlessId>(1u));
+	static_assert(EMemberlessId::Max == EMemberlessId::InitialReference);
 
 	static FName Names[] =
 	{
 		FName("Invalid"),
-		FName("Collector"), // ETokenlessId::Collector
-		FName("Class"), // ETokenlessId::Class
-		FName("Outer"), // ETokenlessId::Outer
-		FName("ExternalPackage"), // ETokenlessId::ExternalPackage
-		FName("ClassOuter"), // ETokenlessId::ClassOuter
-		FName("InitialReference"), // ETokenlessId::InitialReference
+		FName("Collector"), // EMemberlessId::Collector
+		FName("Class"), // EMemberlessId::Class
+		FName("Outer"), // EMemberlessId::Outer
+		FName("ExternalPackage"), // EMemberlessId::ExternalPackage
+		FName("ClassOuter"), // EMemberlessId::ClassOuter
+		FName("InitialReference"), // EMemberlessId::InitialReference
 	};
 
 	check(static_cast<uint32>(Id) < UE_ARRAY_COUNT(Names));
 	return Names[static_cast<uint32>(Id)];
 }
 
-void FTokenId::StaticAssert()
-{
-	static_assert(1u << TokenlessIdBits >= static_cast<uint32>(ETokenlessId::Max), "Need to bump TokenlessIdBits");
-}
-
-//////////////////////////////////////////////////////////////////////////
-
-FTokenInfo FTokenStreamOwner::GetTokenInfo(FTokenId Id) const
-{
-	if (Id.IsTokenless())
-	{
-		// Technically offsetof(UObjectBase, ClassPrivate/OuterPrivate) could be exposed via some new API or friend declaration
-		return {-1, ToName(Id.AsTokenless()) };
-	}
-
-	uint32 Idx = Id.GetIndex();
-
-	check(Idx < Strong.Num);
-	FTokenInfo DebugInfo;
-	DebugInfo.Offset = FGCReferenceInfo(Strong.Tokens[Idx]).Offset;
-#if ENABLE_GC_TOKEN_DEBUG_INFO
-	DebugInfo.Name = Strong.GetDebugNames()[Idx];
-#else
-	static const FName DefaultTokenName("Token");	
-	DebugInfo.Name = DefaultTokenName; 
-#if !UE_FNAME_OUTLINE_NUMBER
-	DebugInfo.Name.SetNumber(Idx);
-#endif
-#endif
-	return DebugInfo;
-}
-
-void FTokenStreamOwner::Reset()
-{
-	UpdateStreamReference(/* out */ Strong, FTokenStreamView());
-}
-
-int64 FTokenStreamOwner::GetTokenAllocatedSize() const
-{
-	return Strong.Num * sizeof(uint32);
-}
-
-int64 FTokenStreamOwner::GetDebugInfoAllocatedSize() const
-{
-#if ENABLE_GC_TOKEN_DEBUG_INFO
-	return GetTokenAllocatedSize() * sizeof(FName) / sizeof(uint32);
-#else
-	return 0;
-#endif
-}
-
-TMap<UClass*, TUniquePtr<FTokenStreamBuilder>> GIntrinsicClassTokens;
-
-FTokenStreamBuilder& FIntrinsicClassTokens::AllocateBuilder(UClass* Class)
-{
-	TUniquePtr<FTokenStreamBuilder>& Builder = GIntrinsicClassTokens.FindOrAdd(Class);
-	checkf(!Builder, TEXT("Allocated multiple builders for same class"));
-	Builder.Reset(new FTokenStreamBuilder(*Class));
-	return *Builder;
-}
-
-TUniquePtr<FTokenStreamBuilder> FIntrinsicClassTokens::ConsumeBuilder(UClass* Class)
+void FMemberId::StaticAssert()
 {
-	TUniquePtr<FTokenStreamBuilder>* Builder = GIntrinsicClassTokens.Find(Class);
-	return Builder && *Builder ? MoveTemp(*Builder) : TUniquePtr<FTokenStreamBuilder>();
+	static_assert((1u << MemberlessIdBits) >= static_cast<uint32>(EMemberlessId::Max), "Need to bump MemberlessIdBits");
 }
 
 //////////////////////////////////////////////////////////////////////////
+#if !UE_BUILD_SHIPPING
 
-void DumpMemoryStats(FOutputDevice& OutputDevice)
+static void DumpMemoryStats(FOutputDevice& OutputDevice)
 {
 	FContextPoolScope Pool;
-	int32 TotalWeakCapacity = 0;
-	int32 NumContexts = Pool.PeekFree().Num();
+	uint32 TotalWeakCapacity = 0;
+	uint32 NumContexts = Pool.PeekFree().Num();
 	for (const TUniquePtr<FWorkerContext>& Context : Pool.PeekFree())
 	{
 		TotalWeakCapacity += Context->WeakReferences.Max();
 	}
 
-	int32 ContextsKB = (NumContexts * sizeof(FWorkerContext) + TotalWeakCapacity * sizeof(UObject**))/1024;
-	int32 ScratchKB = static_cast<int32>(GScratchPages.CountBytes() / 1024);
-	int32 GlobalsKB = (sizeof(GScratchPages) + sizeof(GAROBlocks) + sizeof(GWorkstealingManager) + sizeof(GSlowARO)) / 1024;
-	// Excluding debug names in non-shipping configs
-	int32 TokensKB = (GNumTokens * sizeof(uint32) + GNumTokenStreams * sizeof(FTokenStreamHeader)) / 1024; 
-	int32 TotalKB = ContextsKB + ScratchKB + GlobalsKB + TokensKB;
+	uint32 NumSchemaWords = 0;
+	uint32 NumSchemas = CountSchemas(NumSchemaWords);
+	uint32 SchemaKB = (NumSchemas * sizeof(FSchemaHeader) + NumSchemaWords * 8) / 1024;
+	uint32 ContextsKB = (NumContexts * sizeof(FWorkerContext) + TotalWeakCapacity * sizeof(UObject**))/1024;
+	uint32 ScratchKB = static_cast<int32>(GScratchPages.CountBytes() / 1024);
+	uint32 GlobalsKB = (sizeof(GScratchPages) + sizeof(GAROBlocks) + sizeof(GWorkstealingManager) + sizeof(GSlowARO)) / 1024;
+	
+	uint32 TotalKB = ContextsKB + ScratchKB + GlobalsKB + SchemaKB;
 
-	OutputDevice.Logf(TEXT("GC mem usage: %dKB total, %dKB scratch, %dKB for %d contexts, %dKB globals (code size) and %dKB for %d tokens in %d streams"),
-						TotalKB, ScratchKB, ContextsKB, NumContexts, GlobalsKB, TokensKB, GNumTokens, GNumTokenStreams);
+	OutputDevice.Logf(TEXT("GC mem usage: %dKB total, %dKB scratch, %dKB for %d contexts and %dKB globals (code size) and %dKB for %d schemas"),
+						TotalKB, ScratchKB, ContextsKB, NumContexts, GlobalsKB, SchemaKB, NumSchemas);
 }
 
+static FAutoConsoleCommandWithOutputDevice GDumpMemoryStats(TEXT("gc.DumpMemoryStats"), TEXT("Print GC memory usage"), FConsoleCommandWithOutputDeviceDelegate::CreateStatic(&DumpMemoryStats));
+
+#endif //!UE_BUILD_SHIPPING
 //////////////////////////////////////////////////////////////////////////
 
 // Coordinates worker starting, tail spinning and stopping
diff --git a/Engine/Source/Runtime/CoreUObject/Private/UObject/GarbageCollectionSchema.cpp b/Engine/Source/Runtime/CoreUObject/Private/UObject/GarbageCollectionSchema.cpp
new file mode 100644
--- /dev/null
+++ b/Engine/Source/Runtime/CoreUObject/Private/UObject/GarbageCollectionSchema.cpp
@@ -0,0 +1,970 @@
+// Copyright Epic Games, Inc. All Rights Reserved.
+
+#include "UObject/GarbageCollectionSchema.h"
+#include "UObject/FastReferenceCollector.h"
+#include "UObject/UObjectIterator.h"
+#include "UObject/UnrealType.h"
+
+namespace UE::GC {
+
+//////////////////////////////////////////////////////////////////////////
+
+static constexpr uint32 ToBit(EMemberType Type)		{ return 1u << static_cast<uint32>(Type); }
+static constexpr uint32 InnerSchemaTypes =			ToBit(EMemberType::StructArray) | ToBit(EMemberType::SparseStructArray) | ToBit(EMemberType::FreezableStructArray) | ToBit(EMemberType::Optional);
+static constexpr uint32 AROTypes =					ToBit(EMemberType::ARO) | ToBit(EMemberType::SlowARO) | ToBit(EMemberType::MemberARO);
+static constexpr uint32 AROPointerTypes =			ToBit(EMemberType::ARO) | ToBit(EMemberType::MemberARO);
+static constexpr uint32 ExtraWordTypes =			InnerSchemaTypes | AROPointerTypes | ToBit(EMemberType::StridedArray);
+static constexpr uint32 StopTypes =					ToBit(EMemberType::Stop) | ToBit(EMemberType::ARO) | ToBit(EMemberType::SlowARO);
+static constexpr uint32 GeneratedTypes =			StopTypes | ToBit(EMemberType::Jump) | ToBit(EMemberType::StridedArray);
+static constexpr uint32 MemberTypeCount =			uint8(EMemberType::Count);
+
+template<uint32 Types>
+static constexpr bool IsIn(EMemberType Type)
+{
+	return !!(Types & ToBit(Type));
+}
+
+static uint8 ToInt(EMemberType Type)
+{
+	check((uint8)Type < MemberTypeCount);
+	return (uint8)Type;
+}
+
+FName ToName(EMemberType Type)
+{
+	struct FMemberTypeNames
+	{
+		FMemberTypeNames()
+		{
+			Names[(uint8)EMemberType::Stop] =							"Stop";
+			Names[(uint8)EMemberType::Jump] =							"Jump";
+			Names[(uint8)EMemberType::Reference] =						"Reference";
+			Names[(uint8)EMemberType::ReferenceArray] =					"ReferenceArray";
+			Names[(uint8)EMemberType::StructArray] =					"StructArray";
+			Names[(uint8)EMemberType::StridedArray] =					"StridedArray";
+			Names[(uint8)EMemberType::SparseStructArray] =				"SparseStructArray";
+			Names[(uint8)EMemberType::FieldPath] =						"FieldPath";
+			Names[(uint8)EMemberType::FieldPathArray] =					"FieldPathArray";
+			Names[(uint8)EMemberType::FreezableReferenceArray] =		"FreezableReferenceArray";
+			Names[(uint8)EMemberType::FreezableStructArray] =			"FreezableStructArray";
+			Names[(uint8)EMemberType::Optional] =						"Optional";
+			Names[(uint8)EMemberType::DynamicallyTypedValue] =			"DynamicallyTypedValue";
+			Names[(uint8)EMemberType::ARO] =							"ARO";
+			Names[(uint8)EMemberType::SlowARO] =						"SlowARO";
+			Names[(uint8)EMemberType::MemberARO] =						"MemberARO";
+		}
+
+		FName Names[MemberTypeCount];
+	};
+
+	static FMemberTypeNames MemberTypes;
+	check(!MemberTypes.Names[ToInt(Type)].IsNone());
+	return MemberTypes.Names[ToInt(Type)];
+}
+
+//////////////////////////////////////////////////////////////////////////
+
+static_assert(sizeof(FMemberWord) == 8 && alignof(FMemberWord) == 8);
+
+static bool IsSet(FMemberWord Word)
+{
+	return !!reinterpret_cast<const uint64&>(Word);
+}
+
+static FMemberWord ToWord(ObjectAROFn In)
+{
+	FMemberWord Out;
+	Out.ObjectARO = In;
+	return Out;
+}
+
+//////////////////////////////////////////////////////////////////////////
+
+FSchemaHeader& FSchemaView::GetHeader()
+{
+	check(GetWords());
+	return *(reinterpret_cast<FSchemaHeader*>(Handle & ~OriginBit) - 1);
+}
+
+FSchemaHeader* FSchemaView::TryGetHeader()
+{
+	return IsEmpty() ? nullptr : &GetHeader();
+}
+
+//////////////////////////////////////////////////////////////////////////
+
+template<typename Fn>
+void VisitInnerSchemas(FSchemaView Schema, Fn Visitor)
+{
+	check(!Schema.IsEmpty());
+	const FMemberWord* WordIt = Schema.GetWords();
+	while (true)
+	{
+		FMemberWord Word = *WordIt++;
+		for (Private::FMemberUnpacked Member : Word.Members)
+		{
+			if (IsIn<StopTypes>(Member.Type))
+			{
+				return;
+			}
+			else if (IsIn<ExtraWordTypes>(Member.Type))
+			{
+				FMemberWord ExtraWord = *WordIt++;
+				if (IsIn<InnerSchemaTypes>(Member.Type))
+				{
+					Visitor(ExtraWord.InnerSchema);
+				}
+			}
+		}
+	}		
+}
+
+static void AddSchemaReference(FSchemaView Schema)
+{
+	Schema.GetHeader().RefCount.fetch_add(1);
+}
+
+static void DropSchemaReference(FSchemaView Schema)
+{
+	FSchemaHeader& Header = Schema.GetHeader();
+	if (Header.RefCount.fetch_sub(1) == 1)
+	{
+		VisitInnerSchemas(Schema, &DropSchemaReference);
+		FMemory::Free(&Header);
+	}
+}
+
+static void TryAddSchemaReference(FSchemaView Schema)
+{
+	if (!Schema.IsEmpty())
+	{
+		AddSchemaReference(Schema);
+	}
+}
+
+static void TryDropSchemaReference(FSchemaView Schema)
+{
+	if (!Schema.IsEmpty())
+	{
+		DropSchemaReference(Schema);
+	}
+}
+
+static void TryAddSchemaReference(FMemberDeclaration& Member)
+{
+	if (IsIn<InnerSchemaTypes>(Member.Type))
+	{
+		AddSchemaReference(Member.ExtraWord.InnerSchema);
+	}
+}
+
+static void TryDropSchemaReference(FMemberDeclaration& Member)
+{
+	if (IsIn<InnerSchemaTypes>(Member.Type))
+	{
+		DropSchemaReference(Member.ExtraWord.InnerSchema);
+	}
+}
+
+static void ReplaceSchemaReference(FSchemaView& Out, FSchemaView New)
+{
+	if (Out.GetWords() != New.GetWords())
+	{
+		TryAddSchemaReference(New);
+		TryDropSchemaReference(Out);
+		Out = New;
+	}
+}
+
+FSchemaOwner::FSchemaOwner(FSchemaView Schema)
+: SchemaView(reinterpret_cast<uint64&>(Schema))
+{
+	static_assert(sizeof(FSchemaOwner) == sizeof(FSchemaView) && alignof(FSchemaOwner) == alignof(FSchemaView) && offsetof(FSchemaOwner, SchemaView) == 0);
+	TryAddSchemaReference(Schema);
+}
+
+void FSchemaOwner::Set(FSchemaView In)
+{
+	ReplaceSchemaReference(reinterpret_cast<FSchemaView&>(SchemaView), In);
+}
+
+void FSchemaOwner::Reset()
+{
+	TryDropSchemaReference(Get());
+	SchemaView = 0;
+}
+
+//////////////////////////////////////////////////////////////////////////
+
+static FName ToName(EMemberlessId Id)
+{
+	static_assert(EMemberlessId::Collector == static_cast<EMemberlessId>(1u));
+	static_assert(EMemberlessId::Max == EMemberlessId::InitialReference);
+
+	static FName Names[] =
+	{
+		FName(EName::Error),
+		FName("Collector"), // EMemberlessId::Collector
+		FName("Class"), // EMemberlessId::Class
+		FName("Outer"), // EMemberlessId::Outer
+		FName("ExternalPackage"), // EMemberlessId::ExternalPackage
+		FName("ClassOuter"), // EMemberlessId::ClassOuter
+		FName("InitialReference"), // EMemberlessId::InitialReference
+	};
+
+	check(static_cast<uint32>(Id) < UE_ARRAY_COUNT(Names));
+	return Names[static_cast<uint32>(Id)];
+}
+
+FMemberInfo GetMemberDebugInfo(FSchemaView Schema, FMemberId Id)
+{
+	if (Id.IsMemberless())
+	{
+		// Technically offsetof(UObjectBase, ClassPrivate/OuterPrivate) could be exposed via some new API or friend declaration
+		return {-1, ToName(Id.AsMemberless()) };	
+	}
+
+	int32 Offset = -1;
+	int32 DebugNameIdx = -1;
+	const uint32 FindIdx = Id.GetIndex();
+	uint32 MemberIdx = 0;
+	uint32 JumpedWords = 0;
+	uint32 NumJumps = 0;
+	const FMemberWord* WordIt = Schema.GetWords();
+	while (true)
+	{
+		for (Private::FMemberUnpacked Member : WordIt++->Members)
+		{
+			WordIt += IsIn<ExtraWordTypes>(Member.Type);
+
+			if (MemberIdx++ == FindIdx)
+			{
+				if (IsIn<StopTypes | ToBit(EMemberType::Jump)>(Member.Type))
+				{
+					return { IntCastChecked<int32>(Member.WordOffset), ToName(Member.Type) };
+				}
+
+				Offset = IntCastChecked<int32>(8 * (JumpedWords + Member.WordOffset));
+
+#if UE_GC_DEBUGNAMES // Must visit all remaining words to find start of debug names
+				DebugNameIdx = FindIdx - NumJumps;
+#else
+				return { Offset, ToName(Member.Type) };
+#endif
+			}
+			else if (IsIn<StopTypes>(Member.Type))
+			{
+				return { Offset, DebugNameIdx >= 0 ? ((FName*)WordIt)[DebugNameIdx] : EName::Error };
+			}
+			else if (Member.Type == EMemberType::Jump)
+			{
+				JumpedWords += (Member.WordOffset + 1) * FMemberPacked::OffsetRange;
+				++NumJumps;
+			}
+		}
+	}
+}
+
+//////////////////////////////////////////////////////////////////////////
+
+FSchemaBuilder::FSchemaBuilder(uint32 InStride, std::initializer_list<FMemberDeclaration> InMembers)
+	: Members(InMembers)
+	, StructStride(InStride) 
+{
+	check((StructStride % 8) == 0);
+
+	for (FMemberDeclaration& Member : Members)
+	{
+		TryAddSchemaReference(Member);
+	}
+}
+
+FSchemaBuilder::~FSchemaBuilder()
+{
+	// Drop references to declared inner schemas if ownership hasn't been transferred to BuiltSchema
+	if (!BuiltSchema.IsSet())
+	{
+		for (FMemberDeclaration& Member : Members)
+		{
+			TryDropSchemaReference(Member);
+		}
+	}
+}
+
+void FSchemaBuilder::Add(FMemberDeclaration Member)
+{
+	check(IsIn<ExtraWordTypes>(Member.Type) == IsSet(Member.ExtraWord));
+	checkf(IsIn<~GeneratedTypes>(Member.Type), TEXT("Member type %d is generated by Build() and shouldn't be declared"), Member.Type);
+	check(!BuiltSchema.IsSet());
+	
+	Members.Add(Member);
+	TryAddSchemaReference(Member);
+}
+
+static void CopyDebugNames(const FName* In, TArrayView<FMemberDeclaration> Out)
+{
+	for (FMemberDeclaration& Member : Out)
+	{
+		Member.DebugName = UE_GC_DEBUGNAMES ? *In++ : ToName(Member.Type);
+	}
+}
+
+void FSchemaBuilder::Append(FSchemaView SuperSchema)
+{
+	if (SuperSchema.IsEmpty())
+	{
+		return;
+	}
+	
+	uint32 JumpedWords = 0;
+	const int32 OldNum = Members.Num();
+	const FMemberWord* WordIt = SuperSchema.GetWords();
+	while (true)
+	{
+		FMemberWord Word = *WordIt++;
+		for (Private::FMemberUnpacked Member : Word.Members)
+		{
+			FMemberWord ExtraWord = IsIn<ExtraWordTypes>(Member.Type) ? *WordIt++ : FMemberWord{};
+
+			if (IsIn<StopTypes>(Member.Type))
+			{
+				CopyDebugNames((const FName*)WordIt, MakeArrayView(Members).RightChop(OldNum));
+				return;
+			}
+			else if (Member.Type == EMemberType::Jump)
+			{
+				JumpedWords += (Member.WordOffset + 1) * FMemberPacked::OffsetRange;
+			}
+			else
+			{
+				Members.Add({ FName(), (JumpedWords + Member.WordOffset) * 8, Member.Type, ExtraWord});
+				TryAddSchemaReference(Members.Last());
+			}
+		}
+	}
+}
+
+// Build()-helper that packs declarations into words, generates unnamed jumps and collects debug names
+struct FSchemaPacker
+{
+	FSchemaPacker(TArrayView<FMemberDeclaration> Members)
+	{
+		WordIdx = Words.AddZeroed();
+		for (const FMemberDeclaration& Member : Members)
+		{
+			Pack(Member);
+		}
+	}
+	
+	TArray<FMemberWord, TInlineAllocator<16>> Words;
+	TArray<FName, TInlineAllocator<64>> DebugNames;
+
+private:
+	static constexpr uint32 OffsetRange = FMemberPacked::OffsetRange;
+
+	int32 WordIdx = 0;
+	int32 QuadIdx = 0;
+	uint64 JumpWordPos = 0;
+
+	void Pack(const FMemberDeclaration& Member)
+	{
+		uint16 RelativeWordOffset = IsIn<StopTypes>(Member.Type) ? Member.Offset : GenerateJumps(Member.Offset);
+
+		if (TOptional<FStridedLayout> StridedLayout = TrySimplifyStructArray(Member))
+		{
+			AddNamedMember(EMemberType::StridedArray, RelativeWordOffset, Member.DebugName);
+			Words.AddZeroed_GetRef().StridedLayout = StridedLayout.GetValue();
+		}
+		else
+		{
+			AddNamedMember(Member.Type, RelativeWordOffset, Member.DebugName);
+
+			if (IsIn<ExtraWordTypes>(Member.Type))
+			{
+				Words.Add(Member.ExtraWord);
+			}
+		}
+		
+	}
+
+	uint16 GenerateJumps(uint32 ByteOffset)
+	{
+		check(ByteOffset % 8 == 0);
+		check(ByteOffset / 8 >= JumpWordPos);
+
+		const uint32 WordOffsetFromJumpPos = ByteOffset / 8 - JumpWordPos;
+
+		// 1 jump unit is OffsetRange words, a max jump is OffsetRange * OffsetRange * 8B = 32MB with 11 offset bits
+		const uint32 JumpUnits = WordOffsetFromJumpPos / OffsetRange;
+		for (uint32 MaxJumps = JumpUnits / OffsetRange; MaxJumps; --MaxJumps)
+		{
+			GenerateJump(OffsetRange);
+		}
+		if (uint32 LastJumpUnits = JumpUnits % OffsetRange)
+		{
+			GenerateJump(LastJumpUnits);
+		}
+		
+		check(ByteOffset / 8 >= JumpWordPos);
+		check(ByteOffset / 8 <  JumpWordPos + OffsetRange);
+
+		return IntCastChecked<uint16>(WordOffsetFromJumpPos % OffsetRange);
+	}
+
+	void GenerateJump(uint32 JumpUnit)
+	{
+		check(JumpUnit > 0 && JumpUnit < OffsetRange);
+		// -1 since we always jump at least 1 JumpUnit
+		AddMember(EMemberType::Jump, IntCastChecked<uint16>(JumpUnit - 1));
+		JumpWordPos += JumpUnit * OffsetRange;
+	}
+	
+	void AddNamedMember(EMemberType Type, uint16 Offset, FName DebugName)
+	{
+		AddMember(Type, Offset);
+#if UE_GC_DEBUGNAMES
+		DebugNames.Add(DebugName);
+#endif
+	}
+
+	void AddMember(EMemberType Type, uint16 Offset)
+	{
+		static_assert(MemberTypeCount <= 1u << FMemberPacked::TypeBits);
+		check(Offset < OffsetRange);
+
+		if (QuadIdx == 4)
+		{
+			WordIdx = Words.AddZeroed();
+			QuadIdx = 0;
+		}
+		
+		Words[WordIdx].Members[QuadIdx] = { ToInt(Type), Offset };
+		++QuadIdx;
+	}
+
+	static TOptional<FStridedLayout> TrySimplifyStructArray(const FMemberDeclaration& Member)
+	{
+		const FSchemaView& InnerSchema = Member.ExtraWord.InnerSchema;
+		if (Member.Type == EMemberType::StructArray &&
+			InnerSchema.GetWords()[0].Members[0].Type == ToInt(EMemberType::Reference) &&
+			InnerSchema.GetWords()[0].Members[1].Type == ToInt(EMemberType::Stop) &&
+			InnerSchema.GetStructStride() / 8u <= std::numeric_limits<decltype(FStridedLayout::WordStride)>::max())
+		{
+			return FStridedLayout{ InnerSchema.GetWords()[0].Members[0].WordOffset, decltype(FStridedLayout::WordStride)(InnerSchema.GetStructStride() / 8u) };
+		}
+
+		return NullOpt;
+	}
+};
+
+
+int32 FindSlowImplementation(ObjectAROFn ARO);
+
+static FMemberDeclaration GenerateTerminator(ObjectAROFn ARO)
+{
+	if (!ARO)
+	{
+		return { ToName(EMemberType::Stop), 0, EMemberType::Stop };
+	}
+	else if (int32 SlowAROIdx = FindSlowImplementation(ARO); SlowAROIdx != INDEX_NONE)
+	{
+		return { ToName(EMemberType::SlowARO), static_cast<uint32>(SlowAROIdx), EMemberType::SlowARO };
+	}
+	else
+	{
+		return { ToName(EMemberType::ARO), 0, EMemberType::ARO, ToWord(ARO) };
+	}
+}
+
+FSchemaView FSchemaBuilder::Build(ObjectAROFn ARO)
+{
+	if (BuiltSchema.IsSet())
+	{
+		return BuiltSchema.GetValue().Get();
+	}
+	else if (Members.IsEmpty() && ARO == nullptr)
+	{
+		return BuiltSchema.Emplace().Get();
+	}
+
+	// Sort and terminate
+	Algo::SortBy(Members, &FMemberDeclaration::Offset);
+	Members.Add(GenerateTerminator(ARO));
+
+	// Pack into words and generate jump instructions
+	FSchemaPacker Packed(Members);
+	check(Packed.DebugNames.IsEmpty() == !UE_GC_DEBUGNAMES) ;
+
+	// Allocate and initialize schema data
+	const uint32 Bytes = Align(sizeof(FSchemaHeader), sizeof(FMemberWord)) + sizeof(FMemberWord) * Packed.Words.Num() + sizeof(FName) * Packed.DebugNames.Num();
+	FSchemaHeader* Header = new (FMemory::Malloc(Bytes)) FSchemaHeader {StructStride};
+	FMemberWord* FirstWord = (FMemberWord*)(Header + 1);
+	FMemory::Memcpy(FirstWord, Packed.Words.GetData(), sizeof(FMemberWord) * Packed.Words.Num());
+	FMemory::Memcpy(FirstWord + Packed.Words.Num(), Packed.DebugNames.GetData(), sizeof(FName) * Packed.DebugNames.Num());
+	
+	// Construct schema owner
+	return BuiltSchema.Emplace(FSchemaView(FirstWord)).Get();
+}
+
+//////////////////////////////////////////////////////////////////////////
+
+TMap<UClass*, FSchemaOwner> GIntrinsicClassSchemas;
+
+void DeclareIntrinsicSchema(UClass* Class, FSchemaView Schema)
+{
+	FSchemaOwner& Owner = GIntrinsicClassSchemas.FindOrAdd(Class);
+	checkf(Owner.Get().IsEmpty(), TEXT("Declared multiple schemas for same class"));
+	Owner.Set(Schema);
+}
+
+FSchemaView GetIntrinsicSchema(UClass* Class)
+{
+	FSchemaOwner* Owner = GIntrinsicClassSchemas.Find(Class);
+	return Owner ? Owner->Get() : FSchemaView();
+}
+
+//////////////////////////////////////////////////////////////////////////
+
+FString FPropertyStack::GetPropertyPath() const
+{
+	TStringBuilder<128> Result;
+	const FProperty* PreviousProperty = nullptr;
+	const TCHAR DelimiterChar = TEXT('.');
+
+	for (const FProperty* Property : Props)
+	{
+		if (PreviousProperty)
+		{
+			if (Property->GetOwner<FProperty>() == PreviousProperty && Property->GetFName() == PreviousProperty->GetFName())
+			{
+				// Skipping inner properties (inside of containers) if their name matches their owner name - TArrayName.TArrayName doesn't have much value
+				// but we do want to keep TMapName.TMapName_Key
+				continue;
+			}
+			Result += DelimiterChar;
+		}
+		Result += Property->GetName();
+		PreviousProperty = Property;
+	}
+	return Result.ToString();
+}
+
+bool FPropertyStack::ConvertPathToProperties(UClass* ObjectClass, FName InPropertyPath, TArray<FProperty*>& OutProperties)
+{
+	const TCHAR DelimiterChar = TEXT('.');
+	FString PropertyNameOrPath = InPropertyPath.ToString();
+	int32 DelimiterIndex = -1;
+	bool bFullPathConstructed = true;
+
+	if (!PropertyNameOrPath.FindChar(DelimiterChar, DelimiterIndex))
+	{
+		// 99% of the time we're be dealing with just a single property
+		if (FProperty* FoundProperty = ObjectClass->FindPropertyByName(*PropertyNameOrPath))
+		{
+			OutProperties.Add(FoundProperty);
+		}
+		else
+		{
+			bFullPathConstructed = false;
+		}
+	}
+	else
+	{
+		// Try and find the first property as we can't start processing the rest of the path without it
+		FString PropertyName = PropertyNameOrPath.Left(DelimiterIndex);
+		if (FProperty* FoundProperty = ObjectClass->FindPropertyByName(*PropertyName))
+		{
+			OutProperties.Add(FoundProperty);
+
+			int32 StartIndex = DelimiterIndex + 1;
+			const TCHAR DelimiterStr[] = { DelimiterChar, TEXT('\0') };
+			do
+			{
+				// Determine the next property name
+				DelimiterIndex = PropertyNameOrPath.Find(DelimiterStr, ESearchCase::CaseSensitive, ESearchDir::FromStart, StartIndex);
+				PropertyName = PropertyNameOrPath.Mid(StartIndex, DelimiterIndex >= 0 ? (DelimiterIndex - StartIndex) : (PropertyNameOrPath.Len() - StartIndex));
+
+				if (FStructProperty* StructProp = CastField<FStructProperty>(FoundProperty))
+				{
+					// If the previous property was a struct property, the next one belongs to the struct the previous property represented
+					FoundProperty = StructProp->Struct->FindPropertyByName(*PropertyName);
+				}
+				else
+				{
+					// In all other case (though in reality it should only be a TMap) find the inner property
+					FoundProperty = CastField<FProperty>(FoundProperty->GetInnerFieldByName(*PropertyName));
+				}
+
+				if (FoundProperty)
+				{
+					OutProperties.Add(FoundProperty);
+				}
+				else
+				{
+					bFullPathConstructed = false;
+				}
+			} while (DelimiterIndex >= 0 && bFullPathConstructed);
+		}
+		else
+		{
+			bFullPathConstructed = false;
+		}
+	}
+	return bFullPathConstructed;
+}
+
+//////////////////////////////////////////////////////////////////////////
+#if !UE_BUILD_SHIPPING
+
+struct FSchemaStats;
+
+struct FStructMemberStats
+{
+	EMemberType Type = EMemberType::Stop;
+	uint32 WordOffset = 0;
+	FSchemaView Schema;
+	FSchemaStats* Stat = nullptr;
+
+	template<typename Fn>
+	void VisitInstances(const uint8* Instance, Fn Visitor) const
+	{
+		const uint8* Container = Instance + 8 * WordOffset;
+		switch (Type)
+		{ 
+			case EMemberType::StructArray:			VisitArray(*(const FScriptArray*)Container, Visitor);
+			break;
+			case EMemberType::SparseStructArray:	VisitSparseArray(*(const FScriptSparseArray*)Container, Visitor);
+			break;
+			case EMemberType::FreezableStructArray:	VisitArray(*(const FFreezableScriptArray*)Container, Visitor);
+			break;
+			case EMemberType::Optional:				VisitElements(Container, *(const bool*)(Container + Schema.GetStructStride()), Visitor);
+			break;
+			default:								check(false); 
+			break;
+		}
+	}
+
+private:
+	template<typename Fn>
+	void VisitElements(const uint8* It, int32 Num, Fn Visitor) const
+	{
+		const uint32 Stride = Schema.GetStructStride();
+		for (const uint8* End = It + Num * Stride; It != End; It += Stride)
+		{
+			Visitor(It);
+		}
+	}
+	
+	template<class ArrayType, typename Fn>
+	void VisitArray(const ArrayType& Array, Fn Visitor) const
+	{
+		VisitElements((const uint8*)Array.GetData(), Array.Num(), Visitor);
+	}
+
+	template<typename Fn>
+	void VisitSparseArray(const FScriptSparseArray& Array, Fn Visitor) const
+	{
+		const uint8* Data = Private::GetSparseData(Array);
+		const uint32 Stride = Schema.GetStructStride();
+		for (int32 Idx = 0, Num = Array.Num(); Idx < Num; ++Idx)
+		{
+			if (Array.IsValidIndex(Idx))
+			{
+				Visitor(Data + Idx * Stride);
+			}
+		}
+	}
+};
+
+struct FSchemaStats
+{
+	uint32 Counters[MemberTypeCount] = {};
+	uint32 NumMembers = 0;
+	uint32 NumExtraWords = 0;
+	TArray<FStructMemberStats> InnerStructs;
+
+	FSchemaStats() = default;
+	explicit FSchemaStats(FSchemaView Schema, TMap<const void*, FSchemaStats>& /* in-out */ StructStats)
+	: InnerStructs(MakeInnerStats(Schema, StructStats))
+	{
+		const FMemberWord* WordIt = Schema.GetWords();
+		while (true)
+		{
+			for (Private::FMemberUnpacked Member : WordIt++->Members)
+			{
+				++NumMembers;
+				++Counters[ToInt(Member.Type)];
+				NumExtraWords += IsIn<ExtraWordTypes>(Member.Type);
+				WordIt += IsIn<ExtraWordTypes>(Member.Type);
+
+				if (IsIn<StopTypes>(Member.Type))
+				{
+					return;
+				}
+			}
+		}
+	}
+	
+	void Aggregate(const FSchemaStats& In)
+	{
+		for (uint32 Idx = 0; Idx < MemberTypeCount; ++Idx)
+		{
+			Counters[Idx] += In.Counters[Idx];
+		}
+		NumMembers += In.NumMembers;
+		NumExtraWords += In.NumExtraWords;
+	}
+
+	static TArray<FStructMemberStats> MakeInnerStats(FSchemaView Schema, TMap<const void*, FSchemaStats>& /* in-out */ StructStats)
+	{
+		TArray<FStructMemberStats> Out;
+
+		check(!Schema.IsEmpty());
+		uint32 JumpedWords = 0;
+		const FMemberWord* WordIt = Schema.GetWords();
+		while (true)
+		{
+			for (Private::FMemberUnpacked Member : WordIt++->Members)
+			{
+				if (IsIn<StopTypes>(Member.Type))
+				{
+					return Out;
+				}
+				else if (Member.Type == EMemberType::Jump)
+				{
+					JumpedWords += (Member.WordOffset + 1) * FMemberPacked::OffsetRange;
+				}
+				else if (IsIn<ExtraWordTypes>(Member.Type))
+				{
+					FMemberWord ExtraWord = *WordIt++;
+					if (IsIn<InnerSchemaTypes>(Member.Type))
+					{
+						FSchemaView InnerSchema = ExtraWord.InnerSchema;
+						Out.Add({Member.Type, JumpedWords + Member.WordOffset, InnerSchema});
+
+						if (!StructStats.Contains(InnerSchema.GetWords()))
+						{
+							StructStats.Emplace(InnerSchema.GetWords(), FSchemaStats(InnerSchema, StructStats));
+						}
+					}
+				}
+			}
+		}
+	}
+	
+	void LinkInnerStats(TMap<const void*, FSchemaStats>& StructStats)
+	{
+		for (FStructMemberStats& InnerStruct : InnerStructs)
+		{
+			if (InnerStruct.Stat == nullptr)
+			{
+				InnerStruct.Stat = &StructStats.FindChecked(InnerStruct.Schema.GetWords());
+				InnerStruct.Stat->LinkInnerStats(StructStats);
+			}
+		}
+	}
+
+	uint32 NumWords() const
+	{
+		return NumExtraWords + Align(NumMembers, 4) / 4;
+	}
+};
+
+using FLog2Distribution = uint32[32];
+enum class EDumpUnit {Bytes, Number};
+
+static void DumpDistribution(FOutputDevice& Out, const TCHAR* Title, const FLog2Distribution& Bins, uint32 Total, EDumpUnit Unit)
+{
+	if (Bins[0] == Total)
+	{
+		return;
+	}
+
+	uint32 End = 32;
+	for (; End > 0 && Bins[End - 1] == 0; --End);
+
+	uint32 Begin = 0;
+	for (; Begin < End && Bins[Begin] == 0; ++Begin);
+
+	Out.Logf(TEXT(" _____ %-24s _____"), Title);
+
+	for (uint32 Idx = Begin, Sum = 0; Idx < End; ++Idx)
+	{
+		uint32 Num = Bins[Idx];
+		Sum += Num;
+		uint32 Max = Idx == 0 ? 0 : 1u << (Idx - 1);
+		float Percentage = (100.f * Num) / Total;
+		float Cumulative = (100.f * Sum) / Total;
+		if (Unit == EDumpUnit::Number)
+		{
+			Out.Logf(TEXT(" %4d %5.1f%% %5.1f%%  %d"), Max, Percentage, Cumulative, Num);
+		}
+		else
+		{
+			uint32 bK = Idx > 10;
+			Out.Logf(TEXT(" %3d%c %5.1f%% %5.1f%%  %d"), Max >> bK*10, "BK"[bK], Percentage, Cumulative, Num);
+		}
+	}
+}
+static void DumpSchemaStats(FOutputDevice& Out, const FSchemaStats& Stats)
+{
+	static constexpr TCHAR FormatStr[] = TEXT(" %18s %5.1f%% %5.1f%%  %d");
+
+	uint32 Sum = 0;
+	for (uint32 Idx = 0; Idx < MemberTypeCount; ++Idx)
+	{
+		if (uint32 Num = Stats.Counters[Idx])
+		{
+			Sum += Num;
+			float Percentage = (100.f * Num) / Stats.NumMembers;
+			float Cumulative = (100.f * Sum) / Stats.NumMembers;
+			Out.Logf(FormatStr, *ToName(EMemberType(Idx)).ToString(), Percentage, Cumulative, Num);
+		}
+	}
+}
+
+static uint32 GetBinIdx(uint32 N) { return FMath::CeilLogTwo(N) + !!N; }
+
+struct FInstanceStats
+{
+	uint32 Num = 0;
+	FSchemaStats Sum;
+	FLog2Distribution SizeBins = {};
+	FLog2Distribution NumMemberBins = {};
+	FLog2Distribution MemberBins[MemberTypeCount] = {};
+
+	void CountInstance(const uint8* Instance, uint32 InstanceSize, const FSchemaStats& TypeStat)
+	{
+		Sum.Aggregate(TypeStat);
+		++SizeBins[GetBinIdx(InstanceSize)];
+		++NumMemberBins[GetBinIdx(TypeStat.NumMembers)];
+		++Num;
+
+		for (uint32 Idx = 0; Idx < MemberTypeCount; ++Idx)
+		{
+			++MemberBins[Idx][GetBinIdx(TypeStat.Counters[Idx])];
+		}
+	}
+
+	void CountInnerStructInstances(const uint8* Instance, TConstArrayView<FStructMemberStats> InnerStructs)
+	{
+		for (const FStructMemberStats& InnerStruct : InnerStructs)
+		{
+			check(InnerStruct.Stat);
+			InnerStruct.VisitInstances(Instance, [this, &InnerStruct] (const uint8* InnerInstance) 
+			{ 
+				CountInstance(InnerInstance, InnerStruct.Schema.GetStructStride(), *InnerStruct.Stat);
+				CountInnerStructInstances(InnerInstance, InnerStruct.Stat->InnerStructs);
+			});				
+		}
+	}
+};
+
+void DumpSchemaStats(FOutputDevice& Out)
+{
+	struct FTotals
+	{
+		uint32 NumMembers = 0;
+		TMap<const void*, FSchemaStats> TypeStats;
+		FInstanceStats InstanceStats;
+	};
+	FTotals Classes;
+	FTotals Structs;
+	uint32 NumClasses = 0;
+
+	// Create all type stats
+	for (TObjectIterator<UClass> It; It; ++It)
+	{
+		FSchemaView ClassSchema = It->ReferenceSchema.Get();
+		++NumClasses;
+		FSchemaStats& ClassStat = Classes.TypeStats.FindOrAdd(ClassSchema.GetWords());
+		if (ClassStat.NumMembers == 0 && !ClassSchema.IsEmpty())
+		{
+			ClassStat = FSchemaStats(ClassSchema, Structs.TypeStats);
+			Classes.NumMembers += ClassStat.NumMembers;
+		}
+	}
+
+	// Link up stats pointers after all stats have been created
+	for (TPair<const void*, FSchemaStats>& Pair : Classes.TypeStats)
+	{
+		Pair.Value.LinkInnerStats(Structs.TypeStats);
+	}
+
+	for (TPair<const void*, FSchemaStats>& Pair : Structs.TypeStats)
+	{
+		Structs.NumMembers += Pair.Value.NumMembers;
+	}
+
+	// Aggregate instance stats
+	for (TObjectIterator<UObject> It; It; ++It)
+	{
+		const uint8* Object = reinterpret_cast<const uint8*>(*It);
+		const FSchemaStats& ClassStat = Classes.TypeStats.FindChecked(It->GetClass()->ReferenceSchema.Get().GetWords());
+		Classes.InstanceStats.CountInstance(Object, It->GetClass()->GetPropertiesSize(), ClassStat);
+		Structs.InstanceStats.CountInnerStructInstances(Object, ClassStat.InnerStructs);
+	}
+
+	// Log stats
+	Out.Logf(TEXT("----------------------- GC schema statistics -----------------------"));
+	Out.Logf(TEXT(" %d schemas shared by %d classes contain %d members"), Classes.TypeStats.Num(), NumClasses, Classes.NumMembers);
+	Out.Logf(TEXT(" %d inner struct schemas contain %d members"), Structs.TypeStats.Num(), Structs.NumMembers);
+	Out.Logf(TEXT(" Sweeping %d objects and %d struct instances"),	Classes.InstanceStats.Num, Structs.InstanceStats.Num);
+	Out.Logf(TEXT(" will process %d class and %d struct members"),	Classes.InstanceStats.Sum.NumMembers, Structs.InstanceStats.Sum.NumMembers);
+	Out.Logf(TEXT("---------------------- Object distributions ------------------------"));
+	DumpDistribution(Out, TEXT("Sizes"),							Classes.InstanceStats.SizeBins,			Classes.InstanceStats.Num, EDumpUnit::Bytes);
+	DumpDistribution(Out, TEXT("Total members"),					Classes.InstanceStats.NumMemberBins,	Classes.InstanceStats.Num, EDumpUnit::Number);
+	for (uint32 Idx = 0; Idx < MemberTypeCount; ++Idx)
+	{
+		DumpDistribution(Out, *ToName(EMemberType(Idx)).ToString(), Classes.InstanceStats.MemberBins[Idx],	Classes.InstanceStats.Num, EDumpUnit::Number);
+	}
+	Out.Logf(TEXT("------------------ Struct instance distributions -------------------"));
+	DumpDistribution(Out, TEXT("Sizes"),							Structs.InstanceStats.SizeBins,			Structs.InstanceStats.Num, EDumpUnit::Bytes);
+	DumpDistribution(Out, TEXT("Total members"),					Structs.InstanceStats.NumMemberBins,	Structs.InstanceStats.Num, EDumpUnit::Number);
+	for (uint32 Idx = 0; Idx < MemberTypeCount; ++Idx)
+	{
+		DumpDistribution(Out, *ToName(EMemberType(Idx)).ToString(),	Structs.InstanceStats.MemberBins[Idx],	Structs.InstanceStats.Num, EDumpUnit::Number);
+	}
+	Out.Logf(TEXT("--------------------- Class members to process ---------------------"));
+	DumpSchemaStats(Out, Classes.InstanceStats.Sum);
+	Out.Logf(TEXT("--------------------- Struct members to process --------------------"));
+	DumpSchemaStats(Out, Structs.InstanceStats.Sum);
+	Out.Logf(TEXT("--------------------------------------------------------------------"));
+}
+
+static FAutoConsoleCommandWithOutputDevice GDumpSchemaStats(TEXT("gc.DumpSchemaStats"), TEXT("Print GC schema statistics"), FConsoleCommandWithOutputDeviceDelegate::CreateStatic(&DumpSchemaStats));
+
+uint32 CountSchemas(uint32& OutNumWords)
+{
+	// Create all schema stats
+	TMap<const void*, FSchemaStats> ClassStats;
+	TMap<const void*, FSchemaStats> StructStats;
+	for (TObjectIterator<UClass> It; It; ++It)
+	{
+		FSchemaView ClassSchema = It->ReferenceSchema.Get();
+		if (!ClassSchema.IsEmpty())
+		{
+			FSchemaStats& ClassStat = ClassStats.FindOrAdd(ClassSchema.GetWords());
+			if (ClassStat.NumMembers == 0)
+			{
+				ClassStat = FSchemaStats(ClassSchema, StructStats);
+			}
+		}
+	}
+
+	// Calculate total words
+	OutNumWords = 0;
+	for (const TMap<const void*, FSchemaStats>& Stats : {ClassStats, StructStats})
+	{
+		for (const TPair<const void*, FSchemaStats>& Pair : Stats)
+		{
+			OutNumWords += Pair.Value.NumWords();
+		}
+	}
+	
+	return ClassStats.Num() + StructStats.Num();
+}
+
+#endif // !UE_BUILD_SHIPPING
+//////////////////////////////////////////////////////////////////////////
+} // namespace UE::GC
diff --git a/Engine/Source/Runtime/CoreUObject/Private/UObject/GarbageCollectionTesting.cpp b/Engine/Source/Runtime/CoreUObject/Private/UObject/GarbageCollectionTesting.cpp
--- a/Engine/Source/Runtime/CoreUObject/Private/UObject/GarbageCollectionTesting.cpp
+++ b/Engine/Source/Runtime/CoreUObject/Private/UObject/GarbageCollectionTesting.cpp
@@ -1,6 +1,7 @@
 // Copyright Epic Games, Inc. All Rights Reserved.
 
 #include "GarbageCollectionTesting.h"
+#include "UObject/GarbageCollectionSchema.h"
 #include "UObject/UnrealType.h"
 
 /*
@@ -52,6 +53,5 @@ void UnlinkReachabilityStressData(UObjectReachabilityStressData* Data)
 
 IMPLEMENT_CORE_INTRINSIC_CLASS(UObjectReachabilityStressData, UObject,
 	{
-		UE::GC::FTokenStreamBuilder& Builder = UE::GC::FIntrinsicClassTokens::AllocateBuilder(Class);
-		Builder.EmitObjectArrayReference(STRUCT_OFFSET(UObjectReachabilityStressData, Children), TEXT("Children"));
+		UE::GC::DeclareIntrinsicMembers(Class, { UE_GC_MEMBER(UObjectReachabilityStressData, Children) });
 	});
diff --git a/Engine/Source/Runtime/CoreUObject/Private/UObject/GarbageCollectionVerification.cpp b/Engine/Source/Runtime/CoreUObject/Private/UObject/GarbageCollectionVerification.cpp
--- a/Engine/Source/Runtime/CoreUObject/Private/UObject/GarbageCollectionVerification.cpp
+++ b/Engine/Source/Runtime/CoreUObject/Private/UObject/GarbageCollectionVerification.cpp
@@ -52,7 +52,7 @@ public:
 	{
 		return NumErrors.GetValue();
 	}
-	FORCEINLINE_DEBUGGABLE void HandleTokenStreamObjectReference(FGCArrayStruct& ObjectsToSerializeStruct, UObject* ReferencingObject, UObject*& Object, UE::GC::FTokenId TokenIndex, const EGCTokenType TokenType, bool bAllowReferenceElimination)
+	FORCEINLINE_DEBUGGABLE void HandleTokenStreamObjectReference(FGCArrayStruct& ObjectsToSerializeStruct, UObject* ReferencingObject, UObject*& Object, FMemberId MemberId, EOrigin Origin, bool bAllowReferenceElimination)
 	{
 		if (Object)
 		{
@@ -63,23 +63,23 @@ public:
 #endif
 				!Object->IsValidLowLevelFast())
 			{
-				FString TokenDebugInfo;
+				FString DebugInfo;
 				if (UClass *Class = (ReferencingObject ? ReferencingObject->GetClass() : nullptr))
 				{
-					FTokenInfo TokenInfo = Class->ReferenceTokens.GetTokenInfo(TokenIndex);
-					TokenDebugInfo = FString::Printf(TEXT("ReferencingObjectClass: %s, Property Name: %s, Offset: %d"),
-						*Class->GetFullName(), *TokenInfo.Name.GetPlainNameString(), TokenInfo.Offset);
+					UE::GC::FMemberInfo Member = GetMemberDebugInfo(Class->ReferenceSchema.Get(), MemberId);
+					DebugInfo = FString::Printf(TEXT("ReferencingObjectClass: %s, Property Name: %s, Offset: %d"),
+						*Class->GetFullName(), *Member.Name.GetPlainNameString(), Member.Offset);
 				}
 				else
 				{
 					// This means this objects is most likely being referenced by AddReferencedObjects
-					TokenDebugInfo = TEXT("Native Reference");
+					DebugInfo = TEXT("Native Reference");
 				}
 
-				UE_LOG(LogGarbage, Fatal, TEXT("Invalid object while verifying Disregard for GC assumptions: 0x%016llx, ReferencingObject: %s, %s, TokenIndex: %d"),
+				UE_LOG(LogGarbage, Fatal, TEXT("Invalid object while verifying Disregard for GC assumptions: 0x%016llx, ReferencingObject: %s, %s, MemberId: %d"),
 					(int64)(PTRINT)Object,
 					ReferencingObject ? *ReferencingObject->GetFullName() : TEXT("NULL"),
-					*TokenDebugInfo, TokenIndex.AsPrintableIndex());
+					*DebugInfo, MemberId.AsPrintableIndex());
 				}
 #endif // ENABLE_GC_OBJECT_CHECKS
 
@@ -166,12 +166,9 @@ public:
 	/**
 	* Handles UObject reference from the token stream. Performance is critical here so we're FORCEINLINING this function.
 	*
-	* @param ObjectsToSerialize An array of remaining objects to serialize (Obj must be added to it if Obj can be added to cluster)
-	* @param ReferencingObject Object referencing the object to process.
-	* @param TokenIndex Index to the token stream where the reference was found.
 	* @param bAllowReferenceElimination True if reference elimination is allowed (ignored when constructing clusters).
 	*/
-	FORCEINLINE_DEBUGGABLE void HandleTokenStreamObjectReference(FGCArrayStruct& ObjectsToSerializeStruct, UObject* ReferencingObject, UObject*& Object, UE::GC::FTokenId TokenIndex, const EGCTokenType TokenType, bool bAllowReferenceElimination)
+	FORCEINLINE_DEBUGGABLE void HandleTokenStreamObjectReference(FGCArrayStruct& ObjectsToSerializeStruct, UObject* ReferencingObject, UObject*& Object, FMemberId MemberId, EOrigin Origin, bool bAllowReferenceElimination)
 	{
 		if (Object)
 		{
@@ -188,27 +185,27 @@ public:
 #endif
 				!Object->IsValidLowLevelFast())
 			{
-				FString TokenDebugInfo;
+				FString DebugInfo;
 				if (UClass *Class = (ReferencingObject ? ReferencingObject->GetClass() : nullptr))
 				{
-					FTokenInfo TokenInfo = Class->ReferenceTokens.GetTokenInfo(TokenIndex);
-					TokenDebugInfo = FString::Printf(TEXT("ReferencingObjectClass: %s, Property Name: %s, Offset: %d"),
-						*Class->GetFullName(), *TokenInfo.Name.GetPlainNameString(), TokenInfo.Offset);
+					UE::GC::FMemberInfo Member = GetMemberDebugInfo(Class->ReferenceSchema.Get(), MemberId);
+					DebugInfo = FString::Printf(TEXT("ReferencingObjectClass: %s, Property Name: %s, Offset: %d"),
+						*Class->GetFullName(), *Member.Name.GetPlainNameString(), Member.Offset);
 				}
 				else
 				{
 					// This means this objects is most likely being referenced by AddReferencedObjects
-					TokenDebugInfo = TEXT("Native Reference");
+					DebugInfo = TEXT("Native Reference");
 				}
 
 #if UE_GCCLUSTER_VERBOSE_LOGGING
 				DumpClusterToLog(*Cluster, true, true);
 #endif
 
-				UE_LOG(LogGarbage, Fatal, TEXT("Invalid object while verifying cluster assumptions: 0x%016llx, ReferencingObject: %s, %s, TokenIndex: %d"),
+				UE_LOG(LogGarbage, Fatal, TEXT("Invalid object while verifying cluster assumptions: 0x%016llx, ReferencingObject: %s, %s, MemberId: %d"),
 					(int64)(PTRINT)Object,
 					ReferencingObject ? *ReferencingObject->GetFullName() : TEXT("NULL"),
-					*TokenDebugInfo, TokenIndex.AsPrintableIndex());
+					*DebugInfo, MemberId.AsPrintableIndex());
 			}
 #endif // ENABLE_GC_OBJECT_CHECKS
 
diff --git a/Engine/Source/Runtime/CoreUObject/Private/UObject/ObjectRedirector.cpp b/Engine/Source/Runtime/CoreUObject/Private/UObject/ObjectRedirector.cpp
--- a/Engine/Source/Runtime/CoreUObject/Private/UObject/ObjectRedirector.cpp
+++ b/Engine/Source/Runtime/CoreUObject/Private/UObject/ObjectRedirector.cpp
@@ -9,6 +9,7 @@
 #include "UObject/ObjectSaveContext.h"
 #include "UObject/Package.h"
 #include "Templates/Casts.h"
+#include "UObject/GarbageCollectionSchema.h"
 #include "UObject/PropertyPortFlags.h"
 #include "UObject/UnrealType.h"
 
@@ -102,7 +103,6 @@ bool UObjectRedirector::GetNativePropertyValues( TMap<FString,FString>& out_Prop
 
 IMPLEMENT_CORE_INTRINSIC_CLASS(UObjectRedirector, UObject,
 	{
-		UE::GC::FTokenStreamBuilder& Builder = UE::GC::FIntrinsicClassTokens::AllocateBuilder(Class);
-		Builder.EmitObjectReference(STRUCT_OFFSET(UObjectRedirector, DestinationObject), TEXT("DestinationObject"));
+		UE::GC::DeclareIntrinsicMembers(Class, { UE_GC_MEMBER(UObjectRedirector, DestinationObject) });
 	}
 );
diff --git a/Engine/Source/Runtime/CoreUObject/Private/UObject/Package.cpp b/Engine/Source/Runtime/CoreUObject/Private/UObject/Package.cpp
--- a/Engine/Source/Runtime/CoreUObject/Private/UObject/Package.cpp
+++ b/Engine/Source/Runtime/CoreUObject/Private/UObject/Package.cpp
@@ -7,6 +7,7 @@
 #include "Misc/AssetRegistryInterface.h"
 #include "Misc/ITransaction.h"
 #include "Misc/PackageName.h"
+#include "UObject/GarbageCollectionSchema.h"
 #include "UObject/LinkerLoad.h"
 #include "UObject/LinkerSave.h"
 #include "UObject/LinkerManager.h"
@@ -407,8 +408,7 @@ void UPackage::SetLoadedByEditorPropertiesOnly(bool bIsEditorOnly, bool bRecursi
 PRAGMA_DISABLE_DEPRECATION_WARNINGS
 IMPLEMENT_CORE_INTRINSIC_CLASS(UPackage, UObject,
 	{
-		UE::GC::FTokenStreamBuilder& Builder = UE::GC::FIntrinsicClassTokens::AllocateBuilder(Class);
-		Builder.EmitObjectReference(STRUCT_OFFSET(UPackage, MetaData), TEXT("MetaData"));
+		UE::GC::DeclareIntrinsicMembers(Class, { UE_GC_MEMBER(UPackage, MetaData) });
 	}
 );
 PRAGMA_ENABLE_DEPRECATION_WARNINGS
diff --git a/Engine/Source/Runtime/CoreUObject/Private/UObject/Property.cpp b/Engine/Source/Runtime/CoreUObject/Private/UObject/Property.cpp
--- a/Engine/Source/Runtime/CoreUObject/Private/UObject/Property.cpp
+++ b/Engine/Source/Runtime/CoreUObject/Private/UObject/Property.cpp
@@ -717,98 +717,6 @@ const TCHAR* FPropertyHelpers::ReadToken( const TCHAR* Buffer, FStringBuilderBas
 	return Buffer;
 }
 
-FString FGCStackSizeHelper::GetPropertyPath() const
-{
-	FString Result;
-	const FProperty* PreviousProperty = nullptr;
-	const TCHAR DelimiterChar = TEXT('.');
-
-	for (int32 PropertyIndex = 0; PropertyIndex < PropertyStack.Num(); ++PropertyIndex)
-	{
-		const FProperty* Property = PropertyStack[PropertyIndex];
-		check(Property);
-		if (PropertyIndex > 0)
-		{
-			if (Property->GetOwner<FProperty>() == PreviousProperty && Property->GetFName() == PreviousProperty->GetFName())
-			{
-				// Skipping inner properties (inside of containers) if their name matches their owner name - TArrayName.TArrayName doesn't have much value
-				// but we do want to keep TMapName.TMapName_Key
-				continue;
-			}
-			Result += DelimiterChar;
-		}
-		Result += Property->GetName();
-		PreviousProperty = Property;
-	}
-	return Result;
-}
-
-bool FGCStackSizeHelper::ConvertPathToProperties(UClass* ObjectClass, const FName& InPropertyPath, TArray<FProperty*>& OutProperties)
-{
-	const TCHAR DelimiterChar = TEXT('.');
-	FString PropertyNameOrPath = InPropertyPath.ToString();
-	int32 DelimiterIndex = -1;
-	bool bFullPathConstructed = true;
-
-	if (!PropertyNameOrPath.FindChar(DelimiterChar, DelimiterIndex))
-	{
-		// 99% of the time we're be dealing with just a single property
-		FProperty* FoundProperty = ObjectClass->FindPropertyByName(*PropertyNameOrPath);
-		if (FoundProperty)
-		{
-			OutProperties.Add(FoundProperty);
-		}
-		else
-		{
-			bFullPathConstructed = false;
-		}
-	}
-	else
-	{
-		// Try and find the first property as we can't start processing the rest of the path without it
-		FString PropertyName = PropertyNameOrPath.Left(DelimiterIndex);
-		FProperty* FoundProperty = ObjectClass->FindPropertyByName(*PropertyName);
-		if (FoundProperty)
-		{
-			OutProperties.Add(FoundProperty);
-
-			int32 StartIndex = DelimiterIndex + 1;
-			const TCHAR DelimiterStr[] = { DelimiterChar, TEXT('\0') };
-			do
-			{
-				// Determine the next property name
-				DelimiterIndex = PropertyNameOrPath.Find(DelimiterStr, ESearchCase::CaseSensitive, ESearchDir::FromStart, StartIndex);
-				PropertyName = PropertyNameOrPath.Mid(StartIndex, DelimiterIndex >= 0 ? (DelimiterIndex - StartIndex) : (PropertyNameOrPath.Len() - StartIndex));
-
-				if (FStructProperty* StructProp = CastField<FStructProperty>(FoundProperty))
-				{
-					// If the previous property was a struct property, the next one belongs to the struct the previous property represented
-					FoundProperty = StructProp->Struct->FindPropertyByName(*PropertyName);
-				}
-				else
-				{
-					// In all other case (though in reality it should only be a TMap) find the inner property
-					FoundProperty = CastField<FProperty>(FoundProperty->GetInnerFieldByName(*PropertyName));
-				}
-
-				if (FoundProperty)
-				{
-					OutProperties.Add(FoundProperty);
-				}
-				else
-				{
-					bFullPathConstructed = false;
-				}
-			} while (DelimiterIndex >= 0 && bFullPathConstructed);
-		}
-		else
-		{
-			bFullPathConstructed = false;
-		}
-	}
-	return bFullPathConstructed;
-}
-
 /*-----------------------------------------------------------------------------
 	FProperty implementation.
 -----------------------------------------------------------------------------*/
diff --git a/Engine/Source/Runtime/CoreUObject/Private/UObject/PropertyPrivate.cpp b/Engine/Source/Runtime/CoreUObject/Private/UObject/PropertyPrivate.cpp
--- a/Engine/Source/Runtime/CoreUObject/Private/UObject/PropertyPrivate.cpp
+++ b/Engine/Source/Runtime/CoreUObject/Private/UObject/PropertyPrivate.cpp
@@ -19,6 +19,7 @@
 #include "Math/Box2D.h"
 #include "UObject/ReleaseObjectVersion.h"
 #include "UObject/UnrealTypePrivate.h"
+#include "UObject/GarbageCollectionSchema.h"
 #include "UObject/LinkerPlaceholderFunction.h"
 #include "UObject/LinkerPlaceholderClass.h"
 
@@ -172,9 +173,7 @@ namespace UEEnumProperty_Private
 
 IMPLEMENT_CORE_INTRINSIC_CLASS(UEnumProperty, UProperty,
 	{
-		UE::GC::FTokenStreamBuilder& Builder = UE::GC::FIntrinsicClassTokens::AllocateBuilder(Class);
-		Builder.EmitObjectReference(UEEnumProperty_Private::FEnumPropertyFriend::EnumOffset, TEXT("Enum"));
-		Builder.EmitObjectReference(UEEnumProperty_Private::FEnumPropertyFriend::UnderlyingPropOffset, TEXT("UnderlyingProp"));
+		UE::GC::DeclareIntrinsicMembers(Class, { UE_GC_MEMBER(UEnumProperty, Enum), UE_GC_MEMBER(UEnumProperty, UnderlyingProp) });
 	}
 );
 
@@ -193,8 +192,7 @@ void UArrayProperty::AddReferencedObjects(UObject* InThis, FReferenceCollector&
 
 IMPLEMENT_CORE_INTRINSIC_CLASS(UArrayProperty, UProperty,
 	{
-		UE::GC::FTokenStreamBuilder& Builder = UE::GC::FIntrinsicClassTokens::AllocateBuilder(Class);
-		Builder.EmitObjectReference(STRUCT_OFFSET(UArrayProperty, Inner), TEXT("Inner"));
+		UE::GC::DeclareIntrinsicMembers(Class, { UE_GC_MEMBER(UArrayProperty, Inner) });
 
 		// Ensure that TArray and FScriptArray are interchangeable, as FScriptArray will be used to access a native array property
 		// from script that is declared as a TArray in C++.
@@ -256,8 +254,7 @@ void UObjectPropertyBase::AddReferencedObjects(UObject* InThis, FReferenceCollec
 
 IMPLEMENT_CORE_INTRINSIC_CLASS(UObjectPropertyBase, UProperty,
 	{
-		UE::GC::FTokenStreamBuilder& Builder = UE::GC::FIntrinsicClassTokens::AllocateBuilder(Class);
-		Builder.EmitObjectReference(STRUCT_OFFSET(UObjectProperty, PropertyClass), TEXT("PropertyClass"));
+		UE::GC::DeclareIntrinsicMembers(Class, { UE_GC_MEMBER(UObjectProperty, PropertyClass) });
 	}
 );
 
@@ -372,8 +369,7 @@ void UByteProperty::AddReferencedObjects(UObject* InThis, FReferenceCollector& C
 
 IMPLEMENT_CORE_INTRINSIC_CLASS(UByteProperty, UNumericProperty,
 	{
-		UE::GC::FTokenStreamBuilder& Builder = UE::GC::FIntrinsicClassTokens::AllocateBuilder(Class);
-		Builder.EmitObjectReference(STRUCT_OFFSET(UByteProperty, Enum), TEXT("Enum"));
+		UE::GC::DeclareIntrinsicMembers(Class, { UE_GC_MEMBER(UByteProperty, Enum) });
 	}
 );
 
@@ -442,8 +438,7 @@ void UClassProperty::AddReferencedObjects(UObject* InThis, FReferenceCollector&
 
 IMPLEMENT_CORE_INTRINSIC_CLASS(UClassProperty, UObjectProperty,
 	{
-		UE::GC::FTokenStreamBuilder& Builder = UE::GC::FIntrinsicClassTokens::AllocateBuilder(Class);
-		Builder.EmitObjectReference(STRUCT_OFFSET(UClassProperty, MetaClass), TEXT("MetaClass"));
+		UE::GC::DeclareIntrinsicMembers(Class, { UE_GC_MEMBER(UClassProperty, MetaClass) });
 	}
 );
 
@@ -478,8 +473,7 @@ void UDelegateProperty::BeginDestroy()
 
 IMPLEMENT_CORE_INTRINSIC_CLASS(UDelegateProperty, UProperty,
 	{
-		UE::GC::FTokenStreamBuilder& Builder = UE::GC::FIntrinsicClassTokens::AllocateBuilder(Class);
-		Builder.EmitObjectReference(STRUCT_OFFSET(UDelegateProperty, SignatureFunction), TEXT("SignatureFunction"));
+		UE::GC::DeclareIntrinsicMembers(Class, { UE_GC_MEMBER(UDelegateProperty, SignatureFunction) });
 	}
 );
 
@@ -580,8 +574,7 @@ void UInterfaceProperty::AddReferencedObjects(UObject* InThis, FReferenceCollect
 
 IMPLEMENT_CORE_INTRINSIC_CLASS(UInterfaceProperty, UProperty,
 	{
-		UE::GC::FTokenStreamBuilder& Builder = UE::GC::FIntrinsicClassTokens::AllocateBuilder(Class);
-		Builder.EmitObjectReference(STRUCT_OFFSET(UInterfaceProperty, InterfaceClass), TEXT("InterfaceClass"));
+		UE::GC::DeclareIntrinsicMembers(Class, { UE_GC_MEMBER(UInterfaceProperty, InterfaceClass) });
 	}
 );
 
@@ -618,9 +611,7 @@ void UMapProperty::AddReferencedObjects(UObject* InThis, FReferenceCollector& Co
 
 IMPLEMENT_CORE_INTRINSIC_CLASS(UMapProperty, UProperty,
 	{
-		UE::GC::FTokenStreamBuilder& Builder = UE::GC::FIntrinsicClassTokens::AllocateBuilder(Class);
-		Builder.EmitObjectReference(STRUCT_OFFSET(UMapProperty, KeyProp),   TEXT("KeyProp"));
-		Builder.EmitObjectReference(STRUCT_OFFSET(UMapProperty, ValueProp), TEXT("ValueProp"));
+		UE::GC::DeclareIntrinsicMembers(Class, { UE_GC_MEMBER(UMapProperty, KeyProp), UE_GC_MEMBER(UMapProperty, ValueProp) });
 
 		// Ensure that TArray and FScriptMap are interchangeable, as FScriptMap will be used to access a native array property
 		// from script that is declared as a TArray in C++.
@@ -658,8 +649,7 @@ void UMulticastDelegateProperty::BeginDestroy()
 
 IMPLEMENT_CORE_INTRINSIC_CLASS(UMulticastDelegateProperty, UProperty,
 	{
-		UE::GC::FTokenStreamBuilder& Builder = UE::GC::FIntrinsicClassTokens::AllocateBuilder(Class);
-		Builder.EmitObjectReference(STRUCT_OFFSET(UMulticastDelegateProperty, SignatureFunction), TEXT("SignatureFunction"));
+		UE::GC::DeclareIntrinsicMembers(Class, { UE_GC_MEMBER(UMulticastDelegateProperty, SignatureFunction) });
 	}
 );
 
@@ -713,8 +703,7 @@ void USetProperty::AddReferencedObjects(UObject* InThis, FReferenceCollector& Co
 
 IMPLEMENT_CORE_INTRINSIC_CLASS(USetProperty, UProperty,
 	{
-		UE::GC::FTokenStreamBuilder& Builder = UE::GC::FIntrinsicClassTokens::AllocateBuilder(Class);
-		Builder.EmitObjectReference(STRUCT_OFFSET(USetProperty, ElementProp), TEXT("ElementProp"));
+		UE::GC::DeclareIntrinsicMembers(Class, { UE_GC_MEMBER(USetProperty, ElementProp) });
 
 		// Ensure that TArray and FScriptMap are interchangeable, as FScriptMap will be used to access a native array property
 		// from script that is declared as a TArray in C++.
@@ -788,8 +777,7 @@ void USoftClassProperty::AddReferencedObjects(UObject* InThis, FReferenceCollect
 
 IMPLEMENT_CORE_INTRINSIC_CLASS(USoftClassProperty, USoftObjectProperty,
 	{
-		UE::GC::FTokenStreamBuilder& Builder = UE::GC::FIntrinsicClassTokens::AllocateBuilder(Class);
-		Builder.EmitObjectReference(STRUCT_OFFSET(USoftClassProperty, MetaClass), TEXT("MetaClass"));
+		UE::GC::DeclareIntrinsicMembers(Class, { UE_GC_MEMBER(USoftClassProperty, MetaClass) });
 	}
 );
 
@@ -890,8 +878,7 @@ void UStructProperty::AddReferencedObjects(UObject* InThis, FReferenceCollector&
 
 IMPLEMENT_CORE_INTRINSIC_CLASS(UStructProperty, UProperty,
 	{
-		UE::GC::FTokenStreamBuilder& Builder = UE::GC::FIntrinsicClassTokens::AllocateBuilder(Class);
-		Builder.EmitObjectReference(STRUCT_OFFSET(UStructProperty, Struct), TEXT("Struct"));
+		UE::GC::DeclareIntrinsicMembers(Class, { UE_GC_MEMBER(UStructProperty, Struct) });
 	}
 );
 
diff --git a/Engine/Source/Runtime/CoreUObject/Private/UObject/ReferenceChainSearch.cpp b/Engine/Source/Runtime/CoreUObject/Private/UObject/ReferenceChainSearch.cpp
--- a/Engine/Source/Runtime/CoreUObject/Private/UObject/ReferenceChainSearch.cpp
+++ b/Engine/Source/Runtime/CoreUObject/Private/UObject/ReferenceChainSearch.cpp
@@ -236,8 +236,8 @@ namespace UE::ReferenceChainSearch
 			void HandleTokenStreamObjectReference(UE::GC::FWorkerContext& Context,
 				UObject* ReferencingObject,
 				UObject*& Object,
-				UE::GC::FTokenId TokenIndex,
-				EGCTokenType TokenType,
+				UE::GC::FMemberId MemberId,
+				UE::GC::EOrigin Origin,
 				bool bAllowReferenceElimination)
 			{
 				if (Object)
@@ -251,13 +251,13 @@ namespace UE::ReferenceChainSearch
 
 					if (Object != ReferencingObject)
 					{
-						This->HandleObjectReference(ThreadIndex, Object, ReferencingObject, TokenIndex, TokenType);
+						This->HandleObjectReference(ThreadIndex, Object, ReferencingObject, MemberId, Origin);
 					}
 				}
 			}
 		};
 
-		struct FCollector : public FReferenceCollector
+		class FCollector final: public FReferenceCollector
 		{
 		protected:
 			FProcessor& Processor;
@@ -273,22 +273,15 @@ namespace UE::ReferenceChainSearch
 				Processor.HandleTokenStreamObjectReference(Context,
 					const_cast<UObject*>(ReferencingObject),
 					Object,
-					UE::GC::ETokenlessId::Collector,
-					EGCTokenType::Native,
+					UE::GC::EMemberlessId::Collector,
+					UE::GC::EOrigin::Other,
 					false);
 			}
-			virtual void HandleObjectReferences(
-				UObject** InObjects, const int32 ObjectNum, const UObject* ReferencingObject, const FProperty* InReferencingProperty) override
+			virtual void HandleObjectReferences(UObject** Objects, int32 Num, const UObject* ReferencingObject, const FProperty* ReferencingProperty) override
 			{
-				for (int32 ObjectIndex = 0; ObjectIndex < ObjectNum; ++ObjectIndex)
+				for (UObject*& Object : MakeArrayView(Objects, Num))
 				{
-					UObject*& Object = InObjects[ObjectIndex];
-					Processor.HandleTokenStreamObjectReference(Context,
-						const_cast<UObject*>(ReferencingObject),
-						Object,
-						UE::GC::ETokenlessId::Collector,
-						EGCTokenType::Native,
-						false);
+					HandleObjectReference(Object, ReferencingObject, ReferencingProperty);
 				}
 			}
 
@@ -467,7 +460,7 @@ namespace UE::ReferenceChainSearch
 		}
 
 		void HandleObjectReference(
-			int32 ThreadIndex, UObject* Object, UObject* ReferencingObject, UE::GC::FTokenId TokenIndex, EGCTokenType TokenType)
+			int32 ThreadIndex, UObject* Object, UObject* ReferencingObject, UE::GC::FTokenId TokenIndex, UE::GC::EOrigin TokenType)
 		{
 			FThreadData& ThreadData = AllThreadData[ThreadIndex];
 			if (ThreadData.PreviousReferencingObject != ReferencingObject)
@@ -537,7 +530,7 @@ namespace UE::ReferenceChainSearch
 			UObject* Object,
 			UObject* ReferencingObject,
 			UE::GC::FTokenId TokenIndex,
-			EGCTokenType TokenType)
+			UE::GC::EOrigin TokenType)
 		{
 			FThreadData& ThreadData = AllThreadData[ThreadIndex];
 			if (!TargetObjects.Contains(Object))
@@ -1026,7 +1019,7 @@ namespace UE::ReferenceChainSearch
 		}
 
 		void HandleObjectReference(
-			int32 ThreadIndex, UObject* Object, UObject* ReferencingObject, UE::GC::FTokenId TokenIndex, EGCTokenType TokenType)
+			int32 ThreadIndex, UObject* Object, UObject* ReferencingObject, UE::GC::FMemberId MemberId, UE::GC::EOrigin Origin)
 		{
 			if (!Object || Object == ReferencingObject) // Skip self-references just in case
 			{
@@ -1050,10 +1043,10 @@ namespace UE::ReferenceChainSearch
 
 			FGCObjectInfo* ObjectInfo = ObjectToInfoMap.FindRef(Object);
 			*RefInfo = FReferenceChainSearch::FObjectReferenceInfo(ObjectInfo);
-			if (TokenIndex != UE::GC::ETokenlessId::Collector)
+			if (MemberId != UE::GC::EMemberlessId::Collector)
 			{
-				FTokenInfo TokenInfo = ReferencingObject->GetClass()->ReferenceTokens.GetTokenInfo(TokenIndex);
-				RefInfo->ReferencerName = TokenInfo.Name;
+				FName MemberName = GetMemberDebugInfo(ReferencingObject->GetClass()->ReferenceSchema.Get(), MemberId).Name;
+				RefInfo->ReferencerName = MemberName;
 				RefInfo->Type = FReferenceChainSearch::EReferenceType::Property;
 			}
 			else
@@ -1369,7 +1362,7 @@ void FReferenceChainSearch::DumpChain(FReferenceChainSearch::FReferenceChain* Ch
 				UClass* ReferencerClass = Cast<UClass>(ReferencerObject->GetClass()->TryResolveObject());
 				TArray<FProperty*> ReferencingProperties;
 
-				if (ReferencerClass && FGCStackSizeHelper::ConvertPathToProperties(ReferencerClass, ReferenceInfo->ReferencerName, ReferencingProperties))
+				if (ReferencerClass && UE::GC::FPropertyStack::ConvertPathToProperties(ReferencerClass, ReferenceInfo->ReferencerName, ReferencingProperties))
 				{
 					FProperty* InnermostProperty = ReferencingProperties.Last();
 					FProperty* OutermostProperty = ReferencingProperties[0];
diff --git a/Engine/Source/Runtime/CoreUObject/Private/UObject/UObjectClusters.cpp b/Engine/Source/Runtime/CoreUObject/Private/UObject/UObjectClusters.cpp
--- a/Engine/Source/Runtime/CoreUObject/Private/UObject/UObjectClusters.cpp
+++ b/Engine/Source/Runtime/CoreUObject/Private/UObject/UObjectClusters.cpp
@@ -630,7 +630,7 @@ public:
 	 * @param ObjectsToSerialize An array of remaining objects to serialize (Obj must be added to it if Obj can be added to cluster)
 	 * @param bOuterAndClass If true, the Obj's Outer and Class will also be added to the cluster
 	 */
-	void AddObjectToCluster(int32 ObjectIndex, FUObjectItem* ObjectItem, UObject* Obj, FGCArrayStruct& ObjectsToSerializeStruct, bool bOuterAndClass)
+	void AddObjectToCluster(int32 ObjectIndex, FUObjectItem* ObjectItem, UObject* Obj, FWorkerContext& ObjectsToSerializeStruct, bool bOuterAndClass)
 	{
 		// If we haven't finished loading, we can't be sure we know all the references
 		checkf(!Obj->HasAnyFlags(RF_NeedLoad), TEXT("%s hasn't been loaded (%s) but is being added to cluster %s"), 
@@ -652,14 +652,14 @@ public:
 				UObject* ObjOuter = Obj->GetOuter();
 				if (ObjOuter)
 				{
-					HandleTokenStreamObjectReference(ObjectsToSerializeStruct, Obj, ObjOuter, UE::GC::ETokenlessId::Outer, EGCTokenType::Native, true);
+					HandleTokenStreamObjectReference(ObjectsToSerializeStruct, Obj, ObjOuter, UE::GC::EMemberlessId::Outer, EOrigin::Other, true);
 				}
 				if (!Obj->GetClass()->HasAllClassFlags(CLASS_Native))
 				{
 					UObject* ObjectClass = Obj->GetClass();
-					HandleTokenStreamObjectReference(ObjectsToSerializeStruct, Obj, ObjectClass, UE::GC::ETokenlessId::Class, EGCTokenType::Native, true);
+					HandleTokenStreamObjectReference(ObjectsToSerializeStruct, Obj, ObjectClass, UE::GC::EMemberlessId::Class, EOrigin::Other, true);
 					UObject* ObjectClassOuter = Obj->GetClass()->GetOuter();
-					HandleTokenStreamObjectReference(ObjectsToSerializeStruct, Obj, ObjectClassOuter, UE::GC::ETokenlessId::ClassOuter, EGCTokenType::Native, true);
+					HandleTokenStreamObjectReference(ObjectsToSerializeStruct, Obj, ObjectClassOuter, UE::GC::EMemberlessId::ClassOuter, EOrigin::Other, true);
 				}
 			}
 		}
@@ -673,7 +673,7 @@ public:
 	* @param TokenIndex Index to the token stream where the reference was found.
 	* @param bAllowReferenceElimination True if reference elimination is allowed (ignored when constructing clusters).
 	*/
-	FORCEINLINE void HandleTokenStreamObjectReference(FGCArrayStruct& ObjectsToSerializeStruct, UObject* ReferencingObject, UObject*& Object, UE::GC::FTokenId TokenIndex, const EGCTokenType TokenType, bool bAllowReferenceElimination)
+	FORCEINLINE void HandleTokenStreamObjectReference(FGCArrayStruct& ObjectsToSerializeStruct, UObject* ReferencingObject, UObject*& Object, FMemberId MemberId, EOrigin Origin, bool bAllowReferenceElimination)
 	{
 		if (!Object)
 		{
diff --git a/Engine/Source/Runtime/CoreUObject/Private/UObject/UObjectGlobals.cpp b/Engine/Source/Runtime/CoreUObject/Private/UObject/UObjectGlobals.cpp
--- a/Engine/Source/Runtime/CoreUObject/Private/UObject/UObjectGlobals.cpp
+++ b/Engine/Source/Runtime/CoreUObject/Private/UObject/UObjectGlobals.cpp
@@ -5107,7 +5107,6 @@ private:
 	 */
 	void AddToObjectList( const UObject* ReferencingObject, const FProperty* ReferencingProperty, UObject* Object )
 	{
-#if ENABLE_GC_DEBUG_OUTPUT
 		// this message is to help track down culprits behind "Object in PIE world still referenced" errors
 		if ( GIsEditor && !GIsPlayInEditorWorld && !CurrentObject->HasAnyFlags(RF_Transient) && Object->RootPackageHasAnyFlags(PKG_PlayInEditor) )
 		{
@@ -5119,7 +5118,6 @@ private:
 				UE_LOG(LogGarbage, Warning, TEXT("  NON-PIE object: %s"), *CurrentObject->GetFullName());
 			}
 		}
-#endif
 
 		// Mark it as reachable.
 		Object->ThisThreadAtomicallyClearedRFUnreachable();
@@ -6133,7 +6131,7 @@ namespace UECodeGen_Private
 		if ((NewClass->ClassFlags & CLASS_Intrinsic) != CLASS_Intrinsic)
 		{
 			check((NewClass->ClassFlags & CLASS_TokenStreamAssembled) != CLASS_TokenStreamAssembled);
-			NewClass->ReferenceTokens.Reset();
+			NewClass->ReferenceSchema.Reset();
 		}
 		NewClass->CreateLinkAndAddChildFunctionsToMap(Params.FunctionLinkArray, Params.NumFunctions);
 
diff --git a/Engine/Source/Runtime/CoreUObject/Public/UObject/Class.h b/Engine/Source/Runtime/CoreUObject/Public/UObject/Class.h
--- a/Engine/Source/Runtime/CoreUObject/Public/UObject/Class.h
+++ b/Engine/Source/Runtime/CoreUObject/Public/UObject/Class.h
@@ -2886,8 +2886,8 @@ public:
 	 */
 	TArray<FImplementedInterface> Interfaces;
 
-	/** Reference token stream used by realtime garbage collector, finalized in AssembleReferenceTokenStream */
-	UE::GC::FTokenStreamOwner ReferenceTokens;
+	/** GC schema, finalized in AssembleReferenceTokenStream */
+	UE::GC::FSchemaOwner ReferenceSchema;
 
 	/** This class's native functions. */
 	TArray<FNativeFunctionLookup> NativeFunctionLookupTable;
diff --git a/Engine/Source/Runtime/CoreUObject/Public/UObject/FastReferenceCollector.h b/Engine/Source/Runtime/CoreUObject/Public/UObject/FastReferenceCollector.h
--- a/Engine/Source/Runtime/CoreUObject/Public/UObject/FastReferenceCollector.h
+++ b/Engine/Source/Runtime/CoreUObject/Public/UObject/FastReferenceCollector.h
@@ -5,6 +5,7 @@
 #include "CoreMinimal.h"
 #include "Stats/Stats.h"
 #include "UObject/GarbageCollection.h"
+#include "UObject/GarbageCollectionSchema.h"
 #include "UObject/Class.h"
 #include "UObject/Package.h"
 #include "Async/TaskGraphInterfaces.h"
@@ -17,20 +18,21 @@
 #include "UObject/DynamicallyTypedValue.h"
 #include "UObject/GCObject.h"
 
-/** Token stream stack overflow checks are not enabled in test and shipping configs */
-#define UE_ENABLE_TOKENSTREAM_STACKOVERFLOW_CHECKS (!(UE_BUILD_TEST || UE_BUILD_SHIPPING) || 0)
-
-namespace UE::GC { class FWorkCoordinator; }
-
 /*=============================================================================
 	FastReferenceCollector.h: Unreal realtime garbage collection helpers
 =============================================================================*/
 
+namespace UE::GC
+{
+class FWorkCoordinator;
+struct FWorkerContext;
+}
+
 enum class EGCOptions : uint32
 {
 	None = 0,
 	Parallel = 1 << 0,					// Use all task workers to collect references, must be started on main thread
-	AutogenerateTokenStream = 1 << 1,	// Lazily generate token streams for new UClasses
+	AutogenerateSchemas = 1 << 1,		// Assemble schemas for new UClasses
 	WithClusters = 1 << 2,				// Use clusters, see FUObjectCluster
 	WithPendingKill = 1 << 3,			// Internal flag used by reachability analysis
 };
@@ -61,7 +63,7 @@ namespace UE::GC
 
 static constexpr uint32 ObjectLookahead = 16;
 
-// Prefetches ClassPrivate/OuterPrivate/Class' ReferenceTokens/TokenData while iterating over an object array
+// Prefetches ClassPrivate, OuterPrivate, class schema and schema data while iterating over an object array
 //
 // Tuned on a Gen5 console using an internal game replay and an in-game GC pass
 class FPrefetchingObjectIterator
@@ -71,16 +73,16 @@ public:
 	explicit FPrefetchingObjectIterator(TConstArrayView<UObject*> Objects)
 	: It(Objects.begin())
 	, End(Objects.end())
-	, PrefetchedTokenStream(Objects.Num() ? &It[1]->GetClass()->ReferenceTokens.Strong : nullptr)
+	, PrefetchedSchema(Objects.Num() ? &It[1]->GetClass()->ReferenceSchema.Get() : nullptr)
 	{}
 
 	FORCEINLINE_DEBUGGABLE void Advance()
 	{	
-		FPlatformMisc::Prefetch(PrefetchedTokenStream->GetTokenData());
-		PrefetchedTokenStream = &It[2]->GetClass()->ReferenceTokens.Strong;
+		FPlatformMisc::Prefetch(PrefetchedSchema->GetWords());
+		PrefetchedSchema = &It[2]->GetClass()->ReferenceSchema.Get();
 
 		UObjectBase::PrefetchOuter(It[6]);
-		FPlatformMisc::Prefetch(It[6]->GetClass(), offsetof(UClass, ReferenceTokens));
+		FPlatformMisc::Prefetch(It[6]->GetClass(), offsetof(UClass, ReferenceSchema));
 		UObjectBase::PrefetchClass(It[ObjectLookahead]);
 
 		++It;
@@ -92,7 +94,7 @@ public:
 private:
 	UObject*const* It;
 	UObject*const* End;
-	FTokenStreamView* PrefetchedTokenStream;
+	const FSchemaView* PrefetchedSchema;
 };
 
 // Pad object array for FPrefetchingObjectIterator use
@@ -229,7 +231,6 @@ struct FProcessorStats
 	FORCEINLINE void AddObjects(uint32 Num) { NumObjects += Num; }
 	FORCEINLINE void AddReferences(uint32 Num) { NumReferences += Num; }
 	FORCEINLINE void TrackPotentialGarbageReference(bool bDetectedGarbage) { bFoundGarbageRef |= bDetectedGarbage; }
-
 #endif
 
 	void AddStats(FProcessorStats Stats)
@@ -301,145 +302,278 @@ public:
 
 //////////////////////////////////////////////////////////////////////////
 
-/* Debug id for references that lack a token */
-enum class ETokenlessId
+namespace Private {
+
+struct FMemberUnpacked
 {
-	Collector = 1,
-	Class,
-	Outer,
-	ExternalPackage,
-	ClassOuter,
-	InitialReference,
-	Max = InitialReference
+	FMemberUnpacked(FMemberPacked In) 
+	: Type(static_cast<EMemberType>(In.Type))
+	, WordOffset(In.WordOffset)
+	{
+		check(Type < EMemberType::Count);
+	}
+	EMemberType Type;
+	uint32 WordOffset;
 };
 
-/** Debug identifier for a token stream index or a tokenless reference */
-struct FTokenId
+struct FMemberWordUnpacked
 {
-public:
-	FORCEINLINE FTokenId(ETokenlessId Tokenless) : Index(0), TokenlessId((uint32)Tokenless) {}
-	FORCEINLINE explicit FTokenId(uint32 Idx) : Index(Idx), TokenlessId(0) {}
+	FMemberWordUnpacked(const FMemberPacked In[4]) : Members{In[0], In[1], In[2], In[3]} {}
+	FMemberUnpacked Members[4];
+};
+
+struct FStructArray
+{
+	FSchemaView Schema{NoInit};
+	uint8* Data;
+	int32 Num;
+	uint32 Stride;
+};
+
+struct FStridedReferenceArray
+{
+	FScriptArray* Array;
+	FStridedLayout Layout;
+};
+
+struct FStridedReferenceView
+{
+	UObject** Data;
+	int32 Num;
+	uint32 Stride;
+};
+
+struct FStridedReferenceIterator
+{
+	UObject** It;
+	uint32 Stride;
+
+	UObject*& operator*() { return *It; }
+	FStridedReferenceIterator& operator++() { It += Stride; return *this;}
+	bool operator!=(FStridedReferenceIterator Rhs) const { return It != Rhs.It; }
+};
 	
-	bool IsTokenless() const { return TokenlessId != 0; }
-	uint32 GetIndex() const { check(!IsTokenless()); return Index; }
-	int32 AsPrintableIndex() const { return IsTokenless() ? -int32(TokenlessId) : Index; }
-	ETokenlessId AsTokenless() const { check(IsTokenless()); return (ETokenlessId)TokenlessId;}
+FORCEINLINE	FStridedReferenceIterator begin(FStridedReferenceView View) { return { View.Data, View.Stride }; }
+FORCEINLINE	FStridedReferenceIterator end(FStridedReferenceView View) { return { View.Data + View.Stride * View.Num, View.Stride }; }
+FORCEINLINE int32 GetNum(FStridedReferenceView View)  { return View.Num; }
+FORCEINLINE FStridedReferenceView ToView(FStridedReferenceArray In) 
+{ 
+	return { reinterpret_cast<UObject**>(In.Array->GetData()) + In.Layout.WordOffset,  In.Array->Num(), In.Layout.WordStride };
+}
 
-	friend bool operator==(FTokenId A, FTokenId B) { return A.All == B.All; }
-	friend bool operator!=(FTokenId A, FTokenId B) { return A.All != B.All; }
+FORCEINLINE uint8* GetSparseData(FScriptSparseArray& Array)
+{
+	return reinterpret_cast<uint8*>(Array.GetData(0, FScriptSparseArrayLayout{})); 
+}
 
-private:
-	static constexpr uint32 TokenlessIdBits = 3;
-	static constexpr uint32 IndexBits = 32 - TokenlessIdBits;
-	void StaticAssert();
+FORCEINLINE const uint8* GetSparseData(const FScriptSparseArray& Array)
+{
+	return reinterpret_cast<const uint8*>(Array.GetData(0, FScriptSparseArrayLayout{})); 
+}
 
-	union
+template<class DispatcherType>
+FORCENOINLINE void VisitNestedStructMembers(DispatcherType& Dispatcher, FSchemaView Schema, uint8* Instance);
+
+template<class DispatcherType>
+FORCEINLINE_DEBUGGABLE void VisitStructs(DispatcherType& Dispatcher, FSchemaView StructSchema, uint8* It, const int32 Num)
+{
+	check(!StructSchema.IsEmpty());
+	if constexpr (DispatcherType::bBatching)
+	{
+		Dispatcher.QueueStructArray(StructSchema, It, Num);
+	}
+	else
 	{
-		struct  
+		uint32 Stride = StructSchema.GetStructStride();
+		for (uint8* End = It + Num*Stride; It != End; It += Stride)
 		{
-			uint32 Index : IndexBits;
-			uint32 TokenlessId : TokenlessIdBits;
-		};
-		uint32 All;
-	};
-};
+			VisitNestedStructMembers(Dispatcher, StructSchema, It);
+		}
+	}
+}
 
-//////////////////////////////////////////////////////////////////////////
+template<class DispatcherType, class ArrayType>
+FORCEINLINE_DEBUGGABLE void VisitStructArray(DispatcherType& Dispatcher, FSchemaView StructSchema, ArrayType& Array)
+{
+	VisitStructs(Dispatcher, StructSchema, (uint8*)Array.GetData(), Array.Num());
+}
 
-/** Helper struct for stack based approach */
-struct FStackEntry
+template<class DispatcherType>
+FORCEINLINE_DEBUGGABLE void VisitSparseStructArray(DispatcherType& Dispatcher, FSchemaView StructSchema, FScriptSparseArray& Array)
 {
-	/** Current data pointer, incremented by stride */
-	uint8* Data;
-	/** Current container property for data pointer. DO NOT rely on its value being initialized. Instead check ContainerType first. */
-	FProperty* ContainerProperty;
-	/** Pointer to the container being processed by GC. DO NOT rely on its value being initialized. Instead check ContainerType first. */
-	void* ContainerPtr;
-	/** Current index within the container. DO NOT rely on its value being initialized. Instead check ContainerType first. */
-	int32	ContainerIndex;
-	/** Current container helper type */
-	uint32	ContainerType : 5; // The number of bits needs to match FGCReferenceInfo::Type
-	/** Current stride */
-	uint32	Stride : 27; // This will always be bigger (8 bits more) than FGCReferenceInfo::Ofset which is the max offset GC can handle
-	/** Current loop count, decremented each iteration */
-	int32	Count;
-	/** First token index in loop */
-	int32	LoopStartIndex;
-
-	FORCEINLINE bool MoveToNextContainerElementAndCheckIfValid()
+	check(!StructSchema.IsEmpty());
+	if constexpr (DispatcherType::bBatching)
+	{
+		Dispatcher.QueueSparseStructArray(StructSchema, Array);
+	}
+	else if (int32 Num = Array.Num())
 	{
-		if (ContainerType == GCRT_AddTMapReferencedObjects)
+		uint8* It = GetSparseData(Array);
+		const uint32 Stride = StructSchema.GetStructStride();
+		for (int32 Idx = 0, MaxIdx = Array.GetMaxIndex(); Idx < MaxIdx; ++Idx, It += Stride)
 		{
-			return ((FMapProperty*)ContainerProperty)->IsValidIndex(ContainerPtr, ++ContainerIndex);
+			if (Array.IsAllocated(Idx))
+			{
+				VisitNestedStructMembers(Dispatcher, StructSchema, It);
+			}
 		}
-		else if (ContainerType == GCRT_AddTSetReferencedObjects)
+	}
+}
+
+template<class DispatcherType>
+FORCEINLINE_DEBUGGABLE void VisitFieldPath(DispatcherType& Dispatcher, FFieldPath& FieldPath, EOrigin Origin, uint32 MemberIdx)
+{
+	if (FUObjectItem* FieldOwnerItem = FGCInternals::GetResolvedOwner(FieldPath))
+	{
+		UObject* OwnerObject = static_cast<UObject*>(FieldOwnerItem->Object);
+		UObject* PreviousOwner = OwnerObject;
+		Dispatcher.HandleReferenceDirectly(Dispatcher.Context.GetReferencingObject(), OwnerObject, FMemberId(MemberIdx), Origin, true);
+								
+		// Handle reference elimination (PendingKill owner)
+		if (PreviousOwner && !OwnerObject)
 		{
-			return ((FSetProperty*)ContainerProperty)->IsValidIndex(ContainerPtr, ++ContainerIndex);
+			FGCInternals::ClearCachedField(FieldPath);
 		}
-		
-		return true;
 	}
-};
-
-//////////////////////////////////////////////////////////////////////////
-
-/** Helper class to manage GC token stream stack  */
-class FTokenStreamStack
+}
+template<class DispatcherType>
+FORCEINLINE_DEBUGGABLE void VisitFieldPathArray(DispatcherType& Dispatcher, TArray<FFieldPath>& FieldPaths, EOrigin Origin, uint32 MemberIdx)
 {
-	/** Allocated stack memory */
-	TArray<FStackEntry> Stack;
-	/** Current stack frame */
-	FStackEntry* CurrentFrame = nullptr;
-#if UE_ENABLE_TOKENSTREAM_STACKOVERFLOW_CHECKS
-	/** Number of used stack frames (for debugging) */
-	int32 NumberOfUsedStackFrames = 0;
-#endif // UE_ENABLE_TOKENSTREAM_STACKOVERFLOW_CHECKS
-
-public:
-	FTokenStreamStack()
+	for (FFieldPath& FieldPath : FieldPaths)
 	{
-		Stack.AddUninitialized(FTokenStreamBuilder::GetMaxStackSize());
+		VisitFieldPath(Dispatcher, FieldPath, Origin, MemberIdx);
 	}
+}
+
+template<class DispatcherType>
+FORCEINLINE_DEBUGGABLE void VisitOptional(DispatcherType& Dispatcher, FSchemaView StructSchema, uint8* Instance)
+{
+	check(!StructSchema.IsEmpty());
+	uint32 ValueSize = StructSchema.GetStructStride();
+	bool bIsSet = *(bool*)(Instance + ValueSize);
+	VisitStructs(Dispatcher, StructSchema, Instance, bIsSet);
+}
 
-	/** Initializes the stack and returns its first frame */
-	FORCEINLINE FStackEntry* Initialize()
+template<class DispatcherType>
+FORCEINLINE_DEBUGGABLE void VisitDynamicallyTypedValue(DispatcherType& Dispatcher, UE::FDynamicallyTypedValue& Value)
+{
+	Value.GetType().MarkReachable(Dispatcher.Collector);
+	if (Value.GetType().GetContainsReferences() != UE::FDynamicallyTypedValueType::EContainsReferences::DoesNot)
 	{
-#if UE_ENABLE_TOKENSTREAM_STACKOVERFLOW_CHECKS
-		NumberOfUsedStackFrames = 1;
-#endif // UE_ENABLE_TOKENSTREAM_STACKOVERFLOW_CHECKS
-		// Grab te first frame. From now on we'll be using it instead of TArray to move to the next one and back
-		// in order to skip any extra internal TArray checks and overhead
-		CurrentFrame = GetBottomFrame();
-		return CurrentFrame;
+		Value.GetType().MarkValueReachable(Value.GetDataPointer(), Dispatcher.Collector);
 	}
+}
 
-	/** Returns the frame at the bottom of the stack (the first one) */
-	FORCEINLINE FStackEntry* GetBottomFrame()
+template<class DispatcherType>
+FORCEINLINE_DEBUGGABLE void CallARO(DispatcherType& Dispatcher, UObject* Instance, FMemberWord Word)
+{
+	Word.ObjectARO(Instance, Dispatcher.Collector);
+}
+
+template<class DispatcherType>
+FORCEINLINE_DEBUGGABLE void CallARO(DispatcherType& Dispatcher, uint8* Instance, FMemberWord Word)
+{
+	Word.StructARO(Instance, Dispatcher.Collector);
+}
+
+template<class DispatcherType>
+FORCEINLINE_DEBUGGABLE void CallSlowARO(DispatcherType& Dispatcher, uint32 SlowAROIdx, UObject* Instance, uint32 MemberIdx)
+{
+	if constexpr (DispatcherType::bBatching && DispatcherType::bParallel)
 	{
-		return Stack.GetData();
+		if (!FSlowARO::TryQueueCall(SlowAROIdx, Instance, Dispatcher.Context))
+		{
+			FSlowARO::CallSync(SlowAROIdx, Instance, Dispatcher.Collector);
+		}
 	}
-
-	/** Advances to the next frame on the stack */
-	FORCEINLINE FStackEntry* Push()
+	else
 	{
-		checkSlow(NumberOfUsedStackFrames > 0); // sanity check to make sure Push() gets called after Initialize()
-#if UE_ENABLE_TOKENSTREAM_STACKOVERFLOW_CHECKS
-		NumberOfUsedStackFrames++;
-		UE_CLOG(NumberOfUsedStackFrames > Stack.Num(), LogGarbage, Fatal, TEXT("Ran out of stack space for GC Token Stream. FGCReferenceTokenStream::GetMaxStackSize() = %d must be miscalculated. Verify EmitReferenceInfo code."), FTokenStreamBuilder::GetMaxStackSize());
-#endif // UE_ENABLE_TOKENSTREAM_STACKOVERFLOW_CHECKS
-		return ++CurrentFrame;
+		FSlowARO::CallSync(SlowAROIdx, Instance, Dispatcher.Collector);
 	}
+}
+
+FORCENOINLINE static void LogIllegalTypeFatal(EMemberType Type, uint32 Idx, UObject* Instance)
+{
+	UE_LOG(LogGarbage, Fatal, TEXT("Illegal GC object member type %d at %d, class:%s object:%s"), Type, Idx, Instance ? *GetNameSafe(Instance->GetClass()) : TEXT("Unknown"), *GetPathNameSafe(Instance));
+}
+
+FORCENOINLINE static void LogIllegalTypeFatal(EMemberType Type, uint32 Idx, uint8*)
+{
+	UE_LOG(LogGarbage, Fatal, TEXT("Illegal GC struct member type %d at %d"), Type, Idx);
+}
 
-	/** Pops back to the previous frame on the stack */
-	FORCEINLINE FStackEntry* Pop()
+template<class DispatcherType>
+FORCEINLINE void CallSlowARO(DispatcherType&, uint32 SlowAROIdx, uint8* Instance, uint32 MemberIdx)
+{
+	LogIllegalTypeFatal(EMemberType::SlowARO, MemberIdx, Instance);
+}
+
+template<class DispatcherType, typename ObjectType>
+FORCEINLINE_DEBUGGABLE void VisitMembers(DispatcherType& Dispatcher, FSchemaView Schema, ObjectType* Instance)
+{
+	check(!Schema.IsEmpty());
+
+	const EOrigin Origin = Schema.GetOrigin();
+	uint64* InstanceCursor = (uint64*)Instance;	// Advanced via Jump to reach far members
+	uint32 DebugIdx = 0;
+	for (const FMemberWord* WordIt = Schema.GetWords(); true; ++WordIt)
 	{
-#if UE_ENABLE_TOKENSTREAM_STACKOVERFLOW_CHECKS
-		NumberOfUsedStackFrames--;
-		UE_CLOG(NumberOfUsedStackFrames < 1, LogGarbage, Fatal, TEXT("GC token stream stack Pop() was called too many times. Probably a Push() call is missing for one of the tokens."));
-#endif // UE_ENABLE_TOKENSTREAM_STACKOVERFLOW_CHECKS
-		return --CurrentFrame;
-	}
-};
+		const FMemberWordUnpacked Quad(WordIt->Members);
+		for (FMemberUnpacked Member : Quad.Members)
+		{
+			uint8* MemberPtr = (uint8*)(InstanceCursor + Member.WordOffset);
+
+			switch (Member.Type)
+			{
+			case EMemberType::Reference:				Dispatcher.HandleKillableReference(*(UObject**)MemberPtr, FMemberId(DebugIdx), Origin);
+			break;
+			case EMemberType::ReferenceArray:			Dispatcher.HandleKillableArray(*(TArray<UObject*>*)MemberPtr, FMemberId(DebugIdx), Origin);
+			break;
+			case EMemberType::StridedArray:				Dispatcher.HandleKillableArray(FStridedReferenceArray{(FScriptArray*)MemberPtr, (++WordIt)->StridedLayout}, FMemberId(DebugIdx), Origin);
+			break;	
+			case EMemberType::FreezableReferenceArray:	Dispatcher.HandleKillableReferences(*(TArray<UObject*, FMemoryImageAllocator>*)MemberPtr, FMemberId(DebugIdx), Origin);
+			break;
+			case EMemberType::StructArray:				VisitStructArray(			Dispatcher, FSchemaView((++WordIt)->InnerSchema, Origin), *(FScriptArray*)MemberPtr);
+			break;
+			case EMemberType::SparseStructArray:		VisitSparseStructArray(		Dispatcher, FSchemaView((++WordIt)->InnerSchema, Origin), *(FScriptSparseArray*)MemberPtr);
+			break;
+			case EMemberType::FreezableStructArray:		VisitStructArray(			Dispatcher, FSchemaView((++WordIt)->InnerSchema, Origin), *(FFreezableScriptArray*)MemberPtr);
+			break;
+			case EMemberType::Optional:					VisitOptional(				Dispatcher, FSchemaView((++WordIt)->InnerSchema, Origin), MemberPtr);
+			break;
+			case EMemberType::FieldPath:				VisitFieldPath(				Dispatcher, *(FFieldPath*)MemberPtr, Origin, DebugIdx);
+			break;
+			case EMemberType::FieldPathArray:			VisitFieldPathArray(		Dispatcher, *(TArray<FFieldPath>*)MemberPtr, Origin, DebugIdx);
+			break;
+			case EMemberType::DynamicallyTypedValue:	VisitDynamicallyTypedValue(	Dispatcher, *(UE::FDynamicallyTypedValue*)MemberPtr);
+			break;
+			case EMemberType::Jump:						InstanceCursor += (Member.WordOffset + 1) * FMemberPacked::OffsetRange;
+			break;
+			case EMemberType::MemberARO:				CallARO(Dispatcher, MemberPtr, *++WordIt);
+			break; // Struct member ARO isn't an implicit stop
+			case EMemberType::ARO:						CallARO(Dispatcher, Instance, *++WordIt);
+			return; // Instance ARO is an implicit stop
+			case EMemberType::SlowARO:					CallSlowARO(Dispatcher, /* slow ARO index */ Member.WordOffset, Instance, DebugIdx);
+			return; // ARO is an implicit stop
+			case EMemberType::Stop:
+			return; // Stop schema without ARO call
+			default:									LogIllegalTypeFatal(Member.Type, DebugIdx, Instance);
+			return;
+			}
+
+			DebugIdx += UE_GC_DEBUGNAMES;
+		} // for quad members
+	} // for schema member words
+}
+
+template<class DispatcherType>
+void VisitNestedStructMembers(DispatcherType& Dispatcher, FSchemaView Schema, uint8* Instance)
+{
+	static_assert(!DispatcherType::bBatching);
+	VisitMembers(Dispatcher, Schema, Instance);
+}
+
+} // namespace Private
 
 //////////////////////////////////////////////////////////////////////////
 
@@ -447,54 +581,66 @@ public:
 template<class ProcessorType>
 struct TDirectDispatcher
 {
+	static constexpr bool bBatching = false;
+	static constexpr bool bParallel = IsParallel(ProcessorType::Options);
+
 	ProcessorType& Processor;
 	FWorkerContext& Context;
+	FReferenceCollector& Collector;
 
-	FORCEINLINE void HandleReferenceDirectly(UObject* ReferencingObject, UObject*& Object, FTokenId TokenId, EGCTokenType TokenType, bool bAllowReferenceElimination) const
+	FORCEINLINE void HandleReferenceDirectly(UObject* ReferencingObject, UObject*& Object, FMemberId MemberId, EOrigin Origin, bool bAllowReferenceElimination) const
 	{
 		if (IsObjectHandleResolved(*reinterpret_cast<FObjectHandle*>(&Object)))
 		{
-			Processor.HandleTokenStreamObjectReference(Context, ReferencingObject, Object, TokenId, TokenType, bAllowReferenceElimination);
+			Processor.HandleTokenStreamObjectReference(Context, ReferencingObject, Object, MemberId, Origin, bAllowReferenceElimination);
 		}
 		Context.Stats.AddReferences(1);
 	}
 	
-	FORCEINLINE void HandleKillableReference(UObject*& Object, FTokenId TokenId, EGCTokenType TokenType) const
+	FORCEINLINE void HandleKillableReference(UObject*& Object, FMemberId MemberId, EOrigin Origin) const
 	{
-		HandleReferenceDirectly(Context.GetReferencingObject(), Object, TokenId, TokenType, true);
+		HandleReferenceDirectly(Context.GetReferencingObject(), Object, MemberId, Origin, true);
 	}
 
-	FORCEINLINE void HandleImmutableReference(UObject* Object, FTokenId TokenId, EGCTokenType TokenType) const
+	FORCEINLINE void HandleImmutableReference(UObject* Object, FMemberId MemberId, EOrigin Origin) const
 	{
-		HandleReferenceDirectly(Context.GetReferencingObject(), Object, TokenId, TokenType, false);
+		HandleReferenceDirectly(Context.GetReferencingObject(), Object, MemberId, Origin, false);
 	}
-
-	FORCEINLINE void HandleKillableArray(TArray<UObject*>& Array, FTokenId TokenId, EGCTokenType TokenType) const
+	
+	template<class ArrayType>
+	FORCEINLINE void HandleKillableReferences(ArrayType&& Objects , FMemberId MemberId, EOrigin Origin) const
 	{
-		for (UObject*& Object : Array)
+		for (UObject*& Object : Objects)
 		{
-			HandleReferenceDirectly(Context.GetReferencingObject(), Object, TokenId, TokenType, true);
+			HandleReferenceDirectly(Context.GetReferencingObject(), Object, MemberId, Origin, true);
 		}
 	}
 
-	FORCEINLINE void HandleKillableReferences(TArrayView<UObject*> Objects, FTokenId TokenId, EGCTokenType TokenType) const
+	FORCEINLINE void HandleKillableArray(TArray<UObject*>& Array, FMemberId MemberId, EOrigin Origin) const
 	{
-		for (UObject*& Object : Objects)
-		{
-			HandleReferenceDirectly(Context.GetReferencingObject(), Object, TokenId, TokenType, true);
-		}
+		HandleKillableReferences(Array, MemberId, Origin);
 	}
 
-	FORCEINLINE void FlushQueuedReferences() const {}
+	FORCEINLINE void HandleKillableArray(Private::FStridedReferenceArray Array, FMemberId MemberId, EOrigin Origin) const
+	{
+		HandleKillableReferences(ToView(Array), MemberId, Origin);
+	}
 };
 
 // Default implementation is to create new direct dispatcher
 template<class CollectorType, class ProcessorType>
-TDirectDispatcher<ProcessorType> GetDispatcher(CollectorType&, ProcessorType& Processor, FWorkerContext& Context)
+TDirectDispatcher<ProcessorType> GetDispatcher(CollectorType& Collector, ProcessorType& Processor, FWorkerContext& Context)
 {
-	return { Processor, Context };
+	return { Processor, Context, Collector };
 }
 
+template<class CollectorType, class ProcessorType, class = void >
+struct TGetDispatcherType
+{
+	using RetType = decltype(GetDispatcher(*(CollectorType*)nullptr, *(ProcessorType*)nullptr, *(FWorkerContext*)nullptr));
+	using Type = typename std::remove_reference_t<RetType>;
+};
+
 //////////////////////////////////////////////////////////////////////////
 
 enum class ELoot { Nothing, Block, ARO, Context };
@@ -517,31 +663,12 @@ COREUOBJECT_API void ProcessAsync(void (*ProcessSync)(void*, FWorkerContext&), v
 template <typename ProcessorType, typename CollectorType>
 class TFastReferenceCollector : public FGCInternals
 {
-	static constexpr EGCOptions Options = ProcessorType::Options;
-
-	static constexpr FORCEINLINE bool IsParallel()					{ return !!(Options & EGCOptions::Parallel); }
-
-	static FORCEINLINE UE::GC::FTokenStreamView& GetTokenStream(UClass* Class)
-	{
-		if (!!(Options & EGCOptions::AutogenerateTokenStream) && !Class->HasAnyClassFlags(CLASS_TokenStreamAssembled))
-		{			
-			Class->AssembleReferenceTokenStream();
-		}
-
-		return Class->ReferenceTokens.Strong;
-	}
-
-	ProcessorType& Processor;
-
 public:
 	TFastReferenceCollector(ProcessorType& InProcessor) : Processor(InProcessor) {}
 
 	void ProcessObjectArray(FWorkerContext& Context)
 	{
-		static_assert(!EnumHasAllFlags(Options, EGCOptions::Parallel | EGCOptions::AutogenerateTokenStream), "Can't assemble token streams in parallel");
-
-		// Presized "recursion" stack for handling arrays and structs.
-		FTokenStreamStack Stack;
+		static_assert(!EnumHasAllFlags(Options, EGCOptions::Parallel | EGCOptions::AutogenerateSchemas), "Can't assemble token streams in parallel");
 		
 		CollectorType Collector(Processor, Context);
 
@@ -553,428 +680,14 @@ StoleContext:
 		Context.ReferencingObject = FGCObject::GGCObjectReferencer;
 		for (UObject** InitialReference : Context.InitialNativeReferences)
 		{
-			Dispatcher.HandleKillableReference(*InitialReference, ETokenlessId::InitialReference, EGCTokenType::Native);
+			Dispatcher.HandleKillableReference(*InitialReference, EMemberlessId::InitialReference, EOrigin::Other);
 		}
 
 		TConstArrayView<UObject*> CurrentObjects = Context.InitialObjects;
 		while (true)
 		{
 			Context.Stats.AddObjects(CurrentObjects.Num());
-			for (FPrefetchingObjectIterator It(CurrentObjects); It.HasMore(); It.Advance())
-			{
-				UObject* CurrentObject = It.GetCurrentObject();
-				UClass* Class = CurrentObject->GetClass();
-				UObject* Outer = CurrentObject->GetOuter();
-
-				FTokenStreamView TokenStream = GetTokenStream(Class);
-				Context.ReferencingObject = CurrentObject;
-
-				// Emit base references
-				Dispatcher.HandleImmutableReference(Class, ETokenlessId::Class, EGCTokenType::Native);
-				Dispatcher.HandleImmutableReference(Outer, ETokenlessId::Outer, EGCTokenType::Native);
-#if WITH_EDITOR
-				UObject* Package = CurrentObject->GetExternalPackageInternal();
-				Package = Package != CurrentObject ? Package : nullptr;
-				Dispatcher.HandleImmutableReference(Package, ETokenlessId::ExternalPackage, EGCTokenType::Native);
-#endif
-				if (TokenStream.IsEmpty())
-				{
-					continue;
-				}
-
-				Processor.BeginTimingObject(CurrentObject);
-
-				uint32 TokenStreamIndex = 0;
-				// Keep track of index to reference info. Used to avoid LHSs.
-				uint32 ReferenceTokenStreamIndex = 0;
-
-				// Create stack entry and initialize sane values.
-				FStackEntry* StackEntry = Stack.Initialize();
-				uint8* StackEntryData = (uint8*)CurrentObject;
-				StackEntry->Data = StackEntryData;
-				StackEntry->ContainerType = GCRT_None;
-				StackEntry->Stride = 0;
-				StackEntry->Count = -1;
-				StackEntry->LoopStartIndex = -1;
-				
-				// Keep track of token return count in separate integer as arrays need to fiddle with it.
-				uint8 TokenReturnCount = 0;
-
-				// Parse the token stream.
-				while (true)
-				{
-					// Cache current token index as it is the one pointing to the reference info.
-					ReferenceTokenStreamIndex = TokenStreamIndex;
-
-					// Handle returning from an array of structs, array of structs of arrays of ... (yadda yadda)
-					for (uint8 ReturnCount = 0; ReturnCount < TokenReturnCount; ReturnCount++)
-					{
-						// Make sure there's no stack underflow.
-						check(StackEntry->Count != -1);
-
-						// We pre-decrement as we're already through the loop once at this point.
-						if (--StackEntry->Count > 0)
-						{
-							if (StackEntry->ContainerType == GCRT_None)
-							{
-								// Fast path for TArrays of structs
-								// Point data to next entry.
-								StackEntryData = StackEntry->Data + StackEntry->Stride;
-								StackEntry->Data = StackEntryData;
-							}
-							else
-							{
-								// Slower path for other containers
-								// Point data to next valid entry.
-								do
-								{
-									StackEntryData = StackEntry->Data + StackEntry->Stride;
-									StackEntry->Data = StackEntryData;
-								} while (!StackEntry->MoveToNextContainerElementAndCheckIfValid());
-							}
-
-							// Jump back to the beginning of the loop.
-							TokenStreamIndex = StackEntry->LoopStartIndex;
-							ReferenceTokenStreamIndex = StackEntry->LoopStartIndex;
-							// We're not done with this token loop so we need to early out instead of backing out further.
-							break;
-						}
-						else
-						{
-							StackEntry->ContainerType = GCRT_None;
-							StackEntry = Stack.Pop();
-							StackEntryData = StackEntry->Data;
-						}
-					}
-
-					TokenStreamIndex++;
-
-					struct FUnpackedReferenceInfo
-					{
-						explicit FUnpackedReferenceInfo(FGCReferenceInfo In)
-						: ReturnCount(In.ReturnCount)
-						, Type(static_cast<EGCReferenceType>(In.Type))
-						, Offset(In.Offset)
-						{}
-
-						uint8 ReturnCount;
-						EGCReferenceType Type;
-						uint32 Offset;
-					};
-					const FUnpackedReferenceInfo ReferenceInfo(TokenStream.AccessReferenceInfo(ReferenceTokenStreamIndex));
-							
-					const EGCTokenType TokenType = TokenStream.GetTokenType();
-
-					switch (ReferenceInfo.Type)
-					{
-						case GCRT_Object:
-						{								
-							UObject** ObjectPtr = (UObject**)(StackEntryData + ReferenceInfo.Offset);
-							TokenReturnCount = ReferenceInfo.ReturnCount;
-							Dispatcher.HandleKillableReference(*ObjectPtr, FTokenId(ReferenceTokenStreamIndex), TokenType);
-						}
-						break;
-						case GCRT_ArrayObject:
-						{
-							TArray<UObject*>& ObjectArray = *((TArray<UObject*>*)(StackEntryData + ReferenceInfo.Offset));
-							TokenReturnCount = ReferenceInfo.ReturnCount;
-							Dispatcher.HandleKillableArray(ObjectArray, FTokenId(ReferenceTokenStreamIndex), TokenType);
-						}
-						break;
-						case GCRT_ArrayObjectFreezable:
-						{
-							TArray<UObject*, FMemoryImageAllocator>& ObjectArray = *((TArray<UObject*, FMemoryImageAllocator>*)(StackEntryData + ReferenceInfo.Offset));
-							TokenReturnCount = ReferenceInfo.ReturnCount;
-							Dispatcher.HandleKillableReferences(ObjectArray, FTokenId(ReferenceTokenStreamIndex), TokenType);
-						}
-						break;
-						case GCRT_ArrayStruct:
-						{
-							const FScriptArray& Array = *((FScriptArray*)(StackEntryData + ReferenceInfo.Offset));
-							StackEntry = Stack.Push();
-							StackEntryData = (uint8*)Array.GetData();
-							StackEntry->Data = StackEntryData;
-							StackEntry->Stride = TokenStream.ReadStride(TokenStreamIndex);
-							StackEntry->Count = Array.Num();
-							StackEntry->ContainerType = GCRT_None;
-
-							const FGCSkipInfo SkipInfo = TokenStream.ReadSkipInfo(TokenStreamIndex);
-							StackEntry->LoopStartIndex = TokenStreamIndex;
-
-							if (StackEntry->Count == 0)
-							{
-								// Skip empty array by jumping to skip index and set return count to the one about to be read in.
-								TokenStreamIndex = SkipInfo.SkipIndex;
-								TokenReturnCount = TokenStream.GetSkipReturnCount(SkipInfo);
-							}
-							else
-							{
-								// Loop again.
-								check(StackEntry->Data);
-								TokenReturnCount = 0;
-							}
-						}
-						break;
-						case GCRT_ArrayStructFreezable:
-						{
-							const FFreezableScriptArray& Array = *((FFreezableScriptArray*)(StackEntryData + ReferenceInfo.Offset));
-							StackEntry = Stack.Push();
-							StackEntryData = (uint8*)Array.GetData();
-							StackEntry->Data = StackEntryData;
-							StackEntry->Stride = TokenStream.ReadStride(TokenStreamIndex);
-							StackEntry->Count = Array.Num();
-							StackEntry->ContainerType = GCRT_None;
-
-							const FGCSkipInfo SkipInfo = TokenStream.ReadSkipInfo(TokenStreamIndex);
-							StackEntry->LoopStartIndex = TokenStreamIndex;
-
-							if (StackEntry->Count == 0)
-							{
-								// Skip empty array by jumping to skip index and set return count to the one about to be read in.
-								TokenStreamIndex = SkipInfo.SkipIndex;
-								TokenReturnCount = TokenStream.GetSkipReturnCount(SkipInfo);
-							}
-							else
-							{
-								// Loop again.
-								check(StackEntry->Data);
-								TokenReturnCount = 0;
-							}
-						}
-						break;
-						case GCRT_AddReferencedObjects:
-						{
-							void(*AddReferencedObjects)(UObject*, FReferenceCollector&) = (void(*)(UObject*, FReferenceCollector&))TokenStream.ReadPointer(TokenStreamIndex);
-							TokenReturnCount = ReferenceInfo.ReturnCount;
-							AddReferencedObjects(CurrentObject, Collector);
-
-							// ARO is always last and an implicit terminator
-							goto TokensDone;
-						}
-						break;
-						case GCRT_EndOfStream:
-						{
-							// Break out of loop
-							goto TokensDone;
-						}
-						break;
-						case GCRT_SlowAddReferencedObjects:
-						{
-							if (!IsParallel() ||
-								!FSlowARO::TryQueueCall(/* ARO index */ ReferenceInfo.Offset, CurrentObject, Context))
-							{
-								FSlowARO::CallSync(/* ARO index */ ReferenceInfo.Offset, CurrentObject, Collector);
-							}
-								
-							// ARO is always last and an implicit terminator
-							goto TokensDone;
-						}
-						case GCRT_FixedArray:
-						{
-							uint8* PreviousData = StackEntryData;
-							StackEntry = Stack.Push();
-							StackEntryData = PreviousData;
-							StackEntry->Data = PreviousData;
-							StackEntry->Stride = TokenStream.ReadStride(TokenStreamIndex);
-							StackEntry->Count = TokenStream.ReadCount(TokenStreamIndex);
-							StackEntry->LoopStartIndex = TokenStreamIndex;
-							StackEntry->ContainerType = GCRT_None;
-							TokenReturnCount = 0;
-						}
-						break;
-						case GCRT_AddStructReferencedObjects:
-						{
-							void* StructPtr = (void*)(StackEntryData + ReferenceInfo.Offset);
-							TokenReturnCount = ReferenceInfo.ReturnCount;
-							UScriptStruct::ICppStructOps::TPointerToAddStructReferencedObjects Func = (UScriptStruct::ICppStructOps::TPointerToAddStructReferencedObjects)TokenStream.ReadPointer(TokenStreamIndex);
-							Func(StructPtr, Collector);
-						}
-						break;
-						case GCRT_AddTMapReferencedObjects:
-						{
-							void* MapPtr = StackEntryData + ReferenceInfo.Offset;
-							FMapProperty* MapProperty = (FMapProperty*)TokenStream.ReadPointer(TokenStreamIndex);
-							TokenStreamIndex++; // GCRT_EndOfPointer
-
-							StackEntry = Stack.Push();
-							StackEntry->ContainerType = GCRT_AddTMapReferencedObjects;
-							StackEntry->ContainerIndex = 0;
-							StackEntry->ContainerProperty = MapProperty;
-							StackEntry->ContainerPtr = MapPtr;
-							StackEntry->Stride = MapProperty->GetPairStride();
-							StackEntry->Count = MapProperty->GetNum(MapPtr);
-
-							const FGCSkipInfo SkipInfo = TokenStream.ReadSkipInfo(TokenStreamIndex);
-							StackEntry->LoopStartIndex = TokenStreamIndex;
-
-							if (StackEntry->Count == 0)
-							{
-								// The map is empty
-								StackEntryData = nullptr;
-								StackEntry->Data = StackEntryData;
-
-								// Skip empty map by jumping to skip index and set return count to the one about to be read in.
-								TokenStreamIndex = SkipInfo.SkipIndex;
-								TokenReturnCount = TokenStream.GetSkipReturnCount(SkipInfo);
-							}
-							else
-							{
-								// Skip any initial invalid entries in the map. We need a valid index for MapProperty->GetPairPtr()
-								int32 FirstValidIndex = 0;
-								while (!MapProperty->IsValidIndex(MapPtr, FirstValidIndex))
-								{
-									FirstValidIndex++;
-								}
-
-								StackEntry->ContainerIndex = FirstValidIndex;
-								StackEntryData = MapProperty->GetPairPtr(MapPtr, FirstValidIndex);
-								StackEntry->Data = StackEntryData;
-
-								// Loop again.
-								TokenReturnCount = 0;
-							}
-						}
-						break;
-						case GCRT_AddTSetReferencedObjects:
-						{
-							void* SetPtr = StackEntryData + ReferenceInfo.Offset;
-							FSetProperty* SetProperty = (FSetProperty*)TokenStream.ReadPointer(TokenStreamIndex);
-							TokenStreamIndex++; // GCRT_EndOfPointer
-
-							StackEntry = Stack.Push();
-							StackEntry->ContainerProperty = SetProperty;
-							StackEntry->ContainerPtr = SetPtr;
-							StackEntry->ContainerType = GCRT_AddTSetReferencedObjects;
-							StackEntry->ContainerIndex = 0;
-							StackEntry->Stride = SetProperty->GetStride();
-							StackEntry->Count = SetProperty->GetNum(SetPtr);
-
-							const FGCSkipInfo SkipInfo = TokenStream.ReadSkipInfo(TokenStreamIndex);
-							StackEntry->LoopStartIndex = TokenStreamIndex;
-
-							if (StackEntry->Count == 0)
-							{
-								// The set is empty or it doesn't contain any valid elements
-								StackEntryData = nullptr;
-								StackEntry->Data = StackEntryData;
-
-								// Skip empty set by jumping to skip index and set return count to the one about to be read in.
-								TokenStreamIndex = SkipInfo.SkipIndex;
-								TokenReturnCount = TokenStream.GetSkipReturnCount(SkipInfo);
-							}
-							else
-							{
-								// Skip any initial invalid entries in the set. We need a valid index for SetProperty->GetElementPtr()
-								int32 FirstValidIndex = 0;
-								while (!SetProperty->IsValidIndex(SetPtr, FirstValidIndex))
-								{
-									FirstValidIndex++;
-								}
-
-								StackEntry->ContainerIndex = FirstValidIndex;
-								StackEntryData = SetProperty->GetElementPtr(SetPtr, FirstValidIndex);
-								StackEntry->Data = StackEntryData;
-
-								// Loop again.
-								TokenReturnCount = 0;
-							}
-						}
-						break;
-						case GCRT_AddFieldPathReferencedObject:
-						{
-							FFieldPath* FieldPathPtr = (FFieldPath*)(StackEntryData + ReferenceInfo.Offset);
-							FUObjectItem* FieldOwnerItem = GetResolvedOwner(*FieldPathPtr);
-							TokenReturnCount = ReferenceInfo.ReturnCount;
-							if (FieldOwnerItem)
-							{
-								UObject* OwnerObject = static_cast<UObject*>(FieldOwnerItem->Object);
-								UObject* PreviousOwner = OwnerObject;
-								Dispatcher.HandleReferenceDirectly(CurrentObject, OwnerObject, FTokenId(ReferenceTokenStreamIndex), TokenType, true);
-								
-								// Handle reference elimination (PendingKill owner)
-								if (PreviousOwner && !OwnerObject)
-								{
-									ClearCachedField(*FieldPathPtr);
-								}
-							}
-						}
-						break;
-						case GCRT_ArrayAddFieldPathReferencedObject:
-						{
-							TArray<FFieldPath>& FieldArray = *((TArray<FFieldPath>*)(StackEntryData + ReferenceInfo.Offset));
-							TokenReturnCount = ReferenceInfo.ReturnCount;
-							for (int32 FieldIndex = 0, FieldNum = FieldArray.Num(); FieldIndex < FieldNum; ++FieldIndex)
-							{
-								if (FUObjectItem* FieldOwnerItem = GetResolvedOwner(FieldArray[FieldIndex]))
-								{
-									UObject* OwnerObject = static_cast<UObject*>(FieldOwnerItem->Object);
-									UObject* PreviousOwner = OwnerObject;
-									Dispatcher.HandleReferenceDirectly(CurrentObject, OwnerObject, FTokenId(ReferenceTokenStreamIndex), TokenType, true);
-
-									// Handle reference elimination (PendingKill owner)
-									if (PreviousOwner && !OwnerObject)
-									{
-										ClearCachedField(FieldArray[FieldIndex]);
-									}
-								}
-							}
-						}
-						break;
-						case GCRT_Optional:
-						{
-							uint32 ValueSize = TokenStream.ReadStride(TokenStreamIndex); // Size of value in bytes. This is also the offset to the bIsSet variable stored thereafter.
-							const FGCSkipInfo SkipInfo = TokenStream.ReadSkipInfo(TokenStreamIndex);
-							const bool& bIsSet = *((bool*)(StackEntryData + ReferenceInfo.Offset + ValueSize));
-
-							StackEntry = Stack.Push();
-							StackEntryData += ReferenceInfo.Offset;
-							StackEntry->Data = StackEntryData;
-							StackEntry->Stride = ValueSize;
-							StackEntry->Count = 1;
-							StackEntry->LoopStartIndex = TokenStreamIndex;
-
-							if (bIsSet)
-							{
-								// It's set - push a stack entry for processing the value
-								// This is somewhat suboptimal since there is only ever just one value, but this approach avoids any changes to the surrounding code
-								check(StackEntry->Data);
-								TokenReturnCount = 0;
-							}
-							else
-							{
-								// It's unset - keep going by jumping to skip index
-								TokenStreamIndex = SkipInfo.SkipIndex;
-								TokenReturnCount = TokenStream.GetSkipReturnCount(SkipInfo);
-							}
-						}
-						break;
-						case GCRT_EndOfPointer:
-						{
-							TokenReturnCount = ReferenceInfo.ReturnCount;
-						}
-						break;
-						case GCRT_DynamicallyTypedValue:
-						{
-							TokenReturnCount = ReferenceInfo.ReturnCount;
-							UE::FDynamicallyTypedValue* DynamicallyTypedValue = (UE::FDynamicallyTypedValue*)(StackEntryData + ReferenceInfo.Offset);
-							DynamicallyTypedValue->GetType().MarkReachable(Collector);
-							if (DynamicallyTypedValue->GetType().GetContainsReferences() != UE::FDynamicallyTypedValueType::EContainsReferences::DoesNot)
-							{
-								DynamicallyTypedValue->GetType().MarkValueReachable(DynamicallyTypedValue->GetDataPointer(), Collector);
-							}
-						}
-						break;
-						default:
-						{
-							UE_LOG(LogGarbage, Fatal, TEXT("Unknown token. Type:%d ReferenceTokenStreamIndex:%d Class:%s Obj:%s"), ReferenceInfo.Type, ReferenceTokenStreamIndex, CurrentObject ? *GetNameSafe(CurrentObject->GetClass()) : TEXT("Unknown"), *GetPathNameSafe(CurrentObject));
-							break;
-						}
-					} // switch (Type)
-				} // while (true)
-TokensDone:
-				Processor.UpdateDetailedStats(CurrentObject);
-				check(StackEntry == Stack.GetBottomFrame());
-			} // for (FPrefetchingObjectIterator)
+			ProcessObjects(Dispatcher, CurrentObjects);
 
 			// Free finished work block
 			if (CurrentObjects.GetData() != Context.InitialObjects.GetData())
@@ -987,17 +700,25 @@ TokensDone:
 			FWorkBlock* Block = RemainingObjects.PopFullBlock<Options>();
 			if (!Block)
 			{
-				if constexpr (IsParallel())
+				if constexpr (bIsParallel)
 				{
 					FSlowARO::ProcessUnbalancedCalls(Context, Collector);
 				}
 
-StoleARO:		
-				Dispatcher.FlushQueuedReferences();
+StoleARO:
+				if constexpr (DispatcherType::bBatching)
+				{
+					if (Dispatcher.FlushToStructBlocks())
+					{
+						ProcessStructs(Dispatcher);
+					}
+					
+					Dispatcher.FlushQueuedReferences();
+				}
 
 				if (	 Block = RemainingObjects.PopFullBlock<Options>(); Block);
 				else if (Block = RemainingObjects.PopPartialBlock(/* out if successful */ BlockSize); Block);
-				else if (IsParallel()) // if constexpr yields MSVC unreferenced label warning
+				else if (bIsParallel) // if constexpr yields MSVC unreferenced label warning
 				{
 					switch (StealWork(/* in-out */ Context, Collector, /* out */ Block))
 					{
@@ -1019,6 +740,48 @@ StoleARO:
 		
 		Processor.LogDetailedStatsSummary();
 	}
+
+private:
+	using DispatcherType = typename TGetDispatcherType<CollectorType, ProcessorType>::Type;
+	static constexpr EGCOptions Options = ProcessorType::Options;
+	static constexpr bool bIsParallel = IsParallel(Options);
+	
+	ProcessorType& Processor;
+
+	FORCEINLINE_DEBUGGABLE void ProcessObjects(DispatcherType& Dispatcher, TConstArrayView<UObject*> CurrentObjects)
+	{
+		for (FPrefetchingObjectIterator It(CurrentObjects); It.HasMore(); It.Advance())
+		{
+			UObject* CurrentObject = It.GetCurrentObject();
+			UClass* Class = CurrentObject->GetClass();
+			UObject* Outer = CurrentObject->GetOuter();
+
+			if (!!(Options & EGCOptions::AutogenerateSchemas) && !Class->HasAnyClassFlags(CLASS_TokenStreamAssembled))
+			{			
+				Class->AssembleReferenceTokenStream();
+			}
+
+			FSchemaView Schema = Class->ReferenceSchema.Get();
+			Dispatcher.Context.ReferencingObject = CurrentObject;
+
+			// Emit base references
+			Dispatcher.HandleImmutableReference(Class, EMemberlessId::Class, EOrigin::Other);
+			Dispatcher.HandleImmutableReference(Outer, EMemberlessId::Outer, EOrigin::Other);
+#if WITH_EDITOR
+			UObject* Package = CurrentObject->GetExternalPackageInternal();
+			Package = Package != CurrentObject ? Package : nullptr;
+			Dispatcher.HandleImmutableReference(Package, EMemberlessId::ExternalPackage, EOrigin::Other);
+#endif
+			if (!Schema.IsEmpty())
+			{
+				Processor.BeginTimingObject(CurrentObject);
+				Private::VisitMembers(Dispatcher, Schema, CurrentObject);
+				Processor.UpdateDetailedStats(CurrentObject);
+			}
+		}
+	}
+
+	FORCENOINLINE void ProcessStructs(DispatcherType& Dispatcher);
 };
 
 //////////////////////////////////////////////////////////////////////////
@@ -1032,21 +795,21 @@ protected:
 	UE::GC::FWorkerContext& Context;
 
 public:
-	TDefaultCollector(ProcessorType& InProcessor, FGCArrayStruct& InContext)
+	TDefaultCollector(ProcessorType& InProcessor, UE::GC::FWorkerContext& InContext)
 	: Processor(InProcessor)
 	, Context(InContext)
 	{}
 
 	virtual void HandleObjectReference(UObject*& Object, const UObject* ReferencingObject, const FProperty* ReferencingProperty) override
 	{
-		Processor.HandleTokenStreamObjectReference(Context, const_cast<UObject*>(ReferencingObject), Object, ETokenlessId::Collector, EGCTokenType::Native, false);
+		Processor.HandleTokenStreamObjectReference(Context, const_cast<UObject*>(ReferencingObject), Object, EMemberlessId::Collector, EOrigin::Other, false);
 	}
 	virtual void HandleObjectReferences(UObject** InObjects, const int32 ObjectNum, const UObject* ReferencingObject, const FProperty* InReferencingProperty) override
 	{
 		for (int32 ObjectIndex = 0; ObjectIndex < ObjectNum; ++ObjectIndex)
 		{
 			UObject*& Object = InObjects[ObjectIndex];
-			Processor.HandleTokenStreamObjectReference(Context, const_cast<UObject*>(ReferencingObject), Object, ETokenlessId::Collector, EGCTokenType::Native, false);
+			Processor.HandleTokenStreamObjectReference(Context, const_cast<UObject*>(ReferencingObject), Object, EMemberlessId::Collector, EOrigin::Other, false);
 		}
 	}
 
@@ -1054,7 +817,7 @@ public:
 	virtual bool IsIgnoringTransient() const override {	return false; }
 };
 
-inline constexpr EGCOptions DefaultOptions = EGCOptions::AutogenerateTokenStream;
+inline constexpr EGCOptions DefaultOptions = EGCOptions::AutogenerateSchemas;
 
 } // namespace UE::GC
 
@@ -1062,6 +825,11 @@ inline constexpr EGCOptions DefaultOptions = EGCOptions::AutogenerateTokenStream
 class FSimpleReferenceProcessorBase
 {
 public:
+	using FMemberId				= UE::GC::FMemberId;
+	using EOrigin				= UE::GC::EOrigin;
+	using FPropertyStack		= UE::GC::FPropertyStack;
+	using FWorkerContext		= UE::GC::FWorkerContext;
+
 	static constexpr EGCOptions Options = UE::GC::DefaultOptions;
 
 	// These functions are implemented in the GC collectors to generate detail stats on objects considered by GC.
@@ -1071,7 +839,7 @@ public:
 	void LogDetailedStatsSummary() {}
 
 	// Implement this in your derived class, don't make this virtual as it will affect performance!
-	//FORCEINLINE void HandleTokenStreamObjectReference(UE::GC::FWorkerContext& Context, UObject* ReferencingObject, UObject*& Object, UE::GC::FTokenId TokenIndex, EGCTokenType TokenType, bool bAllowReferenceElimination)
+	//FORCEINLINE void HandleTokenStreamObjectReference(FWorkerContext& Context, UObject* ReferencingObject, UObject*& Object, FMemberId MemberId, EOrigin Origin, bool bAllowReferenceElimination)
 };
 
 
@@ -1098,4 +866,12 @@ FORCEINLINE static void CollectReferences(ProcessorType& Processor, UE::GC::FWor
 }
 
 // Get number of workers to use when calling CollectReferences in parallel
-COREUOBJECT_API int32 GetNumCollectReferenceWorkers();
\ No newline at end of file
+COREUOBJECT_API int32 GetNumCollectReferenceWorkers();
+
+// Temporary aliases to old types
+
+using FTokenInfo = UE::GC::FMemberInfo;
+using EGCTokenType = UE::GC::EOrigin;
+using FGCArrayStruct = UE::GC::FWorkerContext;
+namespace UE::GC { using FTokenId = FMemberId; }
+namespace UE::GC { using ETokenlessId = EMemberlessId; }
\ No newline at end of file
diff --git a/Engine/Source/Runtime/CoreUObject/Public/UObject/FieldPathProperty.h b/Engine/Source/Runtime/CoreUObject/Public/UObject/FieldPathProperty.h
--- a/Engine/Source/Runtime/CoreUObject/Public/UObject/FieldPathProperty.h
+++ b/Engine/Source/Runtime/CoreUObject/Public/UObject/FieldPathProperty.h
@@ -20,7 +20,7 @@ class UObject;
 class UStruct;
 namespace UECodeGen_Private { struct FFieldPathPropertyParams; }
 struct FPropertyTag;
-namespace UE::GC { class FTokenStreamBuilder; }
+
 // need to break this out a different type so that the DECLARE_CASTED_CLASS_INTRINSIC macro can digest the comma
 typedef TProperty<FFieldPath, FProperty> FFieldPathProperty_Super;
 
@@ -75,7 +75,7 @@ PRAGMA_ENABLE_DEPRECATION_WARNINGS
 	virtual const TCHAR* ImportText_Internal(const TCHAR* Buffer, void* ContainerOrPropertyPtr, EPropertyPointerType PropertyPointerType, UObject* OwnerObject, int32 PortFlags, FOutputDevice* ErrorText) const override;
 	virtual void Serialize(FArchive& Ar) override;
 	virtual bool ContainsObjectReference(TArray<const FStructProperty*>& EncounteredStructProps, EPropertyObjectReferenceType InReferenceType = EPropertyObjectReferenceType::Strong) const override;
-	virtual void EmitReferenceInfo(UE::GC::FTokenStreamBuilder& TokenStream, int32 BaseOffset, TArray<const FStructProperty*>& EncounteredStructProps, FGCStackSizeHelper& StackSizeHelper) override;
+	virtual void EmitReferenceInfo(UE::GC::FSchemaBuilder& Schema, int32 BaseOffset, TArray<const FStructProperty*>& EncounteredStructProps, UE::GC::FPropertyStack& DebugPath) override;
 	virtual bool SupportsNetSharedSerialization() const override;
 	// End of FProperty interface
 
diff --git a/Engine/Source/Runtime/CoreUObject/Public/UObject/GarbageCollection.h b/Engine/Source/Runtime/CoreUObject/Public/UObject/GarbageCollection.h
--- a/Engine/Source/Runtime/CoreUObject/Public/UObject/GarbageCollection.h
+++ b/Engine/Source/Runtime/CoreUObject/Public/UObject/GarbageCollection.h
@@ -6,37 +6,25 @@
 
 #pragma once
 
-#include "Containers/Array.h"
-#include "Containers/Map.h"
 #include "Containers/UnrealString.h"
 #include "CoreTypes.h"
-#include "HAL/PlatformCrt.h"
-#include "HAL/ThreadSafeBool.h"
 #include "Logging/LogMacros.h"
 #include "Misc/AssertionMacros.h"
-#include "Serialization/ArchiveUObject.h"
 #include "Stats/Stats.h"
 #include "Stats/Stats2.h"
-#include "UObject/NameTypes.h"
 #include "UObject/UObjectGlobals.h"
 #include "UObject/UnrealNames.h"
 
-class FProperty;
 class UObject;
 
 /** Context sensitive keep flags for garbage collection */
 #define GARBAGE_COLLECTION_KEEPFLAGS	(GIsEditor ? RF_Standalone : RF_NoFlags)
 
-#define	ENABLE_GC_DEBUG_OUTPUT					1
-
 /** UObject pointer checks are disabled by default in shipping and test builds as they add roughly 20% overhead to GC times */
 #ifndef ENABLE_GC_OBJECT_CHECKS
 	#define ENABLE_GC_OBJECT_CHECKS (!(UE_BUILD_TEST || UE_BUILD_SHIPPING) || 0)
 #endif
 
-/** Token debug info (token names) enabled in non-shipping builds */
-#define ENABLE_GC_TOKEN_DEBUG_INFO (!UE_BUILD_SHIPPING)
-
 #define ENABLE_GC_HISTORY (!UE_BUILD_SHIPPING)
 
 COREUOBJECT_API DECLARE_LOG_CATEGORY_EXTERN(LogGarbage, Warning, All);
@@ -48,113 +36,6 @@ DECLARE_STATS_GROUP(TEXT("Garbage Collection"), STATGROUP_GC, STATCAT_Advanced);
  */
 #define DO_POINTER_CHECKS_ON_GC 0 
 
-namespace UE::GC
-{
-struct FWorkerContext;
-struct FTokenId;
-}
-
-using FGCArrayStruct = UE::GC::FWorkerContext;
-
-/*-----------------------------------------------------------------------------
-	Realtime garbage collection helper classes.
------------------------------------------------------------------------------*/
-
-/**
- * Enum of different supported reference type tokens.
- */
-enum EGCReferenceType
-{
-	GCRT_None = 0,
-	GCRT_Object,							// Scalar reference, i.e. MyObject* and TObjectPtr<MyObject>
-	GCRT_ExternalPackage,					// UObject external package
-	GCRT_ArrayObject,						// Array of references, e.g. TArray<MyObject*> and TArray<TObjectPtr<MyObject>>
-	GCRT_ArrayStruct,						// Array of structs
-	GCRT_FixedArray,						// C-style fixed arrays, e.g. SomeType Props[4]
-	GCRT_AddStructReferencedObjects,		// AddStructReferencedObjects() function pointer
-	GCRT_AddReferencedObjects,				// AddReferencedObjects() function pointer
-	GCRT_AddTMapReferencedObjects,			// Map whose value or key contain references
-	GCRT_AddTSetReferencedObjects,			// Map whose elements contain references
-	GCRT_AddFieldPathReferencedObject,		// Field path strong reference to owner
-	GCRT_ArrayAddFieldPathReferencedObject, // Array of field paths
-	GCRT_EndOfPointer,						// Comes after pointer
-	GCRT_EndOfStream,						// Terminator
-	GCRT_ArrayObjectFreezable,				// Freezable array of references, e.g. TArray<MyObject*, FMemoryImageAllocator>
-	GCRT_ArrayStructFreezable,				// Freezable array of structs
-	GCRT_Optional,							// TOptional
-	GCRT_DynamicallyTypedValue,				// FDynamicallyTypedValue
-	GCRT_SlowAddReferencedObjects,			// AddReferencedObjects() function using UE::GC::RegisterSlowImplementation
-	GCRT_Count
-};
-
-enum class EGCTokenType : uint32
-{
-	Native = 0,
-	NonNative = 1	
-};
-
-/** Transcodes uint32 token <-> reflection info for one reference property */
-struct FGCReferenceInfo
-{
-	FORCEINLINE FGCReferenceInfo(EGCReferenceType InReferenceType, uint32 InOffset)
-		: ReturnCount(0)
-		, Type(InReferenceType)
-		, Offset(InOffset)
-	{
-		static_assert(sizeof(FGCReferenceInfo) == sizeof(uint32));
-		checkf(InReferenceType != GCRT_None && InReferenceType <= 0x1F, TEXT("Invalid GC Token Reference Type (%d)"), (uint32)InReferenceType);
-		checkf((InOffset & ~0x7FFFF) == 0, TEXT("Invalid GC Token Offset (%d), max is %d"), InOffset, 0x7FFFF);
-	}
-
-	FORCEINLINE explicit FGCReferenceInfo(uint32 InValue) : Value(InValue) {}
-
-	FORCEINLINE operator uint32() const { return Value; }
-
-	union
-	{
-		struct
-		{
-			/** Return depth, e.g. 1 for last entry in an array, 2 for last entry in an array of structs of arrays, ... */
-			uint32 ReturnCount	: 8;
-			/** Type of reference */
-			uint32 Type			: 5; // The number of bits needs to match TFastReferenceCollector::FStackEntry::ContainerHelperType
-			/** Offset into struct/ object */
-			uint32 Offset		: 19;
-		};
-		/** uint32 value of reference info, used for easy conversion to/ from uint32 for token array */
-		uint32 Value;
-	};
-};
-
-/** Transcodes uint32 token <-> info needed to skip a dynamic array */
-struct FGCSkipInfo
-{
-	FORCEINLINE FGCSkipInfo() {}
-	FORCEINLINE explicit FGCSkipInfo(uint32 InValue) : Value(InValue) {}
-
-	FORCEINLINE operator uint32() const { return Value; }
-
-	union
-	{
-		/** Mapping to exactly one uint32 */
-		struct
-		{
-			/** Return depth not taking into account embedded arrays. This is needed to return appropriately when skipping empty dynamic arrays as we never step into them */
-			uint32 InnerReturnCount	: 8;
-			/** Skip index */
-			uint32 SkipIndex			: 24;
-		};
-		/** uint32 value of skip info, used for easy conversion to/ from uint32 for token array */
-		uint32 Value;
-	};
-};
-
-/** Token debug name and offset */
-struct FTokenInfo
-{
-	int32 Offset;
-	FName Name;
-};
 
 namespace UE::GC {
 
@@ -169,281 +50,30 @@ ENUM_CLASS_FLAGS(EAROFlags);
 // Reference collection batches up slow AddReferencedObjects calls
 COREUOBJECT_API void RegisterSlowImplementation(void (*AddReferencedObjects)(UObject*, FReferenceCollector&), EAROFlags Flags = EAROFlags::None);
 
-struct FTokenStreamOwner;
-class FTokenStreamBuilderIterator;
+class FSchemaView;
 
-class /* ensure same cache line */ alignas(16) FTokenStreamView
+/** Type-erasing owner of a ref-counted FSchemaView */
+class FSchemaOwner
 {
 public:
-	FTokenStreamView() : StackSize(0), DroppedNum(0), IsNonNative(0) {}
-
-	FORCEINLINE EGCTokenType GetTokenType() const
-	{
-		return IsNonNative ? EGCTokenType::NonNative : EGCTokenType::Native;
-	}
-
-	FORCEINLINE FGCReferenceInfo AccessReferenceInfo(uint32 CurrentIndex) const
-	{
-		return FGCReferenceInfo(Tokens[CurrentIndex]);
-	}
-
-	FORCEINLINE uint32 ReadStride(uint32& CurrentIndex) const
-	{
-		return Tokens[CurrentIndex++];
-	}
-
-	FORCEINLINE FGCSkipInfo ReadSkipInfo(uint32& CurrentIndex) const
-	{
-		FGCSkipInfo SkipInfo(Tokens[CurrentIndex]);
-		SkipInfo.SkipIndex += CurrentIndex;
-		CurrentIndex++;
-		return SkipInfo;
-	}
-
-	/**
-	 * Read return count stored at the index before the skip index. This is required 
-	 * to correctly return the right amount of levels when skipping over an empty array.
-	 *
-	 * @param SkipIndex index of first token after array
-	 */
-	FORCEINLINE uint8 GetSkipReturnCount( FGCSkipInfo SkipInfo ) const
-	{
-		check( SkipInfo.SkipIndex > 0 && SkipInfo.SkipIndex <= Num );		
-		FGCReferenceInfo ReferenceInfo(Tokens[SkipInfo.SkipIndex-1]);
-		check( ReferenceInfo.Type != GCRT_None );
-		return static_cast<uint8>(ReferenceInfo.ReturnCount - SkipInfo.InnerReturnCount);
-	}
-
-	FORCEINLINE uint32 ReadCount(uint32& CurrentIndex) const
-	{
-		return Tokens[CurrentIndex++];
+	FSchemaOwner() = default;
+	FSchemaOwner(const FSchemaOwner&) = delete;
+	FSchemaOwner(FSchemaOwner&& In) : SchemaView(In.SchemaView) { In.SchemaView = 0; }
+	COREUOBJECT_API explicit FSchemaOwner(FSchemaView In);
+	~FSchemaOwner() { Reset(); }
+	FSchemaOwner& operator=(const FSchemaOwner&) = delete;
+	FSchemaOwner& operator=(FSchemaOwner&& In)
+	{
+		Reset();
+		Swap(In.SchemaView, SchemaView);
+		return *this;
 	}
-
-	FORCEINLINE void* ReadPointer(uint32& CurrentIndex) const
-	{
-		UPTRINT Result = UPTRINT(Tokens[CurrentIndex]) | (UPTRINT(Tokens[CurrentIndex+1]) << 32);
-		CurrentIndex += 2;
-		return (void*)Result;
-	}
-
-	bool IsEmpty() const { return Num == 0; }
-	uint32 NumTokens() const { return Num; }
-	void* GetTokenData() const { return Tokens; }
-private:
-	friend FTokenStreamOwner;
-	friend class FTokenStreamBuilder;
-	static constexpr uint32 NumStackSizeBits = 28;
-	static constexpr uint32 NumDroppedNumBits = 3;
-
-	uint32* RESTRICT Tokens = nullptr; // Debug names stored after tokens
-	uint32 Num = 0;
-	uint32 StackSize : NumStackSizeBits;
-	uint32 DroppedNum : NumDroppedNumBits;
-	uint32 IsNonNative : 1;
-
-#if ENABLE_GC_TOKEN_DEBUG_INFO
-	const FName* GetDebugNames() const { return reinterpret_cast<const FName*>(Tokens + Num + DroppedNum); }
-#endif
-};
-
-
-struct FTokenStreamOwner
-{
-	UE_NONCOPYABLE(FTokenStreamOwner);
-	FTokenStreamOwner() = default;
-	~FTokenStreamOwner() { Reset(); }
-
-	FTokenStreamView Strong;
 	
+	COREUOBJECT_API void Set(FSchemaView In);
+	const FSchemaView& Get() const { return reinterpret_cast<const FSchemaView&>(SchemaView); }
 	COREUOBJECT_API void Reset();
-
-	/** Gets the number of bytes allocated for tokens */
-	int64 GetTokenAllocatedSize() const;
-
-	/** Gets the number of bytes allocated for token debug info */
-	int64 GetDebugInfoAllocatedSize() const;
-
-	/** Returns extended information about a token, including its name (if available) */
-	FTokenInfo GetTokenInfo(UE::GC::FTokenId Id) const;
-};
-
-/** Reference token stream class. Used for creating and parsing stream of object references. */
-class FTokenStreamBuilder
-{
-public:
-	explicit FTokenStreamBuilder(UClass& Owner) : DebugOwner(Owner) {}
-
-	uint32 Num() const { return static_cast<uint32>(Tokens.Num()); }
-	bool IsEmpty() const { return Tokens.Num() == 0; }
-	void SetStackSize(int32 InStackSize) { StackSize = InStackSize; }
-	uint32 GetStackSize() const { return StackSize; }
-	void SetTokenType(EGCTokenType InType) { TokenType = InType; }
-	EGCTokenType GetTokenType() const { return TokenType; }
-
-	/** Emit reference info and return index of the reference info */
-	COREUOBJECT_API int32 EmitReferenceInfo(FGCReferenceInfo ReferenceInfo, FName DebugName);
-
-	/**
-	 * Emit placeholder for array skip index, updated in UpdateSkipIndexPlaceholder
-	 *
-	 * @return the index of the skip index, used later in UpdateSkipIndexPlaceholder
-	 */
-	COREUOBJECT_API uint32 EmitSkipIndexPlaceholder();
-
-	/**
-	 * Updates skip index place holder stored and passed in skip index index with passed
-	 * in skip index. The skip index is used to skip over tokens in the case of an empty 
-	 * dynamic array.
-	 * 
-	 * @param SkipIndexIndex index where skip index is stored at.
-	 * @param SkipIndex index to store at skip index index
-	 */
-	COREUOBJECT_API void UpdateSkipIndexPlaceholder(uint32 SkipIndexIndex, uint32 SkipIndex);
-
-	int32 EmitCount(uint32 Count);
-	int32 EmitPointer(const void* Ptr);
-	COREUOBJECT_API int32 EmitStride(uint32 Stride);
-
-	/**
-	 * Increase return count on last token.
-	 *
-	 * @return index of next token
-	 */
-	COREUOBJECT_API uint32 EmitReturn();
-
-	COREUOBJECT_API void EmitObjectReference(int32 Offset, FName DebugName, EGCReferenceType Kind = GCRT_Object);
-	COREUOBJECT_API void EmitObjectArrayReference(int32 Offset, FName DebugName);
-	COREUOBJECT_API uint32 EmitStructArrayBegin(int32 Offset, FName DebugName, int32 Stride);
-
-	/** The index following the current one will be written to the passed in SkipIndexIndex to skip tokens for empty dynamic arrays. */
-	COREUOBJECT_API void EmitStructArrayEnd(uint32 SkipIndexIndex);
-
-	/** All tokens issues between Begin and End will be replayed Count times. */
-	COREUOBJECT_API void EmitFixedArrayBegin(int32 Offset, FName DebugName, int32 Stride, int32 Count);
-	COREUOBJECT_API void EmitFixedArrayEnd();
-
-	void EmitFinalTokens(void(*AddReferencedObjects)(UObject*, FReferenceCollector&));	
-	static FTokenStreamView DropFinalTokens(FTokenStreamView Tokens, void(*DropARO)(UObject*, FReferenceCollector&));
-
-	/** Allocate merged stream with tokens from super class and current class */
-	static void Merge(FTokenStreamView& Out, const FTokenStreamBuilder& Class, FTokenStreamView Super);
-	
-	/** Reads count and advances stream */
-	FORCEINLINE uint32 ReadCount(uint32& CurrentIndex)
-	{
-		return Tokens[CurrentIndex++];
-	}
-
-	/** Reads stride and advances stream */
-	FORCEINLINE uint32 ReadStride(uint32& CurrentIndex)
-	{
-		return Tokens[CurrentIndex++];
-	}
-
-	/** Reads pointer and advance stream */
-	FORCEINLINE void* ReadPointer(uint32& CurrentIndex) const
-	{
-		UPTRINT Result = UPTRINT(Tokens[CurrentIndex++]);
-		Result |= UPTRINT(Tokens[CurrentIndex++]) << 32;
-		return (void*)Result;
-	}
-
-	/** Reads in reference info and advances stream. */
-	FORCEINLINE FGCReferenceInfo ReadReferenceInfo( uint32& CurrentIndex )
-	{
-		return FGCReferenceInfo(Tokens[CurrentIndex++]);
-	}
-
-	/**
-	 * Access reference info at passed in index. Used as helper to eliminate LHS.
-	 *
-	 * @return Reference info at passed in index
-	 */
-	FORCEINLINE FGCReferenceInfo AccessReferenceInfo(uint32 CurrentIndex) const
-	{
-		return FGCReferenceInfo(Tokens[CurrentIndex]);
-	}
-
-	/**
-	 * Read in skip index and advances stream.
-	 *
-	 * @return read in skip index
-	 */
-	FORCEINLINE FGCSkipInfo ReadSkipInfo( uint32& CurrentIndex )
-	{
-		FGCSkipInfo SkipInfo(Tokens[CurrentIndex]);
-		SkipInfo.SkipIndex += CurrentIndex;
-		CurrentIndex++;
-		return SkipInfo;
-	}
-
-	/**
-	 * Read return count stored at the index before the skip index. This is required 
-	 * to correctly return the right amount of levels when skipping over an empty array.
-	 *
-	 * @param SkipIndex index of first token after array
-	 */
-	FORCEINLINE uint8 GetSkipReturnCount( FGCSkipInfo SkipInfo )
-	{
-		check( SkipInfo.SkipIndex > 0 && SkipInfo.SkipIndex <= (uint32)Tokens.Num() );		
-		FGCReferenceInfo ReferenceInfo(Tokens[SkipInfo.SkipIndex-1]);
-		check( ReferenceInfo.Type != GCRT_None );
-		return static_cast<uint8>(ReferenceInfo.ReturnCount - SkipInfo.InnerReturnCount);
-	}
-
-	/**
-	 * Queries the stream for an end of stream condition
-	 *
-	 * @return true if the end of the stream has been reached, false otherwise
-	 */
-	FORCEINLINE bool EndOfStream(uint32 CurrentIndex)
-	{
-		return CurrentIndex >= (uint32)Tokens.Num();
-	}
-
-	/** Gets the maximum stack size required by all token streams */
-	FORCEINLINE static uint32 GetMaxStackSize()
-	{
-		return MaxStackSize;
-	}
-
-	friend bool operator==(const FTokenStreamBuilder& A, const FTokenStreamBuilder& B)
-	{
-		return A.Tokens == B.Tokens && A.StackSize == B.StackSize && A.TokenType == B.TokenType;
-	}
-
 private:
-	static constexpr uint32 SkipIndexPlaceholderMagic =  0xDEDEDEDEu; 
-
-	/**
-	 * Helper function to store a pointer into a preallocated token stream.
-	 *
-	 * @param Stream Preallocated token stream.
-	 * @param Ptr pointer to store
-	 */
-	FORCEINLINE void StorePointer(uint32* Stream, const void* Ptr)
-	{
-		Stream[0] = UPTRINT(Ptr) & 0xffffffffu;
-		Stream[1] = UPTRINT(Ptr) >> 32;
-	}
-
-	/** Stack size required by this token stream */
-	uint32 StackSize = 0;
-	EGCTokenType TokenType = EGCTokenType::Native;
-	TArray<uint32, TInlineAllocator<128>> Tokens;
-#if ENABLE_GC_TOKEN_DEBUG_INFO
-	TArray<FName, TInlineAllocator<128>> DebugNames;
-#endif
-	UClass& DebugOwner;
-
-	/** Maximum stack size for TFastReferenceCollector */
-	COREUOBJECT_API static uint32 MaxStackSize;
-};
-
-struct FIntrinsicClassTokens
-{
-	COREUOBJECT_API static FTokenStreamBuilder& AllocateBuilder(UClass* Class);
-	static TUniquePtr<FTokenStreamBuilder> ConsumeBuilder(UClass* Class);
+	uint64 SchemaView = 0;
 };
 
 } // namespace UE::GC
diff --git a/Engine/Source/Runtime/CoreUObject/Public/UObject/GarbageCollectionHistory.h b/Engine/Source/Runtime/CoreUObject/Public/UObject/GarbageCollectionHistory.h
--- a/Engine/Source/Runtime/CoreUObject/Public/UObject/GarbageCollectionHistory.h
+++ b/Engine/Source/Runtime/CoreUObject/Public/UObject/GarbageCollectionHistory.h
@@ -9,6 +9,7 @@
 #include "Containers/Array.h"
 #include "Containers/Map.h"
 #include "HAL/Platform.h"
+#include "UObject/FastReferenceCollector.h"
 #include "UObject/GCObjectInfo.h"
 #include "UObject/GarbageCollection.h"
 #include "UObject/NameTypes.h"
diff --git a/Engine/Source/Runtime/CoreUObject/Public/UObject/GarbageCollectionSchema.h b/Engine/Source/Runtime/CoreUObject/Public/UObject/GarbageCollectionSchema.h
new file mode 100644
--- /dev/null
+++ b/Engine/Source/Runtime/CoreUObject/Public/UObject/GarbageCollectionSchema.h
@@ -0,0 +1,354 @@
+// Copyright Epic Games, Inc. All Rights Reserved.
+
+#pragma once
+
+#include "Containers/Array.h"
+#include "Containers/Map.h"
+#include "Misc/Optional.h"
+#include "UObject/GarbageCollection.h"
+#include "UObject/NameTypes.h"
+#include "UObject/ObjectPtr.h"
+
+#ifndef UE_GC_DEBUGNAMES
+#define UE_GC_DEBUGNAMES (!UE_BUILD_SHIPPING)
+#endif
+
+namespace UE::GC {
+
+/** Describes member variable types, AddReferencedObject calls and stop/jump instructions used by VisitMembers */
+enum class EMemberType : uint8
+{
+	Stop,								// Null terminator
+	Jump,								// Move base pointer forward to reach members at large offsets
+	Reference,							// Member - Scalar reference, e.g. MyObject* and TObjectPtr<MyObject>
+	ReferenceArray,						// Member - Array of references, e.g. TArray<MyObject*> and TArray<TObjectPtr<MyObject>>
+	StructArray,						// Array of structs
+	StridedArray,						// Array of structs with single reference per struct (~half of struct instances)
+	SparseStructArray,					// TMap/TSet of structs
+	FieldPath,							// Field path strong reference to owner
+	FieldPathArray,					 	// Array of field paths
+	FreezableReferenceArray,			// Freezable array of references
+	FreezableStructArray,				// Freezable array of structs
+	Optional,							// TOptional
+	DynamicallyTypedValue,				// FDynamicallyTypedValue
+	ARO,								// Call Add[Struct]ReferencedObjects() on current object / struct
+	SlowARO,							// Call or queue AddReferencedObjects() on current object
+	MemberARO,							// Call AddStructReferencedObjects() on a struct member in current object / struct
+	Count
+};
+
+COREUOBJECT_API FName ToName(EMemberType Type);
+
+/** Declares if a schema represents a blueprint generated type */
+enum class EOrigin : uint8
+{
+	Other,		// Schema reflects non-blueprint type, such as native C++ classes and structs.
+	Blueprint	// Schema reflects blueprint type. References to garbage-marked objects are cleared regardless of WithPendingKill to preserve old BP semantics.
+};
+
+/////////////////////////////// FPropertyStack ///////////////////////////
+
+/** Helps generate debug property paths for GC schemas and to resolve them back to FProperties */
+class FPropertyStack
+{
+public:
+	UE_NONCOPYABLE(FPropertyStack);
+	FPropertyStack() = default;
+	~FPropertyStack() { check(Props.IsEmpty()); }
+
+	/** Get string representing property stack, e.g. "Member.StructMember.InnerStructMember" */
+	COREUOBJECT_API FString GetPropertyPath() const;
+
+	/** 
+	 * Converts a property path constructed with GetPropertyPath() to an array of properties (from the outermost to the innermost) *
+	 * @param ObjectClass Class that defines the outermost property
+	 * @param InPropertyPath Property path
+	 * @param OutProperties resulting array
+	 * @returns true if the conversion was successful, false otherwise
+	 */
+	COREUOBJECT_API static bool ConvertPathToProperties(UClass* ObjectClass, FName InPropertyPath, TArray<FProperty*>& OutProperties);
+private:
+	friend class FPropertyStackScope;
+	TArray<const FProperty*> Props;
+};
+
+class FPropertyStackScope
+{
+	FPropertyStack& Stack;
+public:
+	UE_NONCOPYABLE(FPropertyStackScope);
+	FPropertyStackScope(FPropertyStack& InStack, FProperty* Property)
+	: Stack(InStack)
+	{		
+		check(Property);
+		Stack.Props.Push(Property);
+	}
+
+	~FPropertyStackScope()
+	{		
+		Stack.Props.Pop();
+	}
+};
+
+/////////////////////////////// FSchemaView //////////////////////////////
+
+struct FSchemaHeader
+{
+	uint32 StructStride; // sizeof(T), required to iterate over struct array
+	std::atomic<int32> RefCount;
+};
+
+union FMemberWord;
+
+/** Describes all strong GC references in a class or struct */
+class FSchemaView
+{
+	static constexpr uint64 OriginBit = 1;
+	uint64 Handle;
+
+public:
+	FSchemaView() : Handle(0) {}
+	FSchemaView(ENoInit) {} 
+	FSchemaView(FSchemaView View, EOrigin Origin) : FSchemaView(View.GetWords(), Origin) {}
+	explicit FSchemaView(const FMemberWord* Data, EOrigin Origin = EOrigin::Other)
+	: Handle(reinterpret_cast<uint64>(Data) | static_cast<uint64>(Origin))
+	{
+		static_assert(sizeof(Handle) >= sizeof(Data));
+	} 
+	
+		
+	const FMemberWord* GetWords() const				{ return reinterpret_cast<FMemberWord*>(Handle & ~OriginBit); }
+	EOrigin GetOrigin() const						{ return static_cast<EOrigin>(Handle & OriginBit); }
+	bool IsEmpty() const							{ return GetWords() == nullptr;}
+	void SetOrigin(EOrigin Origin)					{ Handle = (Handle & ~OriginBit) | static_cast<uint64>(Origin); }
+
+	/// @pre !IsEmpty()
+	uint32 GetStructStride() const					{ return reinterpret_cast<const FSchemaHeader*>(GetWords())[-1].StructStride; }
+	FSchemaHeader& GetHeader();
+	FSchemaHeader* TryGetHeader();
+};
+
+///////////////////////////// Debug helpers //////////////////////////////
+
+/* Debug id for references that lack a schema member declaration */
+enum class EMemberlessId
+{
+	Collector = 1,
+	Class,
+	Outer,
+	ExternalPackage,
+	ClassOuter,
+	InitialReference,
+	Max = InitialReference
+};
+
+/** Debug identifier for a schema index or a memberless reference */
+class FMemberId
+{
+public:
+	FORCEINLINE /* implicit */ FMemberId(EMemberlessId Memberless) : Index(0), MemberlessId((uint32)Memberless) {}
+	FORCEINLINE explicit FMemberId(uint32 Idx) : Index(Idx), MemberlessId(0) {}
+	
+	bool IsMemberless() const { return MemberlessId != 0; }
+	uint32 GetIndex() const { check(!IsMemberless()); return Index; }
+	int32 AsPrintableIndex() const { return IsMemberless() ? -int32(MemberlessId) : Index; }
+	EMemberlessId AsMemberless() const { check(IsMemberless()); return (EMemberlessId)MemberlessId;}
+
+	friend bool operator==(FMemberId A, FMemberId B) { return reinterpret_cast<uint32&>(A) == reinterpret_cast<uint32&>(B); }
+	friend bool operator!=(FMemberId A, FMemberId B) { return reinterpret_cast<uint32&>(A) != reinterpret_cast<uint32&>(B); }
+
+private:
+	static constexpr uint32 MemberlessIdBits = 3;
+	static constexpr uint32 IndexBits = 32 - MemberlessIdBits;
+
+	uint32 Index : IndexBits;
+	uint32 MemberlessId : MemberlessIdBits;
+
+	void StaticAssert();
+};
+
+struct FMemberInfo
+{
+	int32 Offset;
+	FName Name;
+};
+
+COREUOBJECT_API FMemberInfo GetMemberDebugInfo(FSchemaView Schema, FMemberId Id);
+
+#if !UE_BUILD_SHIPPING
+uint32 CountSchemas(uint32& OutNumWords);
+#endif
+
+//////////////////// Built FSchemaView representation ////////////////////
+
+struct FMemberPacked
+{
+	static constexpr uint32 TypeBits = 5;
+	static constexpr uint32 OffsetBits = 16 - TypeBits;
+	static constexpr uint32 OffsetRange = 1u << FMemberPacked::OffsetBits;
+
+	uint16 Type : TypeBits;
+	uint16 WordOffset : OffsetBits;
+};
+
+using ObjectAROFn = void (*)(UObject*, FReferenceCollector&);
+using StructAROFn = void (*)(void*, FReferenceCollector&);
+
+struct alignas(4) FStridedLayout
+{
+	uint16 WordOffset;
+	uint16 WordStride;
+};
+
+union FMemberWord
+{
+	FMemberPacked Members[4];
+	FSchemaView InnerSchema{NoInit};
+	ObjectAROFn ObjectARO;
+	StructAROFn StructARO;
+	FStridedLayout StridedLayout;
+};
+
+inline FMemberWord ToWord(FSchemaView In)
+{
+	FMemberWord Out;
+	Out.InnerSchema = In;
+	return Out;
+}
+
+inline FMemberWord ToWord(StructAROFn In)
+{
+	FMemberWord Out;
+	Out.StructARO = In;
+	return Out;
+}
+
+//////////////////// Schema declaration and building //////////////////////////
+
+struct FMemberDeclaration
+{
+    FName DebugName;
+	uint32 Offset = 0;
+    EMemberType Type = EMemberType::Stop;
+	FMemberWord ExtraWord = {{}};
+};
+
+inline FName ToName(const char* Name) { return FName(Name); }
+inline FName ToName(FPropertyStack& Stack) { return FName(Stack.GetPropertyPath()); }
+
+// Inlined so compiler can drop Name literals when compiling w/o UE_GC_DEBUGNAMES 
+template<typename NameType>
+FMemberDeclaration DeclareMember(NameType&& Name, uint32 Offset, EMemberType Type)
+{
+	FName DebugName = UE_GC_DEBUGNAMES ? ToName(Forward<NameType>(Name)) : ToName(Type);
+	return { DebugName, Offset, Type };
+}
+
+template<typename NameType, typename ExtraWordType>
+FMemberDeclaration DeclareMember(NameType&& Name, uint32 Offset, EMemberType Type, ExtraWordType Extra)
+{
+	FName DebugName = UE_GC_DEBUGNAMES ? ToName(Forward<NameType>(Name)) : ToName(Type);
+	return { DebugName, Offset, Type, ToWord(Extra) };
+}
+
+class FSchemaBuilder
+{
+public:
+	COREUOBJECT_API explicit FSchemaBuilder(uint32 InStride, std::initializer_list<FMemberDeclaration> InMembers = {});
+	COREUOBJECT_API ~FSchemaBuilder();
+	
+	COREUOBJECT_API void Add(FMemberDeclaration Member);
+
+	// Append all but the last Stop/ARO member
+	COREUOBJECT_API void Append(FSchemaView SuperSchema);
+
+	// Emit Stop or ARO and build a view that builder holds reference to. Returns same view if called multiple times.
+	COREUOBJECT_API FSchemaView Build(ObjectAROFn ARO = nullptr);
+
+	int32 NumMembers() const { return Members.Num(); }
+
+private:
+	TArray<FMemberDeclaration, TInlineAllocator<16>> Members;
+	const uint32 StructStride;
+	TOptional<FSchemaOwner> BuiltSchema;
+};
+
+
+//////////// Type-safe wrappers for declaring native schemas //////////////
+
+template<class T>
+struct TMemberDeclaration : FMemberDeclaration
+{
+	template<typename... Ts>
+	TMemberDeclaration(Ts&&... Args) : FMemberDeclaration(DeclareMember(Forward<Ts>(Args)...)) {}
+};
+
+// Type-safe wrapper for declaring native schemas
+template<class T>
+class TSchemaBuilder : public FSchemaBuilder
+{
+public:
+	TSchemaBuilder(std::initializer_list<TMemberDeclaration<T>> InMembers = {}) : FSchemaBuilder(sizeof(T), reinterpret_cast<std::initializer_list<FMemberDeclaration>&&>(InMembers)) {}
+	void Add(TMemberDeclaration<T> Member) { FSchemaBuilder::Add(Member); }
+};
+
+///////////////////////// UE_GC_MEMBER + helpers ///////////////////////
+namespace Private
+{
+
+template<class T, class R>
+TMemberDeclaration<T> MakeMember(const char* Name, uint32 Offset, R T::*)
+{
+	static_assert(std::is_convertible_v<R, UObject*>);
+    return TMemberDeclaration<T>(Name, Offset, EMemberType::Reference);
+}
+
+template<class T, class R>
+TMemberDeclaration<T> MakeMember(const char* Name, uint32 Offset, TArray<R> T::*)
+{
+	static_assert(std::is_convertible_v<R, UObject*>);
+    return TMemberDeclaration<T>(Name, Offset, EMemberType::ReferenceArray);
+}
+
+template<class T, class O>
+TMemberDeclaration<T> MakeMember(const char* Name, uint32 Offset, TObjectPtr<O> T::*)
+{
+    return TMemberDeclaration<T>(Name, Offset, EMemberType::Reference);
+}
+
+template<class T, class O>
+TMemberDeclaration<T> MakeMember(const char* Name, uint32 Offset, TArray<TObjectPtr<O>> T::*)
+{
+    return TMemberDeclaration<T>(Name, Offset, EMemberType::ReferenceArray);
+}
+
+template<class T, class S>
+TMemberDeclaration<T> MakeNestedMember(const char* Name, uint32 Offset, TArray<S> T::*, TSchemaBuilder<S>& InnerSchema)
+{
+    return TMemberDeclaration<T>(Name, Offset, EMemberType::StructArray, InnerSchema.Build());
+}
+
+} // namespace Private
+
+// Helpers to overload/dispatch UE_GC_MEMBER with 2 or 3 parameters
+#define _UE_EXPAND(x) x // MSVC macro evaluation workaround
+#define _UE_GC_MEMBER2(Type, Member)				UE::GC::Private::MakeMember(#Member, offsetof(Type, Member), &Type::Member)
+#define _UE_GC_MEMBER3(Type, Member, InnerSchema)	UE::GC::Private::MakeNestedMember(#Member, offsetof(Type, Member), &Type::Member, InnerSchema)
+#define _UE_GC_MEMBER(_1,_2,_3,CHOSEN_OVERLOAD,...) CHOSEN_OVERLOAD // _UE_GC_MEMBER2 or _UE_GC_MEMBER3
+
+#define UE_GC_MEMBER(...) _UE_EXPAND(_UE_EXPAND(_UE_GC_MEMBER(__VA_ARGS__, _UE_GC_MEMBER3, _UE_GC_MEMBER2)) (__VA_ARGS__))
+
+//////////////////////////////////////////////////////////////////////////
+
+COREUOBJECT_API void DeclareIntrinsicSchema(UClass* Class, FSchemaView IntrinsicSchema);
+
+template<class T>
+void DeclareIntrinsicMembers(UClass* Class, std::initializer_list<TMemberDeclaration<T>> Members)
+{
+	DeclareIntrinsicSchema(Class, TSchemaBuilder<T>(Members).Build());
+}
+
+FSchemaView GetIntrinsicSchema(UClass* Class);
+
+//////////////////////////////////////////////////////////////////////////
+} // namespace UE::GC
\ No newline at end of file
diff --git a/Engine/Source/Runtime/CoreUObject/Public/UObject/UObjectBase.h b/Engine/Source/Runtime/CoreUObject/Public/UObject/UObjectBase.h
--- a/Engine/Source/Runtime/CoreUObject/Public/UObject/UObjectBase.h
+++ b/Engine/Source/Runtime/CoreUObject/Public/UObject/UObjectBase.h
@@ -27,8 +27,6 @@ class UScriptStruct;
 
 DECLARE_DWORD_COUNTER_STAT_EXTERN(TEXT("STAT_UObjectsStatGroupTester"), STAT_UObjectsStatGroupTester, STATGROUP_UObjects, COREUOBJECT_API);
 
-namespace UE::GC { class FTokenStreamBuilder; }
-
 /** 
  * Low level implementation of UObject, should not be used directly in game code 
  */
diff --git a/Engine/Source/Runtime/CoreUObject/Public/UObject/UnrealType.h b/Engine/Source/Runtime/CoreUObject/Public/UObject/UnrealType.h
--- a/Engine/Source/Runtime/CoreUObject/Public/UObject/UnrealType.h
+++ b/Engine/Source/Runtime/CoreUObject/Public/UObject/UnrealType.h
@@ -74,6 +74,11 @@ enum ELifetimeCondition : int;
 struct CGetTypeHashable;
 struct FUObjectSerializeContext;
 template <typename FuncType> class TFunctionRef;
+namespace UE::GC
+{
+	class FPropertyStack;
+	class FSchemaBuilder;
+}
 
 COREUOBJECT_API DECLARE_LOG_CATEGORY_EXTERN(LogType, Log, All);
 
@@ -134,80 +139,6 @@ enum class EPropertyMemoryAccess : uint8
 
 namespace UEProperty_Private { class FProperty_DoNotUse; }
 
-/**
- * Helper class that calculates the maximum stack entry size required by Garbage Collector Token Stream for each class 
- */
-class COREUOBJECT_API FGCStackSizeHelper
-{
-	int32 MaxStackSize = 0;
-	/** Current property stack. Includes properties nested inside of struct member variables */
-	TArray<const FProperty*> PropertyStack;
-	
-public:
-
-	FGCStackSizeHelper() = default;
-	FGCStackSizeHelper(const FGCStackSizeHelper& Other) = delete;
-	FGCStackSizeHelper& operator=(const FGCStackSizeHelper& Other) = delete;
-
-	~FGCStackSizeHelper()
-	{
-		// Sanity check to make sure every Push() was followed by a Pop()
-		check(PropertyStack.Num() == 0);
-	}
-
-	void Push(const FProperty* InProperty)
-	{
-		PropertyStack.Push(InProperty);
-		MaxStackSize = FMath::Max(PropertyStack.Num(), MaxStackSize);
-	}
-
-	void Pop()
-	{
-		PropertyStack.Pop();
-	}
-
-	const TArray<const FProperty*>& GetPropertyStack() const
-	{
-		return PropertyStack;
-	}
-
-	/** Converts the interal property stack to a string representing the current property path (Member.StructMember.InnerStructMember) */
-	FString GetPropertyPath() const;
-
-	/** 
-	 * Converts a property path constructed with GetPropertyPath() to an array of properties (from the outermost to the innermost) *
-	 * @param ObjectClass Class that defines the outermost property
-	 * @param InPropertyPath Property path
-	 * @param OutProperties resulting array
-	 * @returns true if the conversion was successful, false otherwise
-	 */
-	static bool ConvertPathToProperties(UClass* ObjectClass, const FName& InPropertyPath, TArray<FProperty*>& OutProperties);
-
-	int32 GetMaxStackSize() const
-	{
-		return MaxStackSize;
-	}
-};
-
-/** Helper class to push and pop FGCStackSizeHelper stack inside of the current scope */
-class FGCStackSizeHelperScope
-{
-	FGCStackSizeHelper& StackSizeHelper;
-public:
-	explicit FGCStackSizeHelperScope(FGCStackSizeHelper& InHelper, FProperty* InProperty)
-		: StackSizeHelper(InHelper)
-	{		
-		StackSizeHelper.Push(InProperty);
-	}
-
-	FGCStackSizeHelperScope(const FGCStackSizeHelperScope& Other) = delete;
-	FGCStackSizeHelperScope& operator=(const FGCStackSizeHelperScope& Other) = delete;
-
-	~FGCStackSizeHelperScope()
-	{		
-		StackSizeHelper.Pop();
-	}
-};
 
 /** Type of pointer provided for property API functions */
 enum class EPropertyPointerType
@@ -1088,7 +1019,7 @@ public:
 	 * Emits tokens used by realtime garbage collection code to passed in ReferenceTokenStream. The offset emitted is relative
 	 * to the passed in BaseOffset which is used by e.g. arrays of structs.
 	 */
-	virtual void EmitReferenceInfo(UE::GC::FTokenStreamBuilder& TokenStream, int32 BaseOffset, TArray<const FStructProperty*>& EncounteredStructProps, FGCStackSizeHelper& StackSizeHelper);
+	virtual void EmitReferenceInfo(UE::GC::FSchemaBuilder& Schema, int32 BaseOffset, TArray<const FStructProperty*>& EncounteredStructProps, UE::GC::FPropertyStack& DebugPath);
 
     // @TODO: Surely this can have an int32 overflow. This should probably return size_t. Just
     // need to audit all callers to make such a change.
@@ -2928,7 +2859,7 @@ PRAGMA_ENABLE_DEPRECATION_WARNINGS
 
 	// FProperty interface
 	virtual void SerializeItem(FStructuredArchive::FSlot Slot, void* Value, void const* Defaults) const override;
-	virtual void EmitReferenceInfo(UE::GC::FTokenStreamBuilder& TokenStream, int32 BaseOffset, TArray<const FStructProperty*>& EncounteredStructProps, FGCStackSizeHelper& StackSizeHelper) override;
+	virtual void EmitReferenceInfo(UE::GC::FSchemaBuilder& Schema, int32 BaseOffset, TArray<const FStructProperty*>& EncounteredStructProps, UE::GC::FPropertyStack& DebugPath) override;
 	virtual const TCHAR* ImportText_Internal(const TCHAR* Buffer, void* ContainerOrPropertyPtr, EPropertyPointerType PropertyPointerType, UObject* OwnerObject, int32 PortFlags, FOutputDevice* ErrorText) const override;
 	virtual EConvertFromTypeResult ConvertFromType(const FPropertyTag& Tag, FStructuredArchive::FSlot Slot, uint8* Data, UStruct* DefaultsStruct) override;
 
@@ -3528,7 +3459,7 @@ public:
 
 	// UObject interface
 	virtual void Serialize( FArchive& Ar ) override;
-	virtual void EmitReferenceInfo(UE::GC::FTokenStreamBuilder& TokenStream, int32 BaseOffset, TArray<const FStructProperty*>& EncounteredStructProps, FGCStackSizeHelper& StackSizeHelper) override;
+	virtual void EmitReferenceInfo(UE::GC::FSchemaBuilder& Schema, int32 BaseOffset, TArray<const FStructProperty*>& EncounteredStructProps, UE::GC::FPropertyStack& DebugPath) override;
 	virtual void BeginDestroy() override;
 	virtual void AddReferencedObjects(FReferenceCollector& Collector) override;
 	// End of UObject interface
@@ -3785,7 +3716,7 @@ public:
 	virtual bool PassCPPArgsByRef() const override;
 	virtual void InstanceSubobjects( void* Data, void const* DefaultData, UObject* Owner, struct FObjectInstancingGraph* InstanceGraph ) override;
 	virtual bool ContainsObjectReference(TArray<const FStructProperty*>& EncounteredStructProps, EPropertyObjectReferenceType InReferenceType = EPropertyObjectReferenceType::Strong) const override;
-	virtual void EmitReferenceInfo(UE::GC::FTokenStreamBuilder& TokenStream, int32 BaseOffset, TArray<const FStructProperty*>& EncounteredStructProps, FGCStackSizeHelper& StackSizeHelper) override;
+	virtual void EmitReferenceInfo(UE::GC::FSchemaBuilder& Schema, int32 BaseOffset, TArray<const FStructProperty*>& EncounteredStructProps, UE::GC::FPropertyStack& DebugPath) override;
 	virtual bool SameType(const FProperty* Other) const override;
 	virtual EConvertFromTypeResult ConvertFromType(const FPropertyTag& Tag, FStructuredArchive::FSlot Slot, uint8* Data, UStruct* DefaultsStruct) override;
 
@@ -3921,7 +3852,7 @@ public:
 	virtual bool PassCPPArgsByRef() const override;
 	virtual void InstanceSubobjects(void* Data, void const* DefaultData, UObject* Owner, struct FObjectInstancingGraph* InstanceGraph) override;
 	virtual bool ContainsObjectReference(TArray<const FStructProperty*>& EncounteredStructProps, EPropertyObjectReferenceType InReferenceType = EPropertyObjectReferenceType::Strong) const override;
-	virtual void EmitReferenceInfo(UE::GC::FTokenStreamBuilder& TokenStream, int32 BaseOffset, TArray<const FStructProperty*>& EncounteredStructProps, FGCStackSizeHelper& StackSizeHelper) override;
+	virtual void EmitReferenceInfo(UE::GC::FSchemaBuilder& Schema, int32 BaseOffset, TArray<const FStructProperty*>& EncounteredStructProps, UE::GC::FPropertyStack& DebugPath) override;
 	virtual bool SameType(const FProperty* Other) const override;
 	virtual EConvertFromTypeResult ConvertFromType(const FPropertyTag& Tag, FStructuredArchive::FSlot Slot, uint8* Data, UStruct* DefaultsStruct) override;
 	virtual void* GetValueAddressAtIndex_Direct(const FProperty* Inner, void* InValueAddress, int32 Index) const override;
@@ -4045,7 +3976,7 @@ public:
 	virtual bool PassCPPArgsByRef() const override;
 	virtual void InstanceSubobjects(void* Data, void const* DefaultData, UObject* Owner, struct FObjectInstancingGraph* InstanceGraph) override;
 	virtual bool ContainsObjectReference(TArray<const FStructProperty*>& EncounteredStructProps, EPropertyObjectReferenceType InReferenceType = EPropertyObjectReferenceType::Strong) const override;
-	virtual void EmitReferenceInfo(UE::GC::FTokenStreamBuilder& TokenStream, int32 BaseOffset, TArray<const FStructProperty*>& EncounteredStructProps, FGCStackSizeHelper& StackSizeHelper) override;
+	virtual void EmitReferenceInfo(UE::GC::FSchemaBuilder& Schema, int32 BaseOffset, TArray<const FStructProperty*>& EncounteredStructProps, UE::GC::FPropertyStack& DebugPath) override;
 	virtual bool SameType(const FProperty* Other) const override;
 	virtual EConvertFromTypeResult ConvertFromType(const FPropertyTag& Tag, FStructuredArchive::FSlot Slot, uint8* Data, UStruct* DefaultsStruct) override;
 	virtual void* GetValueAddressAtIndex_Direct(const FProperty* Inner, void* InValueAddress, int32 Index) const override;
@@ -5907,7 +5838,7 @@ public:
 	virtual void InstanceSubobjects( void* Data, void const* DefaultData, UObject* Owner, struct FObjectInstancingGraph* InstanceGraph ) override;
 	virtual int32 GetMinAlignment() const override;
 	virtual bool ContainsObjectReference(TArray<const FStructProperty*>& EncounteredStructProps, EPropertyObjectReferenceType InReferenceType = EPropertyObjectReferenceType::Strong) const override;
-	virtual void EmitReferenceInfo(UE::GC::FTokenStreamBuilder& TokenStream, int32 BaseOffset, TArray<const FStructProperty*>& EncounteredStructProps, FGCStackSizeHelper& StackSizeHelper) override;
+	virtual void EmitReferenceInfo(UE::GC::FSchemaBuilder& Schema, int32 BaseOffset, TArray<const FStructProperty*>& EncounteredStructProps, UE::GC::FPropertyStack& DebugPath) override;
 	virtual bool SameType(const FProperty* Other) const override;
 	virtual EConvertFromTypeResult ConvertFromType(const FPropertyTag& Tag, FStructuredArchive::FSlot Slot, uint8* Data, UStruct* DefaultsStruct) override;
 #if WITH_EDITORONLY_DATA
diff --git a/Engine/Source/Runtime/Engine/Private/LevelActorContainer.cpp b/Engine/Source/Runtime/Engine/Private/LevelActorContainer.cpp
--- a/Engine/Source/Runtime/Engine/Private/LevelActorContainer.cpp
+++ b/Engine/Source/Runtime/Engine/Private/LevelActorContainer.cpp
@@ -63,7 +63,7 @@ public:
 	* @param ObjectsToSerialize An array of remaining objects to serialize (Obj must be added to it if Obj can be added to cluster)
 	* @param bOuterAndClass If true, the Obj's Outer and Class will also be added to the cluster
 	*/
-	void AddObjectToCluster(int32 ObjectIndex, FUObjectItem* ObjectItem, UObject* Obj, FGCArrayStruct& ObjectsToSerializeStruct, bool bOuterAndClass)
+	void AddObjectToCluster(int32 ObjectIndex, FUObjectItem* ObjectItem, UObject* Obj, FWorkerContext& Context, bool bOuterAndClass)
 	{
 		// If we haven't finished loading, we can't be sure we know all the references
 		check(!Obj->HasAnyFlags(RF_NeedLoad));
@@ -71,7 +71,7 @@ public:
 		check(Obj->CanBeInCluster());
 		if (ObjectIndex != ClusterRootIndex && ObjectItem->GetOwnerIndex() == 0 && !GUObjectArray.IsDisregardForGC(Obj) && !Obj->IsRooted())
 		{
-			ObjectsToSerializeStruct.ObjectsToSerialize.Add<Options>(Obj);
+			Context.ObjectsToSerialize.Add<Options>(Obj);
 			check(!ObjectItem->HasAnyFlags(EInternalObjectFlags::ClusterRoot));
 			ObjectItem->SetOwnerIndex(ClusterRootIndex);
 			Cluster.Objects.Add(ObjectIndex);
@@ -81,7 +81,7 @@ public:
 				UObject* ObjOuter = Obj->GetOuter();
 				if (CanAddToCluster(ObjOuter))
 				{
-					HandleTokenStreamObjectReference(ObjectsToSerializeStruct, Obj, ObjOuter, UE::GC::ETokenlessId::Outer, EGCTokenType::Native, true);
+					HandleTokenStreamObjectReference(Context, Obj, ObjOuter, UE::GC::EMemberlessId::Outer, EOrigin::Other, true);
 				}
 				else
 				{
@@ -90,9 +90,9 @@ public:
 				if (!Obj->GetClass()->HasAllClassFlags(CLASS_Native))
 				{
 					UObject* ObjectClass = Obj->GetClass();
-					HandleTokenStreamObjectReference(ObjectsToSerializeStruct, Obj, ObjectClass, UE::GC::ETokenlessId::Class, EGCTokenType::Native, true);
+					HandleTokenStreamObjectReference(Context, Obj, ObjectClass, UE::GC::EMemberlessId::Class, EOrigin::Other, true);
 					UObject* ObjectClassOuter = Obj->GetClass()->GetOuter();
-					HandleTokenStreamObjectReference(ObjectsToSerializeStruct, Obj, ObjectClassOuter, UE::GC::ETokenlessId::ClassOuter, EGCTokenType::Native, true);
+					HandleTokenStreamObjectReference(Context, Obj, ObjectClassOuter, UE::GC::EMemberlessId::ClassOuter, EOrigin::Other, true);
 				}
 			}
 		}
@@ -103,10 +103,10 @@ public:
 	*
 	* @param ObjectsToSerialize An array of remaining objects to serialize (Obj must be added to it if Obj can be added to cluster)
 	* @param ReferencingObject Object referencing the object to process.
-	* @param TokenIndex Index to the token stream where the reference was found.
+	* @param MemberId Index to the token stream where the reference was found.
 	* @param bAllowReferenceElimination True if reference elimination is allowed (ignored when constructing clusters).
 	*/
-	FORCEINLINE void HandleTokenStreamObjectReference(FGCArrayStruct& ObjectsToSerializeStruct, UObject* ReferencingObject, UObject*& Object, UE::GC::FTokenId TokenIndex, EGCTokenType TokenType, bool bAllowReferenceElimination)
+	FORCEINLINE void HandleTokenStreamObjectReference(FWorkerContext& Context, UObject* ReferencingObject, UObject*& Object, FMemberId MemberId, EOrigin Origin, bool bAllowReferenceElimination)
 	{
 		if (Object)
 		{
@@ -147,7 +147,7 @@ public:
 					// New object, add it to the cluster.
 					if (CanAddToCluster(Object) && !Object->HasAnyFlags(RF_NeedLoad | RF_NeedPostLoad) && !Object->IsRooted())
 					{
-						AddObjectToCluster(GUObjectArray.ObjectToIndex(Object), ObjectItem, Object, ObjectsToSerializeStruct, true);
+						AddObjectToCluster(GUObjectArray.ObjectToIndex(Object), ObjectItem, Object, Context, true);
 					}
 					else
 					{
diff --git a/Engine/Source/Runtime/Engine/Private/Model.cpp b/Engine/Source/Runtime/Engine/Private/Model.cpp
--- a/Engine/Source/Runtime/Engine/Private/Model.cpp
+++ b/Engine/Source/Runtime/Engine/Private/Model.cpp
@@ -9,6 +9,7 @@
 #include "EngineUtils.h"
 #include "Engine/Polys.h"
 #include "StaticLighting.h"
+#include "UObject/GarbageCollectionSchema.h"
 
 float UModel::BSPTexelScale = 100.0f;
 
@@ -537,36 +538,20 @@ void UModel::GetResourceSizeEx(FResourceSizeEx& CumulativeResourceSize)
 }
 
 #if WITH_EDITOR
-
-IMPLEMENT_INTRINSIC_CLASS(UModel, ENGINE_API, UObject, CORE_API, "/Script/Engine",
-	{
-		Class->CppClassStaticFunctions = UOBJECT_CPPCLASS_STATICFUNCTIONS_FORCLASS(UModel);
-
-		UE::GC::FTokenStreamBuilder& Builder = UE::GC::FIntrinsicClassTokens::AllocateBuilder(Class);
-		Builder.EmitObjectReference(STRUCT_OFFSET(UModel, Polys), TEXT("Polys"));
-		const uint32 SkipIndexIndex = Builder.EmitStructArrayBegin(STRUCT_OFFSET(UModel, Surfs), TEXT("Surfs"), sizeof(FBspSurf));
-		Builder.EmitObjectReference(STRUCT_OFFSET(FBspSurf, Material), TEXT("Material"));
-		Builder.EmitObjectReference(STRUCT_OFFSET(FBspSurf, Actor), TEXT("Actor"));
-		Builder.EmitStructArrayEnd( SkipIndexIndex );
-	}
-);
-
+#define UMODEL_EDITOR_GC_MEMBERS UE_GC_MEMBER(UModel, Polys), 
 #else
+#define UMODEL_EDITOR_GC_MEMBERS
+#endif
 
 IMPLEMENT_INTRINSIC_CLASS(UModel, ENGINE_API, UObject, CORE_API, "/Script/Engine",
 	{
 		Class->CppClassStaticFunctions = UOBJECT_CPPCLASS_STATICFUNCTIONS_FORCLASS(UModel);
 
-		UE::GC::FTokenStreamBuilder& Builder = UE::GC::FIntrinsicClassTokens::AllocateBuilder(Class);
-		const uint32 SkipIndexIndex = Builder.EmitStructArrayBegin(STRUCT_OFFSET(UModel, Surfs), TEXT("Surfs"), sizeof(FBspSurf));
-		Builder.EmitObjectReference(STRUCT_OFFSET(FBspSurf, Material), TEXT("Material"));
-		Builder.EmitObjectReference(STRUCT_OFFSET(FBspSurf, Actor), TEXT("Actor"));
-		Builder.EmitStructArrayEnd( SkipIndexIndex );
+		UE::GC::TSchemaBuilder<FBspSurf> SurfSchema({ UE_GC_MEMBER(FBspSurf, Material), UE_GC_MEMBER(FBspSurf, Actor) });
+		UE::GC::DeclareIntrinsicMembers(Class, { UMODEL_EDITOR_GC_MEMBERS UE_GC_MEMBER(UModel, Surfs, SurfSchema) });
 	}
 );
 
-#endif // WITH_EDITOR
-
 /*---------------------------------------------------------------------------------------
 	UModel implementation.
 ---------------------------------------------------------------------------------------*/
