Speculative fix for bogus MSVC static analyzer warning

#rnx

diff --git a/Engine/Source/Runtime/CoreUObject/Private/UObject/GarbageCollectionSchema.cpp b/Engine/Source/Runtime/CoreUObject/Private/UObject/GarbageCollectionSchema.cpp
--- a/Engine/Source/Runtime/CoreUObject/Private/UObject/GarbageCollectionSchema.cpp
+++ b/Engine/Source/Runtime/CoreUObject/Private/UObject/GarbageCollectionSchema.cpp
@@ -543,6 +543,7 @@ FString FPropertyStack::GetPropertyPath() const
 			}
 			Result += DelimiterChar;
 		}
+		CA_ASSUME(Property);
 		Result += Property->GetName();
 		PreviousProperty = Property;
 	}
