Speculative fix for non-unity name collisions

#rnx

diff --git a/Engine/Source/Runtime/CoreUObject/Private/UObject/GarbageCollection.cpp b/Engine/Source/Runtime/CoreUObject/Private/UObject/GarbageCollection.cpp
--- a/Engine/Source/Runtime/CoreUObject/Private/UObject/GarbageCollection.cpp
+++ b/Engine/Source/Runtime/CoreUObject/Private/UObject/GarbageCollection.cpp
@@ -5318,33 +5318,6 @@ void UClass::AssembleReferenceTokenStreamInternal(bool bForce)
 //////////////////////////////////////////////////////////////////////////
 namespace UE::GC
 {
-
-static FName ToName(EMemberlessId Id)
-{
-	static_assert(EMemberlessId::Collector == static_cast<EMemberlessId>(1u));
-	static_assert(EMemberlessId::Max == EMemberlessId::InitialReference);
-
-	static FName Names[] =
-	{
-		FName("Invalid"),
-		FName("Collector"), // EMemberlessId::Collector
-		FName("Class"), // EMemberlessId::Class
-		FName("Outer"), // EMemberlessId::Outer
-		FName("ExternalPackage"), // EMemberlessId::ExternalPackage
-		FName("ClassOuter"), // EMemberlessId::ClassOuter
-		FName("InitialReference"), // EMemberlessId::InitialReference
-	};
-
-	check(static_cast<uint32>(Id) < UE_ARRAY_COUNT(Names));
-	return Names[static_cast<uint32>(Id)];
-}
-
-void FMemberId::StaticAssert()
-{
-	static_assert((1u << MemberlessIdBits) >= static_cast<uint32>(EMemberlessId::Max), "Need to bump MemberlessIdBits");
-}
-
-//////////////////////////////////////////////////////////////////////////
 #if !UE_BUILD_SHIPPING
 
 static void DumpMemoryStats(FOutputDevice& OutputDevice)
diff --git a/Engine/Source/Runtime/CoreUObject/Private/UObject/GarbageCollectionSchema.cpp b/Engine/Source/Runtime/CoreUObject/Private/UObject/GarbageCollectionSchema.cpp
--- a/Engine/Source/Runtime/CoreUObject/Private/UObject/GarbageCollectionSchema.cpp
+++ b/Engine/Source/Runtime/CoreUObject/Private/UObject/GarbageCollectionSchema.cpp
@@ -216,6 +216,11 @@ static FName ToName(EMemberlessId Id)
 	return Names[static_cast<uint32>(Id)];
 }
 
+void FMemberId::StaticAssert()
+{
+	static_assert((1u << MemberlessIdBits) >= static_cast<uint32>(EMemberlessId::Max), "Need to bump MemberlessIdBits");
+}
+
 FMemberInfo GetMemberDebugInfo(FSchemaView Schema, FMemberId Id)
 {
 	if (Id.IsMemberless())
@@ -807,18 +812,18 @@ static void DumpDistribution(FOutputDevice& Out, const TCHAR* Title, const FLog2
 		}
 	}
 }
-static void DumpSchemaStats(FOutputDevice& Out, const FSchemaStats& Stats)
+static void DumpSchemaStats(FOutputDevice& Out, const FSchemaStats& InStats)
 {
 	static constexpr TCHAR FormatStr[] = TEXT(" %18s %5.1f%% %5.1f%%  %d");
 
 	uint32 Sum = 0;
 	for (uint32 Idx = 0; Idx < MemberTypeCount; ++Idx)
 	{
-		if (uint32 Num = Stats.Counters[Idx])
+		if (uint32 Num = InStats.Counters[Idx])
 		{
 			Sum += Num;
-			float Percentage = (100.f * Num) / Stats.NumMembers;
-			float Cumulative = (100.f * Sum) / Stats.NumMembers;
+			float Percentage = (100.f * Num) / InStats.NumMembers;
+			float Cumulative = (100.f * Sum) / InStats.NumMembers;
 			Out.Logf(FormatStr, *ToName(EMemberType(Idx)).ToString(), Percentage, Cumulative, Num);
 		}
 	}
