Better fix for non-unity name collisions

#rnx

diff --git a/Engine/Source/Runtime/CoreUObject/Private/UObject/GarbageCollection.cpp b/Engine/Source/Runtime/CoreUObject/Private/UObject/GarbageCollection.cpp
--- a/Engine/Source/Runtime/CoreUObject/Private/UObject/GarbageCollection.cpp
+++ b/Engine/Source/Runtime/CoreUObject/Private/UObject/GarbageCollection.cpp
@@ -340,7 +340,7 @@ namespace UE::GC
 		void IncPurgeCount(UObject* Object);
 		void LogPurgeStats(int32 Total);
 	};
-	static FStats Stats;
+	static FStats GStats;
 
 	/**
 	 * Helper function to log the various class to count info maps.
@@ -925,7 +925,7 @@ static bool MarkClusterMutableObjectsAsReachable(FUObjectCluster& Cluster, Conta
 		if (ReferencedMutableObjectIndex >= 0) // Pending kill support
 		{
 			FUObjectItem* ReferencedMutableObjectItem = GUObjectArray.IndexToObjectUnsafeForGC(ReferencedMutableObjectIndex);
-			UE::GC::Stats.IncClusterToObjectRefs(ReferencedMutableObjectItem);
+			UE::GC::GStats.IncClusterToObjectRefs(ReferencedMutableObjectItem);
 			if constexpr (IsParallel(Options))
 			{
 				PRAGMA_DISABLE_DEPRECATION_WARNINGS
@@ -1018,7 +1018,7 @@ static FORCENOINLINE void MarkReferencedClustersAsReachable(int32 ClusterIndex,
 {
 	static_assert(IsWithClusters(Options));
 
-	UE::GC::Stats.IncNumClustersTraversed();
+	UE::GC::GStats.IncNumClustersTraversed();
 
 	// If we run across some PendingKill objects we need to add all objects from this cluster
 	// to ObjectsToSerialize so that we can properly null out all the references.
@@ -1026,7 +1026,7 @@ static FORCENOINLINE void MarkReferencedClustersAsReachable(int32 ClusterIndex,
 
 	bool bAddClusterObjectsToSerialize = false;
 	FUObjectCluster& Cluster = GUObjectClusters[ClusterIndex];
-	UE::GC::Stats.IncClusterToClusterRefs(Cluster.ReferencedClusters.Num());
+	UE::GC::GStats.IncClusterToClusterRefs(Cluster.ReferencedClusters.Num());
 	// Also mark all referenced objects from outside of the cluster as reachable
 	for (int32& ReferncedClusterIndex : Cluster.ReferencedClusters)
 	{
@@ -2755,17 +2755,17 @@ class TReachabilityProcessor
 public:
 	FORCEINLINE void BeginTimingObject(UObject* CurrentObject)
 	{
-		UE::GC::Stats.BeginTimingObject(CurrentObject);
+		UE::GC::GStats.BeginTimingObject(CurrentObject);
 	}
 
 	FORCEINLINE void UpdateDetailedStats(UObject* CurrentObject)
 	{
-		UE::GC::Stats.UpdateDetailedStats(CurrentObject);
+		UE::GC::GStats.UpdateDetailedStats(CurrentObject);
 	}
 
 	FORCEINLINE void LogDetailedStatsSummary()
 	{
-		UE::GC::Stats.LogDetailedStatsSummary();
+		UE::GC::GStats.LogDetailedStatsSummary();
 	}
 
 	static constexpr EGCOptions Options = InOptions;
@@ -2799,7 +2799,7 @@ public:
 	template<EKillable Killable>
 	static FORCEINLINE_DEBUGGABLE void ProcessReferenceDirectly(FWorkerContext& Context, FPermanentObjectPoolExtents PermanentPool, const UObject* ReferencingObject, UObject*& Object, FMemberId MemberId)
 	{
-		UE::GC::Stats.IncreaseObjectRefStats(Object);
+		UE::GC::GStats.IncreaseObjectRefStats(Object);
 		if (ValidateReference(Object, PermanentPool, ReferencingObject, MemberId))
 		{
 			const int32 ObjectIndex = GUObjectArray.ObjectToIndex(Object);
@@ -2821,7 +2821,7 @@ public:
 
 	FORCEINLINE static void HandleBatchedReference(FWorkerContext& Context, FResolvedMutableReference Reference, FReferenceMetadata Metadata)
 	{
-		UE::GC::Stats.IncreaseObjectRefStats(GetObject(Reference));
+		UE::GC::GStats.IncreaseObjectRefStats(GetObject(Reference));
 		if (Metadata.Has(KillFlag))
 		{
 			checkSlow(Metadata.ObjectItem->GetOwnerIndex() <= 0);
@@ -2835,7 +2835,7 @@ public:
 
 	FORCEINLINE static void HandleBatchedReference(FWorkerContext& Context, FImmutableReference Reference, FReferenceMetadata Metadata)
 	{
-		UE::GC::Stats.IncreaseObjectRefStats(GetObject(Reference));
+		UE::GC::GStats.IncreaseObjectRefStats(GetObject(Reference));
 		DetectGarbageReference(Context, Metadata);
 		HandleValidReference(Context, Reference, Metadata);
 	}
@@ -3131,17 +3131,17 @@ class TDebugReachabilityProcessor
 public:
 	FORCEINLINE void BeginTimingObject(UObject* CurrentObject)
 	{
-		UE::GC::Stats.BeginTimingObject(CurrentObject);
+		UE::GC::GStats.BeginTimingObject(CurrentObject);
 	}
 
 	FORCEINLINE void UpdateDetailedStats(UObject* CurrentObject)
 	{
-		UE::GC::Stats.UpdateDetailedStats(CurrentObject);
+		UE::GC::GStats.UpdateDetailedStats(CurrentObject);
 	}
 
 	FORCEINLINE void LogDetailedStatsSummary()
 	{
-		UE::GC::Stats.LogDetailedStatsSummary();
+		UE::GC::GStats.LogDetailedStatsSummary();
 	}
 
 	static constexpr EGCOptions Options = InOptions;
@@ -3156,7 +3156,7 @@ public:
 
 	FORCENOINLINE void HandleTokenStreamObjectReference(FWorkerContext& Context, const UObject* ReferencingObject, UObject*& Object, FMemberId MemberId, EOrigin Origin, bool bAllowReferenceElimination)
 	{
-		UE::GC::Stats.IncreaseObjectRefStats(Object);
+		UE::GC::GStats.IncreaseObjectRefStats(Object);
 		if (ValidateReference(Object, PermanentPool, ReferencingObject, MemberId))
 		{
 			FReferenceMetadata Metadata(GUObjectArray.ObjectToIndex(Object));
@@ -4077,7 +4077,7 @@ bool IncrementalDestroyGarbage(bool bUseTimeLimit, double TimeLimit)
 					// Only proceed with destroying the object if the asynchronous cleanup started by BeginDestroy has finished.
 					if(Object->IsReadyForFinishDestroy())
 					{
-						UE::GC::Stats.IncPurgeCount(Object);
+						UE::GC::GStats.IncPurgeCount(Object);
 						// Send FinishDestroy message.
 						Object->ConditionalFinishDestroy();
 					}
@@ -4137,7 +4137,7 @@ bool IncrementalDestroyGarbage(bool bUseTimeLimit, double TimeLimit)
 					// Only proceed with destroying the object if the asynchronous cleanup started by BeginDestroy has finished.
 					if( Object->IsReadyForFinishDestroy() )
 					{
-						UE::GC::Stats.IncPurgeCount(Object);
+						UE::GC::GStats.IncPurgeCount(Object);
 						// Send FinishDestroy message.
 						Object->ConditionalFinishDestroy();
 
@@ -4302,7 +4302,7 @@ bool IncrementalDestroyGarbage(bool bUseTimeLimit, double TimeLimit)
 				GObjectCountDuringLastMarkPhase.GetValue(), 
 				GObjectCountDuringLastMarkPhase.GetValue() - PurgedObjectCountSinceLastMarkPhase,
 				(FPlatformTime::Seconds() - IncrementalDestroyGarbageStartTime) * 1000);
-			UE::GC::Stats.LogPurgeStats(PurgedObjectCountSinceLastMarkPhase);
+			UE::GC::GStats.LogPurgeStats(PurgedObjectCountSinceLastMarkPhase);
 			GAsyncPurge->ResetObjectsDestroyedSinceLastMarkPhase();
 		}
 	}
diff --git a/Engine/Source/Runtime/CoreUObject/Private/UObject/GarbageCollectionSchema.cpp b/Engine/Source/Runtime/CoreUObject/Private/UObject/GarbageCollectionSchema.cpp
--- a/Engine/Source/Runtime/CoreUObject/Private/UObject/GarbageCollectionSchema.cpp
+++ b/Engine/Source/Runtime/CoreUObject/Private/UObject/GarbageCollectionSchema.cpp
@@ -812,18 +812,18 @@ static void DumpDistribution(FOutputDevice& Out, const TCHAR* Title, const FLog2
 		}
 	}
 }
-static void DumpSchemaStats(FOutputDevice& Out, const FSchemaStats& InStats)
+static void DumpSchemaStats(FOutputDevice& Out, const FSchemaStats& Stats)
 {
 	static constexpr TCHAR FormatStr[] = TEXT(" %18s %5.1f%% %5.1f%%  %d");
 
 	uint32 Sum = 0;
 	for (uint32 Idx = 0; Idx < MemberTypeCount; ++Idx)
 	{
-		if (uint32 Num = InStats.Counters[Idx])
+		if (uint32 Num = Stats.Counters[Idx])
 		{
 			Sum += Num;
-			float Percentage = (100.f * Num) / InStats.NumMembers;
-			float Cumulative = (100.f * Sum) / InStats.NumMembers;
+			float Percentage = (100.f * Num) / Stats.NumMembers;
+			float Cumulative = (100.f * Sum) / Stats.NumMembers;
 			Out.Logf(FormatStr, *ToName(EMemberType(Idx)).ToString(), Percentage, Cumulative, Num);
 		}
 	}
