Make FEditorBulkData thread safe.

*** FLinkerLoad
- Attaching/detaching both FBulkData and FEditorBulkdata to a FLinkerLoad should now be thread safe.
- Note that this does NOT fix the long standing issue when a FPackage is GCed before some of the UObjects (assets) inside of that package, which causes any remaining bulkdata to be detached (and unable to load off disk) without forcing payloads into memory. This will be addressed as its own work item.
-- FEditorBulkData can still attempt to load its payload from a package after detachment as a fall back measure that is preferable to a crash.

*** FEditorBulkData
- Forced to use FRecursiveMutex at the moment some methods that need locking make calls into the bulkdata registry which can make calls back to the calling bulkdata object ending up with recursive locking, at least for now.
- This might look like a lot of locks being introduced, but FEditorBulkData methods are called relatively infrequently., in a test project with a semi large amount of data (~16GB) cooking from a cold DDC benchmarked at roughly the same wall time before/after these changes.
- Added a comment block to the header describing the thread safety expectations of the class.
- Made the class final as we really don't want to deal with inherited bulkdata types as we did with the old system.
- Since FLinkerLoad attachment/detachment is now thread safe we can detach from any pre-existing AttachedAr when running the assignment operator and attach to the AttachedAr of the provided bulkdata. This means copied/assigned bulkdata do not need to force their payloads into memory.
- Cleaned up some comments in the header.

diff --git a/Engine/Source/Runtime/CoreUObject/Private/Serialization/EditorBulkData.cpp b/Engine/Source/Runtime/CoreUObject/Private/Serialization/EditorBulkData.cpp
--- a/Engine/Source/Runtime/CoreUObject/Private/Serialization/EditorBulkData.cpp
+++ b/Engine/Source/Runtime/CoreUObject/Private/Serialization/EditorBulkData.cpp
@@ -2,13 +2,15 @@
 
 #include "Serialization/EditorBulkData.h"
 
+#include "Async/UniqueLock.h"
 #include "Compression/OodleDataCompression.h"
+#include "Experimental/Async/MultiUniqueLock.h"
 #include "HAL/FileManager.h"
 #include "HAL/IConsoleManager.h"
 #include "Misc/ConfigCacheIni.h"
 #include "Misc/PackageSegment.h"
-#include "Misc/SecureHash.h"
 #include "Misc/ScopeLock.h"
+#include "Misc/SecureHash.h"
 #include "ProfilingDebugging/CountersTrace.h"
 #include "Serialization/BulkData.h"
 #include "Serialization/BulkDataRegistry.h"
@@ -18,8 +20,8 @@
 #include "UObject/ObjectSaveContext.h"
 #include "UObject/ObjectVersion.h"
 #include "UObject/PropertyPortFlags.h"
-#include "UObject/UObjectGlobals.h"
 #include "UObject/SavePackage.h"
+#include "UObject/UObjectGlobals.h"
 #include "Virtualization/VirtualizationSystem.h"
 #include "Virtualization/VirtualizationTypes.h"
 
@@ -454,6 +456,17 @@ ECompressedBufferCompressionLevel FCompressionSettings::GetCompressionLevel()
 
 } // namespace Private
 
+FEditorBulkData::~FEditorBulkData()
+{
+	if (HasAttachedArchive())
+	{
+		AttachedAr->DetachBulkData(this, false);
+		AttachedAr = nullptr;
+	}
+
+	OnExitMemory();
+}
+
 FEditorBulkData::FEditorBulkData(FEditorBulkData&& Other)
 {
 	*this = MoveTemp(Other);
@@ -461,19 +474,32 @@ FEditorBulkData::FEditorBulkData(FEditorBulkData&& Other)
 
 FEditorBulkData& FEditorBulkData::operator=(FEditorBulkData&& Other)
 {
+	UE::TMultiUniqueLock<FRecursiveMutex> _({&Mutex, &Other.Mutex });
+
 	// The same as the default move constructor, except we need to handle registration and unregistration
 	Unregister();
 	Other.Unregister();
 
+	if (HasAttachedArchive())
+	{
+		AttachedAr->DetachBulkData(this, false);
+	}
+
 	BulkDataId = MoveTemp(Other.BulkDataId);
 	PayloadContentId = MoveTemp(Other.PayloadContentId);
 	Payload = MoveTemp(Other.Payload);
 	PayloadSize = MoveTemp(Other.PayloadSize);
+	AttachedAr = Other.AttachedAr;
 	OffsetInFile = MoveTemp(Other.OffsetInFile);
 	PackagePath = MoveTemp(Other.PackagePath);
 	Flags = MoveTemp(Other.Flags);
 	CompressionSettings = MoveTemp(Other.CompressionSettings);
 
+	if (HasAttachedArchive())
+	{
+		AttachedAr->AttachBulkData(this);
+	}
+
 	Other.Reset();
 
 	Register(nullptr, TEXT("operator=(FEditorBulkData &&)"), false /* bAllowUpdateId */);
@@ -488,6 +514,8 @@ FEditorBulkData::FEditorBulkData(const FEditorBulkData& Other)
 
 FEditorBulkData& FEditorBulkData::operator=(const FEditorBulkData& Other)
 {
+	UE::TMultiUniqueLock<FRecursiveMutex> _({ &Mutex, &Other.Mutex });
+
 	// Torn-off BulkDatas remain torn-off even when being copied into from a non-torn-off BulkData
 	// Remaining torn-off is a work-around necessary for FTextureSource::CopyTornOff to avoid registering a new
 	// guid before setting the new BulkData to torn-off. The caller can call Reset to clear the torn-off flag.
@@ -512,6 +540,11 @@ FEditorBulkData& FEditorBulkData::operator=(const FEditorBulkData& Other)
 		}
 	}
 
+	if (HasAttachedArchive())
+	{
+		AttachedAr->DetachBulkData(this, false);
+	}
+
 	PayloadContentId = Other.PayloadContentId;
 	Payload = Other.Payload;
 	PayloadSize = Other.PayloadSize;
@@ -520,6 +553,11 @@ FEditorBulkData& FEditorBulkData::operator=(const FEditorBulkData& Other)
 	Flags = Other.Flags;
 	CompressionSettings = Other.CompressionSettings;
 
+	if (HasAttachedArchive())
+	{
+		AttachedAr->AttachBulkData(this);
+	}
+
 	EnumRemoveFlags(Flags, TransientFlags);
 
 	if (bTornOff)
@@ -533,31 +571,26 @@ FEditorBulkData& FEditorBulkData::operator=(const FEditorBulkData& Other)
 	return *this;
 }
 
-FEditorBulkData::~FEditorBulkData()
-{
-	if (HasAttachedArchive())
-	{
-		AttachedAr->DetachBulkData(this, false);
-		AttachedAr = nullptr;
-	}
-
-	OnExitMemory();
-}
-
 FEditorBulkData::FEditorBulkData(const FEditorBulkData& Other, ETornOff)
 {
+	UE::TMultiUniqueLock<FRecursiveMutex> _({ &Mutex, &Other.Mutex });
+
 	EnumAddFlags(Flags, EFlags::IsTornOff);
 	*this = Other; // We rely on operator= preserving the torn-off flag
 }
 
 void FEditorBulkData::TearOff()
 {
+	UE::TUniqueLock _(Mutex);
+
 	Unregister();
 	EnumAddFlags(Flags, EFlags::IsTornOff);
 }
 
 void FEditorBulkData::UpdateRegistrationOwner(UObject* Owner)
 {
+	UE::TUniqueLock _(Mutex);
+
 	UpdateRegistrationData(Owner, TEXT("UpdateRegistrationOwner"), false /* bAllowUpdateId */);
 }
 
@@ -761,6 +794,8 @@ void FEditorBulkData::CreateFromBulkData(FBulkData& InBulkData, const FGuid& InG
 {
 	TRACE_CPUPROFILER_EVENT_SCOPE(FEditorBulkData::CreateFromBulkData);
 
+	UE::TUniqueLock _(Mutex);
+
 	Reset();
 
 #if UE_ALLOW_LINKERLOADER_ATTACHMENT
@@ -811,6 +846,8 @@ void FEditorBulkData::CreateFromBulkData(FBulkData& InBulkData, const FGuid& InG
 
 void FEditorBulkData::CreateLegacyUniqueIdentifier(UObject* Owner)
 {
+	UE::TUniqueLock _(Mutex);
+
 	if (BulkDataId.IsValid())
 	{
 		Unregister();
@@ -825,6 +862,8 @@ void FEditorBulkData::CreateLegacyUniqueIdentifier(UObject* Owner)
 void FEditorBulkData::Serialize(FArchive& Ar, UObject* Owner, bool bAllowRegister)
 {
 	TRACE_CPUPROFILER_EVENT_SCOPE(FEditorBulkData::Serialize);
+	
+	UE::TUniqueLock _(Mutex);
 
 	if (Ar.IsTransacting())
 	{
@@ -1177,6 +1216,8 @@ void FEditorBulkData::Serialize(FArchive& Ar, UObject* Owner, bool bAllowRegiste
 
 void FEditorBulkData::SerializeForRegistry(FArchive& Ar)
 {
+	UE::TUniqueLock _(Mutex);
+
 	if (Ar.IsSaving())
 	{
 		check(CanSaveForRegistry());
@@ -1217,6 +1258,8 @@ void FEditorBulkData::SerializeForRegistry(FArchive& Ar)
 
 bool FEditorBulkData::CanSaveForRegistry() const
 {
+	UE::TUniqueLock _(Mutex);
+
 	return BulkDataId.IsValid() && PayloadSize > 0 && !IsMemoryOnlyPayload()
 		&& EnumHasAnyFlags(Flags, EFlags::IsTornOff) && !EnumHasAnyFlags(Flags, EFlags::HasRegistered);
 }
@@ -1595,11 +1638,14 @@ bool FEditorBulkData::CanLoadDataFromDisk() const
 
 bool FEditorBulkData::IsMemoryOnlyPayload() const
 {
+	UE::TUniqueLock _(Mutex);
 	return !Payload.IsNull() && !IsDataVirtualized() && PackagePath.IsEmpty();
 }
 
 void FEditorBulkData::Reset()
 {
+	UE::TUniqueLock _(Mutex);
+
 	// Unregister rather than allowing the Registry to keep our record, since we are changing the payload
 	Unregister();
 	// Note that we do not reset the BulkDataId
@@ -1621,6 +1667,8 @@ void FEditorBulkData::Reset()
 
 void FEditorBulkData::UnloadData()
 {
+	UE::TUniqueLock _(Mutex);
+
 	if (CanUnloadData())
 	{
 		Payload.Reset();
@@ -1631,9 +1679,16 @@ void FEditorBulkData::DetachFromDisk(FArchive* Ar, bool bEnsurePayloadIsLoaded)
 {
 	TRACE_CPUPROFILER_EVENT_SCOPE(FEditorBulkData::DetachFromDisk);
 
+	UE::TUniqueLock _(Mutex);
+
 	check(Ar != nullptr);
 	check(Ar == AttachedAr || HasAttachedArchive() == false || AttachedAr->IsProxyOf(Ar));
 
+	if (bEnsurePayloadIsLoaded)
+	{
+		int a = 0; a++;
+	}
+
 	// If bEnsurePayloadIsLoaded is true, then we should assume that a change to the 
 	// package file is imminent and we should load the payload into memory if possible.
 	// If it is false then either we are not expecting a change to the package file or
@@ -1675,10 +1730,19 @@ void FEditorBulkData::DetachFromDisk(FArchive* Ar, bool bEnsurePayloadIsLoaded)
 
 FGuid FEditorBulkData::GetIdentifier() const
 {
+	UE::TUniqueLock _(Mutex);
+
 	checkf(GetPayloadSize() == 0 || BulkDataId.IsValid(), TEXT("If bulkdata has a valid payload then it should have a valid BulkDataId"));
 	return BulkDataId;
 }
 
+const FIoHash& FEditorBulkData::GetPayloadId() const
+{
+	UE::TUniqueLock _(Mutex);
+
+	return PayloadContentId;
+}
+
 void FEditorBulkData::SerializeToLegacyPath(FLinkerSave& LinkerSave, FCompressedBuffer PayloadToSerialize, EFlags UpdatedFlags, UObject* Owner)
 {
 	const int64 OffsetPos = LinkerSave.Tell();
@@ -1831,6 +1895,8 @@ void FEditorBulkData::UpdatePayloadImpl(FSharedBuffer&& InPayload, const FIoHash
 {
 	TRACE_CPUPROFILER_EVENT_SCOPE(FEditorBulkData::UpdatePayloadImpl);
 
+	UE::TUniqueLock _(Mutex);
+
 	// Unregister before calling DetachBulkData; DetachFromDisk calls OnExitMemory which is incorrect since
 	// we are changing data rather than leaving memory
 	Unregister();
@@ -1914,8 +1980,16 @@ FCompressedBuffer FEditorBulkData::GetDataInternal() const
 	}
 }
 
+bool FEditorBulkData::DoesPayloadNeedLoading() const
+{
+	UE::TUniqueLock _(Mutex);
+	return Payload.IsNull() && PayloadSize > 0;
+}
+
 TFuture<FSharedBuffer> FEditorBulkData::GetPayload() const
 {
+	UE::TUniqueLock _(Mutex);
+
 	TPromise<FSharedBuffer> Promise;
 	
 	if (GetPayloadSize() == 0)
@@ -1941,6 +2015,8 @@ TFuture<FSharedBuffer> FEditorBulkData::GetPayload() const
 
 TFuture<FCompressedBuffer>FEditorBulkData::GetCompressedPayload() const
 {
+	UE::TUniqueLock _(Mutex);
+
 	TPromise<FCompressedBuffer> Promise;
 
 	FCompressedBuffer CompressedPayload = GetDataInternal();
@@ -1982,6 +2058,8 @@ void FEditorBulkData::UpdatePayload(FSharedBufferWithID InPayload, UObject* Owne
 
 void FEditorBulkData::SetCompressionOptions(ECompressionOptions Option)
 {
+	UE::TUniqueLock _(Mutex);
+
 	switch (Option)
 	{
 	case ECompressionOptions::Disabled:
@@ -2006,6 +2084,8 @@ void FEditorBulkData::SetCompressionOptions(ECompressionOptions Option)
 
 void FEditorBulkData::SetCompressionOptions(ECompressedBufferCompressor Compressor, ECompressedBufferCompressionLevel CompressionLevel)
 {
+	UE::TUniqueLock _(Mutex);
+
 	CompressionSettings.Set(Compressor, CompressionLevel);
 
 	if (CompressionSettings.GetCompressionLevel() == ECompressedBufferCompressionLevel::None)
@@ -2030,24 +2110,31 @@ FCustomVersionContainer FEditorBulkData::GetCustomVersions(FArchive& InlineArchi
 void FEditorBulkData::GetBulkDataVersions(FArchive& InlineArchive, FPackageFileVersion& OutUEVersion,
 	int32& OutLicenseeUEVersion, FCustomVersionContainer& OutCustomVersions) const
 {
-	if (EnumHasAnyFlags(Flags, EFlags::ReferencesWorkspaceDomain))
+	TUniquePtr<FArchive> ExternalArchive;
+
 	{
-		// Read the version data out of the separate package file
-		TUniquePtr<FArchive> ExternalArchive = IPackageResourceManager::Get().OpenReadExternalResource(
-			EPackageExternalResource::WorkspaceDomainFile, PackagePath.GetPackageName());
-		if (ExternalArchive.IsValid())
+		UE::TUniqueLock _(Mutex);
+		if (EnumHasAnyFlags(Flags, EFlags::ReferencesWorkspaceDomain))
 		{
-			FPackageFileSummary PackageFileSummary;
-			*ExternalArchive << PackageFileSummary;
-			if (PackageFileSummary.Tag == PACKAGE_FILE_TAG && !ExternalArchive->IsError())
-			{
-				OutUEVersion = PackageFileSummary.GetFileVersionUE();
-				OutLicenseeUEVersion = PackageFileSummary.GetFileVersionLicenseeUE();
-				OutCustomVersions = PackageFileSummary.GetCustomVersionContainer();
-				return;
-			}
+			// Read the version data out of the separate package file
+			ExternalArchive = IPackageResourceManager::Get().OpenReadExternalResource(EPackageExternalResource::WorkspaceDomainFile, PackagePath.GetPackageName());
+		}
+	}
+
+	if (ExternalArchive.IsValid())
+	{
+		FPackageFileSummary PackageFileSummary;
+		*ExternalArchive << PackageFileSummary;
+		if (PackageFileSummary.Tag == PACKAGE_FILE_TAG && !ExternalArchive->IsError())
+		{
+			OutUEVersion = PackageFileSummary.GetFileVersionUE();
+			OutLicenseeUEVersion = PackageFileSummary.GetFileVersionLicenseeUE();
+			OutCustomVersions = PackageFileSummary.GetCustomVersionContainer();
+
+			return;
 		}
 	}
+	
 	OutUEVersion = InlineArchive.UEVer();
 	OutLicenseeUEVersion = InlineArchive.LicenseeUEVer();
 	OutCustomVersions = InlineArchive.GetCustomVersions();
@@ -2055,11 +2142,15 @@ void FEditorBulkData::GetBulkDataVersions(FArchive& InlineArchive, FPackageFileV
 
 void FEditorBulkData::UpdatePayloadId()
 {
+	UE::TUniqueLock _(Mutex);
+
 	UpdateKeyIfNeeded();
 }
 
 bool FEditorBulkData::LocationMatches(const FEditorBulkData& Other) const
 {
+	UE::TMultiUniqueLock<FRecursiveMutex> _({ &Mutex, &Other.Mutex });
+
 	if (GetIdentifier() != Other.GetIdentifier())
 	{
 		// Different identifiers return false, even if location is the same
diff --git a/Engine/Source/Runtime/CoreUObject/Private/UObject/LinkerLoad.cpp b/Engine/Source/Runtime/CoreUObject/Private/UObject/LinkerLoad.cpp
--- a/Engine/Source/Runtime/CoreUObject/Private/UObject/LinkerLoad.cpp
+++ b/Engine/Source/Runtime/CoreUObject/Private/UObject/LinkerLoad.cpp
@@ -62,6 +62,7 @@
 #include "Misc/EngineBuildSettings.h"
 #include "Internationalization/GatherableTextData.h"
 #include "Async/MappedFileHandle.h"
+#include "Async/UniqueLock.h"
 class FTexture2DResourceMem;
 
 #define LOCTEXT_NAMESPACE "LinkerLoad"
@@ -5790,14 +5791,11 @@ void FLinkerLoad::Detach()
 }
 
 #if WITH_EDITOR
-/**
- * Attaches/ associates the passed in bulk data object with the linker.
- *
- * @param	Owner		UObject owning the bulk data
- * @param	BulkData	Bulk data object to associate
- */
+
 void FLinkerLoad::AttachBulkData( UObject* Owner, FBulkData* BulkData )
 {
+	UE::TUniqueLock _(BulkDataMutex);
+
 	bool bAlreadyInSet = false;
 	BulkDataLoaders.Add(BulkData, &bAlreadyInSet);
 	check(!bAlreadyInSet);
@@ -5805,63 +5803,56 @@ void FLinkerLoad::AttachBulkData( UObject* Owner, FBulkData* BulkData )
 
 void FLinkerLoad::AttachBulkData(UE::Serialization::FEditorBulkData* BulkData)
 {
+	UE::TUniqueLock _(BulkDataMutex);
+
 	bool bAlreadyInSet = false;
 	EditorBulkDataLoaders.Add(BulkData, &bAlreadyInSet);
 	check(!bAlreadyInSet);
 }
 
-/**
- * Detaches the passed in bulk data object from the linker.
- *
- * @param	BulkData	Bulk data object to detach
- * @param	bEnsureBulkDataIsLoaded	Whether to ensure that the bulk data is loaded before detaching
- */
 void FLinkerLoad::DetachBulkData( FBulkData* BulkData, bool bEnsureBulkDataIsLoaded )
 {
-	int32 RemovedCount = BulkDataLoaders.Remove( BulkData );
-	if( RemovedCount!=1 )
+	UE::TUniqueLock _(BulkDataMutex);
+
+	const int32 RemovedCount = BulkDataLoaders.Remove( BulkData );
+	if (RemovedCount!= 1)
 	{	
 		UE_ASSET_LOG(LogLinker, Fatal, PackagePath, TEXT("Detachment inconsistency: %i"), RemovedCount);
 	}
+
 	BulkData->DetachFromArchive( this, bEnsureBulkDataIsLoaded );
 }
 
 void FLinkerLoad::DetachBulkData(UE::Serialization::FEditorBulkData* BulkData, bool bEnsureBulkDataIsLoaded)
 {
-	int32 RemovedCount = EditorBulkDataLoaders.Remove(BulkData);
+	UE::TUniqueLock _(BulkDataMutex);
+
+	const int32 RemovedCount = EditorBulkDataLoaders.Remove(BulkData);
 	if (RemovedCount != 1)
 	{
 		UE_ASSET_LOG(LogLinker, Fatal, PackagePath, TEXT("Detachment inconsistency: %i"), RemovedCount);
 	}
+
 	BulkData->DetachFromDisk(this, bEnsureBulkDataIsLoaded);
 }
 
-/**
- * Detaches all attached bulk  data objects.
- *
- * @param	bEnsureBulkDataIsLoaded	Whether to ensure that the bulk data is loaded before detaching
- */
 void FLinkerLoad::DetachAllBulkData(bool bEnsureAllBulkDataIsLoaded)
 {
-	// Old style bulkdata first
+	UE::TUniqueLock _(BulkDataMutex);
+
+	for (FBulkData* BulkData : BulkDataLoaders)
 	{
-		for (FBulkData* BulkData : BulkDataLoaders)
-		{
-			check(BulkData != nullptr);
-			BulkData->DetachFromArchive(this, bEnsureAllBulkDataIsLoaded);
-		}
-		BulkDataLoaders.Empty();
+		BulkData->DetachFromArchive(this, bEnsureAllBulkDataIsLoaded);
 	}
 
-	// Then virtualized bulkdata
+	BulkDataLoaders.Empty();
+
+	for (UE::Serialization::FEditorBulkData* BulkData : EditorBulkDataLoaders)
 	{
-		for (UE::Serialization::FEditorBulkData* BulkData : EditorBulkDataLoaders)
-		{
-			check(BulkData != nullptr);
-			BulkData->DetachFromDisk(this, bEnsureAllBulkDataIsLoaded);
-		}
-		EditorBulkDataLoaders.Empty();
+		BulkData->DetachFromDisk(this, bEnsureAllBulkDataIsLoaded);
 	}
+
+	EditorBulkDataLoaders.Empty();
 }
 
 #endif // WITH_EDITOR
diff --git a/Engine/Source/Runtime/CoreUObject/Public/Serialization/EditorBulkData.h b/Engine/Source/Runtime/CoreUObject/Public/Serialization/EditorBulkData.h
--- a/Engine/Source/Runtime/CoreUObject/Public/Serialization/EditorBulkData.h
+++ b/Engine/Source/Runtime/CoreUObject/Public/Serialization/EditorBulkData.h
@@ -2,6 +2,8 @@
 
 #pragma once
 
+#include "Async/Mutex.h"
+#include "Async/RecursiveMutex.h"
 #include "Async/Future.h"
 #include "Compression/CompressedBuffer.h"
 #include "Containers/Array.h"
@@ -110,21 +112,34 @@ enum class ECompressionOptions : uint8
  * 
  * For examples see the "System.CoreUObject.Serialization.EditorBulkData.UpdatePayload" test found in
  * Engine\Source\Runtime\CoreUObject\Private\Tests\Serialization\EditorBulkDataTests.cpp.
+ * 
+ * Thread Safety:
+ * 
+ * The class should be thread safe in that you can call the public methods on multiple threads at once and expect
+ * things to work. However if you were to call ::GetPayloadSize() before calling ::GetPayload() on one thread
+ * it would be possible for a call to ::UpdatePayload() to occur on another thread in between those two calls, 
+ * meaning that the size you got might not be the correct size as the payload returned. If you need that level
+ * of thread safety it would be up to you to manage it at a higher level. 
+ * It would however be safe to pass a copy of your FEditorBulkData to a background task rather than a reference.
+ * The copy would either be able to load the payload from disk, pull it from virtualized storage or if the 
+ * source bulkdata is holding the payload in memory it will just take a reference to that payload rather than 
+ * making a copy. In this way you can be sure of the state of the FEditorBulkData on the background task as 
+ * future edits to the original bulkdata will not affect it.
  */
 
 /** The base class with no type */
-class COREUOBJECT_API FEditorBulkData
+class COREUOBJECT_API FEditorBulkData final
 {
 public:
 	FEditorBulkData() = default;
+	~FEditorBulkData();
+
 	FEditorBulkData(FEditorBulkData&& Other);
 	FEditorBulkData& operator=(FEditorBulkData&& Other);
 
 	FEditorBulkData(const FEditorBulkData& Other);
 	FEditorBulkData& operator=(const FEditorBulkData& Other);
 
-	~FEditorBulkData();
-
 	/** 
 	 * Convenience method to make it easier to convert from FBulkData to FEditorBulkData and sets the Guid 
 	 *
@@ -133,7 +148,11 @@ public:
 	 *					This MUST remain the same between sessions so that the payloads key remains consistent!
 	 */
 	void CreateFromBulkData(FBulkData& BulkData, const FGuid& Guid, UObject* Owner);
-	/** Fix legacy content that created the Id from non-unique Guids. */
+	
+	/** 
+	 * Fix legacy content that created the Id from non-unique Guids.
+	 * Run on objects created before FUE5MainStreamObjectVersion::VirtualizedBulkDataHaveUniqueGuids
+	 */
 	void CreateLegacyUniqueIdentifier(UObject* Owner);
 
 	/** 
@@ -151,8 +170,9 @@ public:
 	/** Reset to a truly empty state */
 	void Reset();
 
-	// TODO: Probably can just remove this as there probably isn't a good use case for unloading updated payloads as there is no
-	// way for us to restore it. In that case ::Reset might as well be used.
+	// TODO: Not sure if there is a scenario where we can reload the payload off disk but it is also
+	// in memory.and so this might not really do anything anymore.
+
 	/** Unloads the data (if possible) but leaves it in a state where the data can be reloaded */
 	void UnloadData();
 
@@ -170,28 +190,22 @@ public:
 	FGuid GetIdentifier() const;
 
 	/** Returns an unique identifier for the content of the payload. */
-	const FIoHash& GetPayloadId() const
-	{ 
-		return PayloadContentId; 
-	}
+	const FIoHash& GetPayloadId() const;
 
 	/** Returns the size of the payload in bytes. */
-	int64 GetPayloadSize() const 
+	int64 GetPayloadSize() const
 	{ 
 		return PayloadSize; 
 	}
 
 	/** Returns true if the bulkdata object contains a valid payload greater than zero bytes in size. */
-	bool HasPayloadData() const 
+	bool HasPayloadData() const
 	{ 
 		return PayloadSize > 0; 
 	}
 
 	/** Returns if the payload would require loading in order to be accessed. Returns false if the payload is already in memory or of zero length */
-	bool DoesPayloadNeedLoading() const
-	{
-		return Payload.IsNull() && PayloadSize > 0; 
-	}
+	bool DoesPayloadNeedLoading() const;
 
 	/** Returns an immutable FCompressedBuffer reference to the payload data. */
 	TFuture<FSharedBuffer> GetPayload() const;
@@ -319,14 +333,22 @@ public:
 
 	/** Used to serialize the bulkdata to/from a limited cache system used by the BulkDataRegistry. */
 	void SerializeForRegistry(FArchive& Ar);
+	
 	/** Return true if the bulkdata has a source location that persists between editor processes (package file or virtualization). */
 	bool CanSaveForRegistry() const;
+	
 	/** Return whether the BulkData has legacy payload id that needs to be updated from loaded payload before it can be used in DDC. */
-	bool HasPlaceholderPayloadId() const { return EnumHasAnyFlags(Flags, EFlags::LegacyKeyWasGuidDerived); }
+	bool HasPlaceholderPayloadId() const
+	{
+		return EnumHasAnyFlags(Flags, EFlags::LegacyKeyWasGuidDerived);
+	}
+	
 	/** Return whether the BulkData is an in-memory payload without a persistent source location. */
 	bool IsMemoryOnlyPayload() const;
+	
 	/** Load the payload and set the correct payload id, if the bulkdata has a PlaceholderPayloadId. */
 	void UpdatePayloadId();
+	
 	/** Return whether *this has the same source for the bulkdata (e.g. identical file locations if from file) as Other */
 	bool LocationMatches(const FEditorBulkData& Other) const;
 
@@ -556,6 +578,9 @@ private:
 	/** A 32bit bitfield of flags */
 	EFlags Flags = EFlags::None;
 
+	// Use of IBulkDataRegistry requires a recursive mutex
+	mutable FRecursiveMutex Mutex;
+
 #if UE_ENABLE_VIRTUALIZATION_TOGGLE
 	bool bSkipVirtualization = false;
 #endif //UE_ENABLE_VIRTUALIZATION_TOGGLE
diff --git a/Engine/Source/Runtime/CoreUObject/Public/UObject/LinkerLoad.h b/Engine/Source/Runtime/CoreUObject/Public/UObject/LinkerLoad.h
--- a/Engine/Source/Runtime/CoreUObject/Public/UObject/LinkerLoad.h
+++ b/Engine/Source/Runtime/CoreUObject/Public/UObject/LinkerLoad.h
@@ -2,6 +2,7 @@
 
 #pragma once
 
+#include "Async/Mutex.h"
 #include "Containers/Array.h"
 #include "Containers/Map.h"
 #include "Containers/Set.h"
@@ -288,8 +289,12 @@ public:
 
 	/** The async package associated with this linker */
 	void* AsyncRoot;
+
 #if WITH_EDITOR
-	/** Bulk data that does not need to be loaded when the linker is loaded.												*/
+	/** Used when accessing BulkDataLoaders/EditorBulkDataLoaders to ensure thread safety */
+	UE::FMutex BulkDataMutex;
+
+	/** Bulk data that use the FLinkerLoad to track the state of the file on disk */
 	TSet<FBulkData*> BulkDataLoaders;
 	TSet<UE::Serialization::FEditorBulkData*> EditorBulkDataLoaders;
 #endif // WITH_EDITOR
