Add some basic threading tests to FEditorBulkData

#rnx

diff --git a/Engine/Source/Runtime/CoreUObject/Private/Tests/Serialization/EditorBulkDataTests.cpp b/Engine/Source/Runtime/CoreUObject/Private/Tests/Serialization/EditorBulkDataTests.cpp
--- a/Engine/Source/Runtime/CoreUObject/Private/Tests/Serialization/EditorBulkDataTests.cpp
+++ b/Engine/Source/Runtime/CoreUObject/Private/Tests/Serialization/EditorBulkDataTests.cpp
@@ -1,6 +1,7 @@
 // Copyright Epic Games, Inc. All Rights Reserved.
 
 #include "Algo/AllOf.h"
+#include "Async/ParallelFor.h"
 #include "Compression/OodleDataCompression.h"
 #include "Memory/SharedBuffer.h"
 #include "Serialization/BulkDataRegistry.h"
@@ -11,6 +12,7 @@
 #include "Serialization/LargeMemoryWriter.h"
 #include "Serialization/MemoryReader.h"
 #include "Serialization/MemoryWriter.h"
+#include "Tasks/Task.h"
 #include "Templates/UniquePtr.h"
 
 #include "Misc/AutomationTest.h"
@@ -758,6 +760,110 @@ bool FEditorBulkDataSharedBufferWithID::RunTest(const FString& Parameters)
 	return true;
 }
 
+/** Test a number of threads all updating a FEditorBulkData object at the same time */
+IMPLEMENT_SIMPLE_AUTOMATION_TEST(FEditorBulkDataThreadingBasic, TEXT("System.CoreUObject.Serialization.EditorBulkData.Threading.Basic"), TestFlags)
+bool FEditorBulkDataThreadingBasic::RunTest(const FString& Parameters)
+{
+	// Before thread safety was added the following number of tests/payload tended to result in broken data.
+	// Although trying to induce threading issues is not an exact science.
+	const int32 NumTests = 128;
+	const int32 NumPayloadsToTest = 16;
+
+	for (int32 TestIndex = 0; TestIndex < NumTests; ++TestIndex)
+	{
+		// Create a number of randomly sized payloads
+		TArray<FSharedBuffer> Payloads;
+		Payloads.Reserve(NumPayloadsToTest);
+
+		for (int32 Index = 0; Index < NumPayloadsToTest; ++Index)
+		{
+			const int64 BufferSize = FMath::RandRange(512, 12 * 1024);
+			Payloads.Add(CreateRandomPayload(BufferSize));
+		}
+
+		FEditorBulkData BulkData;
+
+		ParallelFor(NumPayloadsToTest, [&Payloads, &BulkData](int32 Index)
+			{
+				BulkData.UpdatePayload(Payloads[Index]);
+			});
+
+		const FIoHash FinalId = BulkData.GetPayloadId();
+		const uint64 FinalSize = BulkData.GetPayloadSize();
+
+		const FSharedBuffer FinalPayload = BulkData.GetPayload().Get();
+
+		// Make sure that the size of the payload matches the value stored in the object
+		if (!TestEqual(TEXT("Testing the payload size vs the FEditorBulkData size"), FinalPayload.GetSize(), FinalSize))
+		{
+			return false;
+		}
+
+		// Make sure that the hash of the payload matches the value stored in the object
+		if (!TestEqual(TEXT("Testing payload hash vs the FEditorBulkData hash"), FIoHash::HashBuffer(FinalPayload), FinalId))
+		{
+			return false;
+		}
+	}
+
+	return true;
+}
+
+IMPLEMENT_SIMPLE_AUTOMATION_TEST(FEditorBulkDataThreadingAssignment, TEXT("System.CoreUObject.Serialization.EditorBulkData.Threading.Assignment"), TestFlags)
+bool FEditorBulkDataThreadingAssignment::RunTest(const FString& Parameters)
+{
+	const int32 NumThreads = 8;
+	const int32 NumBulkDataToTest = 128;
+	const int32 NumAssignments = 16 * 1024;
+
+	// Create a FEditorBulkData of randomly sized payloads
+	TArray<FEditorBulkData> BulkDatas;
+	BulkDatas.SetNum(NumBulkDataToTest);
+
+	for (int32 Index = 0; Index < NumBulkDataToTest; ++Index)
+	{
+		const int32 PayloadSize = FMath::RandRange(512, 1024);
+		BulkDatas[Index].UpdatePayload(CreateRandomPayload(PayloadSize));
+	}
+
+	TArray<UE::Tasks::FTask> CompletionEvents;
+	CompletionEvents.Reserve(NumThreads);
+
+	for (int Index = 0; Index < NumThreads; ++Index)
+	{
+		UE::Tasks::FTask Task = UE::Tasks::Launch(UE_SOURCE_LOCATION, [&BulkDatas, NumAssignments]()
+			{
+				for (int32 Index = 0; Index < NumAssignments; ++Index)
+				{
+					const int32 DstIndex = FMath::RandRange(0, BulkDatas.Num() - 1);
+					const int32 SrcIndex = FMath::RandRange(0, BulkDatas.Num() - 1);
+
+					BulkDatas[DstIndex] = BulkDatas[SrcIndex];
+				}
+			});
+
+		CompletionEvents.Emplace(MoveTemp(Task));
+	}
+
+	UE::Tasks::Wait(CompletionEvents);
+
+	for (const FEditorBulkData& BulkData : BulkDatas)
+	{
+		FSharedBuffer Payload = BulkData.GetPayload().Get();
+
+		if (!TestEqual(TEXT("PayloadSize"), BulkData.GetPayloadSize(), (int64)Payload.GetSize()))
+		{
+			return false;
+		}
+
+		if (!TestEqual(TEXT("PayloadId"), BulkData.GetPayloadId(), FIoHash::HashBuffer(Payload)))
+		{
+			return false;
+		}
+	}
+
+	return true;
+}
 
 } // namespace UE::Serialization
 
