[autortfm]
- fix Verse for loop semantics
- added some VM opcodes:
  - EX_NothingInt32: takes an int argument and is very handy for being a no-op that annotates the bytecode disassembly for debugging and breakpointing
  - EX_AutoRtfmTransact: takes a transaction id and list of expressions to run within an AutoRTFM transaction
  - EX_AutoRtfmStopTransact: takes a transaction id and ends that transaction if we are inside of it, otherwise is a no-op
  - EX_AutoRtfmAbortIfNot: causes a transaction abort if the boolean expression is false

diff --git a/Engine/Source/Developer/ScriptDisassembler/Private/ScriptDisassembler.cpp b/Engine/Source/Developer/ScriptDisassembler/Private/ScriptDisassembler.cpp
--- a/Engine/Source/Developer/ScriptDisassembler/Private/ScriptDisassembler.cpp
+++ b/Engine/Source/Developer/ScriptDisassembler/Private/ScriptDisassembler.cpp
@@ -632,6 +632,12 @@ void FKismetBytecodeDisassembler::ProcessCommon(int32& ScriptIndex, EExprToken O
 			Ar.Logf(TEXT("%s $%X: EX_Nothing"), *Indents, (int32)Opcode);
 			break;
 		}
+	case EX_NothingInt32:
+		{
+			int32 Value = ReadINT(ScriptIndex);
+			Ar.Logf(TEXT("%s $%X: EX_NothingInt32 %d"), *Indents, (int32)Opcode, Value);
+			break;
+		}
 	case EX_EndOfScript:
 		{
 			Ar.Logf(TEXT("%s $%X: EX_EndOfScript"), *Indents, (int32)Opcode);
@@ -1190,6 +1196,42 @@ void FKismetBytecodeDisassembler::ProcessCommon(int32& ScriptIndex, EExprToken O
 			DropIndent();
 			break;
 		}
+	case EX_AutoRtfmTransact:
+		{
+			// Code offset.
+			int32 Value = ReadINT(ScriptIndex);				
+			CodeSkipSizeType SkipCount = ReadSkipCount(ScriptIndex);
+
+			Ar.Logf(TEXT("%s $%X: AutoRtfmTransact %d to offset 0x%X"), *Indents, (int32)Opcode, Value, SkipCount);
+
+			while (SerializeExpr(ScriptIndex) != EX_AutoRtfmStopTransact)
+			{
+				// Params
+			}
+			break;
+		}
+	case EX_AutoRtfmStopTransact:
+		{
+			int32 Value = ReadINT(ScriptIndex);
+			EAutoRtfmStopTransactMode Mode = EAutoRtfmStopTransactMode(ReadBYTE(ScriptIndex));
+
+			const TCHAR* ModeText = TEXT("");
+			switch(Mode)
+			{
+			case EAutoRtfmStopTransactMode::GracefulExit: ModeText = TEXT("GracefulExit"); break;
+			case EAutoRtfmStopTransactMode::AbortingExit: ModeText = TEXT("AbortingExit"); break;
+			case EAutoRtfmStopTransactMode::AbortingExitAndAbortParent: ModeText = TEXT("AbortingExitAndAbortParent"); break;
+			}
+
+			Ar.Logf(TEXT("%s $%X: EX_AutoRtfmStopTransact (%s) %d"), *Indents, (int32)Opcode, ModeText, Value);
+			break;
+		}
+	case EX_AutoRtfmAbortIfNot:
+		{
+			Ar.Logf(TEXT("%s $%X: EX_AutoRtfmAbortIfNot"), *Indents, (int32)Opcode);
+			SerializeExpr(ScriptIndex);
+			break;
+		}
 	default:
 		{
 			// This should never occur.
diff --git a/Engine/Source/Runtime/Core/Private/AutoRTFM/Transaction.cpp b/Engine/Source/Runtime/Core/Private/AutoRTFM/Transaction.cpp
--- a/Engine/Source/Runtime/Core/Private/AutoRTFM/Transaction.cpp
+++ b/Engine/Source/Runtime/Core/Private/AutoRTFM/Transaction.cpp
@@ -72,7 +72,7 @@ void FTransaction::Undo()
 	UE_LOG(LogAutoRTFM, Verbose, TEXT("Undoing a transaction..."));
 
 	int VerboseCounter = 0;
-
+	int Num = WriteLog.Num();
 	for(auto Iter = WriteLog.rbegin(); Iter != WriteLog.rend(); ++Iter)
     {
 		FWriteLogEntry& Entry = *Iter;
@@ -92,7 +92,7 @@ void FTransaction::Undo()
 		{
 			TStringBuilder<1024> Builder;
 
-			Builder.Appendf(TEXT("%4d [UNDO] %p %4llu : [ "), VerboseCounter, Original, Size);
+			Builder.Appendf(TEXT("%4d [UNDO] %p %4llu : [ "), Num - VerboseCounter - 1, Original, Size);
 
 			unsigned char* Current = (unsigned char*)Original;
 			unsigned char* Old = (unsigned char*)Copy;
diff --git a/Engine/Source/Runtime/Core/Private/AutoRTFM/WriteLog.h b/Engine/Source/Runtime/Core/Private/AutoRTFM/WriteLog.h
--- a/Engine/Source/Runtime/Core/Private/AutoRTFM/WriteLog.h
+++ b/Engine/Source/Runtime/Core/Private/AutoRTFM/WriteLog.h
@@ -165,6 +165,8 @@ namespace AutoRTFM
             return 0 == TotalSize;
         }
 
+		size_t Num() const { return TotalSize; }
+
     private:
 
         FWriteLogEntryBucket* Start = nullptr;
diff --git a/Engine/Source/Runtime/CoreUObject/Private/UObject/ScriptCore.cpp b/Engine/Source/Runtime/CoreUObject/Private/UObject/ScriptCore.cpp
--- a/Engine/Source/Runtime/CoreUObject/Private/UObject/ScriptCore.cpp
+++ b/Engine/Source/Runtime/CoreUObject/Private/UObject/ScriptCore.cpp
@@ -2419,6 +2419,12 @@ DEFINE_FUNCTION(UObject::execNothing)
 }
 IMPLEMENT_VM_FUNCTION( EX_Nothing, execNothing );
 
+DEFINE_FUNCTION(UObject::execNothingInt32)
+{
+	int32 Value = Stack.ReadInt<int32>();
+}
+IMPLEMENT_VM_FUNCTION( EX_NothingInt32, execNothingInt32 );
+
 DEFINE_FUNCTION(UObject::execNothingOp4a)
 {
 	// Do nothing.
@@ -4032,4 +4038,97 @@ DEFINE_FUNCTION(UObject::execInterfaceToObject)
 }
 IMPLEMENT_VM_FUNCTION( EX_InterfaceToObjCast, execInterfaceToObject );
 
+DEFINE_FUNCTION(UObject::execAutoRtfmTransact)
+{
+	int32 TransactionId = Stack.ReadInt<int32>();
+	CodeSkipSizeType JumpOffset = Stack.ReadCodeSkipCount();
+
+	uint8* JumpTarget = &Stack.Node->Script[JumpOffset];
+
+	// sometimes if this inner transaction commits, we want to
+	// abort the parent transaction afterwards (logical not does this)
+	bool bAbortParentOnCommit = false;
+
+	// now wrap the next step around a transaction
+	AutoRTFM::ETransactionResult Result = AutoRTFM::Transact([&]()
+	{
+		bool bKeepRunning = true;
+		for (int i = 0; bKeepRunning; i++)
+		{
+			if(*Stack.Code == EX_AutoRtfmStopTransact)
+			{
+				Stack.Code++;
+				int32 Value = Stack.ReadInt<int32>();
+				EAutoRtfmStopTransactMode Mode = Stack.Read<EAutoRtfmStopTransactMode>();
+
+				if (TransactionId == Value)
+				{
+					switch(Mode)
+					{
+					case EAutoRtfmStopTransactMode::GracefulExit:
+						// gracefully terminate this transaction
+						bKeepRunning = false;
+						break;
+
+					case EAutoRtfmStopTransactMode::AbortingExit:
+						// abort this transaction
+						AutoRTFM::AbortTransaction();
+						break;
+
+					case EAutoRtfmStopTransactMode::AbortingExitAndAbortParent:
+						// abort this transaction and also abort the parent
+						UE_AUTORTFM_OPEN({ bAbortParentOnCommit = true; });
+						AutoRTFM::AbortTransaction();
+						break;
+					}
+				}
+			}
+			else
+			{
+				Stack.Step(Stack.Object, RESULT_PARAM);
+			}
+		}
+	});
+
+	P_NATIVE_BEGIN;
+
+	if (Result != AutoRTFM::ETransactionResult::Committed)
+	{
+		// if this transaction didn't commit, move our code pointer to the target
+		Stack.Code = JumpTarget;
+	}
+
+	if (bAbortParentOnCommit)
+	{
+		AutoRTFM::AbortTransaction();
+	}
+
+	P_NATIVE_END;
+}
+IMPLEMENT_VM_FUNCTION( EX_AutoRtfmTransact, execAutoRtfmTransact );
+
+DEFINE_FUNCTION(UObject::execAutoRtfmStopTransact)
+{
+	Stack.ReadInt<int32>();
+	Stack.Read<EAutoRtfmStopTransactMode>();
+
+	// if we were in a transaction, the processing loop inside execAutoRtfmTransact
+	// would handle this opcode specially. if we are handling it here then we
+	// are not in a transaction, so this opcode is a no-op
+}
+IMPLEMENT_VM_FUNCTION( EX_AutoRtfmStopTransact, execAutoRtfmStopTransact );
+
+DEFINE_FUNCTION(UObject::execAutoRtfmAbortIfNot)
+{
+	// evaluate a boolean expression
+	bool Result = false;
+	Stack.Step(Stack.Object, &Result);
+	// if the expression returned false, abort the current transaction
+	if (!Result)
+	{
+		AutoRTFM::AbortTransaction();
+	}
+}
+IMPLEMENT_VM_FUNCTION( EX_AutoRtfmAbortIfNot, execAutoRtfmAbortIfNot );
+
 #undef LOCTEXT_NAMESPACE
diff --git a/Engine/Source/Runtime/CoreUObject/Public/UObject/Object.h b/Engine/Source/Runtime/CoreUObject/Public/UObject/Object.h
--- a/Engine/Source/Runtime/CoreUObject/Public/UObject/Object.h
+++ b/Engine/Source/Runtime/CoreUObject/Public/UObject/Object.h
@@ -1517,6 +1517,7 @@ public:
 
 	// Do Nothing 
 	DECLARE_FUNCTION(execNothing);
+	DECLARE_FUNCTION(execNothingInt32);
 	DECLARE_FUNCTION(execNothingOp4a);
 
 	/** Breakpoint; only observed in the editor; executing it at any other time is a NOP */
@@ -1679,6 +1680,10 @@ public:
 
 	DECLARE_FUNCTION(execArrayGetByRef);
 
+	DECLARE_FUNCTION(execAutoRtfmTransact);
+	DECLARE_FUNCTION(execAutoRtfmStopTransact);
+	DECLARE_FUNCTION(execAutoRtfmAbortIfNot);
+
 	/** Wrapper struct to hold the entrypoint in the right memory address */
 	struct Object_eventExecuteUbergraph_Parms
 	{
diff --git a/Engine/Source/Runtime/CoreUObject/Public/UObject/Script.h b/Engine/Source/Runtime/CoreUObject/Public/UObject/Script.h
--- a/Engine/Source/Runtime/CoreUObject/Public/UObject/Script.h
+++ b/Engine/Source/Runtime/CoreUObject/Public/UObject/Script.h
@@ -196,7 +196,7 @@ enum EExprToken : uint8
 	EX_Assert				= 0x09,	// Assertion.
 	//						= 0x0A,
 	EX_Nothing				= 0x0B,	// No operation.
-	//						= 0x0C,
+	EX_NothingInt32			= 0x0C, // No operation with an int32 argument (useful for debugging script disassembly)
 	//						= 0x0D,
 	//						= 0x0E,
 	EX_Let					= 0x0F,	// Assign an arbitrary size value to a variable.
@@ -294,9 +294,21 @@ enum EExprToken : uint8
 	EX_ArrayGetByRef		= 0x6B,
 	EX_ClassSparseDataVariable = 0x6C, // Sparse data variable
 	EX_FieldPathConst		= 0x6D,
+	//						= 0x6E,
+	//						= 0x6F,
+	EX_AutoRtfmTransact     = 0x70, // AutoRTFM: run following code in a transaction
+	EX_AutoRtfmStopTransact = 0x71, // AutoRTFM: if in a transaction, abort or break, otherwise no operation
+	EX_AutoRtfmAbortIfNot   = 0x72, // AutoRTFM: evaluate bool condition, abort transaction on false
 	EX_Max					= 0xFF,
 };
 
+enum EAutoRtfmStopTransactMode : uint8
+{
+	GracefulExit,
+	AbortingExit,
+	AbortingExitAndAbortParent,
+};
+
 enum ECastToken : uint8
 {
 	CST_ObjectToInterface		= 0x00,
diff --git a/Engine/Source/Runtime/CoreUObject/Public/UObject/ScriptSerialization.h b/Engine/Source/Runtime/CoreUObject/Public/UObject/ScriptSerialization.h
--- a/Engine/Source/Runtime/CoreUObject/Public/UObject/ScriptSerialization.h
+++ b/Engine/Source/Runtime/CoreUObject/Public/UObject/ScriptSerialization.h
@@ -249,10 +249,15 @@
 			break;
 		}
 		case EX_PushExecutionFlow:
-			{
-				XFER(CodeSkipSizeType);		// location to push
-				break;
-			}
+		{
+			XFER(CodeSkipSizeType);		// location to push
+			break;
+		}
+		case EX_NothingInt32:
+		{
+			XFER(int32);
+			break;
+		}
 		case EX_Nothing:
 		case EX_EndOfScript:
 		case EX_EndFunctionParms:
@@ -610,6 +615,24 @@
 			SerializeExpr( iCode, Ar );
 			break;
 		}
+		case EX_AutoRtfmTransact:
+		{
+			XFER(int32); // Transaction id
+			XFER(CodeSkipSizeType); // Code offset.
+			while( SerializeExpr( iCode, Ar ) != EX_AutoRtfmStopTransact ); // Parms.
+			break;
+		}
+		case EX_AutoRtfmStopTransact:
+		{
+			XFER(int32); // transaction id
+			XFER(int8); // stop mode
+			break;
+		}
+		case EX_AutoRtfmAbortIfNot:
+		{
+			SerializeExpr(iCode,Ar);
+			break;
+		}
 		default:
 		{
 			// This should never occur.
