Fixed multicast delegate indexing in FMulticastDelegateProperty.

diff --git a/Engine/Source/Runtime/CoreUObject/Private/UObject/PropertyMulticastDelegate.cpp b/Engine/Source/Runtime/CoreUObject/Private/UObject/PropertyMulticastDelegate.cpp
--- a/Engine/Source/Runtime/CoreUObject/Private/UObject/PropertyMulticastDelegate.cpp
+++ b/Engine/Source/Runtime/CoreUObject/Private/UObject/PropertyMulticastDelegate.cpp
@@ -43,8 +43,8 @@ void FMulticastDelegateProperty::InstanceSubobjects(void* Data, void const* Defa
 		for( int32 i=0; i<ArrayDim; i++ )
 		{
 			// Fix up references to the class default object (if necessary)
-			FMulticastScriptDelegate::InvocationListType::TIterator CurInvocation(GetInvocationList((uint8*)Data + i));
-			FMulticastScriptDelegate::InvocationListType::TIterator DefaultInvocation(GetInvocationList((uint8*)DefaultData + i));
+			FMulticastScriptDelegate::InvocationListType::TIterator CurInvocation(GetInvocationList((uint8*)Data, i));
+			FMulticastScriptDelegate::InvocationListType::TIterator DefaultInvocation(GetInvocationList((uint8*)DefaultData, i));
 			for(; CurInvocation && DefaultInvocation; ++CurInvocation, ++DefaultInvocation )
 			{
 				FScriptDelegate& DestDelegateInvocation = *CurInvocation;
@@ -76,7 +76,7 @@ void FMulticastDelegateProperty::InstanceSubobjects(void* Data, void const* Defa
 	{
 		for( int32 i=0; i<ArrayDim; i++ )
 		{
-			for( FMulticastScriptDelegate::InvocationListType::TIterator CurInvocation(GetInvocationList((uint8*)Data + i)); CurInvocation; ++CurInvocation )
+			for( FMulticastScriptDelegate::InvocationListType::TIterator CurInvocation(GetInvocationList((uint8*)Data, i)); CurInvocation; ++CurInvocation )
 			{
 				FScriptDelegate& DestDelegateInvocation = *CurInvocation;
 				UObject* CurrentUObject = DestDelegateInvocation.GetUObject();
@@ -93,8 +93,8 @@ void FMulticastDelegateProperty::InstanceSubobjects(void* Data, void const* Defa
 
 bool FMulticastDelegateProperty::Identical( const void* A, const void* B, uint32 PortFlags ) const
 {
-	const FMulticastScriptDelegate::InvocationListType& ListA = GetInvocationList(A);
-	const FMulticastScriptDelegate::InvocationListType& ListB = GetInvocationList(B);
+	const FMulticastScriptDelegate::InvocationListType& ListA = GetInvocationList(A, 0);
+	const FMulticastScriptDelegate::InvocationListType& ListB = GetInvocationList(B, 0);
 
 	const int32 ListASize = ListA.Num();
 	if (ListASize != ListB.Num())
@@ -187,11 +187,11 @@ void FMulticastDelegateProperty::ExportText_Internal( FString& ValueStr, const v
 	{
 		FMulticastScriptDelegate Delegate;
 		GetValue_InContainer(PropertyValueOrContainer, &Delegate);
-		InvocationList = &GetInvocationList(&Delegate);
+		InvocationList = &GetInvocationList(&Delegate, 0);
 	}
 	else
 	{
-		InvocationList = &GetInvocationList(PointerToValuePtr(PropertyValueOrContainer, PropertyPointerType));
+		InvocationList = &GetInvocationList(PointerToValuePtr(PropertyValueOrContainer, PropertyPointerType), 0);
 	}
 
 	// Start delegate array with open paren
@@ -394,9 +394,9 @@ void FMulticastInlineDelegateProperty::SetMulticastDelegate(void* PropertyValue,
 	*(FMulticastScriptDelegate*)PropertyValue = MoveTemp(ScriptDelegate);
 }
 
-FMulticastScriptDelegate::InvocationListType& FMulticastInlineDelegateProperty::GetInvocationList(const void* PropertyValue) const
+FMulticastScriptDelegate::InvocationListType& FMulticastInlineDelegateProperty::GetInvocationList(const void* PropertyValue, int32 Index) const
 {
-	return (PropertyValue ? ((FMulticastScriptDelegate*)PropertyValue)->InvocationList : EmptyList);
+	return (PropertyValue ? ((FMulticastScriptDelegate*)PropertyValue + Index)->InvocationList : EmptyList);
 }
 
 void FMulticastInlineDelegateProperty::SerializeItem(FStructuredArchive::FSlot Slot, void* Value, void const* Defaults) const
@@ -505,10 +505,11 @@ void FMulticastSparseDelegateProperty::SetMulticastDelegate(void* PropertyValue,
 }
 
 
-FMulticastScriptDelegate::InvocationListType& FMulticastSparseDelegateProperty::GetInvocationList(const void* PropertyValue) const
+FMulticastScriptDelegate::InvocationListType& FMulticastSparseDelegateProperty::GetInvocationList(const void* PropertyValue, int32 Index) const
 {
 	if (FSparseDelegate* SparseDelegate = (FSparseDelegate*)PropertyValue)
 	{
+		SparseDelegate += Index;
 		if (SparseDelegate->IsBound())
 		{
 			USparseDelegateFunction* SparseDelegateFunc = CastChecked<USparseDelegateFunction>(SignatureFunction);
diff --git a/Engine/Source/Runtime/CoreUObject/Public/UObject/UnrealType.h b/Engine/Source/Runtime/CoreUObject/Public/UObject/UnrealType.h
--- a/Engine/Source/Runtime/CoreUObject/Public/UObject/UnrealType.h
+++ b/Engine/Source/Runtime/CoreUObject/Public/UObject/UnrealType.h
@@ -6005,7 +6005,7 @@ protected:
 	friend class FProperty;
 
 	static FMulticastScriptDelegate::InvocationListType EmptyList;
-	virtual FMulticastScriptDelegate::InvocationListType& GetInvocationList(const void* PropertyValue) const PURE_VIRTUAL(FMulticastDelegateProperty::GetInvocationList, return EmptyList;);
+	virtual FMulticastScriptDelegate::InvocationListType& GetInvocationList(const void* PropertyValue, int32 Index) const PURE_VIRTUAL(FMulticastDelegateProperty::GetInvocationList, return EmptyList;);
 
 
 	const TCHAR* ImportText_Add( const TCHAR* Buffer, void* PropertyValue, int32 PortFlags, UObject* Parent, FOutputDevice* ErrorText ) const;
@@ -6120,7 +6120,7 @@ PRAGMA_ENABLE_DEPRECATION_WARNINGS
 	virtual void ClearDelegate(UObject* Parent = nullptr, void* PropertyValue = nullptr) const override;
 
 protected:
-	virtual FMulticastScriptDelegate::InvocationListType& GetInvocationList(const void* PropertyValue) const;
+	virtual FMulticastScriptDelegate::InvocationListType& GetInvocationList(const void* PropertyValue, int32 Index) const;
 	// End of FMulticastDelegateProperty interface
 };
 
@@ -6171,7 +6171,7 @@ PRAGMA_ENABLE_DEPRECATION_WARNINGS
 	virtual void ClearDelegate(UObject* Parent = nullptr, void* PropertyValue = nullptr) const override;
 
 protected:
-	virtual FMulticastScriptDelegate::InvocationListType& GetInvocationList(const void* PropertyValue) const;
+	virtual FMulticastScriptDelegate::InvocationListType& GetInvocationList(const void* PropertyValue, int32 Index) const;
 	// End of FMulticastDelegateProperty interface
 
 private:
