Standardize default internal flags exclusion for all object iterators

#rnx

diff --git a/Engine/Source/Runtime/CoreUObject/Public/UObject/UObjectIterator.h b/Engine/Source/Runtime/CoreUObject/Public/UObject/UObjectIterator.h
--- a/Engine/Source/Runtime/CoreUObject/Public/UObject/UObjectIterator.h
+++ b/Engine/Source/Runtime/CoreUObject/Public/UObject/UObjectIterator.h
@@ -55,6 +55,20 @@ enum class EObjectIteratorThreadSafetyOptions : uint8
 	ThreadSafeAdvance = 2 // Can be used for global iterators but locks the global UObjectArray each time the iterator advances to the next object (can be slow but does not result in deadlocks).
 };
 
+inline EInternalObjectFlags GetObjectIteratorDefaultInternalExclusionFlags(EInternalObjectFlags InternalExclusionFlags)
+{
+	InternalExclusionFlags = UObjectBaseUtility::FixGarbageOrPendingKillInternalObjectFlags(InternalExclusionFlags);
+
+	InternalExclusionFlags |= EInternalObjectFlags::Unreachable | EInternalObjectFlags::PendingConstruction;
+	if (!IsInAsyncLoadingThread())
+	{
+		// We don't want to return any objects that are currently being background loaded unless we're using the object iterator during async loading.
+		InternalExclusionFlags |= EInternalObjectFlags::AsyncLoading;
+	}
+
+	return InternalExclusionFlags;
+}
+
 /**
  * Class for iterating through all objects, including class default objects.
  * Note that when Playing In Editor, this will find objects in the
@@ -97,14 +111,8 @@ public:
 		: FUObjectArray::TIterator(GUObjectArray, bOnlyGCedObjects)
 		, Class(InClass)
 		, ExclusionFlags(AdditionalExclusionFlags)
-		, InternalExclusionFlags(UObjectBaseUtility::FixGarbageOrPendingKillInternalObjectFlags(InInternalExclusionFlags))
+		, InternalExclusionFlags(GetObjectIteratorDefaultInternalExclusionFlags(InInternalExclusionFlags))
 	{
-		// We don't want to return any objects that are currently being background loaded unless we're using the object iterator during async loading.
-		InternalExclusionFlags |= EInternalObjectFlags::Unreachable | EInternalObjectFlags::PendingConstruction;
-		if (!IsInAsyncLoadingThread())
-		{
-			InternalExclusionFlags |= EInternalObjectFlags::AsyncLoading;
-		}
 		check(Class);
 
 		if (IteratorThreadSafety == EObjectIteratorThreadSafetyOptions::ThreadSafe)
@@ -260,10 +268,10 @@ public:
 	/**
 	 * Constructor
 	 */
-	explicit TObjectIterator(EObjectFlags AdditionalExclusionFlags = RF_ClassDefaultObject, bool bIncludeDerivedClasses = true, EInternalObjectFlags InternalExclusionFlags = EInternalObjectFlags::None)
+	explicit TObjectIterator(EObjectFlags AdditionalExclusionFlags = RF_ClassDefaultObject, bool bIncludeDerivedClasses = true, EInternalObjectFlags InInternalExclusionFlags = EInternalObjectFlags::None)
 		: Index(-1)
 	{
-		GetObjectsOfClass(T::StaticClass(), ObjectArray, bIncludeDerivedClasses, AdditionalExclusionFlags, InternalExclusionFlags);
+		GetObjectsOfClass(T::StaticClass(), ObjectArray, bIncludeDerivedClasses, AdditionalExclusionFlags, GetObjectIteratorDefaultInternalExclusionFlags(InInternalExclusionFlags));
 		Advance();
 	}
 
@@ -364,8 +372,8 @@ public:
 	 * @param	AdditionalExclusionFlags	RF_* flags that should not be included in results
 	 * @param	InInternalExclusionFlags	EInternalObjectFlags flagged objects that should not be included in results
 	 */
-	explicit TObjectIterator(EObjectFlags AdditionalExclusionFlags = RF_ClassDefaultObject, bool bIncludeDerivedClasses = true, EInternalObjectFlags InternalExclusionFlags = EInternalObjectFlags::None)
-		: FThreadSafeObjectIterator(UObject::StaticClass(), false, AdditionalExclusionFlags, InternalExclusionFlags)
+	explicit TObjectIterator(EObjectFlags AdditionalExclusionFlags = RF_ClassDefaultObject, bool bIncludeDerivedClasses = true, EInternalObjectFlags InInternalExclusionFlags = EInternalObjectFlags::None)
+		: FThreadSafeObjectIterator(UObject::StaticClass(), false, AdditionalExclusionFlags, GetObjectIteratorDefaultInternalExclusionFlags(InInternalExclusionFlags))
 	{
 	}
 
