Changed multicast delegates to work on empty delegates rather than on the internal invocation list, so that they can use the read/write scopes.

diff --git a/Engine/Source/Runtime/CoreUObject/Private/UObject/PropertyMulticastDelegate.cpp b/Engine/Source/Runtime/CoreUObject/Private/UObject/PropertyMulticastDelegate.cpp
--- a/Engine/Source/Runtime/CoreUObject/Private/UObject/PropertyMulticastDelegate.cpp
+++ b/Engine/Source/Runtime/CoreUObject/Private/UObject/PropertyMulticastDelegate.cpp
@@ -9,7 +9,7 @@
 #include "UObject/LinkerPlaceholderFunction.h"
 #include "Serialization/ArchiveUObjectFromStructuredArchive.h"
 
-FMulticastScriptDelegate::InvocationListType FMulticastDelegateProperty::EmptyList;
+FMulticastScriptDelegate FMulticastDelegateProperty::EmptyDelegate;
 
 FMulticastDelegateProperty::FMulticastDelegateProperty(FFieldVariant InOwner, const UECodeGen_Private::FMulticastDelegatePropertyParams& Prop, EPropertyFlags AdditionalPropertyFlags /*= CPF_None*/)
 	: FProperty(InOwner, (const UECodeGen_Private::FPropertyParamsBaseWithOffset&)Prop, AdditionalPropertyFlags)
@@ -38,21 +38,34 @@ void FMulticastDelegateProperty::PostDuplicate(const FField& InField)
 -----------------------------------------------------------------------------*/
 void FMulticastDelegateProperty::InstanceSubobjects(void* Data, void const* DefaultData, UObject* InOwner, FObjectInstancingGraph* InstanceGraph )
 {
+	check(Data);
 	if (DefaultData)
 	{
 		for( int32 i=0; i<ArrayDim; i++ )
 		{
 			// Fix up references to the class default object (if necessary)
-			FMulticastScriptDelegate::InvocationListType::TIterator CurInvocation(GetInvocationList((uint8*)Data, i));
-			FMulticastScriptDelegate::InvocationListType::TIterator DefaultInvocation(GetInvocationList((uint8*)DefaultData, i));
+
+			FMulticastScriptDelegate& Delegate = GetMulticastScriptDelegate((uint8*)Data, i);
+			FMulticastScriptDelegate& DefaultDelegate = GetMulticastScriptDelegate((uint8*)DefaultData, i);
+
+			if (&Delegate == &EmptyDelegate || &DefaultDelegate == &EmptyDelegate)
+			{
+				continue;
+			}
+
+			FMulticastScriptDelegate::FWriteAccessScope DelegateWriteScope = Delegate.GetWriteAccessScope();
+			FMulticastScriptDelegate::FReadAccessScope DefaultDelegateReadScope = DefaultDelegate.GetReadAccessScope();
+
+			FMulticastScriptDelegate::InvocationListType::TIterator CurInvocation(Delegate.InvocationList);
+			FMulticastScriptDelegate::InvocationListType::TIterator DefaultInvocation(DefaultDelegate.InvocationList);
 			for(; CurInvocation && DefaultInvocation; ++CurInvocation, ++DefaultInvocation )
 			{
-				FScriptDelegate& DestDelegateInvocation = *CurInvocation;
+				FMulticastScriptDelegate::UnicastDelegateType& DestDelegateInvocation = *CurInvocation;
 				UObject* CurrentUObject = DestDelegateInvocation.GetUObject();
 
 				if (CurrentUObject)
 				{
-					FScriptDelegate& DefaultDelegateInvocation = *DefaultInvocation;
+					FMulticastScriptDelegate::UnicastDelegateType& DefaultDelegateInvocation = *DefaultInvocation;
 					UObject *Template = DefaultDelegateInvocation.GetUObject();
 					UObject* NewUObject = InstanceGraph->InstancePropertyValue(Template, CurrentUObject, InOwner, EInstancePropertyValueFlags::AllowSelfReference | EInstancePropertyValueFlags::DoNotCreateNewInstance);
 					DestDelegateInvocation.BindUFunction(NewUObject, DestDelegateInvocation.GetFunctionName());
@@ -61,7 +74,7 @@ void FMulticastDelegateProperty::InstanceSubobjects(void* Data, void const* Defa
 			// now finish up the ones for which there is no default
 			for(; CurInvocation; ++CurInvocation )
 			{
-				FScriptDelegate& DestDelegateInvocation = *CurInvocation;
+				FMulticastScriptDelegate::UnicastDelegateType& DestDelegateInvocation = *CurInvocation;
 				UObject* CurrentUObject = DestDelegateInvocation.GetUObject();
 
 				if (CurrentUObject)
@@ -76,9 +89,17 @@ void FMulticastDelegateProperty::InstanceSubobjects(void* Data, void const* Defa
 	{
 		for( int32 i=0; i<ArrayDim; i++ )
 		{
-			for( FMulticastScriptDelegate::InvocationListType::TIterator CurInvocation(GetInvocationList((uint8*)Data, i)); CurInvocation; ++CurInvocation )
+			FMulticastScriptDelegate& Delegate = GetMulticastScriptDelegate((uint8*)Data, i);
+			if (&Delegate == &EmptyDelegate)
 			{
-				FScriptDelegate& DestDelegateInvocation = *CurInvocation;
+				continue;
+			}
+
+			FMulticastScriptDelegate::FWriteAccessScope DelegateWriteScope = Delegate.GetWriteAccessScope();
+
+			for (FMulticastScriptDelegate::InvocationListType::TIterator CurInvocation(Delegate.InvocationList); CurInvocation; ++CurInvocation)
+			{
+				FMulticastScriptDelegate::UnicastDelegateType& DestDelegateInvocation = *CurInvocation;
 				UObject* CurrentUObject = DestDelegateInvocation.GetUObject();
 
 				if (CurrentUObject)
@@ -93,19 +114,33 @@ void FMulticastDelegateProperty::InstanceSubobjects(void* Data, void const* Defa
 
 bool FMulticastDelegateProperty::Identical( const void* A, const void* B, uint32 PortFlags ) const
 {
-	const FMulticastScriptDelegate::InvocationListType& ListA = GetInvocationList(A, 0);
-	const FMulticastScriptDelegate::InvocationListType& ListB = GetInvocationList(B, 0);
+	if (!A || !B)
+	{
+		return A == B;
+	}
 
-	const int32 ListASize = ListA.Num();
-	if (ListASize != ListB.Num())
+	const FMulticastScriptDelegate& DelegateA = GetMulticastScriptDelegate(A, 0);
+	const FMulticastScriptDelegate& DelegateB = GetMulticastScriptDelegate(B, 0);
+
+	if (&DelegateA == &DelegateB)
+	{
+		return true;
+	}
+
+	FMulticastScriptDelegate::FReadAccessScope DelegateA_ReadScope = DelegateA.GetReadAccessScope();
+	FMulticastScriptDelegate::FReadAccessScope DelegateB_ReadScope = DelegateB.GetReadAccessScope();
+
+
+	const int32 ListASize = DelegateA.InvocationList.Num();
+	if (ListASize != DelegateB.InvocationList.Num())
 	{
 		return false;
 	}
 
 	for (int32 CurInvocationIndex = 0; CurInvocationIndex != ListASize; ++CurInvocationIndex)
 	{
-		const FScriptDelegate& BindingA = ListA[CurInvocationIndex];
-		const FScriptDelegate& BindingB = ListB[CurInvocationIndex];
+		const FMulticastScriptDelegate::UnicastDelegateType& BindingA = DelegateA.InvocationList[CurInvocationIndex];
+		const FMulticastScriptDelegate::UnicastDelegateType& BindingB = DelegateB.InvocationList[CurInvocationIndex];
 
 		if (BindingA.GetUObject() != BindingB.GetUObject())
 		{
@@ -181,37 +216,41 @@ FString FMulticastDelegateProperty::GetCPPTypeForwardDeclaration() const
 
 void FMulticastDelegateProperty::ExportText_Internal( FString& ValueStr, const void* PropertyValueOrContainer, EPropertyPointerType PropertyPointerType, const void* DefaultValue, UObject* Parent, int32 PortFlags, UObject* ExportRootScope ) const
 {
-	const FMulticastScriptDelegate::InvocationListType* InvocationList = nullptr;
-	
+	FMulticastScriptDelegate DelegateInContainer;
+	FMulticastScriptDelegate* Delegate;
 	if (PropertyPointerType == EPropertyPointerType::Container && HasGetter())
 	{
-		FMulticastScriptDelegate Delegate;
-		GetValue_InContainer(PropertyValueOrContainer, &Delegate);
-		InvocationList = &GetInvocationList(&Delegate, 0);
+		GetValue_InContainer(PropertyValueOrContainer, &DelegateInContainer);
+		Delegate = &GetMulticastScriptDelegate(&DelegateInContainer, 0);
 	}
 	else
 	{
-		InvocationList = &GetInvocationList(PointerToValuePtr(PropertyValueOrContainer, PropertyPointerType), 0);
+		Delegate = &GetMulticastScriptDelegate(PointerToValuePtr(PropertyValueOrContainer, PropertyPointerType), 0);
 	}
 
 	// Start delegate array with open paren
 	ValueStr += TEXT( "(" );
 
-	bool bIsFirstFunction = true;
-	for (FMulticastScriptDelegate::InvocationListType::TConstIterator CurInvocation(*InvocationList); CurInvocation; ++CurInvocation)
+	if (Delegate && Delegate != &EmptyDelegate)
 	{
-		if (CurInvocation->IsBound())
+		FMulticastScriptDelegate::FReadAccessScope DelegateReadScope = Delegate->GetReadAccessScope();
+
+		bool bIsFirstFunction = true;
+		for (FMulticastScriptDelegate::InvocationListType::TConstIterator CurInvocation(Delegate->InvocationList); CurInvocation; ++CurInvocation)
 		{
-			if (!bIsFirstFunction)
+			if (CurInvocation->IsBound())
 			{
-				ValueStr += TEXT(",");
-			}
-			bIsFirstFunction = false;
+				if (!bIsFirstFunction)
+				{
+					ValueStr += TEXT(",");
+				}
+				bIsFirstFunction = false;
 
-			bool bDelegateHasValue = CurInvocation->GetFunctionName() != NAME_None;
-			ValueStr += FString::Printf(TEXT("%s.%s"),
-				CurInvocation->GetUObject() != NULL ? *CurInvocation->GetUObject()->GetName() : TEXT("(null)"),
-				*CurInvocation->GetFunctionName().ToString());
+				bool bDelegateHasValue = CurInvocation->GetFunctionName() != NAME_None;
+				ValueStr += FString::Printf(TEXT("%s.%s"),
+					CurInvocation->GetUObject() != NULL ? *CurInvocation->GetUObject()->GetName() : TEXT("(null)"),
+					*CurInvocation->GetFunctionName().ToString());
+			}
 		}
 	}
 
@@ -394,9 +433,9 @@ void FMulticastInlineDelegateProperty::SetMulticastDelegate(void* PropertyValue,
 	*(FMulticastScriptDelegate*)PropertyValue = MoveTemp(ScriptDelegate);
 }
 
-FMulticastScriptDelegate::InvocationListType& FMulticastInlineDelegateProperty::GetInvocationList(const void* PropertyValue, int32 Index) const
+FMulticastScriptDelegate& FMulticastInlineDelegateProperty::GetMulticastScriptDelegate(const void* PropertyValue, int32 Index) const
 {
-	return (PropertyValue ? ((FMulticastScriptDelegate*)PropertyValue + Index)->InvocationList : EmptyList);
+	return (PropertyValue ? *((FMulticastScriptDelegate*)PropertyValue + Index) : EmptyDelegate);
 }
 
 void FMulticastInlineDelegateProperty::SerializeItem(FStructuredArchive::FSlot Slot, void* Value, void const* Defaults) const
@@ -505,7 +544,7 @@ void FMulticastSparseDelegateProperty::SetMulticastDelegate(void* PropertyValue,
 }
 
 
-FMulticastScriptDelegate::InvocationListType& FMulticastSparseDelegateProperty::GetInvocationList(const void* PropertyValue, int32 Index) const
+FMulticastScriptDelegate& FMulticastSparseDelegateProperty::GetMulticastScriptDelegate(const void* PropertyValue, int32 Index) const
 {
 	if (FSparseDelegate* SparseDelegate = (FSparseDelegate*)PropertyValue)
 	{
@@ -516,11 +555,11 @@ FMulticastScriptDelegate::InvocationListType& FMulticastSparseDelegateProperty::
 			UObject* OwningObject = FSparseDelegateStorage::ResolveSparseOwner(*SparseDelegate, SparseDelegateFunc->OwningClassName, SparseDelegateFunc->DelegateName);
 			if (FMulticastScriptDelegate* Delegate = FSparseDelegateStorage::GetMulticastDelegate(OwningObject, SparseDelegateFunc->DelegateName))
 			{
-				return Delegate->InvocationList;
+				return *Delegate;
 			}
 		}
 	}
-	return EmptyList;
+	return EmptyDelegate;
 }
 
 void FMulticastSparseDelegateProperty::SerializeItem(FStructuredArchive::FSlot Slot, void* Value, void const* Defaults) const
@@ -565,12 +604,12 @@ void FMulticastSparseDelegateProperty::SerializeItemInternal(FArchive& Ar, void*
 			}
 			else
 			{
-				Ar << EmptyList;
+				Ar << EmptyDelegate;
 			}
 		}
 		else
 		{
-			Ar << EmptyList;
+			Ar << EmptyDelegate;
 		}
 	}
 }
diff --git a/Engine/Source/Runtime/CoreUObject/Public/UObject/UnrealType.h b/Engine/Source/Runtime/CoreUObject/Public/UObject/UnrealType.h
--- a/Engine/Source/Runtime/CoreUObject/Public/UObject/UnrealType.h
+++ b/Engine/Source/Runtime/CoreUObject/Public/UObject/UnrealType.h
@@ -6004,8 +6004,8 @@ public:
 protected:
 	friend class FProperty;
 
-	static FMulticastScriptDelegate::InvocationListType EmptyList;
-	virtual FMulticastScriptDelegate::InvocationListType& GetInvocationList(const void* PropertyValue, int32 Index) const PURE_VIRTUAL(FMulticastDelegateProperty::GetInvocationList, return EmptyList;);
+	static FMulticastScriptDelegate EmptyDelegate;
+	virtual FMulticastScriptDelegate& GetMulticastScriptDelegate(const void* PropertyValue, int32 Index) const PURE_VIRTUAL(FMulticastDelegateProperty::GetMulticastScriptDelegate, return EmptyDelegate;);
 
 
 	const TCHAR* ImportText_Add( const TCHAR* Buffer, void* PropertyValue, int32 PortFlags, UObject* Parent, FOutputDevice* ErrorText ) const;
@@ -6120,7 +6120,7 @@ PRAGMA_ENABLE_DEPRECATION_WARNINGS
 	virtual void ClearDelegate(UObject* Parent = nullptr, void* PropertyValue = nullptr) const override;
 
 protected:
-	virtual FMulticastScriptDelegate::InvocationListType& GetInvocationList(const void* PropertyValue, int32 Index) const;
+	virtual FMulticastScriptDelegate& GetMulticastScriptDelegate(const void* PropertyValue, int32 Index) const;
 	// End of FMulticastDelegateProperty interface
 };
 
@@ -6171,7 +6171,7 @@ PRAGMA_ENABLE_DEPRECATION_WARNINGS
 	virtual void ClearDelegate(UObject* Parent = nullptr, void* PropertyValue = nullptr) const override;
 
 protected:
-	virtual FMulticastScriptDelegate::InvocationListType& GetInvocationList(const void* PropertyValue, int32 Index) const;
+	virtual FMulticastScriptDelegate& GetMulticastScriptDelegate(const void* PropertyValue, int32 Index) const;
 	// End of FMulticastDelegateProperty interface
 
 private:
