Add a log category and verbosity check for logs output from ConstructorHelpers::FailedToFind and ConstructorHelpers::CheckFoundViaRedirect.
Suppress them during startup of ReplicationSystemLowLevelTests.

diff --git a/Engine/Source/Programs/ReplicationSystemLowLevelTests/Private/ReplicationSystemLowLevelTests.cpp b/Engine/Source/Programs/ReplicationSystemLowLevelTests/Private/ReplicationSystemLowLevelTests.cpp
--- a/Engine/Source/Programs/ReplicationSystemLowLevelTests/Private/ReplicationSystemLowLevelTests.cpp
+++ b/Engine/Source/Programs/ReplicationSystemLowLevelTests/Private/ReplicationSystemLowLevelTests.cpp
@@ -57,10 +57,10 @@ GROUP_BEFORE_GLOBAL(Catch::DefaultGroup)
 	UE::Net::SetUseIrisReplication(true);
 
 	{
-		// Silence some warnings during initialization unrelated to the replication system
+		// Silence some errors & warnings during initialization unrelated to the replication system
 		LOG_SCOPE_VERBOSITY_OVERRIDE(LogSlate, ELogVerbosity::Error);
 		LOG_SCOPE_VERBOSITY_OVERRIDE(LogSlateStyle, ELogVerbosity::Error);
-		LOG_SCOPE_VERBOSITY_OVERRIDE(LogUObjectGlobals, ELogVerbosity::Error);
+		LOG_SCOPE_VERBOSITY_OVERRIDE(LogUObjectGlobals, ELogVerbosity::Fatal);
 		InitAll(true, true);
 	}
 
diff --git a/Engine/Source/Runtime/CoreUObject/Private/UObject/UObjectGlobals.cpp b/Engine/Source/Runtime/CoreUObject/Private/UObject/UObjectGlobals.cpp
--- a/Engine/Source/Runtime/CoreUObject/Private/UObject/UObjectGlobals.cpp
+++ b/Engine/Source/Runtime/CoreUObject/Private/UObject/UObjectGlobals.cpp
@@ -4470,7 +4470,12 @@ void ConstructorHelpers::FailedToFind(const TCHAR* ObjectToFind)
 		(CurrentInitializer && CurrentInitializer->GetClass()) ? *CurrentInitializer->GetClass()->GetName() : TEXT("Unknown"),
 		ObjectToFind);
 	FPlatformMisc::LowLevelOutputDebugString(*Message);
-	UClass::GetDefaultPropertiesFeedbackContext().Log(ELogVerbosity::Error, *Message);
+#if !NO_LOGGING
+	if (UE_LOG_ACTIVE(LogUObjectGlobals, Error))
+	{
+		UClass::GetDefaultPropertiesFeedbackContext().Log(LogUObjectGlobals.GetCategoryName(), ELogVerbosity::Error, *Message);
+	}
+#endif
 }
 
 void ConstructorHelpers::CheckFoundViaRedirect(UObject *Object, const FString& PathName, const TCHAR* ObjectToFind)
@@ -4488,7 +4493,12 @@ void ConstructorHelpers::CheckFoundViaRedirect(UObject *Object, const FString& P
 			ObjectToFind, *NewString);
 
 		FPlatformMisc::LowLevelOutputDebugString(*Message);
-		UClass::GetDefaultPropertiesFeedbackContext().Log(ELogVerbosity::Warning, *Message);
+#if !NO_LOGGING
+		if (UE_LOG_ACTIVE(LogUObjectGlobals, Warning))
+		{
+			UClass::GetDefaultPropertiesFeedbackContext().Log(LogUObjectGlobals.GetCategoryName(), ELogVerbosity::Warning, *Message);
+		}
+#endif
 	}
 }
 
