Remove DeferredPostLoadGroup ensure validating a wrong assumption

diff --git a/Engine/Source/Runtime/CoreUObject/Private/Serialization/AsyncLoading2.cpp b/Engine/Source/Runtime/CoreUObject/Private/Serialization/AsyncLoading2.cpp
--- a/Engine/Source/Runtime/CoreUObject/Private/Serialization/AsyncLoading2.cpp
+++ b/Engine/Source/Runtime/CoreUObject/Private/Serialization/AsyncLoading2.cpp
@@ -8537,10 +8537,6 @@ void FAsyncLoadingThread2::FlushLoading(int32 RequestId)
 						CurrentlyExecutingPackage->PostLoadGroup->PackagesWithExportsToSerializeCount = 1;
 					}
 				}
-				else
-				{
-					ensureMsgf(CurrentlyExecutingPackage->DeferredPostLoadGroup == nullptr, TEXT("Flushing while executing `%s` async PostLoad which is illegal"), *CurrentlyExecutingPackage->Desc.UPackageName.ToString());
-				}
 			}
 		}
 
