Add a missing PIE scope when getting asset registry tags which shouldn't influence soft object path resolving for registry tags

diff --git a/Engine/Source/Runtime/CoreUObject/Private/UObject/Obj.cpp b/Engine/Source/Runtime/CoreUObject/Private/UObject/Obj.cpp
--- a/Engine/Source/Runtime/CoreUObject/Private/UObject/Obj.cpp
+++ b/Engine/Source/Runtime/CoreUObject/Private/UObject/Obj.cpp
@@ -46,6 +46,7 @@
 #include "Serialization/BulkData.h"
 #include "UObject/LinkerLoad.h"
 #include "Misc/RedirectCollector.h"
+#include "Misc/PlayInEditorLoadingScope.h"
 #include "UObject/GCScopeLock.h"
 #include "ProfilingDebugging/CookStats.h"
 #include "ProfilingDebugging/LoadTimeTracker.h"
@@ -2129,6 +2130,7 @@ const FName FPrimaryAssetId::PrimaryAssetNameTag(TEXT("PrimaryAssetName"));
 void UObject::GetAssetRegistryTags(TArray<FAssetRegistryTag>& OutTags) const
 {
 	using namespace UE::Object::Private;
+	UE::Core::Private::FPlayInEditorLoadingScope Scope(INDEX_NONE);
 
 	// Add primary asset info if valid
 	FPrimaryAssetId PrimaryAssetId = GetPrimaryAssetId();
