Deprecate four FBulkData flags that no longer serve any useful purposes: BULKDATA_Unused, BULKDATA_ForceStreamPayload, BULKDATA_SerializeCompressedBitWindow and BULKDATA_BadDataVersion

- BULKDATA_ForceStreamPayload and BULKDATA_SerializeCompressedBitWindow have not done anything for sometime.
- BULKDATA_BadDataVersion was a bug added some years ago which could've caused DDC pollution. Since then the DDC keys for the problematic assets have been changed several times and there should no longer be any danger.
- BULKDATA_Unused technically works but is unused and there does not seem to be any demand for it. This flag will probably remain deprecated for longer than the others to check for any valid external feedback.

diff --git a/Engine/Source/Runtime/CoreUObject/Private/Serialization/AsyncLoading2.cpp b/Engine/Source/Runtime/CoreUObject/Private/Serialization/AsyncLoading2.cpp
--- a/Engine/Source/Runtime/CoreUObject/Private/Serialization/AsyncLoading2.cpp
+++ b/Engine/Source/Runtime/CoreUObject/Private/Serialization/AsyncLoading2.cpp
@@ -1849,7 +1849,9 @@ public:
 		const bool bIsInline = Meta.HasAnyFlags(BULKDATA_PayloadAtEndOfFile) == false;
 		if (bIsInline)
 		{
+PRAGMA_DISABLE_DEPRECATION_WARNINGS
 			if (int64 PayloadSize = Meta.GetSize(); PayloadSize > 0 && Meta.HasAnyFlags(BULKDATA_Unused) == false)
+PRAGMA_ENABLE_DEPRECATION_WARNINGS
 			{
 				// Set the offset from the beginning of the I/O chunk in order to make inline bulk data reloadable
 				const int64 ExportBundleChunkOffset = (ActiveFPLB->StartFastPathLoadBuffer - ActiveFPLB->OriginalFastPathLoadBuffer);
diff --git a/Engine/Source/Runtime/CoreUObject/Private/Serialization/BulkData.cpp b/Engine/Source/Runtime/CoreUObject/Private/Serialization/BulkData.cpp
--- a/Engine/Source/Runtime/CoreUObject/Private/Serialization/BulkData.cpp
+++ b/Engine/Source/Runtime/CoreUObject/Private/Serialization/BulkData.cpp
@@ -150,14 +150,18 @@ FArchive& operator<<(FArchive& Ar, FBulkMetaResource& BulkMeta)
 		Ar << BulkMeta.ElementCount;
 		Ar << BulkMeta.SizeOnDisk;
 		Ar << BulkMeta.Offset;
-
+PRAGMA_DISABLE_DEPRECATION_WARNINGS
 		if (UNLIKELY(BulkMeta.Flags & BULKDATA_BadDataVersion))
+PRAGMA_ENABLE_DEPRECATION_WARNINGS
 		{
 			if (Ar.IsLoading())
 			{
 				uint16 DummyValue;
 				Ar << DummyValue;
+
+PRAGMA_DISABLE_DEPRECATION_WARNINGS
 				BulkMeta.Flags = static_cast<EBulkDataFlags>(BulkMeta.Flags & ~BULKDATA_BadDataVersion);
+PRAGMA_ENABLE_DEPRECATION_WARNINGS
 			}
 		}
 
@@ -174,13 +178,17 @@ FArchive& operator<<(FArchive& Ar, FBulkMetaResource& BulkMeta)
 		SerializeAsInt32(Ar, BulkMeta.SizeOnDisk);
 		Ar << BulkMeta.Offset;
 		
+PRAGMA_DISABLE_DEPRECATION_WARNINGS
 		if (UNLIKELY(BulkMeta.Flags & BULKDATA_BadDataVersion))
+PRAGMA_ENABLE_DEPRECATION_WARNINGS
 		{
 			if (Ar.IsLoading())
 			{
 				uint16 DummyValue;
 				Ar << DummyValue;
+PRAGMA_DISABLE_DEPRECATION_WARNINGS
 				BulkMeta.Flags = static_cast<EBulkDataFlags>(BulkMeta.Flags & ~BULKDATA_BadDataVersion);
+PRAGMA_ENABLE_DEPRECATION_WARNINGS
 			}
 		}
 		
@@ -785,13 +793,11 @@ bool FBulkData::IsAsyncLoadingComplete() const
 	return (GetBulkDataFlags() & BULKDATA_HasAsyncReadPending) == 0;
 }
 
-/**
-* Returns whether this bulk data is used
-* @return true if BULKDATA_Unused is not set
-*/
 bool FBulkData::IsAvailableForUse() const
 {
+PRAGMA_DISABLE_DEPRECATION_WARNINGS
 	return BulkMeta.HasAnyFlags(BULKDATA_Unused) == false;
+PRAGMA_ENABLE_DEPRECATION_WARNINGS
 }
 
 /*-----------------------------------------------------------------------------
@@ -1375,7 +1381,10 @@ void FBulkData::SetFlagsFromDiskWrittenValues(
 	int64 InBulkDataSizeOnDisk,
 	int64 LinkerSummaryBulkDataStartOffset)
 {
+PRAGMA_DISABLE_DEPRECATION_WARNINGS
 	check(!(InBulkDataFlags & BULKDATA_BadDataVersion)); // This is a legacy flag that should no longer be set when saving
+PRAGMA_ENABLE_DEPRECATION_WARNINGS
+
 	if (GIsEditor)
 	{
 		InBulkDataFlags = static_cast<EBulkDataFlags>(InBulkDataFlags & ~BULKDATA_SingleUse);
@@ -1508,7 +1517,9 @@ void FBulkData::SerializeBulkData(FArchive& Ar, void* Data, int64 DataSize, EBul
 #endif
 
 	// skip serializing of unused data
+PRAGMA_DISABLE_DEPRECATION_WARNINGS
 	if (InBulkDataFlags & BULKDATA_Unused)
+PRAGMA_ENABLE_DEPRECATION_WARNINGS
 	{
 		return;
 	}
@@ -1719,7 +1730,9 @@ void FUntypedBulkData::SerializeBulkData(FArchive& Ar, void* Data, int64 DataSiz
 	SCOPED_LOADTIMER(BulkData_SerializeBulkData);
 
 	// skip serializing of unused data
+PRAGMA_DISABLE_DEPRECATION_WARNINGS
 	if (InBulkDataFlags & BULKDATA_Unused)
+PRAGMA_ENABLE_DEPRECATION_WARNINGS
 	{
 		return;
 	}
diff --git a/Engine/Source/Runtime/CoreUObject/Public/Serialization/BulkData.h b/Engine/Source/Runtime/CoreUObject/Public/Serialization/BulkData.h
--- a/Engine/Source/Runtime/CoreUObject/Public/Serialization/BulkData.h
+++ b/Engine/Source/Runtime/CoreUObject/Public/Serialization/BulkData.h
@@ -119,13 +119,13 @@ enum EBulkDataFlags : uint32
 	/** Bulk data is only used once at runtime in the game. */
 	BULKDATA_SingleUse = 1 << 3,
 	/** Bulk data won't be used and doesn't need to be loaded. */
-	BULKDATA_Unused = 1 << 5,
+	BULKDATA_Unused UE_DEPRECATED(5.3, "This feature is being removed") = 1 << 5,
 	/** Forces the payload to be saved inline, regardless of its size. */
 	BULKDATA_ForceInlinePayload = 1 << 6,
 	/** Flag to check if either compression mode is specified. */
 	BULKDATA_SerializeCompressed = (BULKDATA_SerializeCompressedZLIB),
 	/** Forces the payload to be always streamed, regardless of its size. */
-	BULKDATA_ForceStreamPayload = 1 << 7,
+	BULKDATA_ForceStreamPayload UE_DEPRECATED(5.3, "This flag has had no purpose for sometime") = 1 << 7,
 	/**
 	 * INTERNAL SET ONLY - callers of bulkdata should not set this flag on the bulk data
 	 * It is overwritten according to global configuration by Serialize.
@@ -133,7 +133,7 @@ enum EBulkDataFlags : uint32
 	 * */
 	BULKDATA_PayloadInSeperateFile = 1 << 8,
 	/** DEPRECATED: If set, payload is compressed using platform specific bit window. */
-	BULKDATA_SerializeCompressedBitWindow = 1 << 9,
+	BULKDATA_SerializeCompressedBitWindow UE_DEPRECATED(5.3, "This flag has had no purpose for sometime") = 1 << 9,
 	/** There is a new default to inline unless you opt out. */
 	BULKDATA_Force_NOT_InlinePayload = 1 << 10,
 	/** This payload is optional and may not be on device. */
@@ -145,7 +145,7 @@ enum EBulkDataFlags : uint32
 	/** Duplicate non-optional payload in optional bulk data. */
 	BULKDATA_DuplicateNonOptionalPayload = 1 << 14,
 	/** Indicates that an old ID is present in the data, at some point when the DDCs are flushed we can remove this. */
-	BULKDATA_BadDataVersion = 1 << 15,
+	BULKDATA_BadDataVersion UE_DEPRECATED(5.3, "This flag has had no purpose for sometime") = 1 << 15,
 	/** BulkData did not have it's offset changed during the cook and does not need the fix up at load time */
 	BULKDATA_NoOffsetFixUp = 1 << 16,
 	/**
@@ -644,6 +644,7 @@ public:
 	* Returns whether this bulk data is used
 	* @return true if BULKDATA_Unused is not set
 	*/
+	UE_DEPRECATED(5.3, "The feature is being removed, it is assumed that all bulkdata is available for use")
 	bool IsAvailableForUse() const;
 
 	/**
