Fix GC schema shutdown leak

Nested GC schemas transformed from StructArray -> StridedArray failed to drop their reference. Fixed by undoing refcount optimization where FSchemaBuilder transfers ownership to its built schema.

#rnx

diff --git a/Engine/Source/Runtime/CoreUObject/Private/UObject/GarbageCollectionSchema.cpp b/Engine/Source/Runtime/CoreUObject/Private/UObject/GarbageCollectionSchema.cpp
--- a/Engine/Source/Runtime/CoreUObject/Private/UObject/GarbageCollectionSchema.cpp
+++ b/Engine/Source/Runtime/CoreUObject/Private/UObject/GarbageCollectionSchema.cpp
@@ -286,13 +286,9 @@ FSchemaBuilder::FSchemaBuilder(uint32 InStride, std::initializer_list<FMemberDec
 
 FSchemaBuilder::~FSchemaBuilder()
 {
-	// Drop references to declared inner schemas if ownership hasn't been transferred to BuiltSchema
-	if (!BuiltSchema.IsSet())
+	for (FMemberDeclaration& Member : Members)
 	{
-		for (FMemberDeclaration& Member : Members)
-		{
-			TryDropSchemaReference(Member);
-		}
+		TryDropSchemaReference(Member);
 	}
 }
 
@@ -507,7 +503,8 @@ FSchemaView FSchemaBuilder::Build(ObjectAROFn ARO)
 	FMemory::Memcpy(FirstWord, Packed.Words.GetData(), sizeof(FMemberWord) * Packed.Words.Num());
 	FMemory::Memcpy(FirstWord + Packed.Words.Num(), Packed.DebugNames.GetData(), sizeof(FName) * Packed.DebugNames.Num());
 	
-	// Construct schema owner
+	// Add references to nested schemas and construct schema owner
+	VisitInnerSchemas(FSchemaView(FirstWord), [](FSchemaView Inner){ AddSchemaReference(Inner); });
 	return BuiltSchema.Emplace(FSchemaView(FirstWord)).Get();
 }
 
