fixed buffer overrun in ArchiveStackTrace that was causing mem corruption during cooking

diff --git a/Engine/Source/Runtime/CoreUObject/Private/Serialization/ArchiveStackTrace.cpp b/Engine/Source/Runtime/CoreUObject/Private/Serialization/ArchiveStackTrace.cpp
--- a/Engine/Source/Runtime/CoreUObject/Private/Serialization/ArchiveStackTrace.cpp
+++ b/Engine/Source/Runtime/CoreUObject/Private/Serialization/ArchiveStackTrace.cpp
@@ -192,8 +192,8 @@ FArchiveCallstacks::FCallstackData FArchiveCallstacks::FCallstackData::Clone() c
 	TUniquePtr<ANSICHAR[]> CallstackCopy;
 	if (const int32 Len = FCStringAnsi::Strlen(Callstack.Get()); Len > 0)
 	{
-		CallstackCopy = MakeUnique<ANSICHAR[]>(Len);
-		FCStringAnsi::Strcpy(CallstackCopy.Get(), Len, Callstack.Get());
+		CallstackCopy = MakeUnique<ANSICHAR[]>(Len + 1);
+		FMemory::Memcpy(CallstackCopy.Get(), Callstack.Get(), Len + 1);
 	}
 
 	FCallstackData Clone(MoveTemp(CallstackCopy), nullptr, SerializedProp);
