Initial implementation of DiffOnly mode for ZenServer as Cooked Output Store.  Lacking detailed header diffing, but stubbed in for future addition.  When using ZenServer as Cooked Output Store the packages use the ZenPackageSummary header format instead of the PackageFileSummary header format.  This is how data is stored in ZenServer and the format it is when fetched.  Data produced in process is transformed into this format ahead of diffing.

The code for parsing a view of the ZenPackageSummary has been extracted from AsyncLoading2.cpp into a private h/cpp and is possible to use from the header diffing code when printing out detailed diffs, but is not used yet.

diff --git a/Engine/Source/Developer/IoStoreUtilities/Private/ZenStoreWriter.cpp b/Engine/Source/Developer/IoStoreUtilities/Private/ZenStoreWriter.cpp
--- a/Engine/Source/Developer/IoStoreUtilities/Private/ZenStoreWriter.cpp
+++ b/Engine/Source/Developer/IoStoreUtilities/Private/ZenStoreWriter.cpp
@@ -66,6 +66,11 @@ FZenStoreWriter::FPackageDataEntry::~FPackageDataEntry()
 
 }
 
+FZenStoreWriter::FPendingPackageState::~FPendingPackageState()
+{
+
+}
+
 struct FZenStoreWriter::FZenCommitInfo
 {
 	IPackageWriter::FCommitPackageInfo CommitInfo;
@@ -220,22 +225,36 @@ void FZenStoreWriter::WritePackageData(const FPackageInfo& Info, FLargeMemoryWri
 	FPendingPackageState& ExistingState = GetPendingPackage(Info.PackageName);
 	FPackageDataEntry& Entry = ExistingState.PackageData.AddDefaulted_GetRef();
 
-	int64 DataSize = ExportsArchive.TotalSize();
-	FIoBuffer PackageData(FIoBuffer::AssumeOwnership, ExportsArchive.ReleaseOwnership(), DataSize);
-
 	TRACE_CPUPROFILER_EVENT_SCOPE(FZenStoreWriter::WritePackageData);
 
-	FIoBuffer CookedHeaderBuffer = FIoBuffer(PackageData.Data(), Info.HeaderSize, PackageData);
-	FIoBuffer CookedExportsBuffer = FIoBuffer(PackageData.Data() + Info.HeaderSize, PackageData.DataSize() - Info.HeaderSize, PackageData);
-	Entry.OptimizedPackage.Reset(PackageStoreOptimizer->CreatePackageFromCookedHeader(Info.PackageName, CookedHeaderBuffer));
+	FIoBuffer PackageBuffer;
+	if (ExistingState.PreOptimizedPackage.IsValid())
+	{
+		// If we are writing output data after having done a diff operation, we may already have pre-optimized package data in memory and
+		// we should use that instead of generating it again.
+		Entry.OptimizedPackage = MoveTemp(ExistingState.PreOptimizedPackage);
+		PackageBuffer = FIoBuffer(FIoBuffer::Clone, ExportsArchive.GetData(), Info.HeaderSize);
+	}
+	else
+	{
+		ExistingState.OriginalHeaderSize = Info.HeaderSize;
+
+		int64 DataSize = ExportsArchive.TotalSize();
+		FIoBuffer PackageData(FIoBuffer::AssumeOwnership, ExportsArchive.ReleaseOwnership(), DataSize);
+
+		FIoBuffer CookedHeaderBuffer = FIoBuffer(PackageData.Data(), Info.HeaderSize, PackageData);
+		FIoBuffer CookedExportsBuffer = FIoBuffer(PackageData.Data() + Info.HeaderSize, PackageData.DataSize() - Info.HeaderSize, PackageData);
+		Entry.OptimizedPackage.Reset(PackageStoreOptimizer->CreatePackageFromCookedHeader(Info.PackageName, CookedHeaderBuffer));
+		PackageBuffer = PackageStoreOptimizer->CreatePackageBuffer(Entry.OptimizedPackage.Get(), CookedExportsBuffer);
+	}
+
 	Entry.FileRegions = FileRegions;
 	for (FFileRegion& Region : Entry.FileRegions)
 	{
 		// Adjust regions so they are relative to the start of the export bundle buffer
-		Region.Offset -= Info.HeaderSize;
+		Region.Offset -= ExistingState.OriginalHeaderSize;
 		Region.Offset += Entry.OptimizedPackage->GetHeaderSize();
 	}
-	FIoBuffer PackageBuffer = PackageStoreOptimizer->CreatePackageBuffer(Entry.OptimizedPackage.Get(), CookedExportsBuffer);
 
 	// Commit to Zen build store
 
@@ -1191,6 +1210,55 @@ void FZenStoreWriter::MarkPackagesUpToDate(TArrayView<const FName> UpToDatePacka
 	}
 }
 
+bool FZenStoreWriter::GetPreviousCookedBytes(const FPackageInfo& Info, FPreviousCookedBytesData& OutData)
+{
+	if (!Info.ChunkId.IsValid())
+	{
+		return false;
+	}
+
+	FIoReadOptions ReadOptions;
+	TIoStatusOr<FIoBuffer> Status = HttpClient->ReadChunk(Info.ChunkId, ReadOptions.GetOffset(), ReadOptions.GetSize());
+	if (!Status.IsOk())
+	{
+		return false;
+	}
+
+	FIoBuffer Buffer = Status.ConsumeValueOrDie();
+	OutData.HeaderSize = reinterpret_cast<const FZenPackageSummary*>(Buffer.Data())->HeaderSize;
+	OutData.Size = Buffer.GetSize();
+	OutData.StartOffset = 0;
+	Buffer.EnsureOwned();
+	OutData.Data.Reset(Buffer.Release().ConsumeValueOrDie());
+
+	return true;
+}
+
+void FZenStoreWriter::CompleteExportsArchiveForDiff(FPackageInfo& Info, FLargeMemoryWriter& ExportsArchive)
+{
+	check(Info.ChunkId.IsValid());
+	FPendingPackageState& ExistingState = GetPendingPackage(Info.PackageName);
+
+	uint64 OptimizedHeaderSize = 0;
+	TUniquePtr<FPackageStorePackage> PackageStorePackage;
+	FIoBuffer PackageBuffer;
+	{
+		FIoBuffer CookedHeaderBuffer(FIoBuffer::Wrap, ExportsArchive.GetData(), Info.HeaderSize);
+		FIoBuffer CookedExportsBuffer(FIoBuffer::Wrap, ExportsArchive.GetData() + Info.HeaderSize, ExportsArchive.TotalSize() - Info.HeaderSize);
+		PackageStorePackage.Reset(PackageStoreOptimizer->CreatePackageFromCookedHeader(Info.PackageName, CookedHeaderBuffer));
+		OptimizedHeaderSize = PackageStorePackage->GetHeaderSize();
+		PackageBuffer = PackageStoreOptimizer->CreatePackageBuffer(PackageStorePackage.Get(), CookedExportsBuffer);
+	}
+
+	ExistingState.OriginalHeaderSize = Info.HeaderSize;
+	ExportsArchive.Seek(0);
+	FMemory::Free(ExportsArchive.ReleaseOwnership());
+	ExportsArchive.Reserve(PackageBuffer.DataSize());
+	ExportsArchive.Serialize(PackageBuffer.GetData(), PackageBuffer.DataSize());
+	Info.HeaderSize = OptimizedHeaderSize;
+	ExistingState.PreOptimizedPackage = MoveTemp(PackageStorePackage);
+}
+
 TFuture<FCbObject> FZenStoreWriter::WriteMPCookMessageForPackage(FName PackageName)
 {
 	TArray<FString> AdditionalFiles;
@@ -1274,6 +1342,15 @@ bool FZenStoreWriter::TryReadMPCookMessageForPackage(FName PackageName, FCbObjec
 	return bOk;
 }
 
+FZenStoreWriter::FPendingPackageState& FZenStoreWriter::AddPendingPackage(const FName& PackageName)
+{
+	FScopeLock _(&PackagesCriticalSection);
+	checkf(!PendingPackages.Contains(PackageName), TEXT("Trying to add package that is already pending"));
+	TUniquePtr<FPendingPackageState>& Package = PendingPackages.Add(PackageName, MakeUnique<FPendingPackageState>());
+	check(Package.IsValid());
+	return *Package;
+}
+
 void FZenStoreWriter::CreateProjectMetaData(FCbPackage& Pkg, FCbWriter& PackageObj)
 {
 	// File Manifest
diff --git a/Engine/Source/Developer/IoStoreUtilities/Public/ZenStoreWriter.h b/Engine/Source/Developer/IoStoreUtilities/Public/ZenStoreWriter.h
--- a/Engine/Source/Developer/IoStoreUtilities/Public/ZenStoreWriter.h
+++ b/Engine/Source/Developer/IoStoreUtilities/Public/ZenStoreWriter.h
@@ -45,6 +45,7 @@ class FCbAttachment;
 class FCbPackage;
 class FCbWriter;
 class FPackageStoreOptimizer;
+class FPackageStorePackage;
 class FZenFileSystemManifest;
 
 /** 
@@ -74,6 +75,14 @@ public:
 		uint16 HostPort;
 	};
 
+	IOSTOREUTILITIES_API virtual FCookCapabilities GetCookCapabilities() const override
+	{
+		FCookCapabilities Result;
+		Result.bDiffModeSupported = true;
+		Result.HeaderFormat = EPackageHeaderFormat::ZenPackageSummary;
+		return Result;
+	}
+
 	IOSTOREUTILITIES_API ZenHostInfo GetHostInfo() const;
 
 	IOSTOREUTILITIES_API virtual void BeginPackage(const FBeginPackageInfo& Info) override;
@@ -115,6 +124,8 @@ public:
 	IOSTOREUTILITIES_API virtual void RemoveCookedPackages(TArrayView<const FName> PackageNamesToRemove) override;
 	IOSTOREUTILITIES_API virtual void RemoveCookedPackages() override;
 	IOSTOREUTILITIES_API virtual void MarkPackagesUpToDate(TArrayView<const FName> UpToDatePackages) override;
+	IOSTOREUTILITIES_API bool GetPreviousCookedBytes(const FPackageInfo& Info, FPreviousCookedBytesData& OutData) override;
+	IOSTOREUTILITIES_API void CompleteExportsArchiveForDiff(FPackageInfo& Info, FLargeMemoryWriter& ExportsArchive) override;
 	IOSTOREUTILITIES_API virtual TFuture<FCbObject> WriteMPCookMessageForPackage(FName PackageName) override;
 	IOSTOREUTILITIES_API virtual bool TryReadMPCookMessageForPackage(FName PackageName, FCbObjectView Message) override;
 
@@ -140,7 +151,7 @@ private:
 		TFuture<FCompressedBuffer> CompressedPayload;
 		FPackageInfo Info;
 		FCbObjectId ChunkId;
-		TUniquePtr<class FPackageStorePackage> OptimizedPackage;
+		TUniquePtr<FPackageStorePackage> OptimizedPackage;
 		TArray<FFileRegion> FileRegions;
 		bool IsValid = false;
 	};
@@ -155,12 +166,18 @@ private:
 
 	struct FPendingPackageState
 	{
+		~FPendingPackageState();
+
 		FName PackageName;
 		TArray<FPackageDataEntry> PackageData;
 		TArray<FBulkDataEntry> BulkData;
 		TArray<FFileDataEntry> FileData;
 		TRefCountPtr<FPackageHashes> PackageHashes;
 		TUniquePtr<TPromise<int>> PackageHashesCompletionPromise;
+
+		/** Solely for use in DiffOnly mode */
+		uint64 OriginalHeaderSize = 0;
+		TUniquePtr<FPackageStorePackage> PreOptimizedPackage;
 	};
 
 	FPendingPackageState& GetPendingPackage(const FName& PackageName)
@@ -171,14 +188,7 @@ private:
 		return *Package;
 	}
 	
-	FPendingPackageState& AddPendingPackage(const FName& PackageName)
-	{
-		FScopeLock _(&PackagesCriticalSection);
-		checkf(!PendingPackages.Contains(PackageName), TEXT("Trying to add package that is already pending"));
-		TUniquePtr<FPendingPackageState>& Package = PendingPackages.Add(PackageName, MakeUnique<FPendingPackageState>());
-		check(Package.IsValid());
-		return *Package;
-	}
+	FPendingPackageState& AddPendingPackage(const FName& PackageName);
 	
 	TUniquePtr<FPendingPackageState> RemovePendingPackage(const FName& PackageName)
 	{
diff --git a/Engine/Source/Editor/UnrealEd/Private/Cooker/DiffPackageWriter.cpp b/Engine/Source/Editor/UnrealEd/Private/Cooker/DiffPackageWriter.cpp
--- a/Engine/Source/Editor/UnrealEd/Private/Cooker/DiffPackageWriter.cpp
+++ b/Engine/Source/Editor/UnrealEd/Private/Cooker/DiffPackageWriter.cpp
@@ -144,11 +144,12 @@ void FDiffPackageWriter::CommitPackage(FCommitPackageInfo&& Info)
 void FDiffPackageWriter::WritePackageData(const FPackageInfo& Info, FLargeMemoryWriter& ExportsArchive,
 	const TArray<FFileRegion>& FileRegions)
 {
-	Inner->CompleteExportsArchiveForDiff(Info, ExportsArchive);
+	FPackageInfo LocalInfo(Info);
+	Inner->CompleteExportsArchiveForDiff(LocalInfo, ExportsArchive);
 
 	FArchiveStackTrace& Writer = static_cast<FArchiveStackTrace&>(ExportsArchive);
 	ICookedPackageWriter::FPreviousCookedBytesData PreviousInnerData;
-	Inner->GetPreviousCookedBytes(Info, PreviousInnerData);
+	Inner->GetPreviousCookedBytes(LocalInfo, PreviousInnerData);
 
 	FArchiveStackTrace::FPackageData PreviousPackageData;
 	PreviousPackageData.Data = PreviousInnerData.Data.Get();
@@ -164,22 +165,22 @@ void FDiffPackageWriter::WritePackageData(const FPackageInfo& Info, FLargeMemory
 
 		TMap<FName, FArchiveDiffStats> PackageDiffStats;
 		const TCHAR* CutoffString = TEXT("UEditorEngine::Save()");
-		Writer.CompareWith(PreviousPackageData, *Info.LooseFilePath, Info.HeaderSize, CutoffString,
-			MaxDiffsToLog, PackageDiffStats);
+		Writer.CompareWith(PreviousPackageData, *LocalInfo.LooseFilePath, LocalInfo.HeaderSize, CutoffString,
+			MaxDiffsToLog, PackageDiffStats, Inner->GetCookCapabilities().HeaderFormat);
 
 		//COOK_STAT(FSavePackageStats::NumberOfDifferentPackages++);
 		//COOK_STAT(FSavePackageStats::MergeStats(PackageDiffStats));
 	}
 	else
 	{
-		Writer.GetCallstacks().Append(*ExportsCallstacks, Info.HeaderSize);
-		check(Info.MultiOutputIndex < 2);
-		ExportsDiffMapOffset[Info.MultiOutputIndex] = Info.HeaderSize;
+		Writer.GetCallstacks().Append(*ExportsCallstacks, LocalInfo.HeaderSize);
+		check(LocalInfo.MultiOutputIndex < 2);
+		ExportsDiffMapOffset[LocalInfo.MultiOutputIndex] = LocalInfo.HeaderSize;
 
-		bIsDifferent = !Writer.GenerateDiffMap(PreviousPackageData, Info.HeaderSize, MaxDiffsToLog, DiffMap[Info.MultiOutputIndex]);
+		bIsDifferent = !Writer.GenerateDiffMap(PreviousPackageData, LocalInfo.HeaderSize, MaxDiffsToLog, DiffMap[LocalInfo.MultiOutputIndex]);
 	}
 
-	Inner->WritePackageData(Info, ExportsArchive, FileRegions);
+	Inner->WritePackageData(LocalInfo, ExportsArchive, FileRegions);
 }
 
 TUniquePtr<FLargeMemoryWriter> FDiffPackageWriter::CreateLinkerArchive(FName PackageName, UObject* Asset, uint16 MultiOutputIndex)
diff --git a/Engine/Source/Editor/UnrealEd/Private/Cooker/LooseCookedPackageWriter.cpp b/Engine/Source/Editor/UnrealEd/Private/Cooker/LooseCookedPackageWriter.cpp
--- a/Engine/Source/Editor/UnrealEd/Private/Cooker/LooseCookedPackageWriter.cpp
+++ b/Engine/Source/Editor/UnrealEd/Private/Cooker/LooseCookedPackageWriter.cpp
@@ -116,7 +116,7 @@ void FLooseCookedPackageWriter::AsyncSave(FRecord& Record, const FCommitPackageI
 	AsyncSaveOutputFiles(Record, Context);
 }
 
-void FLooseCookedPackageWriter::CompleteExportsArchiveForDiff(const FPackageInfo& Info, FLargeMemoryWriter& ExportsArchive)
+void FLooseCookedPackageWriter::CompleteExportsArchiveForDiff(FPackageInfo& Info, FLargeMemoryWriter& ExportsArchive)
 {
 	FPackageWriterRecords::FPackage& BaseRecord = Records.FindRecordChecked(Info.PackageName);
 	FRecord& Record = static_cast<FRecord&>(BaseRecord);
diff --git a/Engine/Source/Editor/UnrealEd/Private/Cooker/LooseCookedPackageWriter.h b/Engine/Source/Editor/UnrealEd/Private/Cooker/LooseCookedPackageWriter.h
--- a/Engine/Source/Editor/UnrealEd/Private/Cooker/LooseCookedPackageWriter.h
+++ b/Engine/Source/Editor/UnrealEd/Private/Cooker/LooseCookedPackageWriter.h
@@ -64,7 +64,7 @@ public:
 	virtual TFuture<FCbObject> WriteMPCookMessageForPackage(FName PackageName) override;
 	virtual bool TryReadMPCookMessageForPackage(FName PackageName, FCbObjectView Message) override;
 	virtual bool GetPreviousCookedBytes(const FPackageInfo& Info, FPreviousCookedBytesData& OutData) override;
-	virtual void CompleteExportsArchiveForDiff(const FPackageInfo& Info, FLargeMemoryWriter& ExportsArchive) override;
+	virtual void CompleteExportsArchiveForDiff(FPackageInfo& Info, FLargeMemoryWriter& ExportsArchive) override;
 	virtual void CommitPackageInternal(FPackageWriterRecords::FPackage&& BaseRecord,
 		const FCommitPackageInfo& Info) override;
 	virtual FPackageWriterRecords::FPackage* ConstructRecord() override;
diff --git a/Engine/Source/Runtime/Core/Public/Serialization/PackageWriter.h b/Engine/Source/Runtime/Core/Public/Serialization/PackageWriter.h
--- a/Engine/Source/Runtime/Core/Public/Serialization/PackageWriter.h
+++ b/Engine/Source/Runtime/Core/Public/Serialization/PackageWriter.h
@@ -216,10 +216,19 @@ class ICookedPackageWriter : public IPackageWriter
 public:
 	virtual ~ICookedPackageWriter() = default;
 
+	enum class EPackageHeaderFormat
+	{
+		PackageFileSummary,
+		ZenPackageSummary
+	};
+
 	struct FCookCapabilities
 	{
 		/** Whether this writer implements -diffonly and -linkerdiff. */
 		bool bDiffModeSupported = false;
+
+		/** What header format is produced as output by this writer. */
+		EPackageHeaderFormat HeaderFormat = EPackageHeaderFormat::PackageFileSummary;
 	};
 
 	/** Return cook capabilities/settings this PackageWriter has/requires
@@ -312,7 +321,7 @@ public:
 		return false;
 	}
 	/** Append all data to the Exports archive that would normally be done in CommitPackage, used for diffing. */
-	virtual void CompleteExportsArchiveForDiff(const FPackageInfo& Info, FLargeMemoryWriter& ExportsArchive)
+	virtual void CompleteExportsArchiveForDiff(FPackageInfo& Info, FLargeMemoryWriter& ExportsArchive)
 	{
 		// The subclass must override this method if it returns bDiffModeSupported=true in GetCookCapabilities
 		unimplemented();
diff --git a/Engine/Source/Runtime/CoreUObject/Private/Serialization/ArchiveStackTrace.cpp b/Engine/Source/Runtime/CoreUObject/Private/Serialization/ArchiveStackTrace.cpp
--- a/Engine/Source/Runtime/CoreUObject/Private/Serialization/ArchiveStackTrace.cpp
+++ b/Engine/Source/Runtime/CoreUObject/Private/Serialization/ArchiveStackTrace.cpp
@@ -22,6 +22,7 @@
 #include "UObject/UObjectGlobals.h"
 #include "Misc/ScopeExit.h"
 #include "Compression/CompressionUtil.h"
+#include "Serialization/ZenPackageHeader.h"
 
 DEFINE_LOG_CATEGORY_STATIC(LogArchiveDiff, Log, All);
 
@@ -900,7 +901,8 @@ void FArchiveStackTrace::CompareWith(const TCHAR* InFilename, const int64 TotalH
 }
 
 void FArchiveStackTrace::CompareWith(const FPackageData& SourcePackage, const TCHAR* FileDisplayName, const int64 TotalHeaderSize,
-	const TCHAR* CallstackCutoffText, const int32 MaxDiffsToLog, TMap<FName, FArchiveDiffStats>&OutStats)
+	const TCHAR* CallstackCutoffText, const int32 MaxDiffsToLog, TMap<FName, FArchiveDiffStats>&OutStats,
+	const FArchiveStackTraceWriter::EPackageHeaderFormat PackageHeaderFormat /* = FArchiveStackTraceWriter::EPackageHeaderFormat::PackageFileSummary */)
 {
 	const FName AssetClass = Callstacks.GetAssetClass();
 	OutStats.FindOrAdd(AssetClass).NewFileTotalSize = TotalSize();
@@ -936,7 +938,7 @@ void FArchiveStackTrace::CompareWith(const FPackageData& SourcePackage, const TC
 
 	if (TotalHeaderSize > 0 && OutStats.FindOrAdd(AssetClass).NumDiffs > 0)
 	{
-		FArchiveStackTraceWriter::DumpPackageHeaderDiffs(SourcePackage, DestPackage, FileDisplayName, MaxDiffsToLog);
+		FArchiveStackTraceWriter::DumpPackageHeaderDiffs(SourcePackage, DestPackage, FileDisplayName, MaxDiffsToLog, PackageHeaderFormat);
 	}
 
 	FPackageData SourcePackageExports = SourcePackage;
@@ -1546,7 +1548,11 @@ static void DumpTableDifferences(
 #endif // !NO_LOGGING
 }
 
-void FArchiveStackTraceWriter::DumpPackageHeaderDiffs(const FPackageData& SourcePackage, const FPackageData& DestPackage, const FString& AssetFilename, const int32 MaxDiffsToLog)
+static void DumpPackageHeaderDiffs_LinkerLoad(
+	const FArchiveStackTraceWriter::FPackageData& SourcePackage,
+	const FArchiveStackTraceWriter::FPackageData& DestPackage,
+	const FString& AssetFilename,
+	const int32 MaxDiffsToLog)
 {
 #if !NO_LOGGING
 	FString AssetPathName = FPaths::Combine(*FPaths::GetPath(AssetFilename.Mid(AssetFilename.Find(TEXT(":"), ESearchCase::CaseSensitive) + 1)), *FPaths::GetBaseFilename(AssetFilename));
@@ -1570,14 +1576,14 @@ void FArchiveStackTraceWriter::DumpPackageHeaderDiffs(const FPackageData& Source
 	{
 		TRefCountPtr<FUObjectSerializeContext> LinkerLoadContext(FUObjectThreadContext::Get().GetSerializeContext());
 		BeginLoad(LinkerLoadContext);
-		SourceLinker = CreateLinkerForPackage(LinkerLoadContext, SourceAssetPackageName, AssetFilename, SourcePackage);
+		SourceLinker = FArchiveStackTraceWriter::CreateLinkerForPackage(LinkerLoadContext, SourceAssetPackageName, AssetFilename, SourcePackage);
 		EndLoad(SourceLinker ? SourceLinker->GetSerializeContext() : LinkerLoadContext.GetReference());
 	}
 	
 	{
 		TRefCountPtr<FUObjectSerializeContext> LinkerLoadContext(FUObjectThreadContext::Get().GetSerializeContext());
 		BeginLoad(LinkerLoadContext);
-		DestLinker = CreateLinkerForPackage(LinkerLoadContext, DestAssetPackageName, AssetFilename, DestPackage);
+		DestLinker = FArchiveStackTraceWriter::CreateLinkerForPackage(LinkerLoadContext, DestAssetPackageName, AssetFilename, DestPackage);
 		EndLoad(DestLinker ? DestLinker->GetSerializeContext() : LinkerLoadContext.GetReference());
 	}
 
@@ -1610,6 +1616,39 @@ void FArchiveStackTraceWriter::DumpPackageHeaderDiffs(const FPackageData& Source
 #endif // !NO_LOGGING
 }
 
+static void DumpPackageHeaderDiffs_ZenPackage(
+	const FArchiveStackTraceWriter::FPackageData& SourcePackage,
+	const FArchiveStackTraceWriter::FPackageData& DestPackage,
+	const FString& AssetFilename,
+	const int32 MaxDiffsToLog)
+{
+#if !NO_LOGGING
+	// TODO: Fill in detailed diffing of Zen Package Summary
+#endif // !NO_LOGGING
+}
+
+void FArchiveStackTraceWriter::DumpPackageHeaderDiffs(
+	const FPackageData& SourcePackage,
+	const FPackageData& DestPackage,
+	const FString& AssetFilename,
+	const int32 MaxDiffsToLog,
+	const EPackageHeaderFormat PackageHeaderFormat /* = EPackageHeaderFormat::PackageFileSummary */)
+{
+#if !NO_LOGGING
+	switch (PackageHeaderFormat)
+	{
+	case EPackageHeaderFormat::PackageFileSummary:
+		DumpPackageHeaderDiffs_LinkerLoad(SourcePackage, DestPackage, AssetFilename, MaxDiffsToLog);
+		break;
+	case EPackageHeaderFormat::ZenPackageSummary:
+		DumpPackageHeaderDiffs_ZenPackage(SourcePackage, DestPackage, AssetFilename, MaxDiffsToLog);
+		break;
+	default:
+		unimplemented();
+	}
+#endif // !NO_LOGGING
+}
+
 FArchiveStackTraceReader::FSerializeData::FSerializeData(int64 InOffset, int64 InSize, UObject* InObject, FProperty* InProperty)
 : Offset(InOffset)
 , Size(InSize)
diff --git a/Engine/Source/Runtime/CoreUObject/Private/Serialization/AsyncLoading2.cpp b/Engine/Source/Runtime/CoreUObject/Private/Serialization/AsyncLoading2.cpp
--- a/Engine/Source/Runtime/CoreUObject/Private/Serialization/AsyncLoading2.cpp
+++ b/Engine/Source/Runtime/CoreUObject/Private/Serialization/AsyncLoading2.cpp
@@ -75,6 +75,7 @@
 #include "Containers/SpscQueue.h"
 #include "IO/IoPriorityQueue.h"
 #include "UObject/CoreRedirects.h"
+#include "ZenPackageHeader.h"
 
 #include <atomic>
 
@@ -1407,42 +1408,13 @@ public:
 	}
 };
 
-struct FAsyncPackageHeaderData
-{
-	uint32 CookedHeaderSize = 0;
-	uint32 ExportCount = 0; // Need to keep this count around after ExportMap is cleared
-	TOptional<FZenPackageVersioningInfo> VersioningInfo;
-	FNameMap NameMap;
-	FName PackageName;
-	// Backed by IoBuffer
-	const FZenPackageSummary* PackageSummary = nullptr;
-	TArrayView<const uint64> ImportedPublicExportHashes;
-	TArrayView<const FPackageObjectIndex> ImportMap;
-	TArrayView<const FExportMapEntry> ExportMap;
-	TArrayView<const FBulkDataMapEntry> BulkDataMap;
-	TArrayView<const FExportBundleEntry> ExportBundleEntries;
-	TArrayView<const FDependencyBundleHeader> DependencyBundleHeaders;
-	TArrayView<const FDependencyBundleEntry> DependencyBundleEntries;
-
+struct FAsyncPackageHeaderData: public FZenPackageHeader
+{
 	// Backed by allocation in FAsyncPackageData
 	TArrayView<FPackageId> ImportedPackageIds;
 	TArrayView<FAsyncPackage2*> ImportedAsyncPackagesView;
 	TArrayView<FExportObject> ExportsView;
 	TArrayView<FExportBundleEntry> ExportBundleEntriesCopyForPostLoad; // TODO: Can we use ConstructedObjects or Exports instead for posloading?
-
-	TArray<FName> ImportedPackageNames;
-
-	void OnReleaseHeaderBuffer()
-	{
-		PackageSummary = nullptr;
-		ImportedPublicExportHashes = TArrayView<const uint64>();
-		ImportMap = TArrayView<const FPackageObjectIndex>();
-		ExportMap = TArrayView<const FExportMapEntry>();
-		BulkDataMap = TArrayView<const FBulkDataMapEntry>();
-		ExportBundleEntries = TArrayView<const FExportBundleEntry>();
-		DependencyBundleHeaders = TArrayView<const FDependencyBundleHeader>();
-		DependencyBundleEntries = TArrayView<const FDependencyBundleEntry>();
-	}
 };
 
 #if ALT2_ENABLE_LINKERLOAD_SUPPORT
@@ -5067,76 +5039,6 @@ void FAsyncPackage2::StartLoading(FAsyncLoadingThreadState2& ThreadState, FIoBat
 	}
 }
 
-static void ReadAsyncPackageHeader(FAsyncPackageSerializationState& SerializationState, FAsyncPackageHeaderData& HeaderData)
-{
-	const uint8* PackageHeaderDataPtr = SerializationState.IoRequest.GetResultOrDie().Data();
-	const FZenPackageSummary* PackageSummary = reinterpret_cast<const FZenPackageSummary*>(PackageHeaderDataPtr);
-	HeaderData.PackageSummary = PackageSummary;
-
-	TArrayView<const uint8> PackageHeaderDataView(PackageHeaderDataPtr + sizeof(FZenPackageSummary), PackageSummary->HeaderSize - sizeof(FZenPackageSummary));
-	FMemoryReaderView PackageHeaderDataReader(PackageHeaderDataView);
-	if (PackageSummary->bHasVersioningInfo)
-	{
-		HeaderData.VersioningInfo.Emplace();
-		PackageHeaderDataReader << HeaderData.VersioningInfo.GetValue();
-	}
-
-	{
-		TRACE_CPUPROFILER_EVENT_SCOPE(LoadPackageNameMap);
-		HeaderData.NameMap.Load(PackageHeaderDataReader, FMappedName::EType::Package);
-	}
-	HeaderData.PackageName = HeaderData.NameMap.GetName(PackageSummary->Name);
-	
-	const FZenPackageVersioningInfo* VersioningInfo = HeaderData.VersioningInfo.GetPtrOrNull();
-	if (VersioningInfo == nullptr || VersioningInfo->PackageVersion >= EUnrealEngineObjectUE5Version::DATA_RESOURCES)
-	{
-		int64 BulkDataMapSize = 0;
-		PackageHeaderDataReader << BulkDataMapSize;
-		const uint8* BulkDataMapData = PackageHeaderDataPtr + sizeof(FZenPackageSummary) + PackageHeaderDataReader.Tell();
-		HeaderData.BulkDataMap = MakeArrayView(reinterpret_cast<const FBulkDataMapEntry*>(BulkDataMapData), BulkDataMapSize / sizeof(FBulkDataMapEntry));
-	}
-
-	HeaderData.CookedHeaderSize = PackageSummary->CookedHeaderSize;
-	HeaderData.ImportedPublicExportHashes = TArrayView<const uint64>(
-		reinterpret_cast<const uint64*>(PackageHeaderDataPtr + PackageSummary->ImportedPublicExportHashesOffset),
-		(PackageSummary->ImportMapOffset - PackageSummary->ImportedPublicExportHashesOffset) / sizeof(uint64));
-	HeaderData.ImportMap = TArrayView<const FPackageObjectIndex>(
-		reinterpret_cast<const FPackageObjectIndex*>(PackageHeaderDataPtr + PackageSummary->ImportMapOffset),
-		(PackageSummary->ExportMapOffset - PackageSummary->ImportMapOffset) / sizeof(FPackageObjectIndex));
-	HeaderData.ExportMap = TArrayView<const FExportMapEntry>(
-		reinterpret_cast<const FExportMapEntry*>(PackageHeaderDataPtr + PackageSummary->ExportMapOffset),
-		(PackageSummary->ExportBundleEntriesOffset - PackageSummary->ExportMapOffset) / sizeof(FExportMapEntry));
-	HeaderData.ExportCount = HeaderData.ExportMap.Num();
-
-	const uint64 ExportBundleEntriesSize = PackageSummary->DependencyBundleHeadersOffset - PackageSummary->ExportBundleEntriesOffset;
-	const int32 ExportBundleEntriesCount = static_cast<int32>(ExportBundleEntriesSize / sizeof(FExportBundleEntry));
-
-	if (ExportBundleEntriesCount != HeaderData.ExportCount * FExportBundleEntry::ExportCommandType_Count)
-	{
-		UE_LOG(LogStreaming, Fatal, TEXT("Corrupt Zen header in package %s"), *HeaderData.PackageName.ToString());
-		return;
-	}
-
-	HeaderData.ExportBundleEntries = TArrayView<const FExportBundleEntry>(
-		reinterpret_cast<const FExportBundleEntry*>(PackageHeaderDataPtr + PackageSummary->ExportBundleEntriesOffset),
-		ExportBundleEntriesCount);
-	HeaderData.DependencyBundleHeaders = TArrayView<const FDependencyBundleHeader>(
-		reinterpret_cast<const FDependencyBundleHeader*>(PackageHeaderDataPtr + PackageSummary->DependencyBundleHeadersOffset),
-		(PackageSummary->DependencyBundleEntriesOffset - PackageSummary->DependencyBundleHeadersOffset) / sizeof(FDependencyBundleHeader));
-	HeaderData.DependencyBundleEntries = TArrayView<const FDependencyBundleEntry>(
-		reinterpret_cast<const FDependencyBundleEntry*>(PackageHeaderDataPtr + PackageSummary->DependencyBundleEntriesOffset),
-		(PackageSummary->ImportedPackageNamesOffset - PackageSummary->DependencyBundleEntriesOffset) / sizeof(FDependencyBundleEntry));
-
-#if WITH_EDITOR
-	// Only read the imported package names in editor builds where we need to interact with uncooked packages
-	TArrayView<const uint8> ImportedPackageNamesDataView(PackageHeaderDataPtr + PackageSummary->ImportedPackageNamesOffset, PackageSummary->HeaderSize - PackageSummary->ImportedPackageNamesOffset);
-	FMemoryReaderView ImportedPackageNamesDataReader(ImportedPackageNamesDataView);
-	FZenPackageImportedPackageNamesContainer Container;
-	ImportedPackageNamesDataReader << Container;
-	HeaderData.ImportedPackageNames = MoveTemp(Container.Names);
-#endif
-}
-
 #if ALT2_ENABLE_LINKERLOAD_SUPPORT
 EEventLoadNodeExecutionResult FAsyncPackage2::ProcessLinkerLoadPackageSummary(FAsyncLoadingThreadState2& ThreadState)
 {
@@ -5561,12 +5463,12 @@ EEventLoadNodeExecutionResult FAsyncPackage2::Event_ProcessPackageSummary(FAsync
 		TRACE_LOADTIME_PROCESS_SUMMARY_SCOPE(Package);
 		check(Package->ExportBundleEntryIndex == 0);
 
-		ReadAsyncPackageHeader(Package->SerializationState, Package->HeaderData);
+		static_cast<FZenPackageHeader&>(Package->HeaderData) = FZenPackageHeader::MakeView(Package->SerializationState.IoRequest.GetResultOrDie().GetView());
 #if WITH_EDITOR
 		FAsyncPackageHeaderData* OptionalSegmentHeaderData = Package->OptionalSegmentHeaderData.GetPtrOrNull();
 		if (OptionalSegmentHeaderData)
 		{
-			ReadAsyncPackageHeader(*Package->OptionalSegmentSerializationState, *OptionalSegmentHeaderData);
+			static_cast<FZenPackageHeader&>(*OptionalSegmentHeaderData) = FZenPackageHeader::MakeView(Package->OptionalSegmentSerializationState->IoRequest.GetResultOrDie().GetView());
 		}
 #endif
 		if (Package->Desc.bCanBeImported)
@@ -5778,12 +5680,12 @@ EEventLoadNodeExecutionResult FAsyncPackage2::Event_ProcessExportBundle(FAsyncLo
 	if (++Package->ProcessedExportBundlesCount == Package->Data.TotalExportBundleCount)
 	{
 		Package->ProcessedExportBundlesCount = 0;
-		Package->HeaderData.OnReleaseHeaderBuffer();
+		Package->HeaderData.Reset();
 		Package->SerializationState.ReleaseIoRequest();
 #if WITH_EDITOR
 		if (Package->OptionalSegmentHeaderData.IsSet())
 		{
-			Package->OptionalSegmentHeaderData->OnReleaseHeaderBuffer();
+			Package->OptionalSegmentHeaderData->Reset();
 			Package->OptionalSegmentSerializationState->ReleaseIoRequest();
 		}
 #endif
diff --git a/Engine/Source/Runtime/CoreUObject/Private/Serialization/ZenPackageHeader.cpp b/Engine/Source/Runtime/CoreUObject/Private/Serialization/ZenPackageHeader.cpp
new file mode 100644
--- /dev/null
+++ b/Engine/Source/Runtime/CoreUObject/Private/Serialization/ZenPackageHeader.cpp
@@ -0,0 +1,90 @@
+// Copyright Epic Games, Inc. All Rights Reserved.
+
+#include "ZenPackageHeader.h"
+#include "Serialization/MemoryReader.h"
+
+
+FZenPackageHeader FZenPackageHeader::MakeView(FMemoryView Memory)
+{
+	FZenPackageHeader PackageHeader;
+	const uint8* PackageHeaderDataPtr = reinterpret_cast<const uint8*>(Memory.GetData());
+	const FZenPackageSummary* PackageSummary = reinterpret_cast<const FZenPackageSummary*>(PackageHeaderDataPtr);
+	PackageHeader.PackageSummary = PackageSummary;
+
+	TArrayView<const uint8> PackageHeaderDataView(PackageHeaderDataPtr + sizeof(FZenPackageSummary), PackageHeader.PackageSummary->HeaderSize - sizeof(FZenPackageSummary));
+	FMemoryReaderView PackageHeaderDataReader(PackageHeaderDataView);
+	if (PackageHeader.PackageSummary->bHasVersioningInfo)
+	{
+		PackageHeader.VersioningInfo.Emplace();
+		PackageHeaderDataReader << PackageHeader.VersioningInfo.GetValue();
+	}
+
+	{
+		TRACE_CPUPROFILER_EVENT_SCOPE(LoadPackageNameMap);
+		PackageHeader.NameMap.Load(PackageHeaderDataReader, FMappedName::EType::Package);
+	}
+	PackageHeader.PackageName = PackageHeader.NameMap.GetName(PackageHeader.PackageSummary->Name);
+	
+	const FZenPackageVersioningInfo* LocalVersioningInfo = PackageHeader.VersioningInfo.GetPtrOrNull();
+	if (LocalVersioningInfo == nullptr || LocalVersioningInfo->PackageVersion >= EUnrealEngineObjectUE5Version::DATA_RESOURCES)
+	{
+		int64 BulkDataMapSize = 0;
+		PackageHeaderDataReader << BulkDataMapSize;
+		const uint8* BulkDataMapData = PackageHeaderDataPtr + sizeof(FZenPackageSummary) + PackageHeaderDataReader.Tell();
+		PackageHeader.BulkDataMap = MakeArrayView(reinterpret_cast<const FBulkDataMapEntry*>(BulkDataMapData), BulkDataMapSize / sizeof(FBulkDataMapEntry));
+	}
+
+	PackageHeader.CookedHeaderSize = PackageHeader.PackageSummary->CookedHeaderSize;
+	PackageHeader.ImportedPublicExportHashes = TArrayView<const uint64>(
+		reinterpret_cast<const uint64*>(PackageHeaderDataPtr + PackageSummary->ImportedPublicExportHashesOffset),
+		(PackageSummary->ImportMapOffset - PackageSummary->ImportedPublicExportHashesOffset) / sizeof(uint64));
+	PackageHeader.ImportMap = TArrayView<const FPackageObjectIndex>(
+		reinterpret_cast<const FPackageObjectIndex*>(PackageHeaderDataPtr + PackageSummary->ImportMapOffset),
+		(PackageSummary->ExportMapOffset - PackageSummary->ImportMapOffset) / sizeof(FPackageObjectIndex));
+	PackageHeader.ExportMap = TArrayView<const FExportMapEntry>(
+		reinterpret_cast<const FExportMapEntry*>(PackageHeaderDataPtr + PackageSummary->ExportMapOffset),
+		(PackageSummary->ExportBundleEntriesOffset - PackageSummary->ExportMapOffset) / sizeof(FExportMapEntry));
+	PackageHeader.ExportCount = PackageHeader.ExportMap.Num();
+
+	const uint64 ExportBundleEntriesSize = PackageHeader.PackageSummary->DependencyBundleHeadersOffset - PackageHeader.PackageSummary->ExportBundleEntriesOffset;
+	const int32 ExportBundleEntriesCount = static_cast<int32>(ExportBundleEntriesSize / sizeof(FExportBundleEntry));
+
+	if (ExportBundleEntriesCount != PackageHeader.ExportCount * FExportBundleEntry::ExportCommandType_Count)
+	{
+		UE_LOG(LogStreaming, Fatal, TEXT("Corrupt Zen header in package %s"), *PackageHeader.PackageName.ToString());
+		return PackageHeader;
+	}
+
+	PackageHeader.ExportBundleEntries = TArrayView<const FExportBundleEntry>(
+		reinterpret_cast<const FExportBundleEntry*>(PackageHeaderDataPtr + PackageHeader.PackageSummary->ExportBundleEntriesOffset),
+		ExportBundleEntriesCount);
+	PackageHeader.DependencyBundleHeaders = TArrayView<const FDependencyBundleHeader>(
+		reinterpret_cast<const FDependencyBundleHeader*>(PackageHeaderDataPtr + PackageSummary->DependencyBundleHeadersOffset),
+		(PackageSummary->DependencyBundleEntriesOffset - PackageSummary->DependencyBundleHeadersOffset) / sizeof(FDependencyBundleHeader));
+	PackageHeader.DependencyBundleEntries = TArrayView<const FDependencyBundleEntry>(
+		reinterpret_cast<const FDependencyBundleEntry*>(PackageHeaderDataPtr + PackageSummary->DependencyBundleEntriesOffset),
+		(PackageSummary->ImportedPackageNamesOffset - PackageSummary->DependencyBundleEntriesOffset) / sizeof(FDependencyBundleEntry));
+
+#if WITH_EDITOR
+	// Only read the imported package names in editor builds where we need to interact with uncooked packages
+	TArrayView<const uint8> ImportedPackageNamesDataView(PackageHeaderDataPtr + PackageSummary->ImportedPackageNamesOffset, PackageSummary->HeaderSize - PackageSummary->ImportedPackageNamesOffset);
+	FMemoryReaderView ImportedPackageNamesDataReader(ImportedPackageNamesDataView);
+	FZenPackageImportedPackageNamesContainer Container;
+	ImportedPackageNamesDataReader << Container;
+	PackageHeader.ImportedPackageNames = MoveTemp(Container.Names);
+#endif
+
+	return PackageHeader;
+}
+
+void FZenPackageHeader::Reset()
+{
+	PackageSummary = nullptr;
+	ImportedPublicExportHashes = TArrayView<const uint64>();
+	ImportMap = TArrayView<const FPackageObjectIndex>();
+	ExportMap = TArrayView<const FExportMapEntry>();
+	BulkDataMap = TArrayView<const FBulkDataMapEntry>();
+	ExportBundleEntries = TArrayView<const FExportBundleEntry>();
+	DependencyBundleHeaders = TArrayView<const FDependencyBundleHeader>();
+	DependencyBundleEntries = TArrayView<const FDependencyBundleEntry>();
+}
diff --git a/Engine/Source/Runtime/CoreUObject/Private/Serialization/ZenPackageHeader.h b/Engine/Source/Runtime/CoreUObject/Private/Serialization/ZenPackageHeader.h
new file mode 100644
--- /dev/null
+++ b/Engine/Source/Runtime/CoreUObject/Private/Serialization/ZenPackageHeader.h
@@ -0,0 +1,35 @@
+// Copyright Epic Games, Inc. All Rights Reserved.
+
+#pragma once
+
+#include "Containers/Array.h"
+#include "Containers/ArrayView.h"
+#include "IO/PackageId.h"
+#include "Memory/MemoryFwd.h"
+#include "Serialization/AsyncLoading2.h"
+#include "Serialization/MappedName.h"
+#include "UObject/NameTypes.h"
+
+struct FZenPackageHeader
+{
+	uint32 CookedHeaderSize = 0;
+	uint32 ExportCount = 0; // Need to keep this count around after ExportMap is cleared
+	TOptional<FZenPackageVersioningInfo> VersioningInfo;
+	FNameMap NameMap;
+	FName PackageName;
+
+	// Backed by IoBuffer
+	const FZenPackageSummary* PackageSummary = nullptr;
+	TArrayView<const uint64> ImportedPublicExportHashes;
+	TArrayView<const FPackageObjectIndex> ImportMap;
+	TArrayView<const FExportMapEntry> ExportMap;
+	TArrayView<const FBulkDataMapEntry> BulkDataMap;
+	TArrayView<const FExportBundleEntry> ExportBundleEntries;
+	TArrayView<const FDependencyBundleHeader> DependencyBundleHeaders;
+	TArrayView<const FDependencyBundleEntry> DependencyBundleEntries;
+
+	TArray<FName> ImportedPackageNames;
+
+	static FZenPackageHeader MakeView(FMemoryView Memory);
+	void Reset();
+};
diff --git a/Engine/Source/Runtime/CoreUObject/Public/Serialization/ArchiveStackTrace.h b/Engine/Source/Runtime/CoreUObject/Public/Serialization/ArchiveStackTrace.h
--- a/Engine/Source/Runtime/CoreUObject/Public/Serialization/ArchiveStackTrace.h
+++ b/Engine/Source/Runtime/CoreUObject/Public/Serialization/ArchiveStackTrace.h
@@ -11,6 +11,7 @@
 #include "Serialization/Archive.h"
 #include "Serialization/LargeMemoryReader.h"
 #include "Serialization/LargeMemoryWriter.h"
+#include "Serialization/PackageWriter.h"
 #include "Templates/RefCounting.h"
 #include "Templates/UniquePtr.h"
 #include "UObject/NameTypes.h"
@@ -284,6 +285,8 @@ public:
 		int64 StartOffset = 0;
 	};
 
+	using EPackageHeaderFormat = ICookedPackageWriter::EPackageHeaderFormat;
+
 	/** Compares two packages and logs the differences and calltacks. */
 	static void Compare(
 		const FPackageData& SourcePackage,
@@ -310,7 +313,8 @@ public:
 		const FPackageData& SourcePackage,
 		const FPackageData& DestPackage,
 		const FString& AssetFilename,
-		const int32 MaxDiffsToLog);
+		const int32 MaxDiffsToLog,
+		const EPackageHeaderFormat PackageHeaderFormat = EPackageHeaderFormat::PackageFileSummary);
 
 	/** Returns a new linker for loading the specified package. */
 	static FLinkerLoad* CreateLinkerForPackage(
@@ -385,7 +389,8 @@ public:
 	void CompareWith(const TCHAR* InFilename, const int64 TotalHeaderSize, const TCHAR* CallstackCutoffText,
 		const int32 MaxDiffsToLog, TMap<FName, FArchiveDiffStats>& OutStats);
 	void CompareWith(const FPackageData& SourcePackage, const TCHAR* FileDisplayName, const int64 TotalHeaderSize,
-		const TCHAR* CallstackCutoffText, const int32 MaxDiffsToLog, TMap<FName, FArchiveDiffStats>& OutStats);
+		const TCHAR* CallstackCutoffText, const int32 MaxDiffsToLog, TMap<FName, FArchiveDiffStats>& OutStats,
+		const FArchiveStackTraceWriter::EPackageHeaderFormat PackageHeaderFormat = FArchiveStackTraceWriter::EPackageHeaderFormat::PackageFileSummary);
 
 	/** Generates a map of all differences between this archive and the given bytes from disk or FPackageData. */
 	bool GenerateDiffMap(const TCHAR* InFilename, int64 TotalHeaderSize, int32 MaxDiffsToFind, FArchiveDiffMap& OutDiffMap);
