#UE Actually print all content roots before fatal error about failing to convert a filename. It was printed before this change, but truncated because the string is too big to print on one line

diff --git a/Engine/Source/Runtime/CoreUObject/Private/Misc/PackageName.cpp b/Engine/Source/Runtime/CoreUObject/Private/Misc/PackageName.cpp
--- a/Engine/Source/Runtime/CoreUObject/Private/Misc/PackageName.cpp
+++ b/Engine/Source/Runtime/CoreUObject/Private/Misc/PackageName.cpp
@@ -838,14 +838,15 @@ FString FPackageName::FilenameToLongPackageName(const FString& InFilename)
 	FString Result;
 	if (!TryConvertFilenameToLongPackageName(InFilename, Result, &FailureReason))
 	{
-		TStringBuilder<128> ContentRoots;
+		TArray<FString> ContentRootsArrayRelative;
+		TArray<FString> ContentRootsArrayAbsolute;
 		{
 			const auto& Paths = FLongPackagePathsSingleton::Get();
 			FReadScopeLock ScopeLock(Paths.MountLock);
 			for (const TUniquePtr<FMountPoint>& MountPoint : Paths.MountPoints)
 			{
-				ContentRoots << TEXT("\n\t\t") << MountPoint->ContentPathRelative;
-				ContentRoots << TEXT("\n\t\t\t") << MountPoint->ContentPathAbsolute;
+				ContentRootsArrayRelative.Add(MountPoint->ContentPathRelative);
+				ContentRootsArrayAbsolute.Add(MountPoint->ContentPathAbsolute);
 			}
 		}
 
@@ -855,11 +856,22 @@ FString FPackageName::FilenameToLongPackageName(const FString& InFilename)
 			TEXT("\n\tConvertRelativePathToFull=%s")
 			TEXT("\n\tRootDir=%s")
 			TEXT("\n\tBaseDir=%s")
-			TEXT("\n\tContentRoots=%s"),
+			TEXT("\n\tContentRoots listed below..."),
 			*InFilename, *IFileManager::Get().ConvertToRelativePath(*InFilename),
-			*FPaths::ConvertRelativePathToFull(InFilename), FPlatformMisc::RootDir(), FPlatformProcess::BaseDir(),
-			ContentRoots.ToString()
-		);
+			*FPaths::ConvertRelativePathToFull(InFilename), FPlatformMisc::RootDir(), FPlatformProcess::BaseDir()
+			);
+		
+		if (ensure(ContentRootsArrayRelative.Num() == ContentRootsArrayAbsolute.Num()))
+		{
+			for (int32 RootIdx = 0; RootIdx < ContentRootsArrayRelative.Num(); ++RootIdx)
+			{
+				const FString& RelativeRoot = ContentRootsArrayRelative[RootIdx];
+				const FString& AbsoluteRoot = ContentRootsArrayAbsolute[RootIdx];
+				UE_LOG(LogPackageName, Display, TEXT("\t\t%s"), *RelativeRoot);
+				UE_LOG(LogPackageName, Display, TEXT("\t\t\t%s"), *AbsoluteRoot)
+			}
+		}
+		
 		UE_LOG(LogPackageName, Fatal, TEXT("%s"), *FailureReason);
 	}
 	return Result;
