#UE Fix to no longer consider BP classes and CDOs as editor only

diff --git a/Engine/Source/Runtime/CoreUObject/Private/UObject/SavePackage2.cpp b/Engine/Source/Runtime/CoreUObject/Private/UObject/SavePackage2.cpp
--- a/Engine/Source/Runtime/CoreUObject/Private/UObject/SavePackage2.cpp
+++ b/Engine/Source/Runtime/CoreUObject/Private/UObject/SavePackage2.cpp
@@ -359,7 +359,7 @@ ESavePackageResult HarvestPackage(FSaveContext& SaveContext)
 	{
 		ForEachObjectWithPackage(SaveContext.GetPackage(), [&TryHarvestRootObject, TopLevelFlags](UObject* InObject)
 			{
-				if (InObject->HasAnyFlags(TopLevelFlags))
+				if (InObject->HasAnyFlags(TopLevelFlags | RF_ClassDefaultObject))
 				{
 					TryHarvestRootObject(InObject);
 				}
