Ensure that ObjectHandle access tracking has an inline and lockless fast exit path when there are zero handle read callbacks installed.

diff --git a/Engine/Source/Runtime/CoreUObject/Private/UObject/ObjectHandleTracking.cpp b/Engine/Source/Runtime/CoreUObject/Private/UObject/ObjectHandleTracking.cpp
--- a/Engine/Source/Runtime/CoreUObject/Private/UObject/ObjectHandleTracking.cpp
+++ b/Engine/Source/Runtime/CoreUObject/Private/UObject/ObjectHandleTracking.cpp
@@ -10,6 +10,8 @@ namespace UE::CoreUObject
 {
 	namespace Private
 	{
+		COREUOBJECT_API std::atomic<int32> HandleReadCallbackQuantity = 0;
+
 		struct ObjectHandleCallbacks
 		{
 			static ObjectHandleCallbacks& Get()
@@ -71,6 +73,7 @@ namespace UE::CoreUObject
 				FObjectHandleTrackingCallbackId Result;
 				UE_AUTORTFM_OPEN(
 				{
+					HandleReadCallbackQuantity.fetch_add(1, std::memory_order_release);
 					FWriteScopeLock _(HandleLock);
 					NextHandleId++;
 					ReadHandleCallbacks.Add({ NextHandleId, Func });
@@ -83,6 +86,7 @@ namespace UE::CoreUObject
 			{
 				UE_AUTORTFM_OPEN(
 				{
+					HandleReadCallbackQuantity.fetch_sub(1, std::memory_order_release);
 					FWriteScopeLock _(HandleLock);
 					for (int32 i = ReadHandleCallbacks.Num() - 1; i >= 0; --i)
 					{
@@ -193,7 +197,7 @@ namespace UE::CoreUObject
 			FRWLock HandleLock;
 		};
 
-		void OnHandleRead(TArrayView<const UObject* const> Objects)
+		void OnHandleReadInternal(TArrayView<const UObject* const> Objects)
 		{
 			ObjectHandleCallbacks::Get().OnHandleRead(Objects);
 		}
diff --git a/Engine/Source/Runtime/CoreUObject/Public/UObject/ObjectHandleTracking.h b/Engine/Source/Runtime/CoreUObject/Public/UObject/ObjectHandleTracking.h
--- a/Engine/Source/Runtime/CoreUObject/Public/UObject/ObjectHandleTracking.h
+++ b/Engine/Source/Runtime/CoreUObject/Public/UObject/ObjectHandleTracking.h
@@ -4,6 +4,7 @@
 #include "Delegates/Delegate.h"
 #include "UObject/ObjectHandleDefines.h"
 #include "Templates/Function.h"
+#include <atomic>
 
 struct FObjectRef;
 
@@ -102,12 +103,22 @@ namespace UE::CoreUObject
 #if UE_WITH_OBJECT_HANDLE_TRACKING
 namespace UE::CoreUObject::Private
 {
-	
-	COREUOBJECT_API void OnHandleRead(TArrayView<const UObject* const> Objects);
+	extern COREUOBJECT_API std::atomic<int32> HandleReadCallbackQuantity;
+	COREUOBJECT_API void OnHandleReadInternal(TArrayView<const UObject* const> Objects);
+	inline void OnHandleRead(TArrayView<const UObject* const> Objects)
+	{
+		if (HandleReadCallbackQuantity.load(std::memory_order_acquire) > 0)
+		{
+			OnHandleReadInternal(Objects);
+		}
+	}
 	inline void OnHandleRead(const UObject* Object)
 	{
-		TArrayView<const UObject* const> View(&Object, 1);
-		OnHandleRead(View);
+		if (HandleReadCallbackQuantity.load(std::memory_order_acquire) > 0)
+		{
+			TArrayView<const UObject* const> View(&Object, 1);
+			OnHandleReadInternal(View);
+		}
 	}
 	COREUOBJECT_API void OnClassReferenceResolved(const FObjectRef& ObjectRef, UPackage* ClassPackage, UClass* Class);
 	COREUOBJECT_API void OnReferenceResolved(const FObjectRef& ObjectRef, UPackage* ObjectPackage, UObject* Object);
