Remove TObjectPtr on UStruct::SuperStruct as it causes a high quantity of access tracking callbacks that are not valuable.  Will discuss options to mitigate this while restoring it to a TObjectPtr.

diff --git a/Engine/Source/Runtime/CoreUObject/Public/UObject/Class.h b/Engine/Source/Runtime/CoreUObject/Public/UObject/Class.h
--- a/Engine/Source/Runtime/CoreUObject/Public/UObject/Class.h
+++ b/Engine/Source/Runtime/CoreUObject/Public/UObject/Class.h
@@ -392,7 +392,7 @@ protected:
 	friend struct Z_Construct_UClass_UStruct_Statics;
 private:
 	/** Struct this inherits from, may be null */
-	TObjectPtr<UStruct> SuperStruct;
+	UStruct* SuperStruct;
 public:
 	/** Pointer to start of linked list of child fields */
 	TObjectPtr<UField> Children;
