Minor fix for verify cached archetype logging

diff --git a/Engine/Source/Runtime/CoreUObject/Private/UObject/UObjectArchetype.cpp b/Engine/Source/Runtime/CoreUObject/Private/UObject/UObjectArchetype.cpp
--- a/Engine/Source/Runtime/CoreUObject/Private/UObject/UObjectArchetype.cpp
+++ b/Engine/Source/Runtime/CoreUObject/Private/UObject/UObjectArchetype.cpp
@@ -189,7 +189,7 @@ UObject* UObject::GetArchetype() const
 		UObject* ExpectedArchetype = GetArchetypeFromRequiredInfo(GetClass(), GetOuter(), GetFName(), GetFlags());
 		if (ExpectedArchetype != Archetype)
 		{
-			UE_LOG(LogClass, Fatal, TEXT("Cached archetype mismatch, expected: %s, cached: %s"), *ExpectedArchetype->GetFullName(), *Archetype->GetFullName());
+			UE_LOG(LogClass, Fatal, TEXT("Cached archetype mismatch, expected: %s, cached: %s"), *GetFullNameSafe(ExpectedArchetype), *GetFullNameSafe(Archetype));
 		}
 #endif // UE_VERIFY_CACHED_ARCHETYPE
 	}
