Fix potential truncation issues in GarbageCollection.

#rnx

*** Garbage Collection
- Change 'TimeLimit' parameters in FAsyncPurge to double as the value is a double before being passed in and is compred to other doubles inside of the ::Tick methods anyway.
- FWorkCoordinator constructor was doing pointer arithmetic to find the array index. Since we know the array uses int32 for the size (TArrayView) we can just use UE_PTRDIFF_TO_INT32 for the cast.

*** Garbage Collection Schema
- Used IntCastChecked in places where we are going from larger types to smaller ones to fit in with the style of the rest of the code file (INtCastChecked is being used elsewhere here for similar cases)
- Store the size of memory to be allocated in a SIZE_T, the return type of Align and the type that FMemory::Malloc takes.

diff --git a/Engine/Source/Runtime/CoreUObject/Private/UObject/GarbageCollection.cpp b/Engine/Source/Runtime/CoreUObject/Private/UObject/GarbageCollection.cpp
--- a/Engine/Source/Runtime/CoreUObject/Private/UObject/GarbageCollection.cpp
+++ b/Engine/Source/Runtime/CoreUObject/Private/UObject/GarbageCollection.cpp
@@ -531,7 +531,7 @@ class FAsyncPurge : public FRunnable
 
 	/** [PURGE/GAME THREAD] Destroys objects that are unreachable */
 	template <bool bMultithreaded> // Having this template argument lets the compiler strip unnecessary checks
-	bool TickDestroyObjects(bool bUseTimeLimit, float TimeLimit, double StartTime)
+	bool TickDestroyObjects(bool bUseTimeLimit, double TimeLimit, double StartTime)
 	{
 		TRACE_CPUPROFILER_EVENT_SCOPE(FAsyncPurge::TickDestroyObjects);
 		const int32 TimeLimitEnforcementGranularityForDeletion = 100;
@@ -582,7 +582,7 @@ class FAsyncPurge : public FRunnable
 	}
 
 	/** [GAME THREAD] Destroys objects that are unreachable and couldn't be destroyed on the worker thread */
-	bool TickDestroyGameThreadObjects(bool bUseTimeLimit, float TimeLimit, double StartTime)
+	bool TickDestroyGameThreadObjects(bool bUseTimeLimit, double TimeLimit, double StartTime)
 	{
 		TRACE_CPUPROFILER_EVENT_SCOPE(FAsyncPurge::TickDestroyGameThreadObjects);
 		const int32 TimeLimitEnforcementGranularityForDeletion = 100;
@@ -711,7 +711,7 @@ public:
 	}
 
 	/** [GAME THREAD] Ticks the purge process on the game thread */
-	void TickPurge(bool bUseTimeLimit, float TimeLimit, double StartTime)
+	void TickPurge(bool bUseTimeLimit, double TimeLimit, double StartTime)
 	{
 		bool bCanStartDestroyingGameThreadObjects = true;
 		if (!Thread)
@@ -5415,7 +5415,7 @@ public:
 	{
 		for (FStealableContext& Context : Contexts)
 		{
-			Context.store(InContexts[&Context - Contexts.GetData()], std::memory_order_relaxed);
+			Context.store(InContexts[UE_PTRDIFF_TO_INT32(&Context - Contexts.GetData())], std::memory_order_relaxed);
 		}
 	}
 
diff --git a/Engine/Source/Runtime/CoreUObject/Private/UObject/GarbageCollectionSchema.cpp b/Engine/Source/Runtime/CoreUObject/Private/UObject/GarbageCollectionSchema.cpp
--- a/Engine/Source/Runtime/CoreUObject/Private/UObject/GarbageCollectionSchema.cpp
+++ b/Engine/Source/Runtime/CoreUObject/Private/UObject/GarbageCollectionSchema.cpp
@@ -369,7 +369,7 @@ private:
 
 	void Pack(const FMemberDeclaration& Member)
 	{
-		uint16 RelativeWordOffset = IsIn<StopTypes>(Member.Type) ? Member.Offset : GenerateJumps(Member.Offset);
+		uint16 RelativeWordOffset = IsIn<StopTypes>(Member.Type) ? IntCastChecked<uint16>(Member.Offset) : GenerateJumps(Member.Offset);
 
 		if (TOptional<FStridedLayout> StridedLayout = TrySimplifyStructArray(Member))
 		{
@@ -393,7 +393,7 @@ private:
 		check(ByteOffset % 8 == 0);
 		check(ByteOffset / 8 >= JumpWordPos);
 
-		const uint32 WordOffsetFromJumpPos = ByteOffset / 8 - JumpWordPos;
+		const uint32 WordOffsetFromJumpPos = ByteOffset / 8 - IntCastChecked<uint32>(JumpWordPos);
 
 		// 1 jump unit is OffsetRange words, a max jump is OffsetRange * OffsetRange * 8B = 32MB with 11 offset bits
 		const uint32 JumpUnits = WordOffsetFromJumpPos / OffsetRange;
@@ -497,7 +497,7 @@ FSchemaView FSchemaBuilder::Build(ObjectAROFn ARO)
 	check(Packed.DebugNames.IsEmpty() == !UE_GC_DEBUGNAMES) ;
 
 	// Allocate and initialize schema data
-	const uint32 Bytes = Align(sizeof(FSchemaHeader), sizeof(FMemberWord)) + sizeof(FMemberWord) * Packed.Words.Num() + sizeof(FName) * Packed.DebugNames.Num();
+	const SIZE_T Bytes = Align(sizeof(FSchemaHeader), sizeof(FMemberWord)) + sizeof(FMemberWord) * Packed.Words.Num() + sizeof(FName) * Packed.DebugNames.Num();
 	FSchemaHeader* Header = new (FMemory::Malloc(Bytes)) FSchemaHeader {StructStride};
 	FMemberWord* FirstWord = (FMemberWord*)(Header + 1);
 	FMemory::Memcpy(FirstWord, Packed.Words.GetData(), sizeof(FMemberWord) * Packed.Words.Num());
