Removing legacy hash table locks when collecting garbage as they are no longer used

diff --git a/Engine/Source/Runtime/CoreUObject/Private/UObject/GarbageCollection.cpp b/Engine/Source/Runtime/CoreUObject/Private/UObject/GarbageCollection.cpp
--- a/Engine/Source/Runtime/CoreUObject/Private/UObject/GarbageCollection.cpp
+++ b/Engine/Source/Runtime/CoreUObject/Private/UObject/GarbageCollection.cpp
@@ -97,77 +97,23 @@ bool GIsGarbageCollecting = false;
 */
 void NotifyUnreachableObjects(const TArrayView<FUObjectItem*>& UnreachableObjects);
 
-enum FGCLockBehavior : int32
-{
-	Default = 0,
-	Legacy = 1,
-};
-static int32 GGCLockBehavior = FGCLockBehavior::Default;
-static FAutoConsoleVariable CVarGCLockBehavior(
-	TEXT("gc.LockBehavior"),
-	GGCLockBehavior,
-	TEXT("Set the GC lock behavior: 0=Default, 1=Legacy (more restrictive and less performant)."),
-	FConsoleVariableDelegate::CreateLambda([](IConsoleVariable* Variable)
-	{
-		GGCLockBehavior = FMath::Clamp(Variable->GetInt(), 0, 1);
-	}),
-	ECVF_Default);
-
-
 
 /** Locks all UObject hash tables for performing GC reachability analysis.
  * Can be queried with IsGarbageCollectingAndLockingUObjectHashTables().
  * */
 class FGCHashTableScopeLock
 {
-	bool bEnabled = (GGCLockBehavior == FGCLockBehavior::Default);
 public:
 
 	FORCEINLINE FGCHashTableScopeLock()
 	{
-		if (bEnabled)
-		{
-			GIsGarbageCollectingAndLockingUObjectHashTables = true;
-			LockUObjectHashTables();
-		}
+		GIsGarbageCollectingAndLockingUObjectHashTables = true;
+		LockUObjectHashTables();
 	}
 	FORCEINLINE ~FGCHashTableScopeLock()
 	{
-		if (bEnabled)
-		{
-			UnlockUObjectHashTables();
-			GIsGarbageCollectingAndLockingUObjectHashTables = false;
-		}
-	}
-};
-
-/** Locks all UObject hash tables when performing GC */
-class FGCLegacyHashTableScopeLock
-{
-	bool bEnabled = (GGCLockBehavior == FGCLockBehavior::Legacy);
-	/** Previous value of the GetGarbageCollectingFlag() */
-	bool bPreviousGabageCollectingFlagValue = GIsGarbageCollectingAndLockingUObjectHashTables;
-public:
-
-	/**
-	 * We're storing the value of GetGarbageCollectingFlag in the constructor, it's safe as only
-	 * one thread is ever going to be setting it and calling this code - the game thread.
-	 **/
-	FORCEINLINE FGCLegacyHashTableScopeLock()
-	{
-		if (bEnabled)
-		{
-			LockUObjectHashTables();
-			GIsGarbageCollectingAndLockingUObjectHashTables = true;
-		}
-	}
-	FORCEINLINE ~FGCLegacyHashTableScopeLock()
-	{
-		if (bEnabled)
-		{
-			GIsGarbageCollectingAndLockingUObjectHashTables = bPreviousGabageCollectingFlagValue;
-			UnlockUObjectHashTables();
-		}
+		UnlockUObjectHashTables();
+		GIsGarbageCollectingAndLockingUObjectHashTables = false;
 	}
 };
 
@@ -3915,28 +3861,6 @@ void ReleaseGCLock()
 	FGCCSyncObject::Get().GCUnlock();
 }
 
-/** Locks GC within a scope but only if it hasn't been locked already */
-struct FLegacyConditionalIncrementalGCLock
-{
-	bool bNeedsUnlock;
-	FLegacyConditionalIncrementalGCLock()
-		: bNeedsUnlock(false)
-	{
-		if ((GGCLockBehavior == FGCLockBehavior::Legacy) && !FGCCSyncObject::Get().IsGCLocked())
-		{
-			AcquireGCLock();
-			bNeedsUnlock = true;
-		}
-	}
-	~FLegacyConditionalIncrementalGCLock()
-	{
-		if (bNeedsUnlock)
-		{
-			ReleaseGCLock();
-		}
-	}
-};
-
 static bool IncrementalDestroyGarbage(bool bUseTimeLimit, double TimeLimit);
 
 /**
@@ -3997,9 +3921,7 @@ void IncrementalPurgeGarbage(bool bUseTimeLimit, double TimeLimit)
 	}
 	else
 	{
-		// Lock before settting GCStartTime as it could be slow to lock if async loading is in progress
-		// but we still want to perform some GC work otherwise we'd be keeping objects in memory for a long time
-		FLegacyConditionalIncrementalGCLock ScopedGCLock;
+		// Set 'I'm garbage collecting' flag - might be checked inside various functions.
 		TGuardValue<bool> GuardIsGarbageCollecting(GIsGarbageCollecting, true);
 
 		// Keep track of start time to enforce time limit unless bForceFullPurge is true;
@@ -4081,7 +4003,6 @@ bool IncrementalDestroyGarbage(bool bUseTimeLimit, double TimeLimit)
 	const int32	TimeLimitEnforcementGranularityForDeletion = 100;
 
 	// Set 'I'm garbage collecting' flag - might be checked inside UObject::Destroy etc.
-	FGCLegacyHashTableScopeLock GCLegacyHashTableLock;
 	TGuardValue<bool> GuardIsGarbageCollecting(GIsGarbageCollecting, true);
 
 	if( !GObjFinishDestroyHasBeenRoutedToAllObjects && !bTimeLimitReached )
@@ -4666,7 +4587,6 @@ void CollectGarbageImpl(EObjectFlags KeepFlags)
 
 		// Set 'I'm garbage collecting' flag - might be checked inside various functions.
 		// This has to be unlocked before we call post GC callbacks
-		FGCLegacyHashTableScopeLock GCLegacyHashTableLock;
 		TGuardValue<bool> GuardIsGarbageCollecting(GIsGarbageCollecting, true);
 
 		UE_LOG(LogGarbage, Log, TEXT("Collecting garbage%s"), IsAsyncLoading() ? TEXT(" while async loading") : TEXT(""));
@@ -4779,15 +4699,11 @@ void CollectGarbageImpl(EObjectFlags KeepFlags)
 
 			NotifyUnreachableObjects(GUnreachableObjects);
 		}
-		if (GGCLockBehavior == FGCLockBehavior::Default)
-		{
-			// With the new FGCLockBehavior::Default behavior the hash tables lock was released when exiting the reachability analysis scope above.
-			// BeginDestroy, FinishDestroy, destructors and callbacks are allowed to call functions like StaticAllocateObject and StaticFindObject.
-			// Now release the GC lock to allow async loading and other threads to perform UObject operations under the FGCScopeGuard.
-			ReleaseGCLock();
-		}
-		// The old FGCLockBehavior::Legacy behavior keeps the GC and hash table locks until GC is completed,
-		// and no threads are allowed to perform UObject operations yet.
+
+		// The hash tables lock was released when exiting the reachability analysis scope above.
+		// BeginDestroy, FinishDestroy, destructors and callbacks are allowed to call functions like StaticAllocateObject and StaticFindObject.
+		// Now release the GC lock to allow async loading and other threads to perform UObject operations under the FGCScopeGuard.
+		ReleaseGCLock();
 
 		// Fire post-reachability analysis hooks
 		{
@@ -4955,12 +4871,7 @@ void CollectGarbage(EObjectFlags KeepFlags, bool bPerformFullPurge)
 	// Perform actual garbage collection
 	UE::GC::CollectGarbageInternal(KeepFlags, bPerformFullPurge);
 
-	if (GGCLockBehavior == FGCLockBehavior::Legacy)
-	{
-		// Release the GC lock to allow async loading and other threads to perform UObject operations under the FGCScopeGuard.
-		ReleaseGCLock();
-	}
-	// With the new FGCLockBehavior::Default behavior the lock was released after reachability analysis inside CollectGarbageInternal
+	// GC lock was released after reachability analysis inside CollectGarbageInternal
 }
 
 bool TryCollectGarbage(EObjectFlags KeepFlags, bool bPerformFullPurge)
@@ -4990,12 +4901,7 @@ bool TryCollectGarbage(EObjectFlags KeepFlags, bool bPerformFullPurge)
 		// Perform actual garbage collection
 		UE::GC::CollectGarbageInternal(KeepFlags, bPerformFullPurge);
 
-		if (GGCLockBehavior == FGCLockBehavior::Legacy)
-		{
-			// Release the GC lock to allow async loading and other threads to perform UObject operations under the FGCScopeGuard.
-			ReleaseGCLock();
-		}
-		// With the new FGCLockBehavior::Default behavior the lock was released after reachability analysis inside CollectGarbageInternal
+		// GC lock was released after reachability analysis inside CollectGarbageInternal
 	}
 	else
 	{
