Add a CVar so we can skip async package loader for the content cooker without disrupting other workflows

#rnx

diff --git a/Engine/Source/Runtime/CoreUObject/Private/UObject/LinkerLoad.cpp b/Engine/Source/Runtime/CoreUObject/Private/UObject/LinkerLoad.cpp
--- a/Engine/Source/Runtime/CoreUObject/Private/UObject/LinkerLoad.cpp
+++ b/Engine/Source/Runtime/CoreUObject/Private/UObject/LinkerLoad.cpp
@@ -165,6 +165,14 @@ static FAutoConsoleVariableRef CVarAllowCookedDataInEditorBuilds(
 	ECVF_Default
 );
 
+int32 GSkipAsyncLoaderForCookedData = 0;
+static FAutoConsoleVariableRef CVarSkipAsyncLoaderForCookedData(
+	TEXT("cook.SkipAsyncLoaderForCookedData"),
+	GSkipAsyncLoaderForCookedData,
+	TEXT("If true, skip the async loader and load package header synchronously to reduce ping/pong between threads."),
+	ECVF_Default
+);
+
 int32 GEnforcePackageCompatibleVersionCheck = 1;
 static FAutoConsoleVariableRef CEnforcePackageCompatibleVersionCheck(
 	TEXT("s.EnforcePackageCompatibleVersionCheck"),
@@ -1143,7 +1151,7 @@ FLinkerLoad::ELinkerStatus FLinkerLoad::CreateLoader(
 
 		// If want to be able to load cooked data in the editor we need to use FAsyncArchive which supports EDL cooked packages,
 		// otherwise the generic file reader is faster in the editor so use that
-		bool bCanUseAsyncLoader = FPlatformProperties::RequiresCookedData() || GAllowCookedDataInEditorBuilds;
+		bool bCanUseAsyncLoader = (FPlatformProperties::RequiresCookedData() || GAllowCookedDataInEditorBuilds) && !GSkipAsyncLoaderForCookedData;
 
 		if (bCanUseAsyncLoader)
 		{
