Turn off access tracking for UStruct::SuperStruct field (optimization)

diff --git a/Engine/Source/Runtime/CoreUObject/Private/UObject/Class.cpp b/Engine/Source/Runtime/CoreUObject/Private/UObject/Class.cpp
--- a/Engine/Source/Runtime/CoreUObject/Private/UObject/Class.cpp
+++ b/Engine/Source/Runtime/CoreUObject/Private/UObject/Class.cpp
@@ -1821,7 +1821,7 @@ void UStruct::Serialize(FArchive& Ar)
 	UStruct* SuperStructBefore = GetSuperStruct();
 #endif
 
-	Ar << SuperStruct;
+	Ar << SuperStruct.GetAccessTrackedObjectPtr();
 
 #if USTRUCT_FAST_ISCHILDOF_IMPL == USTRUCT_ISCHILDOF_STRUCTARRAY
 	if (Ar.IsLoading())
@@ -2033,7 +2033,7 @@ void UStruct::AddReferencedObjects(UObject* InThis, FReferenceCollector& Collect
 	if( GIsEditor )
 	{
 		// Required by the unified GC when running in the editor
-		Collector.AddReferencedObject( This->SuperStruct, This );
+		Collector.AddReferencedObject( This->SuperStruct.GetAccessTrackedObjectPtr(), This );
 		Collector.AddReferencedObject( This->Children, This );
 		Collector.AddReferencedObjects(This->ScriptAndPropertyObjectReferences, This);
 	}
@@ -2046,7 +2046,7 @@ void UStruct::AddReferencedObjects(UObject* InThis, FReferenceCollector& Collect
 
 void UStruct::SetSuperStruct(UStruct* NewSuperStruct)
 {
-	SuperStruct = NewSuperStruct;
+	SuperStruct.Set(NewSuperStruct);
 #if USTRUCT_FAST_ISCHILDOF_IMPL == USTRUCT_ISCHILDOF_STRUCTARRAY
 	this->ReinitializeBaseChainArray();
 #endif
diff --git a/Engine/Source/Runtime/CoreUObject/Public/UObject/Class.h b/Engine/Source/Runtime/CoreUObject/Public/UObject/Class.h
--- a/Engine/Source/Runtime/CoreUObject/Public/UObject/Class.h
+++ b/Engine/Source/Runtime/CoreUObject/Public/UObject/Class.h
@@ -392,7 +392,7 @@ protected:
 	friend struct Z_Construct_UClass_UStruct_Statics;
 private:
 	/** Struct this inherits from, may be null */
-	UStruct* SuperStruct;
+	ObjectPtr_Private::TNonAccessTrackedObjectPtr<UStruct> SuperStruct;
 public:
 	/** Pointer to start of linked list of child fields */
 	TObjectPtr<UField> Children;
diff --git a/Engine/Source/Runtime/CoreUObject/Public/UObject/ObjectPtr.h b/Engine/Source/Runtime/CoreUObject/Public/UObject/ObjectPtr.h
--- a/Engine/Source/Runtime/CoreUObject/Public/UObject/ObjectPtr.h
+++ b/Engine/Source/Runtime/CoreUObject/Public/UObject/ObjectPtr.h
@@ -596,6 +596,48 @@ namespace ObjectPtr_Private
 		{
 			Slot << InObjectPtr.ObjectPtr;
 		}
+
+		template <typename T>
+		FORCEINLINE static T* NoAccessTrackingGet(const TObjectPtr<T>& Ptr)
+		{
+			return reinterpret_cast<T*>(UE::CoreUObject::Private::ResolveObjectHandleNoRead(Ptr.ObjectPtr.GetHandleRef()));
+		}
+	};
+	
+	template <typename T>
+	class TNonAccessTrackedObjectPtr
+	{
+	public:
+		TNonAccessTrackedObjectPtr() = default;
+		explicit TNonAccessTrackedObjectPtr(T* Ptr) : ObjectPtr{Ptr} {}
+
+		void Set(T* Value)
+		{
+			ObjectPtr = Value;
+		}
+	
+		T* Get() const
+		{
+			return ObjectPtr_Private::Friend::NoAccessTrackingGet(ObjectPtr);
+		}
+
+		TObjectPtr<T>& GetAccessTrackedObjectPtr() 
+		{
+			return ObjectPtr;
+		}
+
+		operator T*() const
+		{
+			return Get();
+		}
+		
+		T* operator->() const
+		{
+			return Get();
+		}		
+	
+	private:
+		TObjectPtr<T> ObjectPtr{};
 	};
 }
 
