ARO: Fix handling of TWeakObjectPtr

diff --git a/Engine/Source/Runtime/CoreUObject/Private/UObject/UObjectGlobals.cpp b/Engine/Source/Runtime/CoreUObject/Private/UObject/UObjectGlobals.cpp
--- a/Engine/Source/Runtime/CoreUObject/Private/UObject/UObjectGlobals.cpp
+++ b/Engine/Source/Runtime/CoreUObject/Private/UObject/UObjectGlobals.cpp
@@ -4611,7 +4611,7 @@ void FReferenceCollector::AddReferencedObjects(const UScriptStruct*& ScriptStruc
 
 void FReferenceCollector::AddReferencedObjects(TWeakObjectPtr<const UScriptStruct>& ScriptStruct, void* Instance, const UObject* ReferencingObject, const FProperty* ReferencingProperty)
 {
-	const UScriptStruct* Ptr = ScriptStruct.Get();
+	const UScriptStruct* Ptr = ScriptStruct.GetEvenIfUnreachable();
 	AROPrivate::AddReferencedObjects(*this, Ptr, Instance, ReferencingObject, ReferencingProperty);	 
 	ScriptStruct = Ptr;
 }
@@ -4623,7 +4623,7 @@ void FReferenceCollector::AddReferencedObjects(TObjectPtr<const UScriptStruct>&
 
 void FReferenceCollector::AddReferencedObject(FWeakObjectPtr& P, const UObject* ReferencingObject, const FProperty* ReferencingProperty)
 {
-	UObject* Ptr = P.Get();
+	UObject* Ptr = P.GetEvenIfUnreachable();
 	AROPrivate::AddReferencedObject(*this, Ptr, ReferencingObject, ReferencingProperty);
 	P = Ptr;
 }
