Fix reference replacement of set

diff --git a/Engine/Source/Runtime/CoreUObject/Private/UObject/PropertyMap.cpp b/Engine/Source/Runtime/CoreUObject/Private/UObject/PropertyMap.cpp
--- a/Engine/Source/Runtime/CoreUObject/Private/UObject/PropertyMap.cpp
+++ b/Engine/Source/Runtime/CoreUObject/Private/UObject/PropertyMap.cpp
@@ -543,6 +543,16 @@ void FMapProperty::SerializeItem(FStructuredArchive::FSlot Slot, void* Value, co
 					--Num;
 				}
 			}
+
+			// Check if a rehash is needed
+			if (UnderlyingArchive.ArIsModifyingWeakAndStrongReferences)
+			{
+				TArray<const FStructProperty*> EncounteredStructProps;
+				if (ContainsObjectReference(EncounteredStructProps, EPropertyObjectReferenceType::Strong | EPropertyObjectReferenceType::Weak))
+				{
+					MapHelper.Rehash();
+				}
+			}
 		}
 	}
 }
diff --git a/Engine/Source/Runtime/CoreUObject/Private/UObject/PropertySet.cpp b/Engine/Source/Runtime/CoreUObject/Private/UObject/PropertySet.cpp
--- a/Engine/Source/Runtime/CoreUObject/Private/UObject/PropertySet.cpp
+++ b/Engine/Source/Runtime/CoreUObject/Private/UObject/PropertySet.cpp
@@ -475,6 +475,16 @@ void FSetProperty::SerializeItem(FStructuredArchive::FSlot Slot, void* Value, co
 					--Num;
 				}
 			}
+
+			// Check if a rehash is needed
+			if (UnderlyingArchive.ArIsModifyingWeakAndStrongReferences)
+			{
+				TArray<const FStructProperty*> EncounteredStructProps;
+				if (ContainsObjectReference(EncounteredStructProps, EPropertyObjectReferenceType::Strong | EPropertyObjectReferenceType::Weak))
+				{
+					SetHelper.Rehash();
+				}
+			}
 		}
 	}
 }
