[Feature] added GetPropertyInstanceWithin to CPPStructOps so that classes like FInstancedStruct can have their internal properties reflected by property paths

diff --git a/Engine/Plugins/Experimental/StructUtils/Source/StructUtils/Private/InstancedStruct.cpp b/Engine/Plugins/Experimental/StructUtils/Source/StructUtils/Private/InstancedStruct.cpp
--- a/Engine/Plugins/Experimental/StructUtils/Source/StructUtils/Private/InstancedStruct.cpp
+++ b/Engine/Plugins/Experimental/StructUtils/Source/StructUtils/Private/InstancedStruct.cpp
@@ -495,3 +495,23 @@ bool FInstancedStruct::NetSerialize(FArchive& Ar, UPackageMap* Map, bool& bOutSu
 
 #endif // WITH_ENGINE
 }
+
+bool FInstancedStruct::FindInnerPropertyInstance(FName PropertyName, const FProperty*& OutProp, const void*& OutData) const
+{
+	if (!ScriptStruct || !StructMemory)
+	{
+		return false;
+	}
+	
+	for (const FProperty* Prop : TFieldRange<FProperty>(ScriptStruct))
+	{
+		if( Prop->GetFName() == PropertyName )
+		{
+			OutProp = Prop;
+			OutData = StructMemory;
+			return true;
+		}
+	}
+
+	return false;
+}
diff --git a/Engine/Plugins/Experimental/StructUtils/Source/StructUtils/Public/InstancedStruct.h b/Engine/Plugins/Experimental/StructUtils/Source/StructUtils/Public/InstancedStruct.h
--- a/Engine/Plugins/Experimental/StructUtils/Source/StructUtils/Public/InstancedStruct.h
+++ b/Engine/Plugins/Experimental/StructUtils/Source/StructUtils/Public/InstancedStruct.h
@@ -152,6 +152,7 @@ public:
 	bool SerializeFromMismatchedTag(const FPropertyTag& Tag, FStructuredArchive::FSlot Slot);
 	void GetPreloadDependencies(TArray<UObject*>& OutDeps);
 	bool NetSerialize(FArchive& Ar, class UPackageMap* Map, bool& bOutSuccess);
+	bool FindInnerPropertyInstance(FName PropertyName, const FProperty*& OutProp, const void*& OutData) const;
 
 	/** Returns struct type. */
 	const UScriptStruct* GetScriptStruct() const
@@ -264,6 +265,7 @@ struct TStructOpsTypeTraits<FInstancedStruct> : public TStructOpsTypeTraitsBase2
 		WithStructuredSerializeFromMismatchedTag = true,
 		WithGetPreloadDependencies = true,
 		WithNetSerializer = true,
+		WithFindInnerPropertyInstance = true,
 	};
 };
 
diff --git a/Engine/Source/Runtime/CoreUObject/Private/UObject/Class.cpp b/Engine/Source/Runtime/CoreUObject/Private/UObject/Class.cpp
--- a/Engine/Source/Runtime/CoreUObject/Private/UObject/Class.cpp
+++ b/Engine/Source/Runtime/CoreUObject/Private/UObject/Class.cpp
@@ -3196,6 +3196,19 @@ void UScriptStruct::InitializeDefaultValue(uint8* InStructData) const
 	InitializeStruct(InStructData);
 }
 
+bool UScriptStruct::FindInnerPropertyInstance(FName PropertyName, const void* Data, const FProperty*& OutProp, const void*& OutData) const
+{
+	if (const UScriptStruct::ICppStructOps* TheCppStructOps = GetCppStructOps())
+	{
+        if (TheCppStructOps->HasFindInnerPropertyInstance())
+        {
+            return TheCppStructOps->FindInnerPropertyInstance(PropertyName, Data, OutProp, OutData);
+        }
+	}
+	
+	return false;
+}
+
 void UScriptStruct::ClearScriptStruct(void* Dest, int32 ArrayDim) const
 {
 	uint8 *Data = (uint8*)Dest;
diff --git a/Engine/Source/Runtime/CoreUObject/Private/UObject/PropertyStruct.cpp b/Engine/Source/Runtime/CoreUObject/Private/UObject/PropertyStruct.cpp
--- a/Engine/Source/Runtime/CoreUObject/Private/UObject/PropertyStruct.cpp
+++ b/Engine/Source/Runtime/CoreUObject/Private/UObject/PropertyStruct.cpp
@@ -157,6 +157,11 @@ bool FStructProperty::UseBinaryOrNativeSerialization(const FArchive& Ar) const
 	return bUseBinarySerialization || bUseNativeSerialization;
 }
 
+bool FStructProperty::FindInnerPropertyInstance(FName PropertyName, const void* Data, const FProperty*& OutProp, const void*& OutData) const
+{
+	return Struct->FindInnerPropertyInstance(PropertyName, Data, OutProp, OutData);
+}
+
 uint32 FStructProperty::GetValueTypeHashInternal(const void* Src) const
 {
 	check(Struct);
diff --git a/Engine/Source/Runtime/CoreUObject/Public/UObject/Class.h b/Engine/Source/Runtime/CoreUObject/Public/UObject/Class.h
--- a/Engine/Source/Runtime/CoreUObject/Public/UObject/Class.h
+++ b/Engine/Source/Runtime/CoreUObject/Public/UObject/Class.h
@@ -863,6 +863,7 @@ struct TStructOpsTypeTraitsBase2
 		WithNetSharedSerialization     = false,                         // struct has a NetSerialize function that does not require the package map to serialize its state.
 		WithGetPreloadDependencies     = false,                         // struct has a GetPreloadDependencies function to return all objects that will be Preload()ed when the struct is serialized at load time.
 		WithPureVirtual                = false,                         // struct has PURE_VIRTUAL functions and cannot be constructed when CHECK_PUREVIRTUALS is true
+		WithFindInnerPropertyInstance  = false,							// struct has a FindInnerPropertyInstance function that can provide an FProperty and data pointer when given a property FName
 		WithCanEditChange			   = false,							// struct has an editor-only CanEditChange function that can conditionally make child properties read-only in the details panel (same idea as UObject::CanEditChange)
 	};
 
@@ -929,6 +930,7 @@ public:
 			bool HasStructuredSerializeFromMismatchedTag : 1;
 			bool HasGetTypeHash : 1;
 			bool IsAbstract : 1;
+			bool HasFindInnerPropertyInstance : 1;
 #if WITH_EDITOR
 			bool HasCanEditChange : 1;
 #endif
@@ -1106,6 +1108,12 @@ public:
 		 */
 		virtual bool ImportTextItem(const TCHAR*& Buffer, void* Data, int32 PortFlags, class UObject* OwnerObject, FOutputDevice* ErrorText) = 0;
 
+		bool HasFindInnerPropertyInstance() const
+		{
+			return GetCapabilities().HasFindInnerPropertyInstance;
+		}
+		virtual bool FindInnerPropertyInstance(FName PropertyName, const void* Data, const FProperty*& OutProp, const void*& OutData) const = 0;
+
 		/** return true if this struct has custom GC code **/
 		bool HasAddStructReferencedObjects() const
 		{
@@ -1219,6 +1227,7 @@ public:
 				TTraits::WithStructuredSerializeFromMismatchedTag,
 				TModels_V<CGetTypeHashable, CPPSTRUCT>,
 				TIsAbstract<CPPSTRUCT>::Value,
+				TTraits::WithFindInnerPropertyInstance,
 #if WITH_EDITOR
 				TTraits::WithCanEditChange,
 #endif
@@ -1425,6 +1434,20 @@ public:
 				return false;
 			}
 		}
+		
+		virtual bool FindInnerPropertyInstance(FName PropertyName, const void* Data, const FProperty*& OutProp, const void*& OutData) const override
+		{
+			check(TTraits::WithFindInnerPropertyInstance); // don't call this if we have indicated it is not necessary
+			if constexpr (TStructOpsTypeTraits<CPPSTRUCT>::WithFindInnerPropertyInstance)
+            {
+            	return ((CPPSTRUCT*)Data)->FindInnerPropertyInstance(PropertyName, OutProp, OutData);
+            }
+            else
+            {
+            	return false;
+            }
+		}
+
 		virtual TPointerToAddStructReferencedObjects AddStructReferencedObjects() override
 		{
 			check(TTraits::WithAddStructReferencedObjects); // don't call this if we have indicated it is not necessary
@@ -1737,6 +1760,17 @@ public:
 	 * @param InStructData		The memory location to initialize
 	 */
 	virtual COREUOBJECT_API void InitializeDefaultValue(uint8* InStructData) const;
+	
+	/**
+	 * Provide an FProperty and data pointer when given an FName of a property that this object owns indirectly (ie FInstancedStruct)
+	 * @param PropertyName	Name of the property to retrieve
+	 * @param Data			Pointer to struct to retrieve property from
+	 * @param OutProp		Filled with the found property
+	 * @param OutData		Filled with a pointer to the data containing the property.
+	 *						Not offset to property instance yet - Call ContainerPtrToValuePtr before use.
+	 * @return				whether the property instance was found
+	 */
+	virtual COREUOBJECT_API bool FindInnerPropertyInstance(FName PropertyName, const void* Data, const FProperty*& OutProp, const void*& OutData) const;
 };
 
 /*-----------------------------------------------------------------------------
diff --git a/Engine/Source/Runtime/CoreUObject/Public/UObject/UnrealType.h b/Engine/Source/Runtime/CoreUObject/Public/UObject/UnrealType.h
--- a/Engine/Source/Runtime/CoreUObject/Public/UObject/UnrealType.h
+++ b/Engine/Source/Runtime/CoreUObject/Public/UObject/UnrealType.h
@@ -5848,6 +5848,7 @@ public:
 	// End of FProperty interface
 
 	bool UseBinaryOrNativeSerialization(const FArchive& Ar) const;
+	bool FindInnerPropertyInstance(FName PropertyName, const void* Data, const FProperty*& OutProp, const void*& OutData) const;
 
 private:
 	virtual uint32 GetValueTypeHashInternal(const void* Src) const;
