When releasing export bundle process node for a set of depedent packages, release non zen packages containing classes prior to non zen packages that do not.

diff --git a/Engine/Source/Runtime/CoreUObject/Private/Serialization/AsyncLoading2.cpp b/Engine/Source/Runtime/CoreUObject/Private/Serialization/AsyncLoading2.cpp
--- a/Engine/Source/Runtime/CoreUObject/Private/Serialization/AsyncLoading2.cpp
+++ b/Engine/Source/Runtime/CoreUObject/Private/Serialization/AsyncLoading2.cpp
@@ -5,6 +5,7 @@
 =============================================================================*/
 
 #include "Serialization/AsyncLoading2.h"
+#include "Algo/Partition.h"
 #include "IO/IoDispatcher.h"
 #include "Serialization/AsyncPackageLoader.h"
 #include "IO/PackageStore.h"
@@ -2649,6 +2650,7 @@ private:
 		int32 MetaDataIndex = -1;
 #endif
 		bool bIsCurrentlyProcessingImports = false;
+		bool bContainsClasses = false;
 
 		FAsyncPackageLinkerLoadHeaderData LinkerLoadHeaderData;
 	};
@@ -5147,6 +5149,11 @@ EEventLoadNodeExecutionResult FAsyncPackage2::ProcessLinkerLoadPackageSummary(FA
 	for (int32 ExportIndex = 0; ExportIndex < LinkerLoadState->Linker->ExportMap.Num(); ++ExportIndex)
 	{
 		FObjectExport& ObjectExport = LinkerLoadState->Linker->ExportMap[ExportIndex];
+		if ((ObjectExport.ObjectFlags & RF_ClassDefaultObject) != 0)
+		{
+			LinkerLoadState->bContainsClasses |= true;
+		}
+
 		//if ((ObjectExport.ObjectFlags & RF_Public) > 0) // We need hashes for all objects, external actors are breaking assumptions about public exports
 		{
 			TArray<int32, TInlineAllocator<128>> FullPath;
@@ -6509,6 +6516,42 @@ void FAsyncPackage2::WaitForAllDependenciesToReachState(FAsyncLoadingThreadState
 void FAsyncPackage2::ConditionalBeginProcessPackageExports(FAsyncLoadingThreadState2& ThreadState)
 {
 	TRACE_CPUPROFILER_EVENT_SCOPE(ConditionalBeginProcessPackageExports);
+#if ALT2_ENABLE_LINKERLOAD_SUPPORT
+	// When releasing export process nodes for packages when supporting linker loader state
+	// partition the set of related packages so that they are released in this order:
+	// zen packages -> linker load packages containing classes -> other linker load packages
+	TArray<FAsyncPackage2*, TInlineAllocator<16>> PendingExportProcessToRelease;
+	WaitForAllDependenciesToReachState(ThreadState, &FAsyncPackage2::AllDependenciesSetupState, EAsyncPackageLoadingState2::DependenciesReady, AsyncLoadingThread.ConditionalBeginProcessExportsTick,
+		[&ThreadState, &PendingExportProcessToRelease](FAsyncPackage2* Package)
+		{
+			if (Package->LinkerLoadState.IsSet())
+			{
+				PendingExportProcessToRelease.Add(Package);
+			}
+			else
+			{
+				for (int32 ExportBundleIndex = 0; ExportBundleIndex < Package->Data.TotalExportBundleCount; ++ExportBundleIndex)
+				{
+					Package->GetExportBundleNode(EEventLoadNode2::ExportBundle_Process, ExportBundleIndex).ReleaseBarrier(&ThreadState);
+				}
+			}
+		});
+	Algo::Partition(PendingExportProcessToRelease, [](FAsyncPackage2* Package)
+		{
+			if (Package->LinkerLoadState->bContainsClasses)
+			{
+				return true;
+			}
+			return false;
+		});
+	for (FAsyncPackage2* Package : PendingExportProcessToRelease)
+	{
+		for (int32 ExportBundleIndex = 0; ExportBundleIndex < Package->Data.TotalExportBundleCount; ++ExportBundleIndex)
+		{
+			Package->GetExportBundleNode(EEventLoadNode2::ExportBundle_Process, ExportBundleIndex).ReleaseBarrier(&ThreadState);
+		}
+	}
+#else
 	WaitForAllDependenciesToReachState(ThreadState, &FAsyncPackage2::AllDependenciesSetupState, EAsyncPackageLoadingState2::DependenciesReady, AsyncLoadingThread.ConditionalBeginProcessExportsTick,
 		[&ThreadState](FAsyncPackage2* Package)
 		{
@@ -6517,6 +6560,7 @@ void FAsyncPackage2::ConditionalBeginProcessPackageExports(FAsyncLoadingThreadSt
 				Package->GetExportBundleNode(EEventLoadNode2::ExportBundle_Process, ExportBundleIndex).ReleaseBarrier(&ThreadState);
 			}
 		});
+#endif
 }
 
 void FAsyncPackage2::ConditionalFinishLoading(FAsyncLoadingThreadState2& ThreadState)
