Add additional checks for object hash errors to be clearer about possible memory corruption.

diff --git a/Engine/Source/Runtime/CoreUObject/Private/UObject/UObjectHash.cpp b/Engine/Source/Runtime/CoreUObject/Private/UObject/UObjectHash.cpp
--- a/Engine/Source/Runtime/CoreUObject/Private/UObject/UObjectHash.cpp
+++ b/Engine/Source/Runtime/CoreUObject/Private/UObject/UObjectHash.cpp
@@ -29,6 +29,19 @@ LLM_DEFINE_TAG(UObjectHash, TEXT("UObject hashtables"));
 // Global UObject array instance
 FUObjectArray GUObjectArray;
 
+static FORCENOINLINE void OnHashFailure(UObjectBaseUtility* Object, const TCHAR* HashName, const TCHAR* FailureKind)
+{
+	UE_LOG(LogUObjectHash, Error, TEXT("UObject %s consistency failure (%s). Checking for memory corruption"), HashName, FailureKind);
+	UE_CLOG(!Object->IsValidLowLevelFast(false), LogUObjectHash, Fatal, TEXT("IsValidLowLevelFast failure"));
+	UE_CLOG(!Object->IsValidLowLevel(), LogUObjectHash, Fatal, TEXT("IsValid failure"));
+	UObjectBaseUtility* Outer = Object->GetOuter();
+	UE_CLOG(Outer && !Outer->IsValidLowLevelFast(false), LogUObjectHash, Fatal, TEXT("Outer IsValidLowLevelFast failure"));
+	UObjectBaseUtility* Class = Object->GetClass();
+	UE_CLOG(Class && !Class->IsValidLowLevelFast(false), LogUObjectHash, Fatal, TEXT("Class IsvalidLowLevelFast failure"));
+	
+	UE_LOG(LogUObjectHash, Fatal, TEXT("Unidentified failure for object %s, hash itself may be corrupted or buggy."), *Object->GetFullName());
+}
+
 /**
  * This implementation will use more space than the UE3 implementation. The goal was to make UObjects smaller to save L2 cache space. 
  * The hash is rarely used at runtime. A more space-efficient implementation is possible.
@@ -983,7 +996,10 @@ FORCEINLINE static void RemoveFromOuterMap(FUObjectHashTables& ThreadHash, UObje
 	FHashBucket& Bucket = ThreadHash.ObjectOuterMap.FindOrAdd(Object->GetOuter());
 	int32 NumRemoved = Bucket.Remove(Object);
 
-	UE_CLOG(NumRemoved != 1, LogUObjectHash, Fatal, TEXT("Internal Error: RemoveFromOuterMap NumRemoved = %d  for %s"), NumRemoved, *GetFullNameSafe((UObjectBaseUtility*)Object));
+	if (!LIKELY(NumRemoved == 1))
+	{
+		OnHashFailure((UObjectBaseUtility*)Object, TEXT("OuterMap"), TEXT("remove miscount"));
+	}
 
 	if (!Bucket.Num())
 	{
@@ -1000,8 +1016,10 @@ FORCEINLINE static void RemoveFromClassMap(FUObjectHashTables& ThreadHash, UObje
 		FHashBucket& ObjectList = ThreadHash.ClassToObjectListMap.FindOrAdd(Object->GetClass());
 		int32 NumRemoved = ObjectList.Remove(Object);
 
-		// must have existed, else something is wrong with the external code
-		UE_CLOG(NumRemoved != 1, LogUObjectHash, Fatal, TEXT("Internal Error: RemoveFromClassMap NumRemoved = %d  for %s"), NumRemoved, *GetFullNameSafe(ObjectWithUtility));
+		if (!LIKELY(NumRemoved == 1))
+		{
+			OnHashFailure((UObjectBaseUtility*)Object, TEXT("ClassMap"), TEXT("remove miscount"));
+		}
 
 		if (!ObjectList.Num())
 		{
@@ -1019,8 +1037,10 @@ FORCEINLINE static void RemoveFromClassMap(FUObjectHashTables& ThreadHash, UObje
 			TSet<UClass*>& ChildList = ThreadHash.ClassToChildListMap.FindOrAdd(SuperClass);
 			int32 NumRemoved = ChildList.Remove(Class);
 
-			// must have existed, else something is wrong with the external code
-			UE_CLOG(NumRemoved != 1, LogUObjectHash, Fatal, TEXT("Internal Error: RemoveFromClassMap NumRemoved = %d  for %s"), NumRemoved, *GetFullNameSafe(ObjectWithUtility));
+			if (!LIKELY(NumRemoved == 1))
+			{
+				OnHashFailure((UObjectBaseUtility*)Object, TEXT("ClassToChildListMap"), TEXT("remove miscount"));
+			}
 
 			if (!ChildList.Num())
 			{
@@ -1038,7 +1058,10 @@ FORCEINLINE static void RemoveFromPackageMap(FUObjectHashTables& ThreadHash, UOb
 	FHashBucket& Bucket = ThreadHash.PackageToObjectListMap.FindOrAdd(Package);
 	int32 NumRemoved = Bucket.Remove(Object);
 
-	UE_CLOG(NumRemoved != 1, LogUObjectHash, Fatal, TEXT("Internal Error: RemoveFromPackageMap NumRemoved = %d  for %s"), NumRemoved, *GetFullNameSafe((UObjectBaseUtility*)Object));
+	if (!LIKELY(NumRemoved == 1))
+	{
+		OnHashFailure((UObjectBaseUtility*)Object, TEXT("PackageMap"), TEXT("remove miscount"));
+	}
 
 	if (!Bucket.Num())
 	{
@@ -1482,8 +1505,10 @@ void HashObject(UObjectBase* Object)
 
 		Hash = GetObjectHash(Name);
 #if !UE_BUILD_TEST && !UE_BUILD_SHIPPING
-		// if it already exists, something is wrong with the external code
-		UE_CLOG(ThreadHash.PairExistsInHash(Hash, Object), LogUObjectHash, Fatal, TEXT("%s already exists in UObject hash!"), *GetFullNameSafe((UObjectBaseUtility*)Object));
+		if (ThreadHash.PairExistsInHash(Hash, Object))
+		{
+			OnHashFailure((UObjectBaseUtility*)Object, TEXT("Hash"), TEXT("double add"));
+		}
 #endif
 		ThreadHash.AddToHash(Hash, Object);
 
@@ -1492,8 +1517,10 @@ void HashObject(UObjectBase* Object)
 			Hash = GetObjectOuterHash(Name, Outer);
 			checkSlow(!ThreadHash.HashOuter.FindPair(Hash, Object->GetUniqueID()));
 #if !UE_BUILD_TEST && !UE_BUILD_SHIPPING
-			// if it already exists, something is wrong with the external code
-			UE_CLOG(ThreadHash.HashOuter.FindPair(Hash, Object->GetUniqueID()), LogUObjectHash, Fatal, TEXT("%s already exists in UObject Outer hash!"), *GetFullNameSafe((UObjectBaseUtility*)Object));
+			if (ThreadHash.HashOuter.FindPair(Hash, Object->GetUniqueID()))
+			{
+				OnHashFailure((UObjectBaseUtility*)Object, TEXT("HashOuter"), TEXT("double add"));
+			}
 #endif
 			ThreadHash.HashOuter.Add(Hash, Object->GetUniqueID());
 
@@ -1527,16 +1554,20 @@ void UnhashObject(UObjectBase* Object)
 		Hash = GetObjectHash(Name);
 		NumRemoved = ThreadHash.RemoveFromHash(Hash, Object);
 
-		// must have existed, else something is wrong with the external code
-		UE_CLOG(NumRemoved != 1, LogUObjectHash, Fatal, TEXT("Internal Error: RemoveFromHash NumRemoved = %d  for %s"), NumRemoved, *GetFullNameSafe((UObjectBaseUtility*)Object));
+		if (!LIKELY(NumRemoved == 1))
+		{
+			OnHashFailure((UObjectBaseUtility*)Object, TEXT("Hash"), TEXT("remove miscount"));
+		}
 
 		if (PTRINT Outer = (PTRINT)Object->GetOuter())
 		{
 			Hash = GetObjectOuterHash(Name, Outer);
 			NumRemoved = ThreadHash.HashOuter.RemoveSingle(Hash, Object->GetUniqueID());
 
-			// must have existed, else something is wrong with the external code
-			UE_CLOG(NumRemoved != 1, LogUObjectHash, Fatal, TEXT("Internal Error: Remove from HashOuter NumRemoved = %d  for %s"), NumRemoved, *GetFullNameSafe((UObjectBaseUtility*)Object));
+			if (!LIKELY(NumRemoved == 1))
+			{
+				OnHashFailure((UObjectBaseUtility*)Object, TEXT("HashOuter"), TEXT("remove miscount"));
+			}
 
 			RemoveFromOuterMap(ThreadHash, Object);
 		}
