[autortfm] add initial support for UE_USE_VIRTUAL_STACK_ALLOCATOR_FOR_SCRIPT_VM - when transactional we bypass the stack allocator and fall back to FMemory. To be replaced later by returning to stack allocator behavior after we have runtime support for it

diff --git a/Engine/Source/Runtime/CoreUObject/Public/UObject/Stack.h b/Engine/Source/Runtime/CoreUObject/Public/UObject/Stack.h
--- a/Engine/Source/Runtime/CoreUObject/Public/UObject/Stack.h
+++ b/Engine/Source/Runtime/CoreUObject/Public/UObject/Stack.h
@@ -20,7 +20,42 @@ DECLARE_LOG_CATEGORY_EXTERN(LogScriptFrame, Warning, All);
 #endif 
 
 #if UE_USE_VIRTUAL_STACK_ALLOCATOR_FOR_SCRIPT_VM
-FORCEINLINE void* UeVstackAllocHelper(FVirtualStackAllocator* Allocator, size_t Size, size_t Align) { return ((Size == 0) ? 0 : Allocator->Allocate(Size, (Align < 16) ? 16 : Align)); }
+FORCEINLINE void* UeVstackAllocHelper(FVirtualStackAllocator* Allocator, size_t Size, size_t Align)
+{
+	if (Size == 0)
+		return nullptr;
+
+	void* Result;
+
+	if (!AutoRTFM::IsClosed())
+	{
+		Result = Allocator->Allocate(Size, (Align < 16) ? 16 : Align);
+	}
+	else
+	{
+		// in transactional code we redirect the virtual stack allocator to FMemory::Malloc
+		// until we properly implement a transaction-friendly stack allocator
+		AutoRTFM::Open([&Result, Size, Align]()
+		{
+			Result = FMemory::Malloc(Size, Align);
+		});
+
+		// for these 'stack' allocations we call Free on both commit and abort, since
+		// the call frame that these would have otherwise existed in will always be
+		// returned from by the time we commit or abort
+		AutoRTFM::OpenCommit([Result]()
+		{
+			FMemory::Free(Result);
+		});
+
+		AutoRTFM::OpenAbort([Result]()
+		{
+			FMemory::Free(Result);
+		});
+	}
+
+	return Result;
+}
 #define UE_VSTACK_MAKE_FRAME(Name, VirtualStackAllocatorPtr) FScopedStackAllocatorBookmark Name = VirtualStackAllocatorPtr->CreateScopedBookmark()
 #define UE_VSTACK_ALLOC(VirtualStackAllocatorPtr, Size) UeVstackAllocHelper((VirtualStackAllocatorPtr), (Size), 0) 
 #define UE_VSTACK_ALLOC_ALIGNED(VirtualStackAllocatorPtr, Size, Align) UeVstackAllocHelper((VirtualStackAllocatorPtr), (Size), (Align))
