Added serial number to archetype cache to prevent returniong stale pointers.

diff --git a/Engine/Source/Runtime/CoreUObject/Private/UObject/UObjectArchetype.cpp b/Engine/Source/Runtime/CoreUObject/Private/UObject/UObjectArchetype.cpp
--- a/Engine/Source/Runtime/CoreUObject/Private/UObject/UObjectArchetype.cpp
+++ b/Engine/Source/Runtime/CoreUObject/Private/UObject/UObjectArchetype.cpp
@@ -20,11 +20,12 @@
 struct FArchetypeInfo
 {
 	/**
-	* default contructor
+	* default constructor
 	* Default constructor must be the default item
 	*/
 	FArchetypeInfo()
 		: ArchetypeIndex(INDEX_NONE)
+		, SerialNumber(INDEX_NONE)
 	{
 	}
 	/**
@@ -40,12 +41,14 @@ struct FArchetypeInfo
 	* Constructor
 	* @param InArchetype Archetype to assign
 	*/
-	FArchetypeInfo(int32 InArchetypeIndex)
+	FArchetypeInfo(int32 InArchetypeIndex, int32 InSerialNumber)
 		: ArchetypeIndex(InArchetypeIndex)
+		, SerialNumber(InSerialNumber)
 	{
 	}
 
 	int32 ArchetypeIndex;
+	int32 SerialNumber;
 };
 
 static FUObjectAnnotationChunked<FArchetypeInfo, true, 8192> ArchetypeAnnotation;
@@ -147,7 +150,8 @@ void CacheArchetypeForObject(UObject* Object, UObject* Archetype)
 	UObject* VerifyArchetype = GetArchetypeFromRequiredInfoImpl(Object->GetClass(), Object->GetOuter(), Object->GetFName(), Object->GetFlags(), bUseUpToDateClass);
 	checkf(Archetype == VerifyArchetype, TEXT("Cached archetype mismatch, expected: %s, cached: %s"), *GetFullNameSafe(VerifyArchetype), *GetFullNameSafe(Archetype));
 #endif
-	ArchetypeAnnotation.AddAnnotation(Object, GUObjectArray.ObjectToIndex(Archetype));
+	int32 ArchetypeIndex = GUObjectArray.ObjectToIndex(Archetype);
+	ArchetypeAnnotation.AddAnnotation(Object, FArchetypeInfo{ ArchetypeIndex, GUObjectArray.AllocateSerialNumber(ArchetypeIndex) });
 #endif
 }
 
@@ -170,16 +174,19 @@ UObject* UObject::GetArchetype() const
 
 #if UE_CACHE_ARCHETYPE
 	UObject* Archetype = nullptr;
-	int32 ArchetypeIndex = ArchetypeAnnotation.GetAnnotation(this).ArchetypeIndex;
-	if (ArchetypeIndex == INDEX_NONE)
+	FArchetypeInfo Annoatation = ArchetypeAnnotation.GetAnnotation(this);
+	int32 ArchetypeIndex = Annoatation.ArchetypeIndex;
+	int32 SerialNumber = ArchetypeIndex == INDEX_NONE ? INDEX_NONE : GUObjectArray.GetSerialNumber(ArchetypeIndex);
+	if ((ArchetypeIndex == INDEX_NONE) || (SerialNumber != Annoatation.SerialNumber))
 	{
 		Archetype = GetArchetypeFromRequiredInfo(GetClass(), GetOuter(), GetFName(), GetFlags());
 		// If the Outer is pending load we can't cache the archetype as it may be inacurate
 		if (Archetype && !(GetOuter() && GetOuter()->HasAnyFlags(RF_NeedLoad)))
 		{
-			ArchetypeAnnotation.AddAnnotation(this, GUObjectArray.ObjectToIndex(Archetype));
+			ArchetypeIndex = GUObjectArray.ObjectToIndex(Archetype);
+			ArchetypeAnnotation.AddAnnotation(this, FArchetypeInfo{ ArchetypeIndex, GUObjectArray.AllocateSerialNumber(ArchetypeIndex) });
 		}
-	}		
+	}
 	else
 	{
 		FUObjectItem* ArchetypeItem = GUObjectArray.IndexToObject(ArchetypeIndex);
