Add missing PlayInEditorID scope in zenloader

#rnx

diff --git a/Engine/Source/Runtime/CoreUObject/Private/Serialization/AsyncLoading2.cpp b/Engine/Source/Runtime/CoreUObject/Private/Serialization/AsyncLoading2.cpp
--- a/Engine/Source/Runtime/CoreUObject/Private/Serialization/AsyncLoading2.cpp
+++ b/Engine/Source/Runtime/CoreUObject/Private/Serialization/AsyncLoading2.cpp
@@ -5451,6 +5451,10 @@ EEventLoadNodeExecutionResult FAsyncPackage2::Event_ProcessPackageSummary(FAsync
 
 	FAsyncPackageScope2 Scope(Package);
 
+#if WITH_EDITOR
+	UE::Core::Private::FPlayInEditorLoadingScope PlayInEditorIDScope(Package->Desc.PIEInstanceID);
+#endif
+
 #if ALT2_ENABLE_LINKERLOAD_SUPPORT
 	if (Package->LinkerLoadState.IsSet())
 	{
