Fix for an editor crash when editing sparse class data fields in Blueprint class defaults.

diff --git a/Engine/Source/Runtime/CoreUObject/Private/UObject/Obj.cpp b/Engine/Source/Runtime/CoreUObject/Private/UObject/Obj.cpp
--- a/Engine/Source/Runtime/CoreUObject/Private/UObject/Obj.cpp
+++ b/Engine/Source/Runtime/CoreUObject/Private/UObject/Obj.cpp
@@ -576,6 +576,10 @@ void UObject::PropagatePreEditChange( TArray<UObject*>& AffectedObjects, FEditPr
 	const FProperty* ChangedProperty = PropertyAboutToChange.GetActiveMemberNode()->GetValue();
 	FPropagatedEditChangeAnnotation Annotation = PropagatedEditChangeAnnotation.GetAnnotation(this);
 
+	// Determine if the changed property belongs to the archetype's class type (or a parent class).
+	// Note: IsChildOf() returns false for a NULL owner class (i.e. non-class struct type changes).
+	const bool bIsArchetypePropertyChange = GetClass()->IsChildOf(ChangedProperty->GetOwnerClass());
+
 	for ( int32 i = 0; i < Instances.Num(); i++ )
 	{
 		UObject* Obj = Instances[i];
@@ -590,14 +594,21 @@ void UObject::PropagatePreEditChange( TArray<UObject*>& AffectedObjects, FEditPr
 			// that case, we won't need to dirty the package after applying the change.
 			Annotation.bIdenticalToArchetype = true;
 
-			// Note that some elements may match and thus will propagate, but we may
-			// need to dirty the package later even if only one element differs here.
-			for (int32 ArrayIdx = 0; ArrayIdx < ChangedProperty->ArrayDim; ++ArrayIdx)
+			// If the property that was changed is not a member of the archetype's class
+			// type, then it means we're propagating a change event to instances without
+			// having also propagated the value change. Thus, there's nothing to compare
+			// in this case since we're not inferring dirty state from a property value.
+			if (bIsArchetypePropertyChange)
 			{
-				if (!ChangedProperty->Identical_InContainer(this, Obj, ArrayIdx, PPF_DeepComparison))
+				// Note that some elements may match and thus will propagate, but we may
+				// need to dirty the package later even if only one element differs here.
+				for (int32 ArrayIdx = 0; ArrayIdx < ChangedProperty->ArrayDim; ++ArrayIdx)
 				{
-					Annotation.bIdenticalToArchetype = false;
-					break;
+					if (!ChangedProperty->Identical_InContainer(this, Obj, ArrayIdx, PPF_DeepComparison))
+					{
+						Annotation.bIdenticalToArchetype = false;
+						break;
+					}
 				}
 			}
 
