Add FReferenceCollector::AddPropertyReferencesWithStructARO class overload

#rnx

diff --git a/Engine/Source/Runtime/CoreUObject/Private/UObject/UObjectGlobals.cpp b/Engine/Source/Runtime/CoreUObject/Private/UObject/UObjectGlobals.cpp
--- a/Engine/Source/Runtime/CoreUObject/Private/UObject/UObjectGlobals.cpp
+++ b/Engine/Source/Runtime/CoreUObject/Private/UObject/UObjectGlobals.cpp
@@ -4747,7 +4747,7 @@ static void CollectStructReferences(FReferenceCollector& Collector, const Struct
 
 	FPlatformMisc::PrefetchBlock(Struct->RefLink, PropertyPrefetchBytes);
 	
-	if constexpr (EnumHasAnyFlags(CollectFlags, EPropertyCollectFlags::CallStructARO))
+	if constexpr (EnumHasAnyFlags(CollectFlags, EPropertyCollectFlags::CallStructARO) && std::is_same_v<StructType, UScriptStruct>)
 	{
 		if (Struct->StructFlags & STRUCT_AddStructReferencedObjects)
 		{
@@ -5018,6 +5018,11 @@ void FReferenceCollector::AddPropertyReferencesWithStructARO(const UScriptStruct
 	CallCollectStructReferences<EPropertyCollectFlags::CallStructARO>(*this, Struct, Instance, ReferencingObject);
 }
 
+void FReferenceCollector::AddPropertyReferencesWithStructARO(const UClass* Class, void* Instance, const UObject* ReferencingObject)
+{
+	CallCollectStructReferences<EPropertyCollectFlags::CallStructARO>(*this, Class, Instance, ReferencingObject);
+}
+
 void FReferenceCollector::AddPropertyReferencesLimitedToObjectProperties(const UStruct* Struct, void* Instance, const UObject* ReferencingObject)
 {
 	CallCollectStructReferences<EPropertyCollectFlags::OnlyObjectProperty>(*this, Struct, Instance, ReferencingObject);
diff --git a/Engine/Source/Runtime/CoreUObject/Public/UObject/ReferenceCollector.h b/Engine/Source/Runtime/CoreUObject/Public/UObject/ReferenceCollector.h
--- a/Engine/Source/Runtime/CoreUObject/Public/UObject/ReferenceCollector.h
+++ b/Engine/Source/Runtime/CoreUObject/Public/UObject/ReferenceCollector.h
@@ -376,6 +376,9 @@ public:
 	/** Same as AddPropertyReferences but also calls AddStructReferencedObjects on Struct and all nested structs */
 	void AddPropertyReferencesWithStructARO(const UScriptStruct* Struct, void* Instance, const UObject* ReferencingObject = nullptr);
 
+	/** Same as AddPropertyReferences but also calls AddStructReferencedObjects on all nested structs */
+	void AddPropertyReferencesWithStructARO(const UClass* Class, void* Instance, const UObject* ReferencingObject = nullptr);
+
 	/** Internal use only. Same as AddPropertyReferences but skips field path and interface properties. Might get removed. */
 	void AddPropertyReferencesLimitedToObjectProperties(const UStruct* Struct, void* Instance, const UObject* ReferencingObject = nullptr);
 
