Add FReferenceCollector::AddPropertyReferencesWithStructARO class overload

#rnx

diff --git a/Engine/Source/Runtime/CoreUObject/Public/UObject/UObjectGlobals.h b/Engine/Source/Runtime/CoreUObject/Public/UObject/UObjectGlobals.h
--- a/Engine/Source/Runtime/CoreUObject/Public/UObject/UObjectGlobals.h
+++ b/Engine/Source/Runtime/CoreUObject/Public/UObject/UObjectGlobals.h
@@ -2543,6 +2543,9 @@ public:
 	/** Same as AddPropertyReferences but also calls AddStructReferencedObjects on Struct and all nested structs */
 	void AddPropertyReferencesWithStructARO(const UScriptStruct* Struct, void* Instance, const UObject* ReferencingObject = nullptr);
 
+	/** Same as AddPropertyReferences but also calls AddStructReferencedObjects on all nested structs */
+	void AddPropertyReferencesWithStructARO(const UClass* Class, void* Instance, const UObject* ReferencingObject = nullptr);
+
 	/** Internal use only. Same as AddPropertyReferences but skips field path and interface properties. Might get removed. */
 	void AddPropertyReferencesLimitedToObjectProperties(const UStruct* Struct, void* Instance, const UObject* ReferencingObject = nullptr);
 
