Added extra checks to help pinpoint potential GC schema memory stomp or use-after-free found once in CI tests

#rnx

diff --git a/Engine/Source/Runtime/CoreUObject/Private/UObject/GarbageCollection.cpp b/Engine/Source/Runtime/CoreUObject/Private/UObject/GarbageCollection.cpp
--- a/Engine/Source/Runtime/CoreUObject/Private/UObject/GarbageCollection.cpp
+++ b/Engine/Source/Runtime/CoreUObject/Private/UObject/GarbageCollection.cpp
@@ -1745,8 +1745,16 @@ struct FStructArrayWithoutStride
 	int32 Num;
 };
 
+FORCEINLINE_DEBUGGABLE void CheckValid(FSchemaView Schema)
+{
+	check(!Schema.IsEmpty());
+	check(Schema.GetHeader().RefCount.load(std::memory_order_relaxed) > 0); // Memory stomp or using destroyed schema
+	check(Schema.GetHeader().StructStride % 8 == 0); // Memory stomp or using destroyed schema
+}
+
 static FStructArray CopyAndLoadStride(FStructArrayWithoutStride In)
 {
+	CheckValid(In.Schema);
 	return {In.Schema, In.Data, In.Num, In.Schema.GetStructStride() };
 }
 
@@ -1762,7 +1770,7 @@ public:
 
 	FORCEINLINE_DEBUGGABLE void PushSparseStructArray(FSchemaView Schema, FScriptSparseArray& Array)
 	{
-		check(!Schema.IsEmpty());
+		CheckValid(Schema);
 
 		UnvalidatedSparseStructArrays.Push({ Schema, &Array, Array.Num(), Schema.GetStructStride() });
 		if (UnvalidatedSparseStructArrays.IsFull())
@@ -1773,7 +1781,7 @@ public:
 	
 	FORCEINLINE_DEBUGGABLE void PushStructArray(FSchemaView Schema, uint8* Data, int32 Num)
 	{
-		check(!Schema.IsEmpty());
+		CheckValid(Schema);
 
 		UnvalidatedStructArrays.Push({Schema, Data, Num});
 		if (UnvalidatedStructArrays.IsFull())
@@ -3367,7 +3375,8 @@ void TFastReferenceCollector<ProcessorType, CollectorType>::ProcessStructs(Dispa
 
 				const FStructArray AoS = *ArrayIt;
 				check(AoS.Num > 0);
-				check(AoS.Stride % 8 == 0);
+				CheckValid(AoS.Schema);
+				check(AoS.Stride == AoS.Schema.GetStructStride());
 				uint8* StructIt = AoS.Data;
 				uint8* StructEnd = AoS.Data + AoS.Num * AoS.Stride;
 				do
diff --git a/Engine/Source/Runtime/CoreUObject/Private/UObject/GarbageCollectionSchema.cpp b/Engine/Source/Runtime/CoreUObject/Private/UObject/GarbageCollectionSchema.cpp
--- a/Engine/Source/Runtime/CoreUObject/Private/UObject/GarbageCollectionSchema.cpp
+++ b/Engine/Source/Runtime/CoreUObject/Private/UObject/GarbageCollectionSchema.cpp
@@ -276,7 +276,7 @@ FSchemaBuilder::FSchemaBuilder(uint32 InStride, std::initializer_list<FMemberDec
 	: Members(InMembers)
 	, StructStride(InStride) 
 {
-	check((StructStride % 8) == 0);
+	checkf(StructStride % 8 == 0, TEXT("Structs containing TObjectPtr/UObject* references must be 8-byte aligned"));
 
 	for (FMemberDeclaration& Member : Members)
 	{
@@ -286,6 +286,10 @@ FSchemaBuilder::FSchemaBuilder(uint32 InStride, std::initializer_list<FMemberDec
 
 FSchemaBuilder::~FSchemaBuilder()
 {
+	check(StructStride % 8 == 0); // Builder memory stomp
+	check(!BuiltSchema || BuiltSchema->Get().IsEmpty() || FSchemaView(BuiltSchema->Get()).GetHeader().RefCount.load() > 0); // Schema ref-count bug or header memory stomp
+	check(!BuiltSchema || BuiltSchema->Get().IsEmpty() || FSchemaView(BuiltSchema->Get()).GetHeader().StructStride == StructStride); // Header memory stomp
+
 	for (FMemberDeclaration& Member : Members)
 	{
 		TryDropSchemaReference(Member);
@@ -479,6 +483,7 @@ static FMemberDeclaration GenerateTerminator(ObjectAROFn ARO)
 
 FSchemaView FSchemaBuilder::Build(ObjectAROFn ARO)
 {
+	check(StructStride % 8 == 0); // Memory stomp or use-after-free
 	if (BuiltSchema.IsSet())
 	{
 		return BuiltSchema.GetValue().Get();
@@ -494,7 +499,7 @@ FSchemaView FSchemaBuilder::Build(ObjectAROFn ARO)
 
 	// Pack into words and generate jump instructions
 	FSchemaPacker Packed(Members);
-	check(Packed.DebugNames.IsEmpty() == !UE_GC_DEBUGNAMES) ;
+	check(Packed.DebugNames.IsEmpty() == !UE_GC_DEBUGNAMES);
 
 	// Allocate and initialize schema data
 	const SIZE_T Bytes = Align(sizeof(FSchemaHeader), sizeof(FMemberWord)) + sizeof(FMemberWord) * Packed.Words.Num() + sizeof(FName) * Packed.DebugNames.Num();
