EditorBulkData: Remove accidentally-submitted debug code.
#rnx

diff --git a/Engine/Source/Runtime/CoreUObject/Private/Serialization/EditorBulkData.cpp b/Engine/Source/Runtime/CoreUObject/Private/Serialization/EditorBulkData.cpp
--- a/Engine/Source/Runtime/CoreUObject/Private/Serialization/EditorBulkData.cpp
+++ b/Engine/Source/Runtime/CoreUObject/Private/Serialization/EditorBulkData.cpp
@@ -1684,11 +1684,6 @@ void FEditorBulkData::DetachFromDisk(FArchive* Ar, bool bEnsurePayloadIsLoaded)
 	check(Ar != nullptr);
 	check(Ar == AttachedAr || HasAttachedArchive() == false || AttachedAr->IsProxyOf(Ar));
 
-	if (bEnsurePayloadIsLoaded)
-	{
-		int a = 0; a++;
-	}
-
 	// If bEnsurePayloadIsLoaded is true, then we should assume that a change to the 
 	// package file is imminent and we should load the payload into memory if possible.
 	// If it is false then either we are not expecting a change to the package file or
