moved to #include PrintStaleReferencesOption to its own file for better compile iterations
https://p4-swarm.epicgames.net/reviews/25725441

diff --git a/Engine/Plugins/Experimental/GameFeatures/Source/GameFeatures/Private/GameFeaturePluginStateMachine.cpp b/Engine/Plugins/Experimental/GameFeatures/Source/GameFeatures/Private/GameFeaturePluginStateMachine.cpp
--- a/Engine/Plugins/Experimental/GameFeatures/Source/GameFeatures/Private/GameFeaturePluginStateMachine.cpp
+++ b/Engine/Plugins/Experimental/GameFeatures/Source/GameFeatures/Private/GameFeaturePluginStateMachine.cpp
@@ -25,6 +25,7 @@
 #include "String/ParseTokens.h"
 #include "GameFeaturesProjectPolicies.h"
 #include "UObject/ObjectRename.h"
+#include "UObject/ReferenceChainSearch.h"
 #include "UObject/UObjectAllocator.h"
 #include "UObject/UObjectIterator.h"
 #include "UObject/NoExportTypes.h"
diff --git a/Engine/Source/Editor/UnrealEd/Private/Cooker/IoStoreCookOnTheFlyRequestManager.cpp b/Engine/Source/Editor/UnrealEd/Private/Cooker/IoStoreCookOnTheFlyRequestManager.cpp
--- a/Engine/Source/Editor/UnrealEd/Private/Cooker/IoStoreCookOnTheFlyRequestManager.cpp
+++ b/Engine/Source/Editor/UnrealEd/Private/Cooker/IoStoreCookOnTheFlyRequestManager.cpp
@@ -24,6 +24,7 @@
 #include "Sockets.h"
 #include "SocketSubsystem.h"
 #include "Templates/Function.h"
+#include "UObject/ReferenceChainSearch.h"
 #include "UObject/SavePackage.h"
 #include "ZenStoreWriter.h"
 
diff --git a/Engine/Source/Runtime/CoreUObject/Public/UObject/PrintStaleReferencesOptions.h b/Engine/Source/Runtime/CoreUObject/Public/UObject/PrintStaleReferencesOptions.h
new file mode 100644
--- /dev/null
+++ b/Engine/Source/Runtime/CoreUObject/Public/UObject/PrintStaleReferencesOptions.h
@@ -0,0 +1,21 @@
+// Copyright Epic Games, Inc. All Rights Reserved.
+
+#pragma once
+
+/* Options for FReferenceChainSearch::FindAndPrintStaleReferencesToObject function */
+enum class EPrintStaleReferencesOptions
+{
+	None = 0,
+	Log = 1,
+	Display = 2,
+	Warning = 3,
+	Error = 4,
+	Fatal = 5,
+	Ensure = 0x00000100,
+
+	// Only search for direct references to the object or one of its inners, not a full reference chain 
+	Minimal = 0x00000200,
+
+	VerbosityMask = 0x000000ff
+};
+ENUM_CLASS_FLAGS(EPrintStaleReferencesOptions);
\ No newline at end of file
diff --git a/Engine/Source/Runtime/CoreUObject/Public/UObject/ReferenceChainSearch.h b/Engine/Source/Runtime/CoreUObject/Public/UObject/ReferenceChainSearch.h
--- a/Engine/Source/Runtime/CoreUObject/Public/UObject/ReferenceChainSearch.h
+++ b/Engine/Source/Runtime/CoreUObject/Public/UObject/ReferenceChainSearch.h
@@ -20,6 +20,7 @@
 #include "UObject/NameTypes.h"
 #include "UObject/ObjectMacros.h"
 #include "UObject/UObjectGlobals.h"
+#include "UObject/PrintStaleReferencesOptions.h"
 
 class FGCObjectInfo;
 class FOutputDevice;
@@ -59,24 +60,6 @@ enum class EReferenceChainSearchMode
 
 ENUM_CLASS_FLAGS(EReferenceChainSearchMode);
 
-/* Options for UEngine::FindAndPrintStaleReferencesToObject function */
-enum class EPrintStaleReferencesOptions
-{
-	None = 0,
-	Log = 1,
-	Display = 2,
-	Warning = 3,
-	Error = 4,
-	Fatal = 5,
-	Ensure = 0x00000100,
-
-	// Only search for direct references to the object or one of its inners, not a full reference chain 
-	Minimal = 0x00000200,
-
-	VerbosityMask = 0x000000ff
-};
-ENUM_CLASS_FLAGS(EPrintStaleReferencesOptions);
-
 class FReferenceChainSearch
 {
 	// Reference chain searching is a very slow operation.
diff --git a/Engine/Source/Runtime/Engine/Classes/Engine/Engine.h b/Engine/Source/Runtime/Engine/Classes/Engine/Engine.h
--- a/Engine/Source/Runtime/Engine/Classes/Engine/Engine.h
+++ b/Engine/Source/Runtime/Engine/Classes/Engine/Engine.h
@@ -4,7 +4,7 @@
 
 #include "CoreMinimal.h"
 #include "Containers/IndirectArray.h"
-#include "UObject/ReferenceChainSearch.h"
+#include "UObject/PrintStaleReferencesOptions.h"
 #include "UObject/ObjectMacros.h"
 #include "UObject/Object.h"
 #include "Misc/Guid.h"
diff --git a/Engine/Source/Runtime/Engine/Private/LevelStreaming.cpp b/Engine/Source/Runtime/Engine/Private/LevelStreaming.cpp
--- a/Engine/Source/Runtime/Engine/Private/LevelStreaming.cpp
+++ b/Engine/Source/Runtime/Engine/Private/LevelStreaming.cpp
@@ -17,6 +17,7 @@
 #include "LevelUtils.h"
 #include "EngineUtils.h"
 #include "UObject/UObjectAnnotation.h"
+#include "UObject/ReferenceChainSearch.h"
 #if WITH_EDITOR
 	#include "Framework/Notifications/NotificationManager.h"
 	#include "Widgets/Notifications/SNotificationList.h"
