fix for not setting up properties correctly.
this matches the uht generated code more closesly
https://p4-swarm.epicgames.net/reviews/25707323

diff --git a/Engine/Source/Runtime/CoreUObject/Tests/ObjectPtrTestClass.cpp b/Engine/Source/Runtime/CoreUObject/Tests/ObjectPtrTestClass.cpp
--- a/Engine/Source/Runtime/CoreUObject/Tests/ObjectPtrTestClass.cpp
+++ b/Engine/Source/Runtime/CoreUObject/Tests/ObjectPtrTestClass.cpp
@@ -77,36 +77,78 @@ IMPLEMENT_CORE_INTRINSIC_CLASS(UObjectPtrTestClassWithRef, UObject,
 #else
 IMPLEMENT_CORE_INTRINSIC_CLASS(UObjectPtrTestClassWithRef, UObject,
 	{
-		//add reflection info for the two properties
-		auto Property1 = new FObjectPtrProperty(Class, TEXT("ObjectPtr"), EObjectFlags::RF_NoFlags);
-		Property1->PropertyClass = UObjectPtrTestClass::StaticClass();
-		Class->AddCppProperty(Property1);
-
-		auto Property2 = new FObjectPtrProperty(Class, TEXT("ObjectPtrNonNullable"), EObjectFlags::RF_NoFlags);
-		Property2->PropertyClass = UObjectPtrTestClass::StaticClass();
-		Property2->SetPropertyFlags(EPropertyFlags::CPF_NonNullable); //make non nullable
-		Class->AddCppProperty(Property2);
+		{
+			UECodeGen_Private::FObjectPropertyParams Params = { };
+			Params.NameUTF8 = "ObjectPtr";
+			Params.Offset = STRUCT_OFFSET(UObjectPtrTestClassWithRef, ObjectPtr);
+			Params.PropertyFlags = CPF_None;
+			Params.ObjectFlags = RF_Public | RF_Transient | RF_MarkAsNative;
+			Params.ClassFunc = nullptr;
+			auto Property = new FObjectPtrProperty(Class, Params);
+			Property->PropertyClass = UObjectPtrTestClass::StaticClass();
+		}
+
+		{
+			UECodeGen_Private::FObjectPropertyParams Params = { };
+			Params.NameUTF8 = "ObjectPtrNonNullable";
+			Params.Offset = STRUCT_OFFSET(UObjectPtrTestClassWithRef, ObjectPtrNonNullable);
+			Params.PropertyFlags = EPropertyFlags::CPF_NonNullable;
+			Params.ObjectFlags = RF_Public | RF_Transient | RF_MarkAsNative;
+			Params.ClassFunc = nullptr;
+			auto Property = new FObjectPtrProperty(Class, Params);
+			Property->PropertyClass = UObjectPtrTestClass::StaticClass();
+		}
+
+		{
+			UECodeGen_Private::FArrayPropertyParams Params = { };
+			Params.NameUTF8 = "ArrayObjPtr";
+			Params.Offset = STRUCT_OFFSET(UObjectPtrTestClassWithRef, ArrayObjPtr);
+			Params.PropertyFlags = CPF_None;
+			Params.ObjectFlags = RF_Public | RF_Transient | RF_MarkAsNative;
+			auto Property = new FArrayProperty(Class, Params);
+			auto InnerProperty = new FObjectPtrProperty(Property, TEXT("Inner"), EObjectFlags::RF_NoFlags);
+			InnerProperty->PropertyClass = UObjectPtrTestClass::StaticClass();
+			Property->AddCppProperty(InnerProperty);
+		}
 	});
 #endif
 
 IMPLEMENT_CORE_INTRINSIC_CLASS(UObjectWithClassProperty, UObject,
 	{
-		auto Property1 = new FClassPtrProperty(Class, TEXT("ClassPtr"), EObjectFlags::RF_NoFlags);
-		Property1->PropertyClass = UClass::StaticClass();
-		Class->AddCppProperty(Property1);
+		UECodeGen_Private::FClassPropertyParams Params = { };
+		Params.NameUTF8 = "ClassPtr";
+		Params.Offset = STRUCT_OFFSET(UObjectWithClassProperty, ClassPtr);
+		Params.PropertyFlags = CPF_None;
+		Params.ObjectFlags = RF_Public | RF_Transient | RF_MarkAsNative;
+		Params.ClassFunc = nullptr;
+		auto Property = new FClassPtrProperty(Class, Params);
+		Property->PropertyClass = UObjectWithClassProperty::StaticClass();
 	}
 );
 
 IMPLEMENT_CORE_INTRINSIC_CLASS(UObjectWithRawProperty, UObject,
 	{
-		auto Property1 = new FObjectProperty(Class, TEXT("ObjectPtr"), EObjectFlags::RF_NoFlags);
-		Property1->PropertyClass = UClass::StaticClass();
-		Class->AddCppProperty(Property1);
-
-		auto Property2 = new FObjectPtrProperty(Class, TEXT("ObjectPtrNonNullable"), EObjectFlags::RF_NoFlags);
-		Property2->PropertyClass = UObjectPtrTestClass::StaticClass();
-		Property2->SetPropertyFlags(EPropertyFlags::CPF_NonNullable); //make non nullable
-		Class->AddCppProperty(Property2);
+		{
+			UECodeGen_Private::FObjectPropertyParams Params = { };
+			Params.NameUTF8 = "ObjectPtr";
+			Params.Offset = STRUCT_OFFSET(UObjectWithRawProperty, ObjectPtr);
+			Params.PropertyFlags = CPF_None;
+			Params.ObjectFlags = RF_Public | RF_Transient | RF_MarkAsNative;
+			Params.ClassFunc = nullptr;
+			auto Property = new FObjectProperty(Class, Params);
+			Property->PropertyClass = UObjectWithRawProperty::StaticClass();
+		}
+
+		{
+			UECodeGen_Private::FObjectPropertyParams Params = { };
+			Params.NameUTF8 = "ObjectPtrNonNullable";
+			Params.Offset = STRUCT_OFFSET(UObjectWithRawProperty, ObjectPtrNonNullable);
+			Params.PropertyFlags = EPropertyFlags::CPF_NonNullable;
+			Params.ObjectFlags = RF_Public | RF_Transient | RF_MarkAsNative;
+			Params.ClassFunc = nullptr;
+			auto Property = new FObjectProperty(Class, Params);
+			Property->PropertyClass = UObjectWithRawProperty::StaticClass();
+		}
 	}
 );
 
