added GetOuter and GetPackage to FObectPtr and TObjectPtr
changed how the Object Handle index stores info to be able to get the class of the outer objects
Packages dont' store a ObjectPathId or ClassDescriptor as it is always known

https://p4-swarm.epicgames.net/reviews/25378953

diff --git a/Engine/Source/Runtime/CoreUObject/Private/UObject/ObjectHandleIndex.cpp b/Engine/Source/Runtime/CoreUObject/Private/UObject/ObjectHandleIndex.cpp
--- a/Engine/Source/Runtime/CoreUObject/Private/UObject/ObjectHandleIndex.cpp
+++ b/Engine/Source/Runtime/CoreUObject/Private/UObject/ObjectHandleIndex.cpp
@@ -64,47 +64,59 @@ namespace UE::CoreUObject::Private
 			return Id == Other.Id;
 		}
 	};
-	enum class EObjectId : uint32
+	
+	struct FObjectId
 	{
-		Invalid = 0,
-	};
-
-	union FObjectId
-	{
-	private:
-		uint32 RawData = 0;
 	public:
-		struct
-		{
-			uint32 DataClassDescriptorId : 8;
-			uint32 ObjectPathId : 24;
-		} Components;
-
 		FObjectId() = default;
-		FObjectId(EObjectId Id) : RawData(static_cast<uint32>(Id)) {}
+		
+		explicit FObjectId(uint32 Index)
+			: PackedIndex(Index)
+		{
+		}
 
-		bool operator==(EObjectId Id) { return RawData == static_cast<uint32>(Id); }
-		bool operator!=(EObjectId Id) { return RawData != static_cast<uint32>(Id); }
+		static FObjectId FromIndex(uint32 Index)
+		{
+			return FObjectId(Index);
+		}
 
 		inline uint32 ToIndex() const
 		{
-			return Components.ObjectPathId - 1;
+			check(PackedIndex != 0);
+			return PackedIndex - 1;
 		}
+
+		inline uint32 ToPackedIndex() const
+		{
+			return PackedIndex;
+		}
+
+		inline bool IsNone() const
+		{
+			return PackedIndex == 0;
+		}
+
+	private:
+		uint32 PackedIndex = 0;
 	};
 
 	static_assert(sizeof(FObjectId) == sizeof(uint32), "FObjectId type must always compile to something equivalent to a uint32 size.");
 
+	struct FObjectDescriptor
+	{
+		FObjectPathId ObjectPath;
+		FObjectHandleDataClassDescriptor ClassDescriptor;
+	};
+
 	struct FObjectHandlePackageData
 	{
 		FMinimalName PackageName;
-		TArray<FObjectPathId> ObjectPaths;
-		TArray<FObjectHandleDataClassDescriptor> DataClassDescriptors;
+		TArray<FObjectDescriptor> ObjectDescriptors;
 		FRWLock Lock;
 	};
 
 	static_assert(STRUCT_OFFSET(FObjectHandlePackageData, PackageName) == STRUCT_OFFSET(FObjectHandlePackageDebugData, PackageName), "FObjectHandlePackageData and FObjectHandlePackageDebugData must match in position of PackageNameField.");
-	static_assert(STRUCT_OFFSET(FObjectHandlePackageData, ObjectPaths) == STRUCT_OFFSET(FObjectHandlePackageDebugData, ObjectPaths), "FObjectHandlePackageData and FObjectHandlePackageDebugData must match in position of ObjectPaths.");
-	static_assert(STRUCT_OFFSET(FObjectHandlePackageData, DataClassDescriptors) == STRUCT_OFFSET(FObjectHandlePackageDebugData, DataClassDescriptors), "FObjectHandlePackageData and FObjectHandlePackageDebugData must match in position of DataClassDescriptors.");
+	static_assert(STRUCT_OFFSET(FObjectHandlePackageData, ObjectDescriptors) == STRUCT_OFFSET(FObjectHandlePackageDebugData, ObjectDescriptors), "FObjectHandlePackageData and FObjectHandlePackageDebugData must match in position of ObjectPaths.");
 	static_assert(sizeof(FObjectHandlePackageData) == sizeof(FObjectHandlePackageDebugData), "FObjectHandlePackageData and FObjectHandlePackageDebugData must match in size.");
 
 	struct FObjectHandleIndex
@@ -170,16 +182,13 @@ namespace UE::CoreUObject::Private
 	{
 		checkf(PackageId.ToIndex() <= 0x7FFFFFFF, TEXT("Package count exceeded the space permitted within packed object references.  This implies over 2 billion packages are in use."));
 		return { static_cast<UPTRINT>(PackageId.ToIndex()) << PackageIdShift |
-				static_cast<UPTRINT>(ObjectId.Components.DataClassDescriptorId) << DataClassDescriptorIdShift |
-				static_cast<UPTRINT>(ObjectId.Components.ObjectPathId) << ObjectPathIdShift |
-				1 };
+				static_cast<UPTRINT>(ObjectId.ToPackedIndex() << ObjectIdShift) | 1 };
 	}
 
 	static inline void Unpack(FPackedObjectRef PackedObjectRef, FPackageId& OutPackageId, FObjectId& OutObjectId)
 	{
 		checkf((PackedObjectRef.EncodedRef & 1) == 1, TEXT("Packed object reference is malformed."));
-		OutObjectId.Components.ObjectPathId = (static_cast<uint32>((PackedObjectRef.EncodedRef >> ObjectPathIdShift) & ObjectPathIdMask));
-		OutObjectId.Components.DataClassDescriptorId = (static_cast<uint32>((PackedObjectRef.EncodedRef >> DataClassDescriptorIdShift) & DataClassDescriptorIdMask));
+		OutObjectId = FObjectId::FromIndex(static_cast<uint32>(PackedObjectRef.EncodedRef >> ObjectIdShift));
 		OutPackageId = FPackageId::FromIndex(static_cast<uint32>((PackedObjectRef.EncodedRef >> PackageIdShift) & PackageIdMask));
 	}
 
@@ -225,67 +234,62 @@ namespace UE::CoreUObject::Private
 		//since the package is new, GlobalLockScope is write locked and nothing can be added or read
 		if (PackageCreated)
 		{
-			if (!ClassName.IsNone() && !ClassPackageName.IsNone())
-			{
-				FObjectHandleDataClassDescriptor DataClassDesc{ NameToMinimalName(ClassPackageName), NameToMinimalName(ClassName) };
-				uint32 DataClassDescriptorIndex = PackageData->DataClassDescriptors.Add(DataClassDesc);
-				checkf(((DataClassDescriptorIndex + 1) & ~DataClassDescriptorIdMask) == 0, TEXT("Data class descriptor id overflowed space in ObjectHandle"));
-				OutObjectId.Components.DataClassDescriptorId = DataClassDescriptorIndex + 1;
-			}
 			if (!ObjectPath.IsNone())
 			{
-				int32 PathIndex = PackageData->ObjectPaths.Add(ObjectPath);
-				checkf(((PathIndex + 1) & ~ObjectPathIdMask) == 0, TEXT("Path id overflowed space in ObjectHandle"));
-				OutObjectId.Components.ObjectPathId = PathIndex + 1;
+				FObjectHandleDataClassDescriptor DataClassDesc;
+				if (!ClassName.IsNone() && !ClassPackageName.IsNone())
+				{
+					DataClassDesc.PackageName = NameToMinimalName(ClassPackageName);
+					DataClassDesc.ClassName = NameToMinimalName(ClassName);
+				}
+			
+				FObjectDescriptor ObjectDesc = { ObjectPath, DataClassDesc };
+				OutObjectId = FObjectId::FromIndex(PackageData->ObjectDescriptors.Add(ObjectDesc) + 1);
 			}
 			return;
 		}
 		else if (ObjectPath.IsNone())
 		{
-			FRWScopeLock LocalLockScope(PackageData->Lock, SLT_ReadOnly);
-			if (!ClassName.IsNone() && !ClassPackageName.IsNone())
-			{
-				FObjectHandleDataClassDescriptor DataClassDesc{ NameToMinimalName(ClassPackageName), NameToMinimalName(ClassName) };
-				uint32 DataClassDescriptorIndex = PackageData->DataClassDescriptors.AddUnique(DataClassDesc);
-				checkf(((DataClassDescriptorIndex + 1) & ~DataClassDescriptorIdMask) == 0, TEXT("Data class descriptor id overflowed space in ObjectHandle"));
-				OutObjectId.Components.DataClassDescriptorId = DataClassDescriptorIndex + 1;
-			}
+			//don't store anything for packages
 			return;
 		}
 
 		FRWScopeLock LocalLockScope(PackageData->Lock, SLT_ReadOnly);
-		int32 PathIndex = PackageData->ObjectPaths.Find(ObjectPath); //linear search is fine as typically there is only one.
+		int32 PathIndex = PackageData->ObjectDescriptors.IndexOfByPredicate([&ObjectPath](const FObjectDescriptor& Desc)
+			{
+				return Desc.ObjectPath == ObjectPath;
+
+			}); //linear search is fine as typically there is only one.
+
 		if (PathIndex == INDEX_NONE)
 		{
 			//ObjectPaths could have been modified when the read lock was released and the write
 			//lock was acquired, so we must check and see if the ObjectPath was added in that window.
 			LocalLockScope.ReleaseReadOnlyLockAndAcquireWriteLock_USE_WITH_CAUTION();
-			PathIndex = PackageData->ObjectPaths.AddUnique(ObjectPath);
+			PathIndex = PathIndex = PackageData->ObjectDescriptors.IndexOfByPredicate([&ObjectPath](const FObjectDescriptor& Desc)
+				{
+					return Desc.ObjectPath == ObjectPath;
+
+				}); //linear search is fine as typically there is only one.
 
-			if (!ClassName.IsNone() && !ClassPackageName.IsNone())
+			if (PathIndex == INDEX_NONE)
 			{
-				FObjectHandleDataClassDescriptor DataClassDesc{ NameToMinimalName(ClassPackageName), NameToMinimalName(ClassName) };
-				uint32 DataClassDescriptorIndex = PackageData->DataClassDescriptors.AddUnique(DataClassDesc);
-				checkf(((DataClassDescriptorIndex + 1) & ~DataClassDescriptorIdMask) == 0, TEXT("Data class descriptor id overflowed space in ObjectHandle"));
-				OutObjectId.Components.DataClassDescriptorId = DataClassDescriptorIndex + 1;
+				FObjectHandleDataClassDescriptor DataClassDesc;
+				if (!ClassName.IsNone() && !ClassPackageName.IsNone())
+				{
+					DataClassDesc.PackageName = NameToMinimalName(ClassPackageName);
+					DataClassDesc.ClassName = NameToMinimalName(ClassName);
+				}
+
+				FObjectDescriptor ObjDesc = { ObjectPath, DataClassDesc };
+				PathIndex = PackageData->ObjectDescriptors.Add(ObjDesc) + 1;
+				OutObjectId = FObjectId::FromIndex(PathIndex);
 			}
 		}
 		else
 		{
-			if (!ClassName.IsNone() && !ClassPackageName.IsNone())
-			{
-				FObjectHandleDataClassDescriptor DataClassDesc{ NameToMinimalName(ClassPackageName), NameToMinimalName(ClassName) };
-				uint32 DataClassDescriptorIndex = PackageData->DataClassDescriptors.Find(DataClassDesc);
-				if (DataClassDescriptorIndex != INDEX_NONE)
-				{
-					checkf(((DataClassDescriptorIndex + 1) & ~DataClassDescriptorIdMask) == 0, TEXT("Data class descriptor id overflowed space in ObjectHandle"));
-					OutObjectId.Components.DataClassDescriptorId = DataClassDescriptorIndex + 1;
-				}
-			}
+			OutObjectId = FObjectId::FromIndex(PathIndex + 1);
 		}
-
-		checkf(((PathIndex + 1) & ~ObjectPathIdMask) == 0, TEXT("Path id overflowed space in ObjectHandle"));
-		OutObjectId.Components.ObjectPathId = PathIndex + 1;
 		return;
 	}
 
@@ -299,13 +303,14 @@ namespace UE::CoreUObject::Private
 
 	static void GetObjectDataFromId(FPackageId PackageId, FObjectId ObjectId, FMinimalName& OutPackageName, FObjectPathId& OutPathId, FMinimalName& OutClassPackageName, FMinimalName& OutClassName)
 	{
-		if ((ObjectId == EObjectId::Invalid) || !PackageId.IsValid())
+		if (!PackageId.IsValid())
 		{
 			return;
 		}
 
 		FRWScopeLock GlobalLockScope(GObjectHandleIndex.Lock, SLT_ReadOnly);
 		const uint32 PackageIndex = PackageId.ToIndex();
+
 		if (PackageIndex >= static_cast<uint32>(GObjectHandleIndex.PackageData.Num()))
 		{
 			//checkf(false, TEXT("FObjectHandle: PackageIndex invalid.  This ObjectHandle is from malformed data."));
@@ -315,35 +320,92 @@ namespace UE::CoreUObject::Private
 		FRWScopeLock LocalLockScope(FoundPackageData.Lock, SLT_ReadOnly);
 		OutPackageName = FoundPackageData.PackageName;
 
-		if ((ObjectId.Components.ObjectPathId >= static_cast<uint32>(FoundPackageData.ObjectPaths.Num() + 1)) ||
-			(ObjectId.Components.DataClassDescriptorId >= static_cast<uint32>(FoundPackageData.DataClassDescriptors.Num() + 1)))
-		{
-			//checkf(false, TEXT("FObjectHandle: ObjectId.Components.ObjectPathId or ObjectId.Components.DataClassDescriptorId invalid.  This ObjectHandle is from malformed data."));
-			return;
-		}
-
-		if (ObjectId.Components.ObjectPathId != 0)
+		if (!ObjectId.IsNone())
 		{
-			OutPathId = FoundPackageData.ObjectPaths[ObjectId.Components.ObjectPathId - 1];
+			int32 ObjectIndex = ObjectId.ToIndex();
+			if (ObjectIndex >= FoundPackageData.ObjectDescriptors.Num())
+			{
+				//must reset to empty as the object id is invalid
+				OutPackageName = FMinimalName();
+				return;
+			}
+			const auto& ObjDesc = FoundPackageData.ObjectDescriptors[ObjectIndex];
+			OutPathId = ObjDesc.ObjectPath;
+			OutClassPackageName = ObjDesc.ClassDescriptor.PackageName;
+			OutClassName = ObjDesc.ClassDescriptor.ClassName;
 		}
 		else
 		{
 			OutPathId = FObjectPathId();
+			OutClassPackageName = FMinimalName(UPackage::StaticClass()->GetPackage()->GetFName());
+			OutClassName = FMinimalName(UPackage::StaticClass()->GetFName());
 		}
+	}
+
+	FPackedObjectRef GetOuter(FPackedObjectRef PackedRef)
+	{
+		if (PackedRef.IsNull())
+		{
+			return PackedRef;
+		}
+		FPackageId PackageId;
+		FObjectId ObjectId;
+		Unpack(PackedRef, PackageId, ObjectId);
+
+		FMinimalName PackageName;
+		FObjectPathId PathId(FObjectPathId::Invalid);
+		FMinimalName ClassPackageName;
+		FMinimalName ClassName;
+		GetObjectDataFromId(PackageId, ObjectId, PackageName, PathId, ClassPackageName, ClassName);
 
-		if (ObjectId.Components.DataClassDescriptorId > 0)
+		if (PathId.IsNone())
+		{
+			return { 0 }; //return null for packages as they have no outer
+		}
+		else if (PathId.IsSimple())
+		{
+			FObjectId PackageObjectId;
+			return Pack(PackageId, PackageObjectId);
+		}
+		else if (PathId.IsWeakObj())
 		{
-			FObjectHandleDataClassDescriptor& Desc = FoundPackageData.DataClassDescriptors[ObjectId.Components.DataClassDescriptorId - 1];
-			OutClassPackageName = Desc.PackageName;
-			OutClassName = Desc.ClassName;
+			return MakePackedObjectRef(PathId.GetWeakObjPtr().Get()->GetOuter());
 		}
 		else
 		{
-			//@TODO OBJPTR: Is this an error case we should check on?
-			//checkf(false, TEXT("FObjectHandle: ObjectId.Components.ObjectPathId or ObjectId.Components.DataClassDescriptorId invalid.  This ObjectHandle is from malformed data."));
-			return;
+			TStringBuilder<256> StringBuilder;
+			FObjectPathId::ResolvedNameContainerType Names;
+			PathId.Resolve(Names);
+			for (int32 i = 0; i < Names.Num() - 1; ++i)
+			{
+				const FName& Name = Names[i];
+				if (i != 0)
+				{
+					StringBuilder << '.'; //separate subobjects by periods a colon is unnecessary to produce a correct FObjectPathId
+				}
+				Name.AppendString(StringBuilder);
+			}
+			FObjectPathId OuterPathId(StringBuilder);
+			FPackedObjectRef OuterPackedRef = MakePackedObjectRef(FName(PackageName), FName(), FName(), OuterPathId);
+
+			return OuterPackedRef;
 		}
 	}
+	
+	FPackedObjectRef GetPackage(FPackedObjectRef PackedRef)
+	{
+		if (PackedRef.IsNull())
+		{
+			return PackedRef;
+		}
+		FPackageId PackageId;
+		FObjectId ObjectId;
+		Unpack(PackedRef, PackageId, ObjectId);
+
+		//object id's for packages are always 0
+		FPackedObjectRef PackagePackedRef = Pack(PackageId, FObjectId());
+		return PackagePackedRef;
+	}
 
 	void FreeObjectHandle(const UObjectBase* Object)
 	{
@@ -388,7 +450,7 @@ namespace UE::CoreUObject::Private
 		}
 
 		//must make a copy to avoid changing data while we are comparing
-		FObjectPathId OldObjectPath = PackageData.ObjectPaths[ObjectId.ToIndex()];
+		FObjectPathId OldObjectPath = PackageData.ObjectDescriptors[ObjectId.ToIndex()].ObjectPath;
 		if (OldObjectPath.IsWeakObj())
 		{
 			return;
@@ -413,9 +475,9 @@ namespace UE::CoreUObject::Private
 		}
 
 		//update all paths that start with OldObjectPath. including itself
-		for (int32 PathIndex = 0; PathIndex < PackageData.ObjectPaths.Num(); ++PathIndex)
+		for (int32 PathIndex = 0; PathIndex < PackageData.ObjectDescriptors.Num(); ++PathIndex)
 		{
-			auto& CurrentObjectPath = PackageData.ObjectPaths[PathIndex];
+			auto& CurrentObjectPath = PackageData.ObjectDescriptors[PathIndex].ObjectPath;
 			if (CurrentObjectPath.IsWeakObj())
 			{
 				continue;
diff --git a/Engine/Source/Runtime/CoreUObject/Private/UObject/ObjectPtr.cpp b/Engine/Source/Runtime/CoreUObject/Private/UObject/ObjectPtr.cpp
--- a/Engine/Source/Runtime/CoreUObject/Private/UObject/ObjectPtr.cpp
+++ b/Engine/Source/Runtime/CoreUObject/Private/UObject/ObjectPtr.cpp
@@ -5,6 +5,12 @@
 
 #if UE_WITH_OBJECT_HANDLE_LATE_RESOLVE
 
+namespace  UE::CoreUObject::Private
+{
+	FPackedObjectRef GetOuter(FPackedObjectRef);
+	FPackedObjectRef GetPackage(FPackedObjectRef);
+}
+
 FString FObjectPtr::GetPathName() const
 {
 	FObjectHandle LocalHandle = Handle;
@@ -47,6 +53,29 @@ FString FObjectPtr::GetFullName(EObjectFullNameFlags Flags) const
 	}
 }
 
+FObjectPtr FObjectPtr::GetOuter() const
+{
+	FObjectHandle LocalHandle = Handle;
+	if (IsObjectHandleResolved(LocalHandle) && !IsObjectHandleNull(LocalHandle))
+	{
+		return FObjectPtr(Get()->GetOuter());
+	}
+	UE::CoreUObject::Private::FPackedObjectRef PackedRef = UE::CoreUObject::Private::GetOuter(ReadObjectHandlePackedObjectRefNoCheck(LocalHandle));
+	FObjectPtr Ptr({ PackedRef.EncodedRef });
+	return Ptr;
+}
+
+FObjectPtr FObjectPtr::GetPackage() const
+{
+	FObjectHandle LocalHandle = Handle;
+	if (IsObjectHandleResolved(LocalHandle) && !IsObjectHandleNull(LocalHandle))
+	{
+		return FObjectPtr(Get()->GetPackage());
+	}
+	UE::CoreUObject::Private::FPackedObjectRef PackedRef = UE::CoreUObject::Private::GetPackage(ReadObjectHandlePackedObjectRefNoCheck(LocalHandle));
+	FObjectPtr Ptr({ PackedRef.EncodedRef });
+	return Ptr;
+}
 #endif
 
 bool FObjectPtr::IsA(const UClass* SomeBase) const
@@ -59,4 +88,4 @@ bool FObjectPtr::IsA(const UClass* SomeBase) const
 	}
 
 	return false;
-}
+}
\ No newline at end of file
diff --git a/Engine/Source/Runtime/CoreUObject/Private/UObject/ObjectRef.cpp b/Engine/Source/Runtime/CoreUObject/Private/UObject/ObjectRef.cpp
--- a/Engine/Source/Runtime/CoreUObject/Private/UObject/ObjectRef.cpp
+++ b/Engine/Source/Runtime/CoreUObject/Private/UObject/ObjectRef.cpp
@@ -92,7 +92,7 @@ void FObjectRef::AppendPathName(FStringBuilderBase& OutPathNameBuilder) const
 		return;
 	}
 
-	PackageName.ToString(OutPathNameBuilder);
+	PackageName.AppendString(OutPathNameBuilder);
 
 	UE::CoreUObject::Private::FObjectPathId ObjectPath = GetObjectPath();
 	UE::CoreUObject::Private::FObjectPathId::ResolvedNameContainerType ResolvedNames;
diff --git a/Engine/Source/Runtime/CoreUObject/Public/UObject/ObjectHandle.h b/Engine/Source/Runtime/CoreUObject/Public/UObject/ObjectHandle.h
--- a/Engine/Source/Runtime/CoreUObject/Public/UObject/ObjectHandle.h
+++ b/Engine/Source/Runtime/CoreUObject/Public/UObject/ObjectHandle.h
@@ -301,8 +301,7 @@ namespace UE::CoreUObject::Private
 	struct FObjectHandlePackageDebugData
 	{
 		FMinimalName PackageName;
-		FScriptArray ObjectPaths;
-		FScriptArray DataClassDescriptors;
+		FScriptArray ObjectDescriptors;
 		uint8 _Padding[sizeof(FRWLock)];
 	};
 
@@ -321,6 +320,12 @@ namespace UE::CoreUObject::Private
 		static constexpr uint32 SimpleNameMask = WeakObjectMask >> 1; //second most significant bits
 	};
 
+	struct FObjectDescriptorDebug
+	{
+		FObjectPathIdDebug ObjectPath;
+		FObjectHandleDataClassDescriptor ClassDescriptor;
+	};
+
 	struct FStoredObjectPathDebug
 	{
 		static constexpr const int32 NumInlineElements = 3;
@@ -333,11 +338,8 @@ namespace UE::CoreUObject::Private
 		};
 	};
 
+	constexpr uint32 ObjectIdShift = 1;
 	constexpr uint32 PackageIdShift = 33;
-	constexpr uint32 ObjectPathIdShift = 1;
-	constexpr uint32 ObjectPathIdMask = 0x00FF'FFFF;
-	constexpr uint32 DataClassDescriptorIdShift = 25;
-	constexpr uint32 DataClassDescriptorIdMask = 0x0000'00FF;
 	constexpr uint32 PackageIdMask = 0x7FFF'FFFF;
 
 #if UE_WITH_OBJECT_HANDLE_LATE_RESOLVE
diff --git a/Engine/Source/Runtime/CoreUObject/Public/UObject/ObjectPtr.h b/Engine/Source/Runtime/CoreUObject/Public/UObject/ObjectPtr.h
--- a/Engine/Source/Runtime/CoreUObject/Public/UObject/ObjectPtr.h
+++ b/Engine/Source/Runtime/CoreUObject/Public/UObject/ObjectPtr.h
@@ -161,6 +161,14 @@ public:
 		return GetFName().ToString();
 	}
 
+	// Gets the Outer UObject for this Object
+	// returns unresolved if this object ptr is unresolved
+	COREUOBJECT_API FObjectPtr GetOuter() const;
+
+	// Gets the package for this Object
+	// returns unresolved if this object ptr is unresolved
+	COREUOBJECT_API FObjectPtr GetPackage() const;
+
 	/** Returns the full name for the object in the form: Class ObjectPath */
 	COREUOBJECT_API FString GetFullName(EObjectFullNameFlags Flags = EObjectFullNameFlags::None) const;
 #else
@@ -181,6 +189,17 @@ public:
 		return GetFName().ToString();
 	}
 
+	FObjectPtr GetOuter() const
+	{
+		const UObject* ResolvedObject = Get();
+		return ResolvedObject ? FObjectPtr(ResolvedObject->GetOuter()) : FObjectPtr(nullptr);
+	}
+
+	FObjectPtr GetPackage() const
+	{
+		const UObject* ResolvedObject = Get();
+		return ResolvedObject ? FObjectPtr(ResolvedObject->GetPackage()) : FObjectPtr(nullptr);
+	}
 	/**
 	 * Returns the fully qualified pathname for this object as well as the name of the class, in the format:
 	 * 'ClassName Outermost.[Outer:]Name'.
@@ -359,6 +378,11 @@ public:
 	{
 	}
 
+	explicit FORCEINLINE TObjectPtr(FObjectPtr ObjPtr)
+		: ObjectPtr(ObjPtr)
+	{
+	}
+
 	template <
 		typename U,
 		decltype(ImplicitConv<T*>(std::declval<U*>()))* = nullptr
@@ -489,6 +513,17 @@ public:
 	//			isn't called as frequently.
 	FORCEINLINE T* Get() const { return (T*)(ObjectPtr.Get()); }
 	FORCEINLINE UClass* GetClass() const { return ObjectPtr.GetClass(); }
+	FORCEINLINE TObjectPtr<UObject> GetOuter() const
+	{ 
+		FObjectPtr Ptr = ObjectPtr.GetOuter();
+		return TObjectPtr<UObject>(Ptr);
+	}
+
+	FORCEINLINE TObjectPtr<UPackage> GetPackage() const
+	{
+		FObjectPtr Ptr = ObjectPtr.GetPackage();
+		return TObjectPtr<UPackage>(Ptr);
+	}
 	
 	FORCEINLINE operator T* () const { return Get(); }
 	template <typename U>
diff --git a/Engine/Source/Runtime/CoreUObject/Tests/ObjectHandleTest.cpp b/Engine/Source/Runtime/CoreUObject/Tests/ObjectHandleTest.cpp
--- a/Engine/Source/Runtime/CoreUObject/Tests/ObjectHandleTest.cpp
+++ b/Engine/Source/Runtime/CoreUObject/Tests/ObjectHandleTest.cpp
@@ -157,6 +157,7 @@ TEST_CASE("CoreUObject::FObjectHandle::Names")
 		TestPackage->RemoveFromRoot();
 	};
 
+	FObjectPtr Test;
 	FObjectPtr PackagePtr(MakeUnresolvedHandle(TestPackage));
 	FObjectPtr Obj1Ptr(MakeUnresolvedHandle(Obj1));
 
@@ -164,14 +165,14 @@ TEST_CASE("CoreUObject::FObjectHandle::Names")
 	CHECK(TestPackage->GetPathName() == PackagePtr.GetPathName());
 	CHECK(TestPackage->GetFName() == PackagePtr.GetFName());
 	CHECK(TestPackage->GetName() == PackagePtr.GetName());
-	//CHECK(TestPackage->GetFullName() == PackagePtr.GetFullName()); TODO always been broken
+	CHECK(TestPackage->GetFullName() == PackagePtr.GetFullName());
 	CHECK(!PackagePtr.IsResolved());
 
 	CHECK(!Obj1Ptr.IsResolved());
 	CHECK(Obj1->GetPathName() == Obj1Ptr.GetPathName());
 	CHECK(Obj1->GetFName() == Obj1Ptr.GetFName());
 	CHECK(Obj1->GetName() == Obj1Ptr.GetName());
-	//CHECK(Obj1->GetFullName() == Obj1Ptr.GetFullName()); TODO always been broken
+	CHECK(Obj1->GetFullName() == Obj1Ptr.GetFullName());
 	CHECK(!Obj1Ptr.IsResolved());
 }
 #endif
@@ -347,6 +348,15 @@ TEST_CASE_METHOD(FObjectHandleTestBase, "CoreUObject::TObjectPtr::Null Behavior"
 
 TEST_CASE_METHOD(FObjectHandleTestBase, "CoreUObject::FObjectHandle::Resolve Malformed Handle", "[CoreUObject][ObjectHandle]")
 {
+	// make one packed ref guarantee something is in the object handle index
+	FObjectRef TargetRef(FName("/Test/DummyPackage"), FName("ClassPackageName"), FName("ClassName"), UE::CoreUObject::Private::FObjectPathId("DummyObjectName"));
+	UE::CoreUObject::Private::MakePackedObjectRef(TargetRef);
+
+	uint32 ObjectId = ~0u;
+	UPTRINT PackedId = ObjectId << 1 | 1;
+	UE::CoreUObject::Private::FPackedObjectRef PackedObjectRef = { PackedId };
+	TestResolveFailure(PackedObjectRef); // packed ref has a valid package id but invalid object id
+
 	TestResolveFailure(UE::CoreUObject::Private::FPackedObjectRef { 0xFFFF'FFFF'FFFF'FFFFull });
 	TestResolveFailure(UE::CoreUObject::Private::FPackedObjectRef { 0xEFEF'EFEF'EFEF'EFEFull });
 }
diff --git a/Engine/Source/Runtime/CoreUObject/Tests/ObjectPtrTest.cpp b/Engine/Source/Runtime/CoreUObject/Tests/ObjectPtrTest.cpp
--- a/Engine/Source/Runtime/CoreUObject/Tests/ObjectPtrTest.cpp
+++ b/Engine/Source/Runtime/CoreUObject/Tests/ObjectPtrTest.cpp
@@ -1084,7 +1084,83 @@ TEST_CASE("CoreUObject::TObjectPtr::ConstArrayViewConversion")
 
 	
 }
+TEST_CASE("CoreUObject::TObjectPtr::GetOuter")
+{
+	UPackage* TestPackage = NewObject<UPackage>(nullptr, "/Test/MyPackage", RF_Transient);
+	TestPackage->AddToRoot();
+	TObjectPtr<UObject> Obj1 = NewObject<UObjectPtrTestClass>(TestPackage, TEXT("Obj1"));
+	TObjectPtr<UObject> Obj2 = NewObject<UObjectPtrTestClass>(Obj1, TEXT("Obj2"));
+	int ResolveCount = 0;
+
+#if UE_WITH_OBJECT_HANDLE_LATE_RESOLVE
+	FObjectPtr Ptr1(MakeUnresolvedHandle(TestPackage));
+	FObjectPtr Ptr2(MakeUnresolvedHandle(Obj1));
+	FObjectPtr Ptr3(MakeUnresolvedHandle(Obj2));
+
+	TObjectPtr<UPackage> PackagePtr = *reinterpret_cast<TObjectPtr<UPackage>*>(&Ptr1);
+	TObjectPtr<UObject> Obj1Ptr = *reinterpret_cast<TObjectPtr<UObject>*>(&Ptr2);
+	TObjectPtr<UObject> Obj2Ptr = *reinterpret_cast<TObjectPtr<UObject>*>(&Ptr3);
+
+
+	auto CallbackHandle = UE::CoreUObject::AddObjectHandleReferenceResolvedCallback([&ResolveCount](const FObjectRef& SourceRef, UPackage* ObjectPackage, UObject* Object)
+		{
+			++ResolveCount;
+		});
+	ON_SCOPE_EXIT
+	{
+		UE::CoreUObject::RemoveObjectHandleReferenceResolvedCallback(CallbackHandle);
+	};
+#else
+	TObjectPtr<UPackage> PackagePtr = TestPackage;
+	TObjectPtr<UObject> Obj1Ptr = Obj1;
+	TObjectPtr<UObject> Obj2Ptr = Obj2;
+#endif
+	TObjectPtr<UObject> Obj1RawOuter = Obj1->GetOuter();
+	TObjectPtr<UObject> Obj2RawOuter = Obj2->GetOuter();
 
+	TObjectPtr<UObject> PackageOuter = PackagePtr.GetOuter();
+	TObjectPtr<UObject> Obj1Outer = Obj1Ptr.GetOuter();
+	TObjectPtr<UObject> Obj2Outer = Obj2Ptr.GetOuter();
+
+#if UE_WITH_OBJECT_HANDLE_LATE_RESOLVE
+	CHECK(!Obj1Outer.IsResolved());
+	CHECK(!Obj2Outer.IsResolved());
+
+	//sanity check that the packed refs are identical
+	CHECK(Obj1Outer.GetHandle().PointerOrRef == PackagePtr.GetHandle().PointerOrRef);
+	CHECK(Obj2Outer.GetHandle().PointerOrRef == Obj1Ptr.GetHandle().PointerOrRef);
+
+#endif
+
+	CHECK(Obj1Outer.GetHandle() == PackagePtr.GetHandle());
+	CHECK(Obj2Outer.GetHandle() == Obj1Ptr.GetHandle());
+
+	CHECK(PackageOuter == TestPackage->GetOuter());
+	CHECK(PackageOuter == nullptr);
+	CHECK(Obj1Outer == Obj1RawOuter);
+	CHECK(Obj1Outer.GetFName() == Obj1RawOuter->GetFName());
+	CHECK(Obj1Outer.GetPathName() == Obj1RawOuter->GetPathName());
+	CHECK(Obj1Outer.GetFullName() == Obj1RawOuter->GetFullName());
+	CHECK(Obj1Outer.GetClass() == Obj1RawOuter->GetClass());
+
+	CHECK(Obj2Outer == Obj2RawOuter);
+	
+	CHECK(Obj2Outer.GetFName() == Obj2RawOuter->GetFName());
+	CHECK(Obj2Outer.GetPathName() == Obj2RawOuter->GetPathName());
+	CHECK(Obj2Outer.GetFullName() == Obj2RawOuter->GetFullName());
+	CHECK(Obj2Outer.GetClass() == Obj2RawOuter->GetClass());
+
+
+	TObjectPtr<UPackage> Package = PackagePtr.GetPackage();
+	TObjectPtr<UPackage> Obj1Package = Obj1Ptr.GetPackage();
+	TObjectPtr<UPackage> Obj2Package = Obj2Ptr.GetPackage();
+
+	CHECK(Package == PackagePtr);
+	CHECK(Obj1Package == PackagePtr);
+	CHECK(Obj2Package == PackagePtr);
+
+	CHECK(ResolveCount == 0);
+}
 // @TODO: OBJPTR: We should have a test that ensures that lazy loading of an object with an external package is handled correctly.
 //				  This should also include external packages in the outer chain of the target object.
 // IMPLEMENT_CUSTOM_SIMPLE_AUTOMATION_TEST(FObjectPtrTestExternalPackages, FObjectPtrTestBase, TEXT(TEST_NAME_ROOT ".ExternalPackages"), ObjectPtrTestFlags)
