SavePackage EditorOnly detection:
HasNonEditorOnlyReferences needs to be applied to suppress the checking of archetype and class for editoronly. It already applies to suppressing IsNeededForClient=false and IsEditorOnly=true, but checking the archetype+class also needs to be suppressed. This is necessary to support instances of UWidgetBlueprint, because the UWidgetBlueprint type and CDO are in the UMGEditor package and so are editoronly. But instances of UWidgetBlueprint, like instances of UBlueprint, should be treated as not editoronly when saving in editor.

Remove the workaround for WidgetBlueprints that was added in HarvestPackage - harvesting CDOs as toplevel objects even though they are not RF_Standalone.

Related but not vital: Deprecate bCheckMarks argument from IsEditorOnlyObject. SavePackage was the only user and no longer uses it.

#rnx

diff --git a/Engine/Source/Runtime/CoreUObject/Private/UObject/SavePackage/PackageHarvester.cpp b/Engine/Source/Runtime/CoreUObject/Private/UObject/SavePackage/PackageHarvester.cpp
--- a/Engine/Source/Runtime/CoreUObject/Private/UObject/SavePackage/PackageHarvester.cpp
+++ b/Engine/Source/Runtime/CoreUObject/Private/UObject/SavePackage/PackageHarvester.cpp
@@ -11,6 +11,8 @@
 
 EObjectMark GenerateMarksForObject(const UObject* InObject, const ITargetPlatform* TargetPlatform)
 {
+	using namespace UE::SavePackageUtilities;
+
 	EObjectMark Marks = OBJECTMARK_NOMARKS;
 
 	// CDOs must be included if their class are, so do not generate any marks for it here, defer exclusion to their outer and class
@@ -46,15 +48,18 @@ EObjectMark GenerateMarksForObject(const UObject* InObject, const ITargetPlatfor
 	}
 #endif
 	
-	if (UE::SavePackageUtilities::IsStrippedEditorOnlyObject(InObject, true, false))
+	// If doing an editor save, HasNonEditorOnlyReferences=true overrides NotForClient, NotForServer, and virtual IsEditorOnly and marks it as UsedInGame
+	bool bApplyHasNonEditorOnlyReferences = TargetPlatform == nullptr;
+	if (IsStrippedEditorOnlyObject(InObject,
+		EEditorOnlyObjectFlags::CheckRecursive |
+		(bApplyHasNonEditorOnlyReferences ? EEditorOnlyObjectFlags::ApplyHasNonEditorOnlyReferences : EEditorOnlyObjectFlags::None)))
 	{
 		Marks = (EObjectMark)(Marks | OBJECTMARK_EditorOnly);
 	}
 	else
 	// If NotForClient and NotForServer, it is implicitly editor only
-	// If doing an editor save, HasNonEditorOnlyReferences can override it to still not be marked as editoronly
 	if ((Marks & OBJECTMARK_NotForClient) && (Marks & OBJECTMARK_NotForServer) &&
-		(TargetPlatform != nullptr || !InObject->HasNonEditorOnlyReferences()))
+		(!bApplyHasNonEditorOnlyReferences || !InObject->HasNonEditorOnlyReferences()))
 	{
 		Marks = (EObjectMark)(Marks | OBJECTMARK_EditorOnly);
 	}
@@ -89,7 +94,9 @@ bool ConditionallyExcludeObjectForRealm(FSaveContext& SaveContext, UObject* Obj,
 	}
 
 	// If the object class is excluded, the object must be excluded too
-	if (ConditionallyExcludeObjectForRealm(SaveContext, Obj->GetClass(), HarvestingContext))
+	bool bApplyHasNonEditorOnlyReferences = TargetPlatform == nullptr;
+	bool bIgnoreEditorOnlyClass = bApplyHasNonEditorOnlyReferences && Obj->HasNonEditorOnlyReferences();
+	if (!bIgnoreEditorOnlyClass && ConditionallyExcludeObjectForRealm(SaveContext, Obj->GetClass(), HarvestingContext))
 	{
 		RealmBeingChecked.AddExcluded(Obj);
 		return true;
@@ -102,26 +109,29 @@ bool ConditionallyExcludeObjectForRealm(FSaveContext& SaveContext, UObject* Obj,
 		return true;
 	}
 
-	// Check parent struct if we have one
-	UStruct* ThisStruct = Cast<UStruct>(Obj);
-	if (ThisStruct)
+	if (!bIgnoreEditorOnlyClass)
 	{
-		UObject* SuperStruct = ThisStruct->GetSuperStruct();
-		if (SuperStruct && ConditionallyExcludeObjectForRealm(SaveContext, SuperStruct, HarvestingContext))
+		// Check parent struct if we have one
+		UStruct* ThisStruct = Cast<UStruct>(Obj);
+		if (ThisStruct)
+		{
+			UObject* SuperStruct = ThisStruct->GetSuperStruct();
+			if (SuperStruct && ConditionallyExcludeObjectForRealm(SaveContext, SuperStruct, HarvestingContext))
+			{
+				RealmBeingChecked.AddExcluded(Obj);
+				return true;
+			}
+		}
+
+		// Check archetype, this may not have been covered in the case of components
+		UObject* Archetype = Obj->GetArchetype();
+		if (Archetype && ConditionallyExcludeObjectForRealm(SaveContext, Archetype, HarvestingContext))
 		{
 			RealmBeingChecked.AddExcluded(Obj);
 			return true;
 		}
 	}
 
-	// Check archetype, this may not have been covered in the case of components
-	UObject* Archetype = Obj->GetArchetype();
-	if (Archetype && ConditionallyExcludeObjectForRealm(SaveContext, Archetype, HarvestingContext))
-	{
-		RealmBeingChecked.AddExcluded(Obj);
-		return true;
-	}
-
 	return false;
 }
 
@@ -392,6 +402,14 @@ void FPackageHarvester::ProcessExport(const FExportWithContext& InProcessContext
 	*this << Class;
 	ForEachExportHarvestingRealm([this, Class, Export](ESaveRealm HarvestingRealm)
 		{
+			if (HarvestingRealm == ESaveRealm::Game && !SaveContext.IsCooking())
+			{
+				// During an EditorSave we mark objects as UsedInGame (!EditorOnly) by adding them to the Game realm.
+				// But objects can be marked as UsedInGame even if their class is EditorOnly if they
+				// return true from HasNonEditorOnlyReferences. So suppress this error that the class is
+				// missing from the realm if the realm is Game and we are doing an editor save.
+				return;
+			}
 			FHarvestedRealm& RealmData = SaveContext.GetHarvestedRealm(HarvestingRealm);
 			if (!RealmData.IsIncluded(Class))
 			{
@@ -548,9 +566,9 @@ void FPackageHarvester::TryHarvestExportInternal(UObject* InObject)
 	FHarvestScope OptionalObjectScope = EnterConditionalOptionalObjectScope(InObject);
 
 #if WITH_EDITORONLY_DATA
-	// Remove the Game realm if the object is editoronly
+	// Remove the Game realm if the object is editoronly and does not override it with HasNonEditorOnlyReferences
 	bool bIsEditorOnlyObject = UE::SavePackageUtilities::IsStrippedEditorOnlyObject(InObject,
-		true /* bCheckRecursive */, false /* bCheckMarks */)
+		UE::SavePackageUtilities::EEditorOnlyObjectFlags::CheckRecursive)
 		&& !InObject->HasNonEditorOnlyReferences();
 	FHarvestScope EditorOnlyScope = EnterConditionalEditorOnlyScope(bIsEditorOnlyObject);
 #endif
diff --git a/Engine/Source/Runtime/CoreUObject/Private/UObject/SavePackage/SavePackageUtilities.cpp b/Engine/Source/Runtime/CoreUObject/Private/UObject/SavePackage/SavePackageUtilities.cpp
--- a/Engine/Source/Runtime/CoreUObject/Private/UObject/SavePackage/SavePackageUtilities.cpp
+++ b/Engine/Source/Runtime/CoreUObject/Private/UObject/SavePackage/SavePackageUtilities.cpp
@@ -595,7 +595,7 @@ void GetCDOSubobjects(UObject* CDO, TArray<UObject*>& Subobjects)
 	}
 }
 	
-bool IsStrippedEditorOnlyObject(const UObject* InObject, bool bCheckRecursive, bool bCheckMarks)
+bool IsStrippedEditorOnlyObject(const UObject* InObject, EEditorOnlyObjectFlags Flags)
 {
 #if WITH_EDITOR
 	// Configurable via ini setting
@@ -614,7 +614,7 @@ bool IsStrippedEditorOnlyObject(const UObject* InObject, bool bCheckRecursive, b
 		return false;
 	}
 	
-	return IsEditorOnlyObject(InObject, bCheckRecursive, bCheckMarks);
+	return IsEditorOnlyObjectInternal(InObject, Flags);
 #else
 	return true;
 #endif
@@ -760,8 +760,37 @@ void FObjectSaveContextData::Set(UPackage* Package, const ITargetPlatform* InTar
 	bUpdatingLoadedPath = UE::SavePackageUtilities::IsUpdatingLoadedPath(InTargetPlatform != nullptr, TargetPath, InSaveFlags);
 }
 
+bool IsEditorOnlyObject(const UObject* InObject, bool bCheckRecursive)
+{
+	using namespace UE::SavePackageUtilities;
+	EEditorOnlyObjectFlags Flags = EEditorOnlyObjectFlags::None;
+	Flags |= bCheckRecursive ? EEditorOnlyObjectFlags::CheckRecursive : EEditorOnlyObjectFlags::None;
+	return IsEditorOnlyObjectInternal(InObject, Flags);
+}
+
 bool IsEditorOnlyObject(const UObject* InObject, bool bCheckRecursive, bool bCheckMarks)
 {
+	using namespace UE::SavePackageUtilities;
+	EEditorOnlyObjectFlags Flags = EEditorOnlyObjectFlags::None;
+	Flags |= bCheckRecursive ? EEditorOnlyObjectFlags::CheckRecursive : EEditorOnlyObjectFlags::None;
+	PRAGMA_DISABLE_DEPRECATION_WARNINGS;
+	Flags |= bCheckMarks ? EEditorOnlyObjectFlags::CheckMarks : EEditorOnlyObjectFlags::None;
+	PRAGMA_ENABLE_DEPRECATION_WARNINGS;
+	return IsEditorOnlyObjectInternal(InObject, Flags);
+}
+
+namespace UE::SavePackageUtilities
+{
+
+bool IsEditorOnlyObjectInternal(const UObject* InObject, EEditorOnlyObjectFlags Flags)
+{
+	bool bCheckRecursive = EnumHasAnyFlags(Flags, EEditorOnlyObjectFlags::CheckRecursive);
+	bool bIgnoreEditorOnlyClass = EnumHasAnyFlags(Flags, EEditorOnlyObjectFlags::ApplyHasNonEditorOnlyReferences) &&
+		InObject->HasNonEditorOnlyReferences();
+	PRAGMA_DISABLE_DEPRECATION_WARNINGS;
+	bool bCheckMarks = EnumHasAnyFlags(Flags, EEditorOnlyObjectFlags::CheckMarks);
+	PRAGMA_ENABLE_DEPRECATION_WARNINGS;
+
 	DECLARE_SCOPE_CYCLE_COUNTER(TEXT("IsEditorOnlyObject"), STAT_IsEditorOnlyObject, STATGROUP_LoadTime);
 	check(InObject);
 
@@ -769,7 +798,8 @@ bool IsEditorOnlyObject(const UObject* InObject, bool bCheckRecursive, bool bChe
 	// Ignore their value of IsEditorOnly
 	if (!InObject->HasAnyFlags(RF_ClassDefaultObject))
 	{
-		if ((bCheckMarks && InObject->HasAnyMarks(OBJECTMARK_EditorOnly)) || InObject->IsEditorOnly())
+		if (!bIgnoreEditorOnlyClass &&
+			((bCheckMarks && InObject->HasAnyMarks(OBJECTMARK_EditorOnly)) || InObject->IsEditorOnly()))
 		{
 			return true;
 		}
@@ -808,37 +838,42 @@ bool IsEditorOnlyObject(const UObject* InObject, bool bCheckRecursive, bool bChe
 		UObject* Outer = InObject->GetOuter();
 		if (Outer && Outer != Package)
 		{
-			if (IsEditorOnlyObject(Outer, true, bCheckMarks))
+			if (IsEditorOnlyObjectInternal(Outer, Flags))
 			{
 				return true;
 			}
 		}
-		const UStruct* InStruct = Cast<UStruct>(InObject);
-		if (InStruct)
+		if (!bIgnoreEditorOnlyClass)
 		{
-			const UStruct* SuperStruct = InStruct->GetSuperStruct();
-			if (SuperStruct && IsEditorOnlyObject(SuperStruct, true, bCheckMarks))
+			const UStruct* InStruct = Cast<UStruct>(InObject);
+			if (InStruct)
 			{
-				return true;
+				const UStruct* SuperStruct = InStruct->GetSuperStruct();
+				if (SuperStruct && IsEditorOnlyObjectInternal(SuperStruct, Flags))
+				{
+					return true;
+				}
 			}
-		}
-		else
-		{
-			if (IsEditorOnlyObject(InObject->GetClass(), true, bCheckMarks))
+			else
 			{
-				return true;
-			}
+				if (IsEditorOnlyObjectInternal(InObject->GetClass(), Flags))
+				{
+					return true;
+				}
 
-			UObject* Archetype = InObject->GetArchetype();
-			if (Archetype && IsEditorOnlyObject(Archetype, true, bCheckMarks))
-			{
-				return true;
+				UObject* Archetype = InObject->GetArchetype();
+				if (Archetype && IsEditorOnlyObjectInternal(Archetype, Flags))
+				{
+					return true;
+				}
 			}
 		}
 	}
 	return false;
 }
 
+} // namespace UE::SavePackageUtilities
+
 void FObjectImportSortHelper::SortImports(FLinkerSave* Linker)
 {
 	TArray<FObjectImport>& Imports = Linker->ImportMap;
diff --git a/Engine/Source/Runtime/CoreUObject/Private/UObject/SavePackage/SavePackageUtilities.h b/Engine/Source/Runtime/CoreUObject/Private/UObject/SavePackage/SavePackageUtilities.h
--- a/Engine/Source/Runtime/CoreUObject/Private/UObject/SavePackage/SavePackageUtilities.h
+++ b/Engine/Source/Runtime/CoreUObject/Private/UObject/SavePackage/SavePackageUtilities.h
@@ -361,8 +361,20 @@ void AsyncWriteFile(EAsyncWriteOptions Options, FSavePackageOutputFile& File);
 
 void GetCDOSubobjects(UObject* CDO, TArray<UObject*>& Subobjects);
 
-/** Returns result of IsEditorOnlyObject if Engine:[Core.System]:CanStripEditorOnlyExportsAndImports (ini) is set to true */
-bool IsStrippedEditorOnlyObject(const UObject* InObject, bool bCheckRecursive = true, bool bCheckMarks = true);
+enum class EEditorOnlyObjectFlags
+{
+	None = 0,
+	CheckRecursive = 1 << 1,
+	ApplyHasNonEditorOnlyReferences = 1 << 2,
+	CheckMarks UE_DEPRECATED(5.3, "CheckMarks is no longer supported") = 1 << 3,
+
+};
+ENUM_CLASS_FLAGS(EEditorOnlyObjectFlags);
+
+/** Returns result of IsEditorOnlyObjectInternal if Engine:[Core.System]:CanStripEditorOnlyExportsAndImports (ini) is set to true */
+bool IsStrippedEditorOnlyObject(const UObject* InObject, EEditorOnlyObjectFlags Flags);
+
+bool IsEditorOnlyObjectInternal(const UObject* InObject, EEditorOnlyObjectFlags Flags);
 
 }
 
diff --git a/Engine/Source/Runtime/CoreUObject/Private/UObject/SavePackage2.cpp b/Engine/Source/Runtime/CoreUObject/Private/UObject/SavePackage2.cpp
--- a/Engine/Source/Runtime/CoreUObject/Private/UObject/SavePackage2.cpp
+++ b/Engine/Source/Runtime/CoreUObject/Private/UObject/SavePackage2.cpp
@@ -359,7 +359,7 @@ ESavePackageResult HarvestPackage(FSaveContext& SaveContext)
 	{
 		ForEachObjectWithPackage(SaveContext.GetPackage(), [&TryHarvestRootObject, TopLevelFlags](UObject* InObject)
 			{
-				if (InObject->HasAnyFlags(TopLevelFlags | RF_ClassDefaultObject))
+				if (InObject->HasAnyFlags(TopLevelFlags))
 				{
 					TryHarvestRootObject(InObject);
 				}
diff --git a/Engine/Source/Runtime/CoreUObject/Public/UObject/UObjectGlobals.h b/Engine/Source/Runtime/CoreUObject/Public/UObject/UObjectGlobals.h
--- a/Engine/Source/Runtime/CoreUObject/Public/UObject/UObjectGlobals.h
+++ b/Engine/Source/Runtime/CoreUObject/Public/UObject/UObjectGlobals.h
@@ -3219,7 +3219,9 @@ struct FAssetMsg
  * or 
  * - if bCheckRecursive is true, if its class, super struct, outer, or archetypes are editor only
  */
-COREUOBJECT_API bool IsEditorOnlyObject(const UObject* InObject, bool bCheckRecursive = true, bool bCheckMarks = true);
+COREUOBJECT_API bool IsEditorOnlyObject(const UObject* InObject, bool bCheckRecursive = true);
+UE_DEPRECATED(5.3, "bCheckMarks argument is no longer supported because we are transitioning away from using ObjectMarks during saving");
+COREUOBJECT_API bool IsEditorOnlyObject(const UObject* InObject, bool bCheckRecursive, bool bCheckMarks);
 #endif //WITH_EDITOR
 
 class FFieldClass;
