AddReferencedObjects: addressing additional review feedback

diff --git a/Engine/Source/Runtime/Core/Public/Misc/CoreDefines.h b/Engine/Source/Runtime/Core/Public/Misc/CoreDefines.h
--- a/Engine/Source/Runtime/Core/Public/Misc/CoreDefines.h
+++ b/Engine/Source/Runtime/Core/Public/Misc/CoreDefines.h
@@ -70,11 +70,11 @@ Can be overridden by setting to 1 or 0 in the project's .Target.cs files.
 #define UE_DEPRECATE_MUTABLE_TOBJECTPTR 0
 #endif
 
-#ifndef UE_DEPRECATE_RAW_UOBJECTPTR_ARO
-#define UE_DEPRECATE_RAW_UOBJECTPTR_ARO 0
+#ifndef UE_REFERENCE_COLLECTOR_REQUIRE_OBJECTPTR
+#define UE_REFERENCE_COLLECTOR_REQUIRE_OBJECTPTR 0
 #endif
 
-#if UE_DEPRECATE_RAW_UOBJECTPTR_ARO
+#if UE_REFERENCE_COLLECTOR_REQUIRE_OBJECTPTR
 #undef UE_DEPRECATE_MUTABLE_TOBJECTPTR
 #define UE_DEPRECATE_MUTABLE_TOBJECTPTR 1
 #endif
diff --git a/Engine/Source/Runtime/CoreUObject/Private/Misc/GCObjectReferencer.cpp b/Engine/Source/Runtime/CoreUObject/Private/Misc/GCObjectReferencer.cpp
--- a/Engine/Source/Runtime/CoreUObject/Private/Misc/GCObjectReferencer.cpp
+++ b/Engine/Source/Runtime/CoreUObject/Private/Misc/GCObjectReferencer.cpp
@@ -163,7 +163,7 @@ class FInitialReferenceCollector final : public FReferenceCollector
 {
 	TArray<UObject**>& Result;
 
-#if !UE_DEPRECATE_RAW_UOBJECTPTR_ARO
+#if !UE_REFERENCE_COLLECTOR_REQUIRE_OBJECTPTR
 	virtual void AddStableReference(UObject** Object) override
 	{
 		Result.Add(Object);
diff --git a/Engine/Source/Runtime/CoreUObject/Private/UObject/GarbageCollection.cpp b/Engine/Source/Runtime/CoreUObject/Private/UObject/GarbageCollection.cpp
--- a/Engine/Source/Runtime/CoreUObject/Private/UObject/GarbageCollection.cpp
+++ b/Engine/Source/Runtime/CoreUObject/Private/UObject/GarbageCollection.cpp
@@ -3051,7 +3051,7 @@ public:
 	virtual void HandleObjectReference(UObject*& InObject, const UObject* InReferencingObject, const FProperty* InReferencingProperty) override;
 	virtual void HandleObjectReferences(UObject** InObjects, const int32 ObjectNum, const UObject* InReferencingObject, const FProperty* InReferencingProperty) override;
 
-#if !UE_DEPRECATE_RAW_UOBJECTPTR_ARO
+#if !UE_REFERENCE_COLLECTOR_REQUIRE_OBJECTPTR
 	virtual void AddStableReference(UObject** Object) override
 	{
 		Dispatcher.QueueReference(Dispatcher.Context.GetReferencingObject(), *Object, EMemberlessId::Collector, MayKill());
@@ -3248,7 +3248,7 @@ public:
 		}
 	}
 
-#if !UE_DEPRECATE_RAW_UOBJECTPTR_ARO
+#if !UE_REFERENCE_COLLECTOR_REQUIRE_OBJECTPTR
 	virtual void AddStableReference(UObject** Object) override
 	{
 		HandleObjectReference(*Object, Context.GetReferencingObject(), nullptr);
diff --git a/Engine/Source/Runtime/CoreUObject/Private/UObject/UObjectGlobals.cpp b/Engine/Source/Runtime/CoreUObject/Private/UObject/UObjectGlobals.cpp
--- a/Engine/Source/Runtime/CoreUObject/Private/UObject/UObjectGlobals.cpp
+++ b/Engine/Source/Runtime/CoreUObject/Private/UObject/UObjectGlobals.cpp
@@ -4570,7 +4570,7 @@ public:
 
 };
 
-#if !UE_DEPRECATE_RAW_UOBJECTPTR_ARO
+#if !UE_REFERENCE_COLLECTOR_REQUIRE_OBJECTPTR
 void FReferenceCollector::AddStableReference(UObject** Object)
 {
 	AROPrivate::AddReferencedObject(*this, *Object);
@@ -4602,7 +4602,7 @@ void FReferenceCollector::AddStableReferenceSet(TSet<TObjectPtr<UObject>>* Objec
 	AddReferencedObjects(*Objects);
 }
 
-#if !UE_DEPRECATE_RAW_UOBJECTPTR_ARO
+#if !UE_REFERENCE_COLLECTOR_REQUIRE_OBJECTPTR
 void FReferenceCollector::AddReferencedObjects(const UScriptStruct*& ScriptStruct, void* StructMemory, const UObject* ReferencingObject /*= nullptr*/, const FProperty* ReferencingProperty /*= nullptr*/)
 {
 	AROPrivate::AddReferencedObjects(*this, ScriptStruct, StructMemory, ReferencingObject, ReferencingProperty);
diff --git a/Engine/Source/Runtime/CoreUObject/Public/UObject/ObjectPtr.h b/Engine/Source/Runtime/CoreUObject/Public/UObject/ObjectPtr.h
--- a/Engine/Source/Runtime/CoreUObject/Public/UObject/ObjectPtr.h
+++ b/Engine/Source/Runtime/CoreUObject/Public/UObject/ObjectPtr.h
@@ -1362,15 +1362,9 @@ FORCEINLINE To* StaticCastPtr(const TObjectPtr<From>& P)
 }
 
 template <typename T>
-FORCEINLINE TObjectPtr<T>& ConstCast(const TObjectPtr<T>& P)
+FORCEINLINE decltype(auto) ConstCast(const TObjectPtr<T>& P)
 {
-	return reinterpret_cast<TObjectPtr<T>&>(const_cast<TObjectPtr<T>&>(P));
-}
-
-template <typename T>
-FORCEINLINE TObjectPtr<T>& ConstCast(const TObjectPtr<const T>& P)
-{
-	return reinterpret_cast<TObjectPtr<T>&>(const_cast<TObjectPtr<const T>&>(P));
+	return reinterpret_cast<TObjectPtr<std::remove_cv_t<T>>&>(const_cast<TObjectPtr<T>&>(P));
 }
 
 #if UE_ENABLE_INCLUDE_ORDER_DEPRECATED_IN_5_3
diff --git a/Engine/Source/Runtime/CoreUObject/Public/UObject/UObjectGlobals.h b/Engine/Source/Runtime/CoreUObject/Public/UObject/UObjectGlobals.h
--- a/Engine/Source/Runtime/CoreUObject/Public/UObject/UObjectGlobals.h
+++ b/Engine/Source/Runtime/CoreUObject/Public/UObject/UObjectGlobals.h
@@ -2189,7 +2189,7 @@ public:
 	template <typename KeyType, typename ValueType, typename Allocator, typename KeyFuncs>
 	FORCEINLINE_DEBUGGABLE void AddStableReferenceMap(TMapBase<KeyType, ValueType, Allocator, KeyFuncs>& Map)
 	{
-#if UE_DEPRECATE_RAW_UOBJECTPTR_ARO
+#if UE_REFERENCE_COLLECTOR_REQUIRE_OBJECTPTR
 		static constexpr bool bKeyReference = TIsTObjectPtr<KeyType>::Value;
 		static constexpr bool bValueReference = TIsTObjectPtr<ValueType>::Value;
 		static_assert(bKeyReference || bValueReference);
@@ -2214,7 +2214,7 @@ public:
 		}
 	}
 
-#if !UE_DEPRECATE_RAW_UOBJECTPTR_ARO
+#if !UE_REFERENCE_COLLECTOR_REQUIRE_OBJECTPTR
 	/** Preferred way to add a reference that allows batching. Object must outlive GC tracing, can't be used for temporary/stack references. */
 	virtual void AddStableReference(UObject** Object);
 	
@@ -2267,7 +2267,7 @@ public:
 		AddStableReferenceSetFwd(reinterpret_cast<TSet<FObjectPtr>*>(Objects));
 	}
 
-#if !UE_DEPRECATE_RAW_UOBJECTPTR_ARO	
+#if !UE_REFERENCE_COLLECTOR_REQUIRE_OBJECTPTR	
 	/**
 	 * Adds object reference.
 	 *
@@ -2460,7 +2460,7 @@ public:
 	{
 		static_assert(sizeof(KeyType) > 0, "AddReferencedObjects: Keys must be pointers to a fully-defined type");
 		static_assert(TPointerIsConvertibleFromTo<KeyType, const UObjectBase>::Value, "AddReferencedObjects: Keys must be pointers to a type derived from UObject");
-		static_assert(!UE_DEPRECATE_RAW_UOBJECTPTR_ARO ||
+		static_assert(!UE_REFERENCE_COLLECTOR_REQUIRE_OBJECTPTR ||
 									!(std::is_pointer_v<ValueType> && std::is_convertible_v<ValueType, const UObjectBase*>));
 		for (auto& It : Map)
 		{
@@ -2475,7 +2475,7 @@ public:
 	{
 		static_assert(sizeof(ValueType) > 0, "AddReferencedObjects: Values must be pointers to a fully-defined type");
 		static_assert(TPointerIsConvertibleFromTo<ValueType, const UObjectBase>::Value, "AddReferencedObjects: Values must be pointers to a type derived from UObject");
-		static_assert(!UE_DEPRECATE_RAW_UOBJECTPTR_ARO ||
+		static_assert(!UE_REFERENCE_COLLECTOR_REQUIRE_OBJECTPTR ||
 									!(std::is_pointer_v<KeyType> && std::is_convertible_v<KeyType, const UObjectBase*>));
 		for (auto& It : Map)
 		{
@@ -2530,7 +2530,7 @@ public:
 	 * They're kept separate initially to maintain exact semantics while replacing the much slower
 	 * SerializeBin/TPropertyValueIterator/GetVerySlowReferenceCollectorArchive paths.
 	 */
-#if !UE_DEPRECATE_RAW_UOBJECTPTR_ARO
+#if !UE_REFERENCE_COLLECTOR_REQUIRE_OBJECTPTR
 	void AddReferencedObjects(const UScriptStruct*& ScriptStruct, void* Instance, const UObject* ReferencingObject = nullptr, const FProperty* ReferencingProperty = nullptr);
 #endif
 
diff --git a/Engine/Source/Runtime/Slate/Public/Widgets/Views/SListView.h b/Engine/Source/Runtime/Slate/Public/Widgets/Views/SListView.h
--- a/Engine/Source/Runtime/Slate/Public/Widgets/Views/SListView.h
+++ b/Engine/Source/Runtime/Slate/Public/Widgets/Views/SListView.h
@@ -2040,7 +2040,7 @@ public:
 	 */
 	virtual void AddReferencedObjects( FReferenceCollector& Collector )
 	{
-#if !UE_DEPRECATE_RAW_UOBJECTPTR_ARO
+#if !UE_REFERENCE_COLLECTOR_REQUIRE_OBJECTPTR
 		TListTypeTraits<ItemType>::AddReferencedObjects( Collector, WidgetGenerator.ItemsWithGeneratedWidgets, SelectedItems, WidgetGenerator.WidgetMapToItem );
 #endif
 	}
