Disable postload group in editor as it cause circular dependencies by trying to flush everything instead of only what's necessary
Only preload the current package instead of the whole hierarchy recursively

#rnx

diff --git a/Engine/Source/Runtime/CoreUObject/Private/Serialization/AsyncLoading2.cpp b/Engine/Source/Runtime/CoreUObject/Private/Serialization/AsyncLoading2.cpp
--- a/Engine/Source/Runtime/CoreUObject/Private/Serialization/AsyncLoading2.cpp
+++ b/Engine/Source/Runtime/CoreUObject/Private/Serialization/AsyncLoading2.cpp
@@ -2870,6 +2870,7 @@ private:
 	EEventLoadNodeExecutionResult ProcessLinkerLoadPackageSummary(FAsyncLoadingThreadState2& ThreadState);
 	bool ProcessLinkerLoadPackageImports(FAsyncLoadingThreadState2& ThreadState);
 	bool ProcessLinkerLoadPackageExports(FAsyncLoadingThreadState2& ThreadState);
+	bool ProcessLinkerLoadPackagePreloads(FAsyncLoadingThreadState2& ThreadState);
 	EEventLoadNodeExecutionResult ProcessLinkerLoadPackageImportsAndExports(FAsyncLoadingThreadState2& ThreadState);
 	EEventLoadNodeExecutionResult PostLoadLinkerLoadPackageExports(FAsyncLoadingThreadState2& ThreadState);
 #endif
@@ -3601,7 +3602,7 @@ private:
 		check(DataPtr - Data.MemoryBuffer1 == MemoryBufferSize);
 		AsyncPackage->CreateExportBundleNodes(EventSpecs.GetData());
 
-		AsyncPackage->ConstructedObjects.Reserve(Data.Exports.Num() + 1); // +1 for UPackage, may grow dynamically beoynd that
+		AsyncPackage->ConstructedObjects.Reserve(Data.Exports.Num() + 1); // +1 for UPackage, may grow dynamically beyond that
 		for (FExportObject& Export : Data.Exports)
 		{
 			Export = FExportObject();
@@ -3897,16 +3898,23 @@ FAsyncPackage2* FAsyncLoadingThread2::FindOrInsertPackage(FAsyncLoadingThreadSta
 			Package->AddCompletionCallback(MoveTemp(PackageLoadedDelegate));
 		}
 	}
+
+	// PostLoadGroup are to protect some race conditions happening in runtime when ALT is active but
+	// it is causing crashes and soft-locks in editor when FlushAsyncLoading is called on specific requests
+	// that are in the same postload group as their outer, preventing fine grained flush. So we disable this
+	// in editor for the time being since ALT is not yet activated by default and we'll sort out the race conditions instead.
+	constexpr bool bIsPostLoadGroupFeatureActive = !ALT2_ENABLE_LINKERLOAD_SUPPORT;
+
 	if (bInserted)
 	{
 		// Created a new package, either create a new post load group or use the one from the importing package
-		FAsyncLoadingPostLoadGroup* PostLoadGroup = ImportedByPackage ? ImportedByPackage->PostLoadGroup : new FAsyncLoadingPostLoadGroup();
+		FAsyncLoadingPostLoadGroup* PostLoadGroup = (bIsPostLoadGroupFeatureActive && ImportedByPackage) ? ImportedByPackage->PostLoadGroup : new FAsyncLoadingPostLoadGroup();
 		++PostLoadGroup->PackagesWithExportsToSerializeCount;
 		PostLoadGroup->Packages.Add(Package);
 		check(!Package->PostLoadGroup);
 		Package->PostLoadGroup = PostLoadGroup;
 	}
-	else if (ImportedByPackage)
+	else if (ImportedByPackage && bIsPostLoadGroupFeatureActive)
 	{
 		// Importing a package that was already being loaded
 		if (!Package->PostLoadGroup)
@@ -5233,6 +5241,30 @@ EEventLoadNodeExecutionResult FAsyncPackage2::ProcessLinkerLoadPackageSummary(FA
 	return EEventLoadNodeExecutionResult::Complete;
 }
 
+bool FAsyncPackage2::ProcessLinkerLoadPackagePreloads(FAsyncLoadingThreadState2& ThreadState)
+{
+	// Serialize exports
+	const int32 ExportCount = LinkerLoadState->Linker->ExportMap.Num();
+	while (LinkerLoadState->SerializeExportIndex < ExportCount)
+	{
+		const int32 ExportIndex = LinkerLoadState->SerializeExportIndex++;
+		FExportObject& ExportObject = Data.Exports[ExportIndex];
+		if (UObject* Object = ExportObject.Object)
+		{
+			if (Object->HasAnyFlags(RF_NeedLoad))
+			{
+				LinkerLoadState->Linker->Preload(Object);
+			}
+		}
+		if (ThreadState.IsTimeLimitExceeded(TEXT("SerializeLinkerLoadExports")))
+		{
+			return false;
+		}
+	}
+
+	return true;
+}
+
 EEventLoadNodeExecutionResult FAsyncPackage2::ProcessLinkerLoadPackageImportsAndExports(FAsyncLoadingThreadState2& ThreadState)
 {
 	ON_SCOPE_EXIT
@@ -5250,6 +5282,10 @@ EEventLoadNodeExecutionResult FAsyncPackage2::ProcessLinkerLoadPackageImportsAnd
 	{
 		return EEventLoadNodeExecutionResult::Timeout;
 	}
+	if (!ProcessLinkerLoadPackagePreloads(ThreadState))
+	{
+		return EEventLoadNodeExecutionResult::Timeout;
+	}
 
 	check(AsyncPackageLoadingState == EAsyncPackageLoadingState2::ProcessExportBundles);
 	if (ExternalReadDependencies.Num() == 0)
@@ -5394,23 +5430,6 @@ bool FAsyncPackage2::ProcessLinkerLoadPackageExports(FAsyncLoadingThreadState2&
 		}
 	}
 
-	// Serialize exports
-	while (LinkerLoadState->SerializeExportIndex < ExportCount)
-	{
-		const int32 ExportIndex = LinkerLoadState->SerializeExportIndex++;
-		FExportObject& ExportObject = Data.Exports[ExportIndex];
-		if (UObject* Object = ExportObject.Object)
-		{
-			if (Object->HasAnyFlags(RF_NeedLoad))
-			{
-				LinkerLoadState->Linker->Preload(Object);
-			}
-		}
-		if (ThreadState.IsTimeLimitExceeded(TEXT("SerializeLinkerLoadExports")))
-		{
-			return false;
-		}
-	}
 
 	return true;
 }
