Remove TagAtLastTest from SoftObjectPtr (and LazyObjectPtr) as it was only used for edge cases involving stale weak pointers and caused bugs when objects are created without calling InvalidateTag
Removing this field saves significant memory and allows removing a per-thread variable to improve performance overall
Fix a redundant dereference in SynchronousLoad and use a faster IsValid function for validating a weak pointer after resolve

diff --git a/Engine/Source/Runtime/CoreUObject/Private/UObject/LazyObjectPtr.cpp b/Engine/Source/Runtime/CoreUObject/Private/UObject/LazyObjectPtr.cpp
--- a/Engine/Source/Runtime/CoreUObject/Private/UObject/LazyObjectPtr.cpp
+++ b/Engine/Source/Runtime/CoreUObject/Private/UObject/LazyObjectPtr.cpp
@@ -93,8 +93,6 @@ FUniqueObjectGuid FUniqueObjectGuid::GetOrCreateIDForObject(const class UObject
 	return ObjectGuid;
 }
 
-FThreadSafeCounter FUniqueObjectGuid::CurrentAnnotationTag(1);
-
 /*-----------------------------------------------------------------------------------------------------------
 	FLazyObjectPtr
 -------------------------------------------------------------------------------------------------------------*/
diff --git a/Engine/Source/Runtime/CoreUObject/Private/UObject/SoftObjectPath.cpp b/Engine/Source/Runtime/CoreUObject/Private/UObject/SoftObjectPath.cpp
--- a/Engine/Source/Runtime/CoreUObject/Private/UObject/SoftObjectPath.cpp
+++ b/Engine/Source/Runtime/CoreUObject/Private/UObject/SoftObjectPath.cpp
@@ -836,7 +836,6 @@ bool FSoftObjectPathThreadContext::GetSerializationOptions(FName& OutPackageName
 	return bFoundAnything;
 }
 
-FThreadSafeCounter FSoftObjectPath::CurrentTag(1);
 TSet<FName> FSoftObjectPath::PIEPackageNames;
 
 #if WITH_LOW_LEVEL_TESTS
diff --git a/Engine/Source/Runtime/CoreUObject/Public/UObject/LazyObjectPtr.h b/Engine/Source/Runtime/CoreUObject/Public/UObject/LazyObjectPtr.h
--- a/Engine/Source/Runtime/CoreUObject/Public/UObject/LazyObjectPtr.h
+++ b/Engine/Source/Runtime/CoreUObject/Public/UObject/LazyObjectPtr.h
@@ -102,14 +102,14 @@ struct COREUOBJECT_API FUniqueObjectGuid
 		Slot << ObjectGuid.Guid;
 	}
 
-	/** Code needed by FLazyPtr internals */
+	/** These will be deprecated in a followup change as they are no longer in use */
 	static int32 GetCurrentTag()
 	{
-		return CurrentAnnotationTag.GetValue();
+		return 0;
 	}
 	static int32 InvalidateTag()
 	{
-		return CurrentAnnotationTag.Increment();
+		return 0;
 	}
 
 	static FUniqueObjectGuid GetOrCreateIDForObject(const class UObject *Object);
@@ -117,9 +117,6 @@ struct COREUOBJECT_API FUniqueObjectGuid
 private:
 	/** Guid representing the object, should be unique */
 	FGuid Guid;
-
-	/** Global counter that determines when we need to re-search for GUIDs because more objects have been loaded **/
-	static FThreadSafeCounter CurrentAnnotationTag;
 };
 
 template<> struct TIsPODType<FUniqueObjectGuid> { enum { Value = true }; };
diff --git a/Engine/Source/Runtime/CoreUObject/Public/UObject/PersistentObjectPtr.h b/Engine/Source/Runtime/CoreUObject/Public/UObject/PersistentObjectPtr.h
--- a/Engine/Source/Runtime/CoreUObject/Public/UObject/PersistentObjectPtr.h
+++ b/Engine/Source/Runtime/CoreUObject/Public/UObject/PersistentObjectPtr.h
@@ -27,20 +27,17 @@ struct TPersistentObjectPtr
 	{
 		WeakPtr.Reset();
 		ObjectID.Reset();
-		TagAtLastTest = 0;
 	}
 
 	/** Resets the weak ptr only, call this when ObjectId may change */
 	FORCEINLINE void ResetWeakPtr()
 	{
 		WeakPtr.Reset();
-		TagAtLastTest = 0;
 	}
 
 	/** Construct from a unique object identifier */
 	explicit FORCEINLINE TPersistentObjectPtr(const TObjectID& InObjectID)
 		: WeakPtr()
-		, TagAtLastTest(0)
 		, ObjectID(InObjectID)
 	{
 	}
@@ -50,7 +47,6 @@ struct TPersistentObjectPtr
 	{
 		WeakPtr.Reset();
 		ObjectID = InObjectID;
-		TagAtLastTest = 0;
 	}
 
 	/** Copy from an object pointer */
@@ -60,7 +56,6 @@ struct TPersistentObjectPtr
 		{
 			ObjectID = TObjectID::GetOrCreateIDForObject(Object);
 			WeakPtr = Object;
-			TagAtLastTest = TObjectID::GetCurrentTag();
 		}
 		else
 		{
@@ -101,20 +96,15 @@ struct TPersistentObjectPtr
 	{
 		UObject* Object = WeakPtr.Get();
 		
-		// Do a full resolve if the returned object is null and either we think we've loaded new objects, or the weak ptr may be stale
-		if (!Object && ObjectID.IsValid() && (TObjectID::GetCurrentTag() != TagAtLastTest || !WeakPtr.IsExplicitlyNull()))
+		// Do a full resolve if the cached object is null but we have a valid object ID that might resolve
+		// This used to check TObjectID::GetCurrentTag() before resolving but that was unreliable and did not improve performance in actual use
+		if (!Object && ObjectID.IsValid())
 		{
 			Object = ObjectID.ResolveObject();
 			WeakPtr = Object;
 
-			// Not safe to update tag during save as ResolveObject may have failed accidentally
-			if (Object || !GIsSavingPackage)
-			{
-				TagAtLastTest = TObjectID::GetCurrentTag();
-			}
-
-			// If this object is pending kill or otherwise invalid, this will return nullptr as expected
-			Object = WeakPtr.Get();
+			// Make sure it isn't garbage to match the default behavior of WeakPtr.Get() without looking it up again
+			return ::GetValid(Object);
 		}
 		return Object;
 	}
@@ -128,18 +118,15 @@ struct TPersistentObjectPtr
 	FORCEINLINE UObject* Get(bool bEvenIfPendingKill) const
 	{
 		UObject* Object = WeakPtr.Get(bEvenIfPendingKill);
-		if (!Object && TObjectID::GetCurrentTag() != TagAtLastTest && ObjectID.IsValid())
+
+		// Do a full resolve if the cached object is null but we have a valid object ID that might resolve
+		// This used to check TObjectID::GetCurrentTag() before resolving but that was unreliable and did not improve performance in actual use
+		if (!Object && ObjectID.IsValid())
 		{
 			Object = ObjectID.ResolveObject();
 			WeakPtr = Object;
 
-			// Not safe to update tag during save as ResolveObject may have failed accidentally
-			if (Object || !GIsSavingPackage)
-			{
-				TagAtLastTest = TObjectID::GetCurrentTag();
-			}
-
-			// If this object is pending kill or otherwise invalid, this will return nullptr as expected
+			// Get the object again using the correct flag
 			Object = WeakPtr.Get(bEvenIfPendingKill);
 		}
 		return Object;
@@ -238,8 +225,6 @@ private:
 
 	/** Once the object has been noticed to be loaded, this is set to the object weak pointer **/
 	mutable FWeakObjectPtr	WeakPtr;
-	/** Compared to CurrentAnnotationTag and if they are not equal, a guid search will be performed **/
-	mutable int32			TagAtLastTest;
 	/** Guid for the object this pointer points to or will point to. **/
 	TObjectID				ObjectID;
 };
diff --git a/Engine/Source/Runtime/CoreUObject/Public/UObject/SoftObjectPath.h b/Engine/Source/Runtime/CoreUObject/Public/UObject/SoftObjectPath.h
--- a/Engine/Source/Runtime/CoreUObject/Public/UObject/SoftObjectPath.h
+++ b/Engine/Source/Runtime/CoreUObject/Public/UObject/SoftObjectPath.h
@@ -318,14 +318,14 @@ PRAGMA_ENABLE_DEPRECATION_WARNINGS;
 		return Hash;
 	}
 
-	/** Code needed by FSoftObjectPtr internals */
+	/** These will be deprecated in a followup change as they are no longer used by FSoftObjectPtr */
 	static int32 GetCurrentTag()
 	{
-		return CurrentTag.GetValue();
+		return 0;
 	}
 	static int32 InvalidateTag()
 	{
-		return CurrentTag.Increment();
+		return 0;
 	}
 	static FSoftObjectPath GetOrCreateIDForObject(const UObject* Object);
 	
@@ -342,9 +342,6 @@ private:
 	/** Optional FString for subobject within an asset. This is the sub path after the : */
 	FString SubPathString;
 
-	/** Global counter that determines when we need to re-search based on path because more objects have been loaded **/
-	static FThreadSafeCounter CurrentTag;
-
 	/** Package names currently being duplicated, needed by FixupForPIE */
 	static TSet<FName> PIEPackageNames;
 
diff --git a/Engine/Source/Runtime/CoreUObject/Public/UObject/SoftObjectPtr.h b/Engine/Source/Runtime/CoreUObject/Public/UObject/SoftObjectPtr.h
--- a/Engine/Source/Runtime/CoreUObject/Public/UObject/SoftObjectPtr.h
+++ b/Engine/Source/Runtime/CoreUObject/Public/UObject/SoftObjectPtr.h
@@ -52,7 +52,7 @@ public:
 	UObject* LoadSynchronous() const
 	{
 		UObject* Asset = Get();
-		if (Asset == nullptr && IsPending())
+		if (Asset == nullptr && !IsNull())
 		{
 			ToSoftObjectPath().TryLoad();
 			
