Added ArchiveUObjectFromStructuredArchive to CoreUObject SharedPCH.

diff --git a/Engine/Source/Runtime/CoreUObject/Public/CoreUObjectSharedPCH.h b/Engine/Source/Runtime/CoreUObject/Public/CoreUObjectSharedPCH.h
--- a/Engine/Source/Runtime/CoreUObject/Public/CoreUObjectSharedPCH.h
+++ b/Engine/Source/Runtime/CoreUObject/Public/CoreUObjectSharedPCH.h
@@ -9,6 +9,7 @@
 #include "Misc/PackageName.h"
 #include "Misc/WorldCompositionUtility.h"
 #include "Serialization/ArchiveUObject.h"
+#include "Serialization/ArchiveUObjectFromStructuredArchive.h"
 #include "Serialization/BulkData.h"
 #include "Serialization/SerializedPropertyScope.h"
 #include "Templates/Casts.h"
