Add missing PlayInEditorLoadingScope in zenloader

#rnx

diff --git a/Engine/Source/Runtime/CoreUObject/Private/Serialization/AsyncLoading2.cpp b/Engine/Source/Runtime/CoreUObject/Private/Serialization/AsyncLoading2.cpp
--- a/Engine/Source/Runtime/CoreUObject/Private/Serialization/AsyncLoading2.cpp
+++ b/Engine/Source/Runtime/CoreUObject/Private/Serialization/AsyncLoading2.cpp
@@ -6780,6 +6780,8 @@ EEventLoadNodeExecutionResult FAsyncPackage2::Event_DeferredPostLoadExportBundle
 		FAsyncLoadingTickScope2 InAsyncLoadingTick(Package->AsyncLoadingThread);
 
 #if WITH_EDITOR
+		UE::Core::Private::FPlayInEditorLoadingScope PlayInEditorIDScope(Package->Desc.PIEInstanceID);
+
 		const FAsyncPackageHeaderData* HeaderData;
 		if (InExportBundleIndex == 1)
 		{
