changed PackageHarvester and SavePackage to use TObjectPtr for future work on Lazy Load
https://p4-swarm.epicgames.net/reviews/25856160

diff --git a/Engine/Source/Editor/UnrealEd/Private/CookOnTheFlyServer.cpp b/Engine/Source/Editor/UnrealEd/Private/CookOnTheFlyServer.cpp
--- a/Engine/Source/Editor/UnrealEd/Private/CookOnTheFlyServer.cpp
+++ b/Engine/Source/Editor/UnrealEd/Private/CookOnTheFlyServer.cpp
@@ -10089,7 +10089,7 @@ UE::Cook::FCookSavePackageContext* UCookOnTheFlyServer::CreateSaveContext(const
 
 		SavePackageSettings.AddExternalImportValidation([SuppressedNativeScriptPackages = MoveTemp(DisabledNativeScriptPackages), TargetPlatform, this](const FImportsValidationContext& ValidationContext)
 			{
-				for (UObject* Object : ValidationContext.Imports)
+				for (TObjectPtr<UObject> Object : ValidationContext.Imports)
 				{
 					UClass* Class = Cast<UClass>(Object);
 					if (Class && Class->IsNative() && SuppressedNativeScriptPackages.Contains(Class->GetPackage()->GetFName()))
diff --git a/Engine/Source/Runtime/CoreUObject/Private/UObject/LinkerSave.cpp b/Engine/Source/Runtime/CoreUObject/Private/UObject/LinkerSave.cpp
--- a/Engine/Source/Runtime/CoreUObject/Private/UObject/LinkerSave.cpp
+++ b/Engine/Source/Runtime/CoreUObject/Private/UObject/LinkerSave.cpp
@@ -238,7 +238,7 @@ int32 FLinkerSave::MapSoftObjectPath(const FSoftObjectPath& SoftObjectPath) cons
 }
 
 
-FPackageIndex FLinkerSave::MapObject( const UObject* Object ) const
+FPackageIndex FLinkerSave::MapObject(TObjectPtr<const UObject> Object) const
 {
 	if (Object)
 	{
@@ -246,7 +246,7 @@ FPackageIndex FLinkerSave::MapObject( const UObject* Object ) const
 		if (Found)
 		{
 			if (IsCooking() && CurrentlySavingExport.IsExport() &&
-				Object->GetOutermost()->GetFName() != GLongCoreUObjectPackageName && // We assume nothing in coreuobject ever loads assets in a constructor
+				Object.GetPackage().GetFName() != GLongCoreUObjectPackageName && // We assume nothing in coreuobject ever loads assets in a constructor
 				*Found != CurrentlySavingExport) // would be weird, but I can't be a dependency on myself
 			{
 				const FObjectExport& SavingExport = Exp(CurrentlySavingExport);
diff --git a/Engine/Source/Runtime/CoreUObject/Private/UObject/SavePackage/PackageHarvester.cpp b/Engine/Source/Runtime/CoreUObject/Private/UObject/SavePackage/PackageHarvester.cpp
--- a/Engine/Source/Runtime/CoreUObject/Private/UObject/SavePackage/PackageHarvester.cpp
+++ b/Engine/Source/Runtime/CoreUObject/Private/UObject/SavePackage/PackageHarvester.cpp
@@ -67,9 +67,9 @@ EObjectMark GenerateMarksForObject(const UObject* InObject, const ITargetPlatfor
 	return Marks;
 }
 
-bool ConditionallyExcludeObjectForRealm(FSaveContext& SaveContext, UObject* Obj, ESaveRealm HarvestingContext)
+bool ConditionallyExcludeObjectForRealm(FSaveContext& SaveContext, TObjectPtr<UObject> Obj, ESaveRealm HarvestingContext)
 {
-	if (!Obj || Obj->GetOutermost()->GetFName() == GLongCoreUObjectPackageName)
+	if (!Obj || Obj.GetPackage().GetFName() == GLongCoreUObjectPackageName)
 	{
 		// No object or in CoreUObject, don't exclude
 		return false;
@@ -103,7 +103,7 @@ bool ConditionallyExcludeObjectForRealm(FSaveContext& SaveContext, UObject* Obj,
 	}
 
 	// If the object outer is excluded, the object must be excluded too
-	if (ConditionallyExcludeObjectForRealm(SaveContext, Obj->GetOuter(), HarvestingContext))
+	if (ConditionallyExcludeObjectForRealm(SaveContext, Obj.GetOuter(), HarvestingContext))
 	{
 		RealmBeingChecked.AddExcluded(Obj);
 		return true;
@@ -255,7 +255,7 @@ FPackageHarvester::FHarvestScope FPackageHarvester::EnterConditionalEditorOnlySc
 	return Scope;
 }
 
-FPackageHarvester::FHarvestScope FPackageHarvester::EnterConditionalOptionalObjectScope(UObject* Object)
+FPackageHarvester::FHarvestScope FPackageHarvester::EnterConditionalOptionalObjectScope(TObjectPtr<UObject> Object)
 {
 	FHarvestScope Scope(*this);
 	if (!HasAnyExportHarvestingRealms())
@@ -302,7 +302,7 @@ FPackageHarvester::FHarvestScope FPackageHarvester::EnterNewExportOnlyScope(UObj
 	return Scope;
 }
 
-FPackageHarvester::FHarvestScope FPackageHarvester::EnterNotExcludedScope(UObject* Object)
+FPackageHarvester::FHarvestScope FPackageHarvester::EnterNotExcludedScope(TObjectPtr<UObject> Object)
 {
 	FHarvestScope Scope(*this);
 	CurrentExportHarvestingRealms.RemoveAllSwap([this, Object](ESaveRealm HarvestingRealm)
@@ -312,7 +312,7 @@ FPackageHarvester::FHarvestScope FPackageHarvester::EnterNotExcludedScope(UObjec
 	return Scope;
 }
 
-FPackageHarvester::FHarvestScope FPackageHarvester::EnterNotPreviouslyExcludedScope(UObject* Object)
+FPackageHarvester::FHarvestScope FPackageHarvester::EnterNotPreviouslyExcludedScope(TObjectPtr<UObject> Object)
 {
 	FHarvestScope Scope(*this);
 	CurrentExportHarvestingRealms.RemoveAllSwap([this, Object](ESaveRealm HarvestingRealm)
@@ -322,7 +322,7 @@ FPackageHarvester::FHarvestScope FPackageHarvester::EnterNotPreviouslyExcludedSc
 	return Scope;
 }
 
-FPackageHarvester::FHarvestScope FPackageHarvester::EnterIncludedScope(UObject* Object)
+FPackageHarvester::FHarvestScope FPackageHarvester::EnterIncludedScope(TObjectPtr<UObject> Object)
 {
 	FHarvestScope Scope(*this);
 	CurrentExportHarvestingRealms.RemoveAllSwap([this, Object](ESaveRealm HarvestingRealm)
@@ -332,7 +332,7 @@ FPackageHarvester::FHarvestScope FPackageHarvester::EnterIncludedScope(UObject*
 	return Scope;
 }
 
-bool FPackageHarvester::IsObjNative(UObject* InObj)
+bool FPackageHarvester::IsObjNative(TObjectPtr<UObject> InObj)
 {
 	bool bIsNative = InObj->IsNative();
 	UObject* Outer = InObj->GetOuter();
@@ -344,13 +344,13 @@ bool FPackageHarvester::IsObjNative(UObject* InObj)
 	return bIsNative;
 };
 
-bool FPackageHarvester::ShouldObjectBeHarvestedInOptionalRealm(UObject* InObj, FSaveContext& InContext)
+bool FPackageHarvester::ShouldObjectBeHarvestedInOptionalRealm(TObjectPtr<UObject> InObj, FSaveContext& InContext)
 {
 	if (!InContext.IsCooking())
 	{
 		return false;
 	}
-	return InObj->GetClass()->HasAnyClassFlags(CLASS_Optional);
+	return InObj.GetClass()->HasAnyClassFlags(CLASS_Optional);
 }
 
 template <typename CallbackType>
@@ -597,7 +597,7 @@ void FPackageHarvester::TryHarvestExportInternal(UObject* InObject)
 	HarvestExport(InObject);
 }
 
-void FPackageHarvester::TryHarvestImport(UObject* InObject)
+void FPackageHarvester::TryHarvestImport(TObjectPtr<UObject> InObject)
 {
 	// Those should have been already validated
 	check(InObject && !InObject->IsInPackage(SaveContext.GetPackage()));
@@ -613,9 +613,9 @@ void FPackageHarvester::TryHarvestImport(UObject* InObject)
 	}
 
 	// Do not add imports in packages excluded from cooking to any realm
-	if (InObject && FCoreUObjectDelegates::ShouldCookPackageForPlatform.IsBound())
+	if (FCoreUObjectDelegates::ShouldCookPackageForPlatform.IsBound())
 	{
-		if (!FCoreUObjectDelegates::ShouldCookPackageForPlatform.Execute(InObject->GetOutermost(), CookingTarget()))
+		if (!FCoreUObjectDelegates::ShouldCookPackageForPlatform.Execute(InObject.GetPackage(), CookingTarget()))
 		{
 			return;
 		}
@@ -632,12 +632,12 @@ void FPackageHarvester::TryHarvestImport(UObject* InObject)
 	ProcessImport(InObject);
 }
 
-void FPackageHarvester::ProcessImport(UObject* InObject)
+void FPackageHarvester::ProcessImport(TObjectPtr<UObject> InObject)
 {
 	bool bIsNative = IsObjNative(InObject);
-	UObject* ObjOuter = InObject->GetOuter();
-	UClass* ObjClass = InObject->GetClass();
-	FName ObjName = InObject->GetFName();
+	TObjectPtr<UObject> ObjOuter = InObject.GetOuter();
+	UClass* ObjClass = InObject.GetClass();
+	FName ObjName = InObject.GetFName();
 	if (SaveContext.IsCooking())
 	{
 		// The ignore dependencies check is is necessary not to have infinite recursive calls
@@ -714,7 +714,7 @@ void FPackageHarvester::MarkSearchableName(const TObjectPtr<const UObject>& Type
 
 	// Serialize object to make sure it ends up in import table
 	// This is doing a const cast to avoid backward compatibility issues
-	UObject* TempObject = const_cast<UObject*>(TypeObject.Get());
+	TObjectPtr<UObject> TempObject = TObjectPtr<UObject>(FObjectPtr(TypeObject.GetHandle()));
 	FPackageHarvester* MutableArchive = const_cast<FPackageHarvester*>(this);
 	MutableArchive->HarvestSearchableName(TempObject, ValueName);
 }
@@ -899,12 +899,12 @@ TMap<UObject*, TSet<FProperty*>> FPackageHarvester::ReleaseTransientPropertyOver
 	return MoveTemp(TransientPropertyOverrides);
 }
 
-void FPackageHarvester::HarvestDependency(UObject* InObj, bool bIsNative)
+void FPackageHarvester::HarvestDependency(TObjectPtr<UObject> InObj, bool bIsNative)
 {
 	// if we aren't currently processing an export or the referenced object is a package, do not harvest the dependency
 	if (CurrentExportDependencies.bIgnoreDependencies ||
 		CurrentExportDependencies.CurrentExport == nullptr ||
-		(InObj->GetOuter() == nullptr && InObj->GetClass()->GetFName() == NAME_Package))
+		(InObj.GetOuter() == nullptr && InObj.GetClass()->GetFName() == NAME_Package))
 	{
 		return;
 	}
@@ -919,7 +919,7 @@ void FPackageHarvester::HarvestDependency(UObject* InObj, bool bIsNative)
 	}
 }
 
-bool FPackageHarvester::CurrentExportHasDependency(UObject* InObj, ESaveRealm HarvestingRealm) const
+bool FPackageHarvester::CurrentExportHasDependency(TObjectPtr<UObject> InObj, ESaveRealm HarvestingRealm) const
 {
 	FHarvestedRealm& RealmData = SaveContext.GetHarvestedRealm(HarvestingRealm);
 	return RealmData.GetObjectDependencies().Contains(InObj) || RealmData.GetNativeObjectDependencies().Contains(InObj);
@@ -941,7 +941,7 @@ void FPackageHarvester::HarvestPackageHeaderName(FName Name)
 		});
 }
 
-void FPackageHarvester::HarvestSearchableName(UObject* TypeObject, FName Name)
+void FPackageHarvester::HarvestSearchableName(TObjectPtr<UObject> TypeObject, FName Name)
 {
 	// Make sure the object is tracked as a dependency
 	bool bAlreadyTrackedInAllHarvestingRealms = true;
@@ -977,7 +977,7 @@ void FPackageHarvester::HarvestExport(UObject* InObject)
 	ExportsToProcess.Enqueue({ InObject, CurrentExportHarvestingRealms });
 }
 
-void FPackageHarvester::HarvestImport(UObject* InObject)
+void FPackageHarvester::HarvestImport(TObjectPtr<UObject> InObject)
 {
 	ForEachExportHarvestingRealm([this, InObject](ESaveRealm HarvestingRealm)
 		{
@@ -1001,7 +1001,7 @@ void FPackageHarvester::AppendCurrentExportDependencies()
 
 FString FPackageHarvester::GetUnsaveableReason(UObject* Required, ESaveRealm RealmInWhichItIsUnsaveable)
 {
-	UObject* Culprit;
+	TObjectPtr<UObject> Culprit;
 	FString Reason;
 	ESaveableStatus Status = GetSaveableStatusForRealm(Required, RealmInWhichItIsUnsaveable, Culprit, Reason);
 	if (Status != ESaveableStatus::Success)
@@ -1022,7 +1022,7 @@ FString FPackageHarvester::GetUnsaveableReason(UObject* Required, ESaveRealm Rea
 }
 
 ESaveableStatus FPackageHarvester::GetSaveableStatusForRealm(UObject* Obj, ESaveRealm RealmInWhichItIsUnsaveable,
-	UObject*& OutCulprit, FString& OutReason)
+	TObjectPtr<UObject>& OutCulprit, FString& OutReason)
 {
 	// Copy some of the code from operator<<(UObject*), TryHarvestExport, TryHarvestImport to find out why the Required object was not included
 	if (!Obj || Obj->GetOutermost()->GetFName() == GLongCoreUObjectPackageName)
diff --git a/Engine/Source/Runtime/CoreUObject/Private/UObject/SavePackage/PackageHarvester.h b/Engine/Source/Runtime/CoreUObject/Private/UObject/SavePackage/PackageHarvester.h
--- a/Engine/Source/Runtime/CoreUObject/Private/UObject/SavePackage/PackageHarvester.h
+++ b/Engine/Source/Runtime/CoreUObject/Private/UObject/SavePackage/PackageHarvester.h
@@ -7,6 +7,7 @@
 #include "Containers/Queue.h"
 #include "Serialization/ArchiveUObject.h"
 #include "UObject/NameTypes.h"
+#include "UObject/ObjectPtr.h"
 
 enum class ESaveableStatus;
 enum class ESaveRealm : uint32;
@@ -67,10 +68,10 @@ public:
 
 	void HarvestExportDataName(FName Name);
 	void HarvestPackageHeaderName(FName Name);
-	void HarvestSearchableName(UObject* TypeObject, FName Name);
-	void HarvestDependency(UObject* InObj, bool bIsNative);
+	void HarvestSearchableName(TObjectPtr<UObject> TypeObject, FName Name);
+	void HarvestDependency(TObjectPtr<UObject> InObj, bool bIsNative);
 
-	bool CurrentExportHasDependency(UObject* InObj, ESaveRealm HarvestingRealm) const;
+	bool CurrentExportHasDependency(TObjectPtr<UObject> InObj, ESaveRealm HarvestingRealm) const;
 	TMap<UObject*, TSet<FProperty*>> ReleaseTransientPropertyOverrides();
 
 	// FArchiveUObject implementation
@@ -85,17 +86,17 @@ public:
 	virtual bool ShouldSkipProperty(const FProperty* InProperty) const override;
 
 	/** Return whether the Object and all its Outers are native. */
-	static bool IsObjNative(UObject* InObj);
+	static bool IsObjNative(TObjectPtr<UObject> InObj);
 	/** Report whether the object should be put in the optional realm even if harvesting realm is non-optional. */
-	static bool ShouldObjectBeHarvestedInOptionalRealm(UObject* InObj, FSaveContext& InSaveContext);
+	static bool ShouldObjectBeHarvestedInOptionalRealm(TObjectPtr<UObject> InObj, FSaveContext& InSaveContext);
 
 private:
 	/** Dependencies collected for the current export, and scope information about how to collect them. */
 	struct FExportDependencies
 	{
 		UObject* CurrentExport = nullptr;
-		TSet<UObject*> ObjectReferences;
-		TSet<UObject*> NativeObjectReferences;
+		TSet<TObjectPtr<UObject>> ObjectReferences;
+		TSet<TObjectPtr<UObject>> NativeObjectReferences;
 		bool bIgnoreDependencies = false;
 	};
 
@@ -136,28 +137,28 @@ private:
 
 	void TryHarvestExportInternal(UObject* InObject);
 	void HarvestExport(UObject* InObject);
-	void TryHarvestImport(UObject* InObject);
-	void HarvestImport(UObject* InObject);
+	void TryHarvestImport(TObjectPtr<UObject> InObject);
+	void HarvestImport(TObjectPtr<UObject> InObject);
 	/** Add objects/names/others that are referenced even when an object is just being imported */
-	void ProcessImport(UObject* InObject);
+	void ProcessImport(TObjectPtr<UObject> InObject);
 
 	void AppendCurrentExportDependencies();
 	FString GetUnsaveableReason(UObject* Required, ESaveRealm RealmInWhichItIsUnsaveable);
 	ESaveableStatus GetSaveableStatusForRealm(UObject* Obj, ESaveRealm RealmInWhichItIsUnsaveable,
-		UObject*& OutCulprit, FString& OutReason);
+		TObjectPtr<UObject>& OutCulprit, FString& OutReason);
 
 	/** If bIsEditorOnly is true, remove CurrentExportHarvestingRealms that do not follow editoronly references. */
 	FHarvestScope EnterConditionalEditorOnlyScope(bool bIsEditorOnly);
 	/** If the Object is optional, clear CurrentExportHarvestingRealms and populate it with the Optional Realm. */
-	FHarvestScope EnterConditionalOptionalObjectScope(UObject* Object);
+	FHarvestScope EnterConditionalOptionalObjectScope(TObjectPtr<UObject> Object);
 	/** Remove CurrentExportHarvestingRealms that already include Export. */
 	FHarvestScope EnterNewExportOnlyScope(UObject* Export);
 	/** Remove CurrentExportHarvestingRealms that return false for ConditionallyExcludeObjectForRealm(Object). */
-	FHarvestScope EnterNotExcludedScope(UObject* Object);
+	FHarvestScope EnterNotExcludedScope(TObjectPtr<UObject> Object);
 	/** Remove CurrentExportHarvestingRealms that return true for IsExcluded(Object). */
-	FHarvestScope EnterNotPreviouslyExcludedScope(UObject* Object);
+	FHarvestScope EnterNotPreviouslyExcludedScope(TObjectPtr<UObject> Object);
 	/** Remove CurrentExportHarvestingRealms that return false for IsIncluded(Object). */
-	FHarvestScope EnterIncludedScope(UObject* Object);
+	FHarvestScope EnterIncludedScope(TObjectPtr<UObject> Object);
 
 	FSaveContext& SaveContext;
 	TQueue<FExportWithContext> ExportsToProcess;
diff --git a/Engine/Source/Runtime/CoreUObject/Private/UObject/SavePackage/SaveContext.cpp b/Engine/Source/Runtime/CoreUObject/Private/UObject/SavePackage/SaveContext.cpp
--- a/Engine/Source/Runtime/CoreUObject/Private/UObject/SavePackage/SaveContext.cpp
+++ b/Engine/Source/Runtime/CoreUObject/Private/UObject/SavePackage/SaveContext.cpp
@@ -43,9 +43,9 @@ void FSaveContext::MarkUnsaveable(UObject* InObject)
 #endif
 }
 
-bool FSaveContext::IsUnsaveable(UObject* InObject, bool bEmitWarning) const
+bool FSaveContext::IsUnsaveable(TObjectPtr<UObject> InObject, bool bEmitWarning) const
 {
-	UObject* Culprit;
+	TObjectPtr<UObject> Culprit;
 	ESaveableStatus CulpritStatus;
 	ESaveableStatus Status = GetSaveableStatus(InObject, &Culprit, &CulpritStatus);
 	if (Status == ESaveableStatus::Success)
@@ -57,18 +57,18 @@ bool FSaveContext::IsUnsaveable(UObject* InObject, bool bEmitWarning) const
 	{
 		check(Culprit);
 		// Only warn if the base object is fine but the outer is invalid. If an object is itself unsaveable, the old behavior is to ignore it
-		if (InObject->GetOutermost() == GetPackage())
+		if (InObject.GetPackage() == GetPackage())
 		{
 			UE_LOG(LogSavePackage, Warning, TEXT("%s has a deprecated or abstract class outer %s, so it will not be saved"),
-				*InObject->GetFullName(), *Culprit->GetFullName());
+				*InObject.GetFullName(), *Culprit.GetFullName());
 		}
 	}
 	return true;
 }
 
-ESaveableStatus FSaveContext::GetSaveableStatus(UObject* InObject, UObject** OutCulprit, ESaveableStatus* OutCulpritStatus) const
+ESaveableStatus FSaveContext::GetSaveableStatus(TObjectPtr<UObject> InObject, TObjectPtr<UObject>* OutCulprit, ESaveableStatus* OutCulpritStatus) const
 {
-	UObject* Obj = InObject;
+	TObjectPtr<UObject> Obj = InObject;
 	while (Obj)
 	{
 		ESaveableStatus Status = GetSaveableStatusNoOuter(Obj);
@@ -84,7 +84,7 @@ ESaveableStatus FSaveContext::GetSaveableStatus(UObject* InObject, UObject** Out
 			}
 			return Obj == InObject ? Status : ESaveableStatus::OuterUnsaveable;
 		}
-		Obj = Obj->GetOuter();
+		Obj = Obj.GetOuter();
 	}
 	if (OutCulprit)
 	{
@@ -97,21 +97,21 @@ ESaveableStatus FSaveContext::GetSaveableStatus(UObject* InObject, UObject** Out
 	return ESaveableStatus::Success;
 }
 
-ESaveableStatus FSaveContext::GetSaveableStatusNoOuter(UObject* Obj) const
+ESaveableStatus FSaveContext::GetSaveableStatusNoOuter(TObjectPtr<UObject> Obj) const
 {
 	// pending kill object are unsaveable
-	if (!IsValidChecked(Obj))
+	if (Obj.IsResolved() && !IsValidChecked(Obj))
 	{
 		return ESaveableStatus::PendingKill;
 	}
 
 	// transient object are considered unsaveable if non native
-	if (Obj->HasAnyFlags(RF_Transient) && !Obj->IsNative())
+	if (Obj.IsResolved() && Obj->HasAnyFlags(RF_Transient) && !Obj->IsNative())
 	{
 		return ESaveableStatus::Transient;
 	}
 
-	UClass* Class = Obj->GetClass();
+	UClass* Class = Obj.GetClass();
 	// if the object class is abstract, has been marked as deprecated, there is a newer version that exist, or the class is marked transient, then the object is unsaveable
 	// @note: Although object instances of a transient class should definitely be unsaveable, it results in discrepancies with the old save algorithm and currently load problems
 	if (Class->HasAnyClassFlags(CLASS_Abstract | CLASS_Deprecated | CLASS_NewerVersionExists /*| CLASS_Transient*/)
@@ -140,11 +140,11 @@ FSavePackageResultStruct FSaveContext::GetFinalResult()
 
 	ResultData.SavedAssets = MoveTemp(SavedAssets);
 	UClass* PackageClass = UPackage::StaticClass();
-	for (UObject* Import : GetImports())
+	for (TObjectPtr<UObject> Import : GetImports())
 	{
-		if (Import->IsA(PackageClass))
+		if (Import.IsA(PackageClass))
 		{
-			ResultData.ImportPackages.Add(Import->GetFName());
+			ResultData.ImportPackages.Add(Import.GetFName());
 		}
 	}
 	TSet<FName>& SoftPackageReferenceList = GetSoftPackageReferenceList();
diff --git a/Engine/Source/Runtime/CoreUObject/Private/UObject/SavePackage/SaveContext.h b/Engine/Source/Runtime/CoreUObject/Private/UObject/SavePackage/SaveContext.h
--- a/Engine/Source/Runtime/CoreUObject/Private/UObject/SavePackage/SaveContext.h
+++ b/Engine/Source/Runtime/CoreUObject/Private/UObject/SavePackage/SaveContext.h
@@ -133,7 +133,7 @@ struct FHarvestedRealm
 		}
 	}
 
-	void AddImport(UObject* InObject)
+	void AddImport(TObjectPtr<UObject> InObject)
 	{
 		Imports.Add(InObject);
 	}
@@ -143,12 +143,12 @@ struct FHarvestedRealm
 		Exports.Add(MoveTemp(InTagObj));
 	}
 
-	void AddExcluded(UObject* InObject)
+	void AddExcluded(TObjectPtr<UObject> InObject)
 	{
 		Excluded.Add(InObject);
 	}
 
-	bool IsImport(UObject* InObject) const
+	bool IsImport(TObjectPtr<UObject> InObject) const
 	{
 		return Imports.Contains(InObject);
 	}
@@ -158,16 +158,16 @@ struct FHarvestedRealm
 		return Exports.Contains(InObject);
 	}
 
-	bool IsIncluded(UObject* InObject) const
+	bool IsIncluded(TObjectPtr<UObject> InObject) const
 	{
-		return IsImport(InObject) || IsExport(InObject);
+		return IsImport(InObject) || (InObject.IsResolved() && IsExport(InObject));
 	}
 
 	/**
 	 * Used during harvesting to early exit from objects we have found referenced earlier but excluded because of
 	 * editoronly or unsaveable or otherwise.
 	 */
-	bool IsExcluded(UObject* InObject) const
+	bool IsExcluded(TObjectPtr<UObject> InObject) const
 	{
 		return Excluded.Contains(InObject);
 	}
@@ -182,7 +182,7 @@ struct FHarvestedRealm
 		return Exports;
 	}
 
-	const TSet<UObject*>& GetImports() const
+	const TSet<TObjectPtr<UObject>>& GetImports() const
 	{
 		return Imports;
 	}
@@ -197,12 +197,12 @@ struct FHarvestedRealm
 		return SoftPackageReferenceList;
 	}
 
-	const TMap<UObject*, TArray<FName>>& GetSearchableNamesObjectMap() const
+	const TMap<TObjectPtr<UObject>, TArray<FName>>& GetSearchableNamesObjectMap() const
 	{
 		return SearchableNamesObjectMap;
 	}
 
-	TMap<UObject*, TArray<FName>>& GetSearchableNamesObjectMap()
+	TMap<TObjectPtr<UObject>, TArray<FName>>& GetSearchableNamesObjectMap()
 	{
 		return SearchableNamesObjectMap;
 	}
@@ -237,22 +237,22 @@ struct FHarvestedRealm
 		return SoftObjectPathList;
 	}
 
-	const TMap<UObject*, TSet<UObject*>>& GetObjectDependencies() const
+	const TMap<TObjectPtr<UObject>, TSet<TObjectPtr<UObject>>>& GetObjectDependencies() const
 	{
 		return ExportObjectDependencies;
 	}
 
-	TMap<UObject*, TSet<UObject*>>& GetObjectDependencies()
+	TMap<TObjectPtr<UObject>, TSet<TObjectPtr<UObject>>>& GetObjectDependencies()
 	{
 		return ExportObjectDependencies;
 	}
 
-	const TMap<UObject*, TSet<UObject*>>& GetNativeObjectDependencies() const
+	const TMap<TObjectPtr<UObject>, TSet<TObjectPtr<UObject>>>& GetNativeObjectDependencies() const
 	{
 		return ExportNativeObjectDependencies;
 	}
 
-	TMap<UObject*, TSet<UObject*>>& GetNativeObjectDependencies()
+	TMap<TObjectPtr<UObject>, TSet<TObjectPtr<UObject>>>& GetNativeObjectDependencies()
 	{
 		return ExportNativeObjectDependencies;
 	}
@@ -365,11 +365,11 @@ private:
 	TOptional<FString> TextFormatTempFilename;
 
 	// Set of objects excluded (import or exports) through marks or otherwise (i.e. transient flags, etc)
-	TSet<UObject*> Excluded;
+	TSet<TObjectPtr<UObject>> Excluded;
 	// Set of objects marked as export
 	TSet<FTaggedExport> Exports;
 	// Set of objects marked as import
-	TSet<UObject*> Imports;
+	TSet<TObjectPtr<UObject>> Imports;
 	// Set of names referenced from export serialization
 	TSet<FNameEntryId> NamesReferencedFromExportData;
 	// Set of names referenced from the package header (import and export table object names etc)
@@ -379,11 +379,11 @@ private:
 	// List of soft package reference found
 	TSet<FName> SoftPackageReferenceList;
 	// Map of objects to their list of searchable names
-	TMap<UObject*, TArray<FName>> SearchableNamesObjectMap;
+	TMap<TObjectPtr<UObject>, TArray<FName>> SearchableNamesObjectMap;
 	// Map of objects to their dependencies
-	TMap<UObject*, TSet<UObject*>> ExportObjectDependencies;
+	TMap<TObjectPtr<UObject>, TSet<TObjectPtr<UObject>>> ExportObjectDependencies;
 	// Map of objects to their native dependencies
-	TMap<UObject*, TSet<UObject*>> ExportNativeObjectDependencies;
+	TMap<TObjectPtr<UObject>, TSet<TObjectPtr<UObject>>> ExportNativeObjectDependencies;
 };
 
 
@@ -744,9 +744,9 @@ public:
 
 	void MarkUnsaveable(UObject* InObject);
 
-	bool IsUnsaveable(UObject* InObject, bool bEmitWarning = true) const;
-	ESaveableStatus GetSaveableStatus(UObject* InObject, UObject** OutCulprit = nullptr, ESaveableStatus* OutCulpritStatus = nullptr) const;
-	ESaveableStatus GetSaveableStatusNoOuter(UObject* InObject) const;
+	bool IsUnsaveable(TObjectPtr<UObject> InObject, bool bEmitWarning = true) const;
+	ESaveableStatus GetSaveableStatus(TObjectPtr<UObject> InObject, TObjectPtr<UObject>* OutCulprit = nullptr, ESaveableStatus* OutCulpritStatus = nullptr) const;
+	ESaveableStatus GetSaveableStatusNoOuter(TObjectPtr<UObject> InObject) const;
 
 	void RecordIllegalReference(UObject* InFrom, UObject* InTo, EIllegalRefReason InReason, FString&& InOptionalReasonText = FString())
 	{
@@ -783,7 +783,7 @@ public:
 		return GetHarvestedRealm().IsExport(InObject);
 	}
 
-	bool IsIncluded(UObject* InObject) const
+	bool IsIncluded(TObjectPtr<UObject> InObject) const
 	{
 		return GetHarvestedRealm().IsIncluded(InObject);
 	}
@@ -793,12 +793,12 @@ public:
 		return GetHarvestedRealm().GetExports();
 	}
 
-	const TSet<UObject*>& GetImports() const
+	const TSet<TObjectPtr<UObject>>& GetImports() const
 	{
 		return GetHarvestedRealm().GetImports();
 	}
 
-	const TSet<UObject*>& GetImportsUsedInGame() const
+	const TSet<TObjectPtr<UObject>>& GetImportsUsedInGame() const
 	{
 		return GetHarvestedRealm(ESaveRealm::Game).GetImports();
 	}
@@ -823,12 +823,12 @@ public:
 		return GetHarvestedRealm(ESaveRealm::Game).GetSoftPackageReferenceList();
 	}
 
-	const TMap<UObject*, TArray<FName>>& GetSearchableNamesObjectMap() const
+	const TMap<TObjectPtr<UObject>, TArray<FName>>& GetSearchableNamesObjectMap() const
 	{
 		return GetHarvestedRealm().GetSearchableNamesObjectMap();
 	}
 
-	TMap<UObject*, TArray<FName>>& GetSearchableNamesObjectMap()
+	TMap<TObjectPtr<UObject>, TArray<FName>>& GetSearchableNamesObjectMap()
 	{
 		return GetHarvestedRealm().GetSearchableNamesObjectMap();
 	}
@@ -848,12 +848,12 @@ public:
 		return GetHarvestedRealm().GetSoftObjectPathList();
 	}
 
-	const TMap<UObject*, TSet<UObject*>>& GetObjectDependencies() const
+	const TMap<TObjectPtr<UObject>, TSet<TObjectPtr<UObject>>>& GetObjectDependencies() const
 	{
 		return GetHarvestedRealm().GetObjectDependencies();
 	}
 
-	const TMap<UObject*, TSet<UObject*>>& GetNativeObjectDependencies() const
+	const TMap<TObjectPtr<UObject>, TSet<TObjectPtr<UObject>>>& GetNativeObjectDependencies() const
 	{
 		return GetHarvestedRealm().GetNativeObjectDependencies();
 	}
@@ -868,17 +868,17 @@ public:
 		return CustomVersions;
 	}
 
-	const TSet<UPackage*>& GetPrestreamPackages() const
+	const TSet<TObjectPtr<UPackage>>& GetPrestreamPackages() const
 	{
 		return PrestreamPackages;
 	}
 
-	TSet<UPackage*>& GetPrestreamPackages()
+	TSet<TObjectPtr<UPackage>>& GetPrestreamPackages()
 	{
 		return PrestreamPackages;
 	}
 
-	bool IsPrestreamPackage(UPackage* InPackage) const
+	bool IsPrestreamPackage(TObjectPtr<UPackage> InPackage) const
 	{
 		return PrestreamPackages.Contains(InPackage);
 	}
@@ -1119,7 +1119,7 @@ private:
 	TArray<FIllegalReference> HarvestedIllegalReferences;
 
 	// Set of harvested prestream packages, should be deprecated
-	TSet<UPackage*> PrestreamPackages;
+	TSet<TObjectPtr<UPackage>> PrestreamPackages;
 
 	// Set of AssetDatas created for the Assets saved into the package
 	TArray<FAssetData> SavedAssets;
diff --git a/Engine/Source/Runtime/CoreUObject/Private/UObject/SavePackage/SavePackageUtilities.cpp b/Engine/Source/Runtime/CoreUObject/Private/UObject/SavePackage/SavePackageUtilities.cpp
--- a/Engine/Source/Runtime/CoreUObject/Private/UObject/SavePackage/SavePackageUtilities.cpp
+++ b/Engine/Source/Runtime/CoreUObject/Private/UObject/SavePackage/SavePackageUtilities.cpp
@@ -883,7 +883,7 @@ void FObjectImportSortHelper::SortImports(FLinkerSave* Linker)
 	}
 
 	// Map of UObject => full name; optimization for sorting.
-	TMap<UObject*, FString> ObjectToFullNameMap;
+	TMap<TObjectPtr<UObject>, FString> ObjectToFullNameMap;
 	ObjectToFullNameMap.Reserve(Imports.Num());
 
 	for (const FObjectImport& Import : Imports)
@@ -980,7 +980,7 @@ FEDLCookChecker::FEDLNodeHash::FEDLNodeHash(const TArray<FEDLNodeData>* InNodes,
 {
 }
 
-FEDLCookChecker::FEDLNodeHash::FEDLNodeHash(const UObject* InObject, EObjectEvent InObjectEvent)
+FEDLCookChecker::FEDLNodeHash::FEDLNodeHash(TObjectPtr<UObject> InObject, EObjectEvent InObjectEvent)
 	: Object(InObject)
 	, bIsNode(false)
 	, ObjectEvent(InObjectEvent)
@@ -996,8 +996,8 @@ bool FEDLCookChecker::FEDLNodeHash::operator==(const FEDLNodeHash& Other) const
 
 	uint32 LocalNodeID;
 	uint32 OtherNodeID;
-	const UObject* LocalObject;
-	const UObject* OtherObject;
+	TObjectPtr<const UObject> LocalObject;
+	TObjectPtr<const UObject> OtherObject;
 	FName LocalName = ObjectNameFirst(*this, LocalNodeID, LocalObject);
 	FName OtherName = ObjectNameFirst(Other, OtherNodeID, OtherObject);
 
@@ -1018,7 +1018,7 @@ uint32 GetTypeHash(const FEDLCookChecker::FEDLNodeHash& A)
 	uint32 Hash = 0;
 
 	uint32 LocalNodeID;
-	const UObject* LocalObject;
+	TObjectPtr<const UObject> LocalObject;
 	FName LocalName = FEDLCookChecker::FEDLNodeHash::ObjectNameFirst(A, LocalNodeID, LocalObject);
 	do
 	{
@@ -1037,7 +1037,7 @@ FName FEDLCookChecker::FEDLNodeHash::GetName() const
 	}
 	else
 	{
-		return Object->GetFName();
+		return Object.GetFName();
 	}
 }
 
@@ -1055,7 +1055,7 @@ bool FEDLCookChecker::FEDLNodeHash::TryGetParent(FEDLCookChecker::FEDLNodeHash&
 	}
 	else
 	{
-		UObject* ParentObject = Object->GetOuter();
+		TObjectPtr<UObject> ParentObject = Object.GetOuter();
 		if (ParentObject)
 		{
 			Parent = FEDLNodeHash(ParentObject, ParentObjectEvent);
@@ -1078,7 +1078,7 @@ void FEDLCookChecker::FEDLNodeHash::SetNodes(const TArray<FEDLNodeData>* InNodes
 	}
 }
 
-FName FEDLCookChecker::FEDLNodeHash::ObjectNameFirst(const FEDLNodeHash& InNode, uint32& OutNodeID, const UObject*& OutObject)
+FName FEDLCookChecker::FEDLNodeHash::ObjectNameFirst(const FEDLNodeHash& InNode, uint32& OutNodeID,  TObjectPtr<const UObject>& OutObject)
 {
 	if (InNode.bIsNode)
 	{
@@ -1088,11 +1088,11 @@ FName FEDLCookChecker::FEDLNodeHash::ObjectNameFirst(const FEDLNodeHash& InNode,
 	else
 	{
 		OutObject = InNode.Object;
-		return OutObject->GetFName();
+		return OutObject.GetFName();
 	}
 }
 
-FName FEDLCookChecker::FEDLNodeHash::ObjectNameNext(const FEDLNodeHash& InNode, uint32& OutNodeID, const UObject*& OutObject)
+FName FEDLCookChecker::FEDLNodeHash::ObjectNameNext(const FEDLNodeHash& InNode, uint32& OutNodeID, TObjectPtr<const UObject>& OutObject)
 {
 	if (InNode.bIsNode)
 	{
@@ -1101,8 +1101,8 @@ FName FEDLCookChecker::FEDLNodeHash::ObjectNameNext(const FEDLNodeHash& InNode,
 	}
 	else
 	{
-		OutObject = OutObject->GetOuter();
-		return OutObject ? OutObject->GetFName() : NAME_None;
+		OutObject = OutObject.GetOuter();
+		return OutObject ? OutObject.GetFName() : NAME_None;
 	}
 }
 
@@ -1208,7 +1208,7 @@ void FEDLCookChecker::Reset()
 
 LLM_DEFINE_TAG(EDLCookChecker);
 
-void FEDLCookChecker::AddImport(UObject* Import, UPackage* ImportingPackage)
+void FEDLCookChecker::AddImport(TObjectPtr<UObject> Import, UPackage* ImportingPackage)
 {
 	if (bIsActive)
 	{
@@ -2073,7 +2073,7 @@ namespace UE::AssetRegistry
 
 // See the corresponding ReadPackageDataMain and ReadPackageDataDependencies defined in PackageReader.cpp in AssetRegistry module
 void WritePackageData(FStructuredArchiveRecord& ParentRecord, bool bIsCooking, const UPackage* Package,
-	FLinkerSave* Linker, const TSet<UObject*>& ImportsUsedInGame, const TSet<FName>& SoftPackagesUsedInGame,
+		FLinkerSave* Linker, const TSet<TObjectPtr<UObject>>& ImportsUsedInGame, const TSet<FName>& SoftPackagesUsedInGame,
 	const ITargetPlatform* TargetPlatform, TArray<FAssetData>* OutAssetDatas)
 {
 	// To avoid large patch sizes, we have frozen cooked package format at the format before VER_UE4_ASSETREGISTRY_DEPENDENCYFLAGS
diff --git a/Engine/Source/Runtime/CoreUObject/Private/UObject/SavePackage/SavePackageUtilities.h b/Engine/Source/Runtime/CoreUObject/Private/UObject/SavePackage/SavePackageUtilities.h
--- a/Engine/Source/Runtime/CoreUObject/Private/UObject/SavePackage/SavePackageUtilities.h
+++ b/Engine/Source/Runtime/CoreUObject/Private/UObject/SavePackage/SavePackageUtilities.h
@@ -14,6 +14,7 @@
 #include "UObject/NameTypes.h"
 #include "UObject/Package.h"
 #include "UObject/UObjectMarks.h"
+#include "UObject/ObjectPtr.h"
 
 // This file contains private utilities shared by UPackage::Save and UPackage::Save2 
 
@@ -153,7 +154,7 @@ struct FEDLCookChecker
 
 	void Reset();
 
-	void AddImport(UObject* Import, UPackage* ImportingPackage);
+	void AddImport(TObjectPtr<UObject> Import, UPackage* ImportingPackage);
 	void AddExport(UObject* Export);
 	void AddArc(UObject* DepObject, bool bDepIsSerialize, UObject* Export, bool bExportIsSerialize);
 	void AddPackageWithUnknownExports(FName LongPackageName);
@@ -188,7 +189,7 @@ public: // FEDLNodeHash is public only so that GetTypeHash can be defined
 	{
 		FEDLNodeHash(); // creates an uninitialized node; only use this to provide as an out parameter
 		FEDLNodeHash(const TArray<FEDLNodeData>* InNodes, FEDLNodeID InNodeID, EObjectEvent InObjectEvent);
-		FEDLNodeHash(const UObject* InObject, EObjectEvent InObjectEvent);
+		FEDLNodeHash(TObjectPtr<UObject> InObject, EObjectEvent InObjectEvent);
 		bool operator==(const FEDLNodeHash& Other) const;
 		friend uint32 GetTypeHash(const FEDLNodeHash& A);
 
@@ -198,8 +199,8 @@ public: // FEDLNodeHash is public only so that GetTypeHash can be defined
 		void SetNodes(const TArray<FEDLNodeData>* InNodes);
 
 	private:
-		static FName ObjectNameFirst(const FEDLNodeHash& InNode, uint32& OutNodeID, const UObject*& OutObject);
-		static FName ObjectNameNext(const FEDLNodeHash& InNode, uint32& OutNodeID, const UObject*& OutObject);
+		static FName ObjectNameFirst(const FEDLNodeHash& InNode, uint32& OutNodeID, TObjectPtr<const UObject>& OutObject);
+		static FName ObjectNameNext(const FEDLNodeHash& InNode, uint32& OutNodeID, TObjectPtr<const UObject>& OutObject);
 			
 		union
 		{
@@ -210,7 +211,7 @@ public: // FEDLNodeHash is public only so that GetTypeHash can be defined
 			 */
 			const TArray<FEDLNodeData>* Nodes;
 			/** Pointer to the Object we are looking up, if this hash was created during lookup-by-objectpath for an object */
-			const UObject* Object;
+			TObjectPtr<const UObject> Object;
 		};
 		/** The identifier for the FEDLNodeData this hash is wrapping. Only used if bIsNode is true. */
 		FEDLNodeID NodeID;
@@ -293,7 +294,7 @@ struct FEDLCookCheckerThreadState : public TThreadSingleton<FEDLCookCheckerThrea
 {
 	FEDLCookCheckerThreadState();
 
-	void AddImport(UObject* Import, UPackage* ImportingPackage)
+	void AddImport(TObjectPtr<UObject> Import, UPackage* ImportingPackage)
 	{
 		Checker.AddImport(Import, ImportingPackage);
 	}
diff --git a/Engine/Source/Runtime/CoreUObject/Private/UObject/SavePackage2.cpp b/Engine/Source/Runtime/CoreUObject/Private/UObject/SavePackage2.cpp
--- a/Engine/Source/Runtime/CoreUObject/Private/UObject/SavePackage2.cpp
+++ b/Engine/Source/Runtime/CoreUObject/Private/UObject/SavePackage2.cpp
@@ -416,9 +416,9 @@ ESavePackageResult HarvestPackage(FSaveContext& SaveContext)
 	{
 		FPackageHarvester::FHarvestScope RootReferenceScope = Harvester.EnterRootReferencesScope();
 		// Trim PrestreamPackage list 
-		TSet<UPackage*>& PrestreamPackages = SaveContext.GetPrestreamPackages();
-		TSet<UPackage*> KeptPrestreamPackages;
-		for (UPackage* Pkg : PrestreamPackages)
+		TSet<TObjectPtr<UPackage>>& PrestreamPackages = SaveContext.GetPrestreamPackages();
+		TSet<TObjectPtr<UPackage>> KeptPrestreamPackages;
+		for (TObjectPtr<UPackage> Pkg : PrestreamPackages)
 		{
 			// If the prestream package hasn't been otherwise already marked as an import, keep it as such and mark it as an import
 			if (!SaveContext.IsImport(Pkg))
@@ -789,26 +789,26 @@ ESavePackageResult ValidateImports(FSaveContext& SaveContext)
 	// Warn for private objects & map object references
 	TArray<UObject*> PrivateObjects;
 	TArray<UObject*> ObjectsInOtherMaps;
-	const TSet<UObject*>& Imports = SaveContext.GetImports();
-	for (UObject* Import : Imports)
+	const TSet<TObjectPtr<UObject>>& Imports = SaveContext.GetImports();
+	for (const TObjectPtr<UObject>& Import : Imports)
 	{
-		UPackage* ImportPackage = Import->GetPackage();
+		TObjectPtr<UPackage> ImportPackage = Import.GetPackage();
 		// All names should be properly harvested at this point
-		ensureAlwaysMsgf(SaveContext.NameExists(Import->GetFName().GetComparisonIndex())
+		ensureAlwaysMsgf(SaveContext.NameExists(Import.GetFName().GetComparisonIndex())
 			, TEXT("Missing import name %s while saving package %s. Did you rename an import during serialization?"), *Import->GetName(), *PackageName);
-		ensureAlwaysMsgf(SaveContext.NameExists(ImportPackage->GetFName().GetComparisonIndex())
+		ensureAlwaysMsgf(SaveContext.NameExists(ImportPackage.GetFName().GetComparisonIndex())
 			, TEXT("Missing import package name %s while saving package %s. Did you rename an import during serialization?"), *ImportPackage->GetName(), *PackageName);
-		ensureAlwaysMsgf(SaveContext.NameExists(Import->GetClass()->GetFName().GetComparisonIndex())
+		ensureAlwaysMsgf(SaveContext.NameExists(Import.GetClass()->GetFName().GetComparisonIndex())
 			, TEXT("Missing import class name %s while saving package %s"), *Import->GetClass()->GetName(), *PackageName);
-		ensureAlwaysMsgf(SaveContext.NameExists(Import->GetClass()->GetOuter()->GetFName().GetComparisonIndex())
-			, TEXT("Missing import class package name %s while saving package %s"), *Import->GetClass()->GetOuter()->GetName(), *PackageName);
+		ensureAlwaysMsgf(SaveContext.NameExists(Import.GetClass()->GetOuter()->GetFName().GetComparisonIndex())
+			, TEXT("Missing import class package name %s while saving package %s"), *Import.GetClass()->GetOuter()->GetName(), *PackageName);
 
 		// if the import is marked as a prestream package, we dont need to validate further
 		if (SaveContext.IsPrestreamPackage(ImportPackage))
 		{
-			ensureAlwaysMsgf(Import == ImportPackage, TEXT("Found an import refrence %s in a prestream package %s while saving package %s"), *Import->GetName(), *ImportPackage->GetName(), *PackageName);
+			ensureAlwaysMsgf(Import == ImportPackage, TEXT("Found an import refrence %s in a prestream package %s while saving package %s"), *Import.GetName(), *ImportPackage.GetName(), *PackageName);
 			// These are not errors
-			UE_LOG(LogSavePackage, Display, TEXT("Prestreaming package %s "), *ImportPackage->GetPathName()); //-V595
+			UE_LOG(LogSavePackage, Display, TEXT("Prestreaming package %s "), *ImportPackage.GetPathName()); //-V595
 			continue;
 		}
 
@@ -1189,14 +1189,14 @@ ESavePackageResult BuildLinker(FSaveContext& SaveContext)
 	{
 		SCOPED_SAVETIMER(UPackage_Save_BuildImportMap);
 
-		for (UObject* Import : SaveContext.GetImports())
+		for (TObjectPtr<UObject> Import : SaveContext.GetImports())
 		{
-			UClass* ImportClass = Import->GetClass();
+			UClass* ImportClass = Import.GetClass();
 			FName ReplacedName = NAME_None;
 			FObjectImport& ObjectImport = Linker->ImportMap.Add_GetRef(FObjectImport(Import, ImportClass));
 
 			// flag the import as optional
-			if (Import->GetClass()->HasAnyClassFlags(CLASS_Optional))
+			if (ImportClass->HasAnyClassFlags(CLASS_Optional))
 			{
 				ObjectImport.bImportOptional = true;
 			}
@@ -1276,11 +1276,11 @@ ESavePackageResult BuildLinker(FSaveContext& SaveContext)
 		{
 			UObject* Object = Linker->ExportMap[ExpIndex].Object;
 			TArray<FPackageIndex>& DependIndices = Linker->DependsMap[ExpIndex];
-			const TSet<UObject*>* SrcDepends = SaveContext.GetObjectDependencies().Find(Object);
+			const TSet<TObjectPtr<UObject>>* SrcDepends = SaveContext.GetObjectDependencies().Find(Object);
 			checkf(SrcDepends, TEXT("Couldn't find dependency map for %s"), *Object->GetFullName());
 			DependIndices.Reserve(SrcDepends->Num());
 
-			for (UObject* DependentObject : *SrcDepends)
+			for (TObjectPtr<UObject> DependentObject : *SrcDepends)
 			{
 				FPackageIndex DependencyIndex = Linker->ObjectIndicesMap.FindRef(DependentObject);
 
@@ -1298,7 +1298,7 @@ ESavePackageResult BuildLinker(FSaveContext& SaveContext)
 		Linker->SoftPackageReferenceList = SaveContext.GetSoftPackageReferenceList().Array();
 
 		// Convert the searchable names map from UObject to packageindex
-		for (TPair<UObject*, TArray<FName>>& SearchableNamePair : SaveContext.GetSearchableNamesObjectMap())
+		for (TPair<TObjectPtr<UObject>, TArray<FName>>& SearchableNamePair : SaveContext.GetSearchableNamesObjectMap())
 		{
 			const FPackageIndex PackageIndex = Linker->MapObject(SearchableNamePair.Key);
 			// This should always be in the imports already
@@ -1399,14 +1399,14 @@ ESavePackageResult BuildLinker(FSaveContext& SaveContext)
 void SavePreloadDependencies(FStructuredArchive::FRecord& StructuredArchiveRoot, FSaveContext& SaveContext)
 {
 	FLinkerSave* Linker = SaveContext.GetLinker();
-	auto IncludeObjectAsDependency = [Linker, &SaveContext](int32 CallSite, TSet<FPackageIndex>& AddTo, UObject* ToTest, UObject* ForObj, bool bMandatory, bool bOnlyIfInLinkerTable)
+	auto IncludeObjectAsDependency = [Linker, &SaveContext](int32 CallSite, TSet<FPackageIndex>& AddTo, TObjectPtr<UObject> ToTest, UObject* ForObj, bool bMandatory, bool bOnlyIfInLinkerTable)
 	{
 		// Skip transient, editor only, and excluded client/server objects
 		if (ToTest)
 		{
-			UPackage* Outermost = ToTest->GetOutermost();
+			TObjectPtr<UPackage> Outermost = ToTest.GetPackage();
 			check(Outermost);
-			if (Outermost->GetFName() == GLongCoreUObjectPackageName)
+			if (Outermost.GetFName() == GLongCoreUObjectPackageName)
 			{
 				return; // We assume nothing in coreuobject ever loads assets in a constructor
 			}
@@ -1417,16 +1417,16 @@ void SavePreloadDependencies(FStructuredArchive::FRecord& StructuredArchiveRoot,
 			}
 			if (!Index.IsNull() && (ToTest->HasAllFlags(RF_Transient) && !ToTest->IsNative()))
 			{
-				UE_LOG(LogSavePackage, Warning, TEXT("A dependency '%s' of '%s' is in the linker table, but is transient. We will keep the dependency anyway (%d)."), *ToTest->GetFullName(), *ForObj->GetFullName(), CallSite);
+				UE_LOG(LogSavePackage, Warning, TEXT("A dependency '%s' of '%s' is in the linker table, but is transient. We will keep the dependency anyway (%d)."), *ToTest.GetFullName(), *ForObj->GetFullName(), CallSite);
 			}
 			if (!Index.IsNull() && !IsValid(ToTest))
 			{
-				UE_LOG(LogSavePackage, Warning, TEXT("A dependency '%s' of '%s' is in the linker table, but is pending kill or garbage. We will keep the dependency anyway (%d)."), *ToTest->GetFullName(), *ForObj->GetFullName(), CallSite);
+				UE_LOG(LogSavePackage, Warning, TEXT("A dependency '%s' of '%s' is in the linker table, but is pending kill or garbage. We will keep the dependency anyway (%d)."), *ToTest.GetFullName(), *ForObj->GetFullName(), CallSite);
 			}
 			bool bIncludedInHarvest = SaveContext.IsIncluded(ToTest);
 			if (bMandatory && !bIncludedInHarvest)
 			{
-				UE_LOG(LogSavePackage, Warning, TEXT("A dependency '%s' of '%s' was filtered, but is mandatory. This indicates a problem with editor only stripping. We will keep the dependency anyway (%d)."), *ToTest->GetFullName(), *ForObj->GetFullName(), CallSite);
+				UE_LOG(LogSavePackage, Warning, TEXT("A dependency '%s' of '%s' was filtered, but is mandatory. This indicates a problem with editor only stripping. We will keep the dependency anyway (%d)."), *ToTest.GetFullName(), *ForObj->GetFullName(), CallSite);
 				bIncludedInHarvest = true;
 			}
 			if (bIncludedInHarvest)
@@ -1578,8 +1578,8 @@ void SavePreloadDependencies(FStructuredArchive::FRecord& StructuredArchiveRoot,
 					}
 				}
 				{
-					const TSet<UObject*>& NativeDeps = SaveContext.GetNativeObjectDependencies()[Export.Object];
-					for (UObject* ToTest : NativeDeps)
+					const TSet<TObjectPtr<UObject>>& NativeDeps = SaveContext.GetNativeObjectDependencies()[Export.Object];
+					for (TObjectPtr<UObject> ToTest : NativeDeps)
 					{
 						if (ToTest != ClassCDO)
 						{
diff --git a/Engine/Source/Runtime/CoreUObject/Public/Misc/AssetRegistryInterface.h b/Engine/Source/Runtime/CoreUObject/Public/Misc/AssetRegistryInterface.h
--- a/Engine/Source/Runtime/CoreUObject/Public/Misc/AssetRegistryInterface.h
+++ b/Engine/Source/Runtime/CoreUObject/Public/Misc/AssetRegistryInterface.h
@@ -192,7 +192,7 @@ namespace UE::AssetRegistry
 
 	// Functions to read and write the data used by the AssetRegistry in each package; the format of this data is separate from the format of the data in the asset registry
 	COREUOBJECT_API void WritePackageData(FStructuredArchiveRecord& ParentRecord, bool bIsCooking,
-		const UPackage* Package, FLinkerSave* Linker, const TSet<UObject*>& ImportsUsedInGame,
+		const UPackage* Package, FLinkerSave* Linker, const TSet<TObjectPtr<UObject>>& ImportsUsedInGame,
 		const TSet<FName>& SoftPackagesUsedInGame, const ITargetPlatform* TargetPlatform,
 		TArray<FAssetData>* OutAssetDatas);
 	// ReadPackageDataMain and ReadPackageDataDependencies are declared in IAssetRegistry.h, in the AssetRegistry module, because they depend upon some structures defined in the AssetRegistry module
diff --git a/Engine/Source/Runtime/CoreUObject/Public/Serialization/Formatters/JsonArchiveOutputFormatter.h b/Engine/Source/Runtime/CoreUObject/Public/Serialization/Formatters/JsonArchiveOutputFormatter.h
--- a/Engine/Source/Runtime/CoreUObject/Public/Serialization/Formatters/JsonArchiveOutputFormatter.h
+++ b/Engine/Source/Runtime/CoreUObject/Public/Serialization/Formatters/JsonArchiveOutputFormatter.h
@@ -87,7 +87,7 @@ public:
 	virtual void Serialize(TArray<uint8>& Value) override;
 	virtual void Serialize(void* Data, uint64 DataSize) override;
 
-	void SetObjectIndicesMap(const TMap<UObject*, FPackageIndex>* InObjectIndicesMap)
+	void SetObjectIndicesMap(const TMap<TObjectPtr<UObject>, FPackageIndex>* InObjectIndicesMap)
 	{
 		ObjectIndicesMap = InObjectIndicesMap;
 	}
@@ -95,7 +95,7 @@ public:
 private:
 	FArchive& Inner;
 
-	const TMap<UObject*, FPackageIndex>* ObjectIndicesMap = nullptr;
+	const TMap<TObjectPtr<UObject>, FPackageIndex>* ObjectIndicesMap = nullptr;
 
 	TArray<ANSICHAR> Newline;
 	bool bNeedsComma   = false;
diff --git a/Engine/Source/Runtime/CoreUObject/Public/UObject/LinkerSave.h b/Engine/Source/Runtime/CoreUObject/Public/UObject/LinkerSave.h
--- a/Engine/Source/Runtime/CoreUObject/Public/UObject/LinkerSave.h
+++ b/Engine/Source/Runtime/CoreUObject/Public/UObject/LinkerSave.h
@@ -62,7 +62,7 @@ public:
 	TArray<FPackageIndex> DepListForErrorChecking;
 
 	/** Index array - location of the resource for a UObject is stored in the ObjectIndices array using the UObject's Index */
-	TMap<UObject *,FPackageIndex> ObjectIndicesMap;
+	TMap<TObjectPtr<UObject>, FPackageIndex> ObjectIndicesMap;
 
 	/** List of Searchable Names, by object containing them. This gets turned into package indices later */
 	TMap<const UObject *, TArray<FName> > SearchableNamesObjectMap;
@@ -153,7 +153,7 @@ public:
 	int32 MapSoftObjectPath(const FSoftObjectPath& SoftObjectPath) const;
 
 	/** Returns the appropriate package index for the source object, or default value if not found in ObjectIndicesMap */
-	FPackageIndex MapObject(const UObject* Object) const;
+	FPackageIndex MapObject(TObjectPtr<const UObject> Object) const;
 
 	// FArchive interface.
 	using FArchiveUObject::operator<<; // For visibility of the overloads we don't override
diff --git a/Engine/Source/Runtime/CoreUObject/Public/UObject/SavePackage.h b/Engine/Source/Runtime/CoreUObject/Public/UObject/SavePackage.h
--- a/Engine/Source/Runtime/CoreUObject/Public/UObject/SavePackage.h
+++ b/Engine/Source/Runtime/CoreUObject/Public/UObject/SavePackage.h
@@ -137,21 +137,21 @@ public:
 	{
 	}
 
-	virtual ESavePackageResult ValidateImports(const UPackage* Package, const TSet<UObject*>& Imports) = 0;
+	virtual ESavePackageResult ValidateImports(const UPackage* Package, const TSet<TObjectPtr<UObject>>& Imports) = 0;
 };
 
 
 /** Param struct for external import validation functions */
 struct FImportsValidationContext
 {
-	FImportsValidationContext(const UPackage* InPackage, const TSet<UObject*>& InImports, FOutputDevice* InOutputDevice)
+	FImportsValidationContext(const UPackage* InPackage, const TSet<TObjectPtr<UObject>>& InImports, FOutputDevice* InOutputDevice)
 		: Package(InPackage)
 		, Imports(InImports)
 		, OutputDevice(InOutputDevice)
 	{}
 
 	const UPackage* Package;
-	const TSet<UObject*>& Imports;
+	const TSet<TObjectPtr<UObject>>& Imports;
 	FOutputDevice* OutputDevice;
 };
 
