Temporary fix so TSAN can see the barrier until we have time to remove that global barrier and do a proper profiling pass of the change on all platforms

#rnx

diff --git a/Engine/Source/Runtime/CoreUObject/Private/UObject/GarbageCollection.cpp b/Engine/Source/Runtime/CoreUObject/Private/UObject/GarbageCollection.cpp
--- a/Engine/Source/Runtime/CoreUObject/Private/UObject/GarbageCollection.cpp
+++ b/Engine/Source/Runtime/CoreUObject/Private/UObject/GarbageCollection.cpp
@@ -2389,6 +2389,9 @@ FORCEINLINE_DEBUGGABLE bool FAROQueue::TryPush(UObject* Object)
 	return true;
 }
 
+static FORCEINLINE void AcquireFence() { std::atomic_thread_fence(std::memory_order_acquire); }
+static FORCEINLINE void NoFence() {}
+
 template<void(*Fence)(), uint32 NumWanted>
 FORCEINLINE_DEBUGGABLE
 TArrayView<UObject*> FAROQueue::PopImpl()
@@ -2413,6 +2416,13 @@ TArrayView<UObject*> FAROQueue::PopImpl()
 		// Atomic-fence synchronization of non-atomic stores via relaxed Head load
 		Fence();
 
+		// This is a temporary solution to make TSAN see the fence until the global fence
+		// can be removed the next time we do a proper profiling pass of these bits on all platforms.
+		if constexpr (USING_THREAD_SANITISER && Fence == AcquireFence)
+		{
+			Head.load(std::memory_order_acquire);
+		}
+
 		UObject** TailData = &GAROBlocks[TailNow];
 		uint32 ExpectedTail = TailNow;
 		if (WantedTail <= LastInTailBlock)
@@ -2432,9 +2442,6 @@ TArrayView<UObject*> FAROQueue::PopImpl()
 	}
 }
 
-static FORCEINLINE void AcquireFence() { std::atomic_thread_fence(std::memory_order_acquire); }
-static FORCEINLINE void NoFence() {}
-
 template<class Quantity> FORCEINLINE TArrayView<UObject*> FAROQueue::Pop()		{ return PopImpl<NoFence,		NumPop<Quantity>>(); }
 template<class Quantity> FORCEINLINE TArrayView<UObject*> FAROQueue::Steal()	{ return PopImpl<AcquireFence,	NumPop<Quantity>>(); }
 
