[CoreUObject]
* Moved dll export from types to methods/staticvar

diff --git a/Engine/Source/Runtime/CoreUObject/Internal/UObject/UObjectGlobalsInternal.h b/Engine/Source/Runtime/CoreUObject/Internal/UObject/UObjectGlobalsInternal.h
--- a/Engine/Source/Runtime/CoreUObject/Internal/UObject/UObjectGlobalsInternal.h
+++ b/Engine/Source/Runtime/CoreUObject/Internal/UObject/UObjectGlobalsInternal.h
@@ -9,11 +9,11 @@
 /**
  * Global CoreUObject delegates for use by other engine modules
  */
-struct COREUOBJECT_API FCoreUObjectInternalDelegates
+struct FCoreUObjectInternalDelegates
 {
 	/** Called before GC verification code rename a package containing a World to try and prevent such errors blocking testing and development. */
 	DECLARE_MULTICAST_DELEGATE_OneParam(FPackageRename, UPackage*);
-	static FPackageRename& GetOnLeakedPackageRenameDelegate();
+	static COREUOBJECT_API FPackageRename& GetOnLeakedPackageRenameDelegate();
 };
 
-    
\ No newline at end of file
+    
diff --git a/Engine/Source/Runtime/CoreUObject/Private/Serialization/AsyncLoadingPrivate.h b/Engine/Source/Runtime/CoreUObject/Private/Serialization/AsyncLoadingPrivate.h
--- a/Engine/Source/Runtime/CoreUObject/Private/Serialization/AsyncLoadingPrivate.h
+++ b/Engine/Source/Runtime/CoreUObject/Private/Serialization/AsyncLoadingPrivate.h
@@ -9,7 +9,7 @@
 
 DECLARE_LOG_CATEGORY_EXTERN(LogLoadingDev, Fatal, All);
 
-class COREUOBJECT_API FAsyncArchive final: public FArchive
+class FAsyncArchive final: public FArchive
 {
 public:
 	enum class ELoadPhase
@@ -30,13 +30,13 @@ public:
 	};
 
 	FAsyncArchive(const FPackagePath& InPackagePath, FLinkerLoad* InOwner, TFunction<void()>&& InSummaryReadyCallback);
-	virtual ~FAsyncArchive ();
+	COREUOBJECT_API virtual ~FAsyncArchive ();
 
 	/** Archive overrides */
-	virtual bool Close() override;
-	virtual bool SetCompressionMap(TArray<FCompressedChunk>* CompressedChunks, ECompressionFlags CompressionFlags) override;
-	virtual bool Precache(int64 PrecacheOffset, int64 PrecacheSize) override;
-	virtual void Serialize(void* Data, int64 Num) override;
+	COREUOBJECT_API virtual bool Close() override;
+	COREUOBJECT_API virtual bool SetCompressionMap(TArray<FCompressedChunk>* CompressedChunks, ECompressionFlags CompressionFlags) override;
+	COREUOBJECT_API virtual bool Precache(int64 PrecacheOffset, int64 PrecacheSize) override;
+	COREUOBJECT_API virtual void Serialize(void* Data, int64 Num) override;
 	FORCEINLINE virtual int64 Tell() override
 	{
 #if DEVIRTUALIZE_FLinkerLoad_Serialize
@@ -45,23 +45,23 @@ public:
 		return CurrentPos;
 #endif
 	}
-	virtual int64 TotalSize() override;
-	virtual void Seek(int64 InPos) override;
-	virtual void FlushCache() override;
+	COREUOBJECT_API virtual int64 TotalSize() override;
+	COREUOBJECT_API virtual void Seek(int64 InPos) override;
+	COREUOBJECT_API virtual void FlushCache() override;
 	virtual FString GetArchiveName() const override 
 	{
 		return PackagePath.GetDebugName();
 	}
 
 	/** AsyncArchive interface */
-	bool PrecacheWithTimeLimit(int64 PrecacheOffset, int64 PrecacheSize, bool bUseTimeLimit, bool bUseFullTimeLimit, double TickStartTime, double TimeLimit);
-	bool PrecacheForEvent(IAsyncReadRequest* Read, int64 PrecacheOffset, int64 PrecacheSize);
-	void FlushPrecacheBlock();
-	bool ReadyToStartReadingHeader(bool bUseTimeLimit, bool bUseFullTimeLimit, double TickStartTime, double TimeLimit);
-	void StartReadingHeader();
-	void EndReadingHeader();
-	IAsyncReadRequest* MakeEventDrivenPrecacheRequest(int64 Offset, int64 BytesToRead, FAsyncFileCallBack* CompleteCallback);
-	void LogItem(const TCHAR* Item, int64 Offset = 0, int64 Size = 0, double StartTime = 0.0);
+	COREUOBJECT_API bool PrecacheWithTimeLimit(int64 PrecacheOffset, int64 PrecacheSize, bool bUseTimeLimit, bool bUseFullTimeLimit, double TickStartTime, double TimeLimit);
+	COREUOBJECT_API bool PrecacheForEvent(IAsyncReadRequest* Read, int64 PrecacheOffset, int64 PrecacheSize);
+	COREUOBJECT_API void FlushPrecacheBlock();
+	COREUOBJECT_API bool ReadyToStartReadingHeader(bool bUseTimeLimit, bool bUseFullTimeLimit, double TickStartTime, double TimeLimit);
+	COREUOBJECT_API void StartReadingHeader();
+	COREUOBJECT_API void EndReadingHeader();
+	COREUOBJECT_API IAsyncReadRequest* MakeEventDrivenPrecacheRequest(int64 Offset, int64 BytesToRead, FAsyncFileCallBack* CompleteCallback);
+	COREUOBJECT_API void LogItem(const TCHAR* Item, int64 Offset = 0, int64 Size = 0, double StartTime = 0.0);
 
 	bool IsCookedForEDLInEditor() const
 	{
@@ -83,19 +83,19 @@ private:
 	/**
 	* Updates CurrentPos based on StartFastPathLoadBuffer and sets both StartFastPathLoadBuffer and EndFastPathLoadBuffer to null
 	*/
-	void DiscardInlineBufferAndUpdateCurrentPos();
+	COREUOBJECT_API void DiscardInlineBufferAndUpdateCurrentPos();
 
-	void SetPosAndUpdatePrecacheBuffer(int64 Pos);
+	COREUOBJECT_API void SetPosAndUpdatePrecacheBuffer(int64 Pos);
 
 #endif
-	void FirstExportStarting();
-	bool WaitRead(double TimeLimit = 0.0);
-	void CompleteRead();
-	void CancelRead();
-	void CompleteCancel();
-	bool WaitForIntialPhases(double TimeLimit = 0.0);
-	void ReadCallback(bool bWasCancelled, IAsyncReadRequest*);
-	bool PrecacheInternal(int64 PrecacheOffset, int64 PrecacheSize, bool bApplyMinReadSize = true, IAsyncReadRequest* Read = nullptr);
+	COREUOBJECT_API void FirstExportStarting();
+	COREUOBJECT_API bool WaitRead(double TimeLimit = 0.0);
+	COREUOBJECT_API void CompleteRead();
+	COREUOBJECT_API void CancelRead();
+	COREUOBJECT_API void CompleteCancel();
+	COREUOBJECT_API bool WaitForIntialPhases(double TimeLimit = 0.0);
+	COREUOBJECT_API void ReadCallback(bool bWasCancelled, IAsyncReadRequest*);
+	COREUOBJECT_API bool PrecacheInternal(int64 PrecacheOffset, int64 PrecacheSize, bool bApplyMinReadSize = true, IAsyncReadRequest* Read = nullptr);
 	
 	FORCEINLINE int64 TotalSizeOrMaxInt64IfNotReady()
 	{
@@ -156,14 +156,14 @@ private:
  * to file relative offsets when exports are cooked to separate
  * archives.
  */
-class COREUOBJECT_API FLinkerExportArchive final
+class FLinkerExportArchive final
 	: public FArchiveProxy
 {
 public:
 	FLinkerExportArchive(class FLinkerLoad& InLinker, int64 InExportSerialOffset, int64 InExportSerialSize);
 
-	virtual int64 Tell() override;
-	virtual void Seek(int64 Position) override;
+	COREUOBJECT_API virtual int64 Tell() override;
+	COREUOBJECT_API virtual void Seek(int64 Position) override;
 
 private:
 	int64 ExportSerialOffset;
diff --git a/Engine/Source/Runtime/CoreUObject/Private/UObject/ObjectPathId.h b/Engine/Source/Runtime/CoreUObject/Private/UObject/ObjectPathId.h
--- a/Engine/Source/Runtime/CoreUObject/Private/UObject/ObjectPathId.h
+++ b/Engine/Source/Runtime/CoreUObject/Private/UObject/ObjectPathId.h
@@ -69,7 +69,7 @@ namespace UE::CoreUObject::Private
 
 		enum EInvalid { Invalid = 0 };
 
-		explicit COREUOBJECT_API FObjectPathId(EInvalid)
+		explicit FObjectPathId(EInvalid)
 			: Index(0)
 			, Number(SimpleNameMask | WeakObjectMask)
 		{
diff --git a/Engine/Source/Runtime/CoreUObject/Public/AssetRegistry/AssetBundleData.h b/Engine/Source/Runtime/CoreUObject/Public/AssetRegistry/AssetBundleData.h
--- a/Engine/Source/Runtime/CoreUObject/Public/AssetRegistry/AssetBundleData.h
+++ b/Engine/Source/Runtime/CoreUObject/Public/AssetRegistry/AssetBundleData.h
@@ -21,7 +21,7 @@ class FOutputDevice;
 class UObject;
 
 /** A struct representing a single AssetBundle */
-struct COREUOBJECT_API FAssetBundleEntry
+struct FAssetBundleEntry
 {
 	/** Declare constructors inline so this can be a header only class */
 	FORCEINLINE FAssetBundleEntry() {}
@@ -76,12 +76,12 @@ PRAGMA_ENABLE_DEPRECATION_WARNINGS
 		return !(*this == Other);
 	}
 
-	bool ExportTextItem(FString& ValueStr, const FAssetBundleEntry& DefaultValue, UObject* Parent, int32 PortFlags, UObject* ExportRootScope) const;
-	bool ImportTextItem(const TCHAR*& Buffer, int32 PortFlags, UObject* Parent, FOutputDevice* ErrorText);
+	COREUOBJECT_API bool ExportTextItem(FString& ValueStr, const FAssetBundleEntry& DefaultValue, UObject* Parent, int32 PortFlags, UObject* ExportRootScope) const;
+	COREUOBJECT_API bool ImportTextItem(const TCHAR*& Buffer, int32 PortFlags, UObject* Parent, FOutputDevice* ErrorText);
 };
 
 /** A struct with a list of asset bundle entries. If one of these is inside a UObject it will get automatically exported as the asset registry tag AssetBundleData */
-struct COREUOBJECT_API FAssetBundleData
+struct FAssetBundleData
 {
 	/** Declare constructors inline so this can be a header only class */
 	FORCEINLINE FAssetBundleData() {}
@@ -105,7 +105,7 @@ struct COREUOBJECT_API FAssetBundleData
 	}
 
 	/** Returns pointer to an entry with given Scope/Name */
-	FAssetBundleEntry* FindEntry(FName SearchName);
+	COREUOBJECT_API FAssetBundleEntry* FindEntry(FName SearchName);
 	
 	UE_DEPRECATED(4.27, "Bundles scopes are removed, please use FindEntry(FName) instead")
 	FAssetBundleEntry* FindEntry(const FPrimaryAssetId& SearchScope, FName SearchName)
@@ -116,7 +116,7 @@ struct COREUOBJECT_API FAssetBundleData
 
 	/** Adds or updates an entry with the given BundleName -> Path. Scope is empty and will be filled in later */
 	UE_DEPRECATED(5.1, "The contents of an asset bundle now use FTopLevelAssetPath instead of FSoftObjectPath.")
-	void AddBundleAsset(FName BundleName, const FSoftObjectPath& AssetPath);
+	COREUOBJECT_API void AddBundleAsset(FName BundleName, const FSoftObjectPath& AssetPath);
 
 	template< typename T >
 	UE_DEPRECATED(5.1, "The contents of an asset bundle now use FTopLevelAssetPath instead of FSoftObjectPath.")
@@ -127,38 +127,38 @@ struct COREUOBJECT_API FAssetBundleData
 
 	/** Adds multiple assets at once */
 	UE_DEPRECATED(5.1, "The contents of an asset bundle now use FTopLevelAssetPath instead of FSoftObjectPath.")
-	void AddBundleAssets(FName BundleName, const TArray<FSoftObjectPath>& AssetPaths);
+	COREUOBJECT_API void AddBundleAssets(FName BundleName, const TArray<FSoftObjectPath>& AssetPaths);
 
 	/** A fast set of asset bundle assets, will destroy copied in path list */
 	UE_DEPRECATED(5.1, "The contents of an asset bundle now use FTopLevelAssetPath instead of FSoftObjectPath.")
-	void SetBundleAssets(FName BundleName, TArray<FSoftObjectPath>&& AssetPaths);
+	COREUOBJECT_API void SetBundleAssets(FName BundleName, TArray<FSoftObjectPath>&& AssetPaths);
 
 	/** Adds or updates an entry with the given BundleName -> Path.*/
-	void AddBundleAsset(FName BundleName, const FTopLevelAssetPath& AssetPath);
+	COREUOBJECT_API void AddBundleAsset(FName BundleName, const FTopLevelAssetPath& AssetPath);
 
 	/** Adds multiple assets at once */
-	void AddBundleAssets(FName BundleName, const TArray<FTopLevelAssetPath>& AssetPaths);
+	COREUOBJECT_API void AddBundleAssets(FName BundleName, const TArray<FTopLevelAssetPath>& AssetPaths);
 
 	/** A fast set of asset bundle assets, will destroy copied in path list */
-	void SetBundleAssets(FName BundleName, TArray<FTopLevelAssetPath>&& AssetPaths);
+	COREUOBJECT_API void SetBundleAssets(FName BundleName, TArray<FTopLevelAssetPath>&& AssetPaths);
 
 	/** Adds or updates an entry with the given BundleName -> Path, truncating to a top-level asset path if necessary.*/
-	void AddBundleAssetTruncated(FName BundleName, const FSoftObjectPath& AssetPath);
+	COREUOBJECT_API void AddBundleAssetTruncated(FName BundleName, const FSoftObjectPath& AssetPath);
 
 	/** Adds multiple assets at once, truncating to top-level asset paths if necessary */
-	void AddBundleAssetsTruncated(FName BundleName, const TArray<FSoftObjectPath>& AssetPaths);
+	COREUOBJECT_API void AddBundleAssetsTruncated(FName BundleName, const TArray<FSoftObjectPath>& AssetPaths);
 
 	/** Replaces the contents of this bundle with the given assets, truncating to top-level asset paths if necessary. */
-	void SetBundleAssetsTruncated(FName BundleName, const TArray<FSoftObjectPath>& AssetPaths);
+	COREUOBJECT_API void SetBundleAssetsTruncated(FName BundleName, const TArray<FSoftObjectPath>& AssetPaths);
 
 	/** Resets the data to defaults */
-	void Reset();
+	COREUOBJECT_API void Reset();
 
 	/** Override Import/Export to not write out empty structs */
-	bool ExportTextItem(FString& ValueStr, FAssetBundleData const& DefaultValue, UObject* Parent, int32 PortFlags, UObject* ExportRootScope) const;
-	bool ImportTextItem(const TCHAR*& Buffer, int32 PortFlags, UObject* Parent, FOutputDevice* ErrorText);
+	COREUOBJECT_API bool ExportTextItem(FString& ValueStr, FAssetBundleData const& DefaultValue, UObject* Parent, int32 PortFlags, UObject* ExportRootScope) const;
+	COREUOBJECT_API bool ImportTextItem(const TCHAR*& Buffer, int32 PortFlags, UObject* Parent, FOutputDevice* ErrorText);
 
-	FString ToDebugString() const;
+	COREUOBJECT_API FString ToDebugString() const;
 };
 
 template<>
diff --git a/Engine/Source/Runtime/CoreUObject/Public/AssetRegistry/AssetData.h b/Engine/Source/Runtime/CoreUObject/Public/AssetRegistry/AssetData.h
--- a/Engine/Source/Runtime/CoreUObject/Public/AssetRegistry/AssetData.h
+++ b/Engine/Source/Runtime/CoreUObject/Public/AssetRegistry/AssetData.h
@@ -66,7 +66,7 @@ COREUOBJECT_API DECLARE_LOG_CATEGORY_EXTERN(LogAssetData, Log, All);
 COREUOBJECT_API extern const FName GAssetBundleDataName;
 
 /** Version used for serializing asset registry caches, both runtime and editor */
-struct COREUOBJECT_API FAssetRegistryVersion
+struct FAssetRegistryVersion
 {
 	enum Type
 	{
@@ -106,10 +106,10 @@ struct COREUOBJECT_API FAssetRegistryVersion
 	};
 
 	/** The GUID for this custom version number */
-	const static FGuid GUID;
+	COREUOBJECT_API const static FGuid GUID;
 
 	/** Read/write the custom version to the archive, should call at the very beginning */
-	static bool SerializeVersion(FArchive& Ar, FAssetRegistryVersion::Type& Version);
+	static COREUOBJECT_API bool SerializeVersion(FArchive& Ar, FAssetRegistryVersion::Type& Version);
 
 private:
 	FAssetRegistryVersion() {}
diff --git a/Engine/Source/Runtime/CoreUObject/Public/AssetRegistry/AssetDataTagMap.h b/Engine/Source/Runtime/CoreUObject/Public/AssetRegistry/AssetDataTagMap.h
--- a/Engine/Source/Runtime/CoreUObject/Public/AssetRegistry/AssetDataTagMap.h
+++ b/Engine/Source/Runtime/CoreUObject/Public/AssetRegistry/AssetDataTagMap.h
@@ -27,7 +27,7 @@ namespace FixedTagPrivate { class FStoreBuilder; }
  * [package].[object]
  * [package]
  */
-struct COREUOBJECT_API FAssetRegistryExportPath
+struct FAssetRegistryExportPath
 {
 PRAGMA_DISABLE_DEPRECATION_WARNINGS // Compilers can complain about deprecated members in compiler generated code
 	FAssetRegistryExportPath() = default;
@@ -35,10 +35,10 @@ PRAGMA_DISABLE_DEPRECATION_WARNINGS // Compilers can complain about deprecated m
 	FAssetRegistryExportPath(const FAssetRegistryExportPath&) = default;
 	FAssetRegistryExportPath& operator=(FAssetRegistryExportPath&&) = default;
 	FAssetRegistryExportPath& operator=(const FAssetRegistryExportPath&) = default;
-PRAGMA_ENABLE_DEPRECATION_WARNINGS
+COREUOBJECT_API PRAGMA_ENABLE_DEPRECATION_WARNINGS
 
 	explicit FAssetRegistryExportPath(FWideStringView String);
-	explicit FAssetRegistryExportPath(FAnsiStringView String);
+	COREUOBJECT_API explicit FAssetRegistryExportPath(FAnsiStringView String);
 
 	FTopLevelAssetPath ClassPath;
 	UE_DEPRECATED(5.1, "Class names are now represented by path names. Please use ClassPath member")
@@ -46,11 +46,11 @@ PRAGMA_ENABLE_DEPRECATION_WARNINGS
 	FName Package;
 	FName Object;
 
-	FString ToString() const;
-	FName ToName() const;
-	void ToString(FStringBuilderBase& Out) const;
-	FString ToPath() const;
-	void ToPath(FStringBuilderBase& Out) const;
+	COREUOBJECT_API FString ToString() const;
+	COREUOBJECT_API FName ToName() const;
+	COREUOBJECT_API void ToString(FStringBuilderBase& Out) const;
+	COREUOBJECT_API FString ToPath() const;
+	COREUOBJECT_API void ToPath(FStringBuilderBase& Out) const;
 	FTopLevelAssetPath ToTopLevelAssetPath() const
 	{
 		return FTopLevelAssetPath(Package, Object);
@@ -114,19 +114,19 @@ namespace FixedTagPrivate
 	};
 
 	// Handle to a tag value owned by a managed FStore
-	struct COREUOBJECT_API FValueHandle
+	struct FValueHandle
 	{
 		uint32 StoreIndex;
 		FValueId Id;
 
-		FString						AsDisplayString() const;
-		FString						AsStorageString() const;
-		FName						AsName() const;
-		FAssetRegistryExportPath	AsExportPath() const;
-		bool						AsText(FText& Out) const;
-		bool						AsMarshalledText(FMarshalledText& Out) const;
-		bool						Equals(FStringView Str) const;
-		bool						Contains(const TCHAR* Str) const;
+		COREUOBJECT_API FString						AsDisplayString() const;
+		COREUOBJECT_API FString						AsStorageString() const;
+		COREUOBJECT_API FName						AsName() const;
+		COREUOBJECT_API FAssetRegistryExportPath	AsExportPath() const;
+		COREUOBJECT_API bool						AsText(FText& Out) const;
+		COREUOBJECT_API bool						AsMarshalledText(FMarshalledText& Out) const;
+		COREUOBJECT_API bool						Equals(FStringView Str) const;
+		COREUOBJECT_API bool						Contains(const TCHAR* Str) const;
 
 	private:
 		template <bool bForStorage>
@@ -192,7 +192,7 @@ namespace FixedTagPrivate
  * Helps avoid needless FString conversions when using fixed / cooked tag values
  * that are stored as FName, FText or FAssetRegistryExportPath.
  */
-class COREUOBJECT_API FAssetTagValueRef
+class FAssetTagValueRef
 {
 	friend class FAssetDataTagMapSharedView;
 
@@ -255,17 +255,17 @@ public:
 
 	bool						IsSet() const { return Bits != 0; }
 
-	FString						AsString() const;
-	FName						AsName() const;
-	FAssetRegistryExportPath	AsExportPath() const;
-	FText						AsText() const;
-	bool						TryGetAsText(FText& Out) const; // @return false if value isn't a localized string
+	COREUOBJECT_API FString						AsString() const;
+	COREUOBJECT_API FName						AsName() const;
+	COREUOBJECT_API FAssetRegistryExportPath	AsExportPath() const;
+	COREUOBJECT_API FText						AsText() const;
+	COREUOBJECT_API bool						TryGetAsText(FText& Out) const; // @return false if value isn't a localized string
 
 	FString						GetValue() const { return AsString(); }
 	/** Coerce the type to a Complex String capable of representing the type */
 	FString						GetStorageString() const { return ToLoose(); }
 
-	bool						Equals(FStringView Str) const;
+	COREUOBJECT_API bool		Equals(FStringView Str) const;
 
 	UE_DEPRECATED(4.27, "Use AsString(), AsName(), AsExportPath() or AsText() instead. ")
 	operator FString () const { return AsString(); }
@@ -277,7 +277,7 @@ private:
 	 * Copy this's value (whether loose or fixed) into the loose format.
 	 * The returned loose value is in StorageFormat (e.g. complex strings) rather than display format.
 	 */
-	FString						ToLoose() const;
+	COREUOBJECT_API FString						ToLoose() const;
 
 	friend class FixedTagPrivate::FStoreBuilder;
 	friend FAssetRegistryState;
diff --git a/Engine/Source/Runtime/CoreUObject/Public/AssetRegistry/AssetDataTagMapSerializationDetails.h b/Engine/Source/Runtime/CoreUObject/Public/AssetRegistry/AssetDataTagMapSerializationDetails.h
--- a/Engine/Source/Runtime/CoreUObject/Public/AssetRegistry/AssetDataTagMapSerializationDetails.h
+++ b/Engine/Source/Runtime/CoreUObject/Public/AssetRegistry/AssetDataTagMapSerializationDetails.h
@@ -22,18 +22,18 @@ namespace FixedTagPrivate
 	 * The AssetRegistry's representation of an FText AssetData Tag value.
 	 * It can be stored and copied without being interpreted as an FText.
 	 */
-	class COREUOBJECT_API FMarshalledText
+	class FMarshalledText
 	{
 	public:
 		FMarshalledText() = default;
-		explicit FMarshalledText(const FString& InString);
-		explicit FMarshalledText(FString&& InString);
-		explicit FMarshalledText(const FText& InText);
-		explicit FMarshalledText(FText&& InText);
+		COREUOBJECT_API explicit FMarshalledText(const FString& InString);
+		COREUOBJECT_API explicit FMarshalledText(FString&& InString);
+		COREUOBJECT_API explicit FMarshalledText(const FText& InText);
+		COREUOBJECT_API explicit FMarshalledText(FText&& InText);
 
-		const FString& GetAsComplexString() const;
-		FText GetAsText() const;
-		int32 CompareToCaseIgnored(const FMarshalledText& Other) const;
+		COREUOBJECT_API const FString& GetAsComplexString() const;
+		COREUOBJECT_API FText GetAsText() const;
+		COREUOBJECT_API int32 CompareToCaseIgnored(const FMarshalledText& Other) const;
 
 	private:
 		FString String;
@@ -255,4 +255,4 @@ namespace FixedTagPrivate
 		TOptional<ELoadOrder> Order;
 	};
 
-} // end namespace FixedTagPrivate
\ No newline at end of file
+} // end namespace FixedTagPrivate
diff --git a/Engine/Source/Runtime/CoreUObject/Public/Blueprint/BlueprintSupport.h b/Engine/Source/Runtime/CoreUObject/Public/Blueprint/BlueprintSupport.h
--- a/Engine/Source/Runtime/CoreUObject/Public/Blueprint/BlueprintSupport.h
+++ b/Engine/Source/Runtime/CoreUObject/Public/Blueprint/BlueprintSupport.h
@@ -30,40 +30,40 @@ class ITargetPlatform;
  * List of asset registry tags used by blueprints. These are here so they can be used by both the asset registry and blueprint code
  * These need to be kept in sync with UBlueprint::GetAssetRegistryTags, and any changes there will require resaving content
  */
-struct COREUOBJECT_API FBlueprintTags
+struct FBlueprintTags
 {
 	/** Full path in export form ClassType'/PackagePath/PackageName.ClassName' of generated blueprint class */
-	static const FName GeneratedClassPath;
+	static COREUOBJECT_API const FName GeneratedClassPath;
 	/** Full path in export form ClassType'/PackagePath/PackageName.ClassName' of the immediate parent, may be a blueprint or native class */
-	static const FName ParentClassPath;
+	static COREUOBJECT_API const FName ParentClassPath;
 	/** Full path in export form Class'/Script/ModuleName.ClassName' of the first found parent native class */
-	static const FName NativeParentClassPath;
+	static COREUOBJECT_API const FName NativeParentClassPath;
 	/** Integer representing bitfield EClassFlags */
-	static const FName ClassFlags;
+	static COREUOBJECT_API const FName ClassFlags;
 	/** String representing enum EBlueprintType */
-	static const FName BlueprintType;
+	static COREUOBJECT_API const FName BlueprintType;
 	/** String with user-entered description of blueprint */
-	static const FName BlueprintDescription;
+	static COREUOBJECT_API const FName BlueprintDescription;
 	/** String with user-entered display name for the blueprint class (used in editor along with the description to identify the Blueprint type) */
-	static const FName BlueprintDisplayName;
+	static COREUOBJECT_API const FName BlueprintDisplayName;
 	/** String with user-entered category for the blueprint */
-	static const FName BlueprintCategory;
+	static COREUOBJECT_API const FName BlueprintCategory;
 	/** String set to True/False, set if this is a data only blueprint */
-	static const FName IsDataOnly;
+	static COREUOBJECT_API const FName IsDataOnly;
 	/** List of implemented interfaces, must be converted to FBPInterfaceDescription */
-	static const FName ImplementedInterfaces;
+	static COREUOBJECT_API const FName ImplementedInterfaces;
 	/** Very large string used to store find in blueprint data for the editor */
-	static const FName FindInBlueprintsData;
+	static COREUOBJECT_API const FName FindInBlueprintsData;
 	/** (Deprecated) Legacy tag that was initially used to store find in blueprint data for the editor */
-	static const FName UnversionedFindInBlueprintsData;
+	static COREUOBJECT_API const FName UnversionedFindInBlueprintsData;
 	/** Number of replicated properties */
-	static const FName NumReplicatedProperties;
+	static COREUOBJECT_API const FName NumReplicatedProperties;
 	/** Number of native components */
-	static const FName NumNativeComponents;
+	static COREUOBJECT_API const FName NumNativeComponents;
 	/** Number of blueprint components */
-	static const FName NumBlueprintComponents;
+	static COREUOBJECT_API const FName NumBlueprintComponents;
 	/** The subpath of a blueprint contained within the asset. Used to determine whether, and where a blueprint exists in a package. */
-	static const FName BlueprintPathWithinPackage;
+	static COREUOBJECT_API const FName BlueprintPathWithinPackage;
 };
 
 struct FBlueprintWarningDeclaration
@@ -146,11 +146,11 @@ struct FBlueprintSupport
  * this registers raw addresses for tracking. This is somewhat less safe, make
  * sure to not register addresses that may change
  */
-struct COREUOBJECT_API FScopedPlaceholderRawContainerTracker
+struct FScopedPlaceholderRawContainerTracker
 {
 public:
-	FScopedPlaceholderRawContainerTracker(void* InData);
-	~FScopedPlaceholderRawContainerTracker();
+	COREUOBJECT_API FScopedPlaceholderRawContainerTracker(void* InData);
+	COREUOBJECT_API ~FScopedPlaceholderRawContainerTracker();
 
 private:
 	void* Data;
@@ -165,11 +165,11 @@ private:
  * recompile any classes that may depend on each other a second time to ensure that that functions and properties
  * are properly resolved
  */
-struct COREUOBJECT_API FScopedClassDependencyGather
+struct FScopedClassDependencyGather
 {
 public:
-	FScopedClassDependencyGather(UClass* ClassToGather, FUObjectSerializeContext* InLoadContext);
-	~FScopedClassDependencyGather();
+	COREUOBJECT_API FScopedClassDependencyGather(UClass* ClassToGather, FUObjectSerializeContext* InLoadContext);
+	COREUOBJECT_API ~FScopedClassDependencyGather();
 
 	/**
 	 * Post load, some systems would like an easy list of dependencies. This will
@@ -178,7 +178,7 @@ public:
 	 * 
 	 * @return The most recent array of tracked dependencies.
 	 */
-	static TArray<UClass*> const& GetCachedDependencies();
+	static COREUOBJECT_API TArray<UClass*> const& GetCachedDependencies();
 
 private:
 	/** Whether or not this dependency gather is the dependency authoritative class, and thus should process all dependencies in the destructor */
diff --git a/Engine/Source/Runtime/CoreUObject/Public/Internationalization/PackageLocalizationCache.h b/Engine/Source/Runtime/CoreUObject/Public/Internationalization/PackageLocalizationCache.h
--- a/Engine/Source/Runtime/CoreUObject/Public/Internationalization/PackageLocalizationCache.h
+++ b/Engine/Source/Runtime/CoreUObject/Public/Internationalization/PackageLocalizationCache.h
@@ -16,7 +16,7 @@
 class FPackageLocalizationCache;
 
 /** Package localization cache for a specific culture (may contain a chain of cultures internally based on their priority) */
-class COREUOBJECT_API FPackageLocalizationCultureCache
+class FPackageLocalizationCultureCache
 {
 public:
 	/**
@@ -25,26 +25,26 @@ public:
 	 * @param InOwnerCache		A pointer to our owner cache. Used to perform an implementation specific search for localized packages (see FPackageLocalizationCache::FindLocalizedPackages).
 	 * @param InCultureName		The name of the culture this cache is for.
 	 */
-	FPackageLocalizationCultureCache(FPackageLocalizationCache* InOwnerCache, const FString& InCultureName);
+	COREUOBJECT_API FPackageLocalizationCultureCache(FPackageLocalizationCache* InOwnerCache, const FString& InCultureName);
 
 	/**
 	 * Update this cache, but only if it is dirty.
 	 */
-	void ConditionalUpdateCache();
+	COREUOBJECT_API void ConditionalUpdateCache();
 
 	/**
 	 * Add a source path to be scanned the next time ConditionalUpdateCache is called.
 	 *
 	 * @param InRootPath		The root source path to add.
 	 */
-	void AddRootSourcePath(const FString& InRootPath);
+	COREUOBJECT_API void AddRootSourcePath(const FString& InRootPath);
 
 	/**
 	 * Remove a source path. This will remove it from the pending list, as well as immediately remove any localized package entries mapped from this source root.
 	 *
 	 * @param InRootPath		The root source path to remove.
 	 */
-	void RemoveRootSourcePath(const FString& InRootPath);
+	COREUOBJECT_API void RemoveRootSourcePath(const FString& InRootPath);
 
 	/**
 	 * Add a package (potentially source or localized) to this cache.
@@ -52,7 +52,7 @@ public:
 	 * @param InPackageName		The name of the package to add.
 	 * @return TRUE if the supplied package is a localized asset and was added to the culture cache, FALSE otherwise
 	 */
-	bool AddPackage(const FString& InPackageName);
+	COREUOBJECT_API bool AddPackage(const FString& InPackageName);
 
 	/**
 	 * Remove a package (potentially source or localized) from this cache.
@@ -60,12 +60,12 @@ public:
 	 * @param InPackageName		The name of the package to remove.
 	 * @return TRUE if the supplied package is a localized asset and was removed from the culture cache, FALSE otherwise
 	 */
-	bool RemovePackage(const FString& InPackageName);
+	COREUOBJECT_API bool RemovePackage(const FString& InPackageName);
 
 	/**
 	 * Restore this cache to an empty state.
 	 */
-	void Empty();
+	COREUOBJECT_API void Empty();
 
 	/**
 	 * Try and find the localized package name for the given source package for culture we represent.
@@ -74,7 +74,7 @@ public:
 	 *
 	 * @return The localized package name, or NAME_None if there is no localized package.
 	 */
-	FName FindLocalizedPackageName(const FName InSourcePackageName);
+	COREUOBJECT_API FName FindLocalizedPackageName(const FName InSourcePackageName);
 
 private:
 	/**
@@ -103,18 +103,18 @@ private:
 };
 
 /** Common implementation for the package localization cache */
-class COREUOBJECT_API FPackageLocalizationCache : public IPackageLocalizationCache
+class FPackageLocalizationCache : public IPackageLocalizationCache
 {
 	friend class FPackageLocalizationCultureCache;
 
 public:
-	FPackageLocalizationCache();
-	virtual ~FPackageLocalizationCache();
+	COREUOBJECT_API FPackageLocalizationCache();
+	COREUOBJECT_API virtual ~FPackageLocalizationCache();
 
 	//~ IPackageLocalizationCache interface
-	virtual void ConditionalUpdateCache() override;
-	virtual FName FindLocalizedPackageName(const FName InSourcePackageName) override;
-	virtual FName FindLocalizedPackageNameForCulture(const FName InSourcePackageName, const FString& InCultureName) override;
+	COREUOBJECT_API virtual void ConditionalUpdateCache() override;
+	COREUOBJECT_API virtual FName FindLocalizedPackageName(const FName InSourcePackageName) override;
+	COREUOBJECT_API virtual FName FindLocalizedPackageNameForCulture(const FName InSourcePackageName, const FString& InCultureName) override;
 
 protected:
 	/**
@@ -142,12 +142,12 @@ protected:
 	 *
 	 * @return A pointer to the cache for the given culture.
 	 */
-	TSharedPtr<FPackageLocalizationCultureCache> FindOrAddCacheForCulture_NoLock(const FString& InCultureName);
+	COREUOBJECT_API TSharedPtr<FPackageLocalizationCultureCache> FindOrAddCacheForCulture_NoLock(const FString& InCultureName);
 
 	/**
 	 * Update the mapping of package names to asset groups (if required).
 	 */
-	void ConditionalUpdatePackageNameToAssetGroupCache_NoLock();
+	COREUOBJECT_API void ConditionalUpdatePackageNameToAssetGroupCache_NoLock();
 
 	/**
 	 * Callback handler for when a new content path is mounted.
@@ -155,7 +155,7 @@ protected:
 	 * @param InAssetPath		The package path that was mounted, eg) /Game
 	 * @param InFilesystemPath	The file-system path for the asset path, eg) ../../../MyGame/Content
 	 */
-	void HandleContentPathMounted(const FString& InAssetPath, const FString& InFilesystemPath);
+	COREUOBJECT_API void HandleContentPathMounted(const FString& InAssetPath, const FString& InFilesystemPath);
 
 	/**
 	 * Callback handler for when an existing content path is dismounted.
@@ -163,12 +163,12 @@ protected:
 	 * @param InAssetPath		The package path that was dismounted, eg) /Game
 	 * @param InFilesystemPath	The file-system path for the asset path, eg) ../../../MyGame/Content
 	 */
-	void HandleContentPathDismounted(const FString& InAssetPath, const FString& InFilesystemPath);
+	COREUOBJECT_API void HandleContentPathDismounted(const FString& InAssetPath, const FString& InFilesystemPath);
 
 	/**
 	 * Callback handler for when the active culture is changed.
 	 */
-	void HandleCultureChanged();
+	COREUOBJECT_API void HandleCultureChanged();
 
 protected:
 	/** Critical section preventing concurrent access to our data. */
diff --git a/Engine/Source/Runtime/CoreUObject/Public/Internationalization/PackageLocalizationManager.h b/Engine/Source/Runtime/CoreUObject/Public/Internationalization/PackageLocalizationManager.h
--- a/Engine/Source/Runtime/CoreUObject/Public/Internationalization/PackageLocalizationManager.h
+++ b/Engine/Source/Runtime/CoreUObject/Public/Internationalization/PackageLocalizationManager.h
@@ -12,7 +12,7 @@
 class IPackageLocalizationCache;
 
 /** Singleton class that manages localized package data. */
-class COREUOBJECT_API FPackageLocalizationManager
+class FPackageLocalizationManager
 {
 public:
 	typedef TFunction<void(FPackageLocalizationManager&)> FLazyInitFunc;
@@ -20,26 +20,26 @@ public:
 	/**
 	 * Initialize the manager from the callback set by InitializeFromLazyCallback. It is expected that this callback calls one of the InitializeFromX functions.
 	 */
-	void PerformLazyInitialization();
+	COREUOBJECT_API void PerformLazyInitialization();
 	
 	/**
 	 * Initialize the manager lazily using the given callback. It is expected that this callback calls one of the InitializeFromX functions.
 	 *
 	 * @param InLazyInitFunc	The function to call to initialize the manager.
 	 */
-	void InitializeFromLazyCallback(FLazyInitFunc InLazyInitFunc);
+	COREUOBJECT_API void InitializeFromLazyCallback(FLazyInitFunc InLazyInitFunc);
 
 	/**
 	 * Initialize the manager using the given cache. This will perform an initial scan for localized packages.
 	 *
 	 * @param InCache	The cache the manager should use.
 	 */
-	void InitializeFromCache(const TSharedRef<IPackageLocalizationCache>& InCache);
+	COREUOBJECT_API void InitializeFromCache(const TSharedRef<IPackageLocalizationCache>& InCache);
 
 	/**
 	 * Initialize this manager using the default cache. This will perform an initial scan for localized packages.
 	 */
-	void InitializeFromDefaultCache();
+	COREUOBJECT_API void InitializeFromDefaultCache();
 
 	/**
 	 * Try and find the localized package name for the given source package for the active culture.
@@ -48,7 +48,7 @@ public:
 	 *
 	 * @return The localized package name, or NAME_None if there is no localized package.
 	 */
-	FName FindLocalizedPackageName(const FName InSourcePackageName);
+	COREUOBJECT_API FName FindLocalizedPackageName(const FName InSourcePackageName);
 
 	/**
 	 * Try and find the localized package name for the given source package for the given culture.
@@ -58,19 +58,19 @@ public:
 	 *
 	 * @return The localized package name, or NAME_None if there is no localized package.
 	 */
-	FName FindLocalizedPackageNameForCulture(const FName InSourcePackageName, const FString& InCultureName);
+	COREUOBJECT_API FName FindLocalizedPackageNameForCulture(const FName InSourcePackageName, const FString& InCultureName);
 
 	/**
 	 * Update this cache, but only if it is dirty.
 	 */
-	void ConditionalUpdateCache();
+	COREUOBJECT_API void ConditionalUpdateCache();
 
 	/**
 	 * Singleton accessor.
 	 *
 	 * @return The singleton instance of the localization manager.
 	 */
-	static FPackageLocalizationManager& Get();
+	static COREUOBJECT_API FPackageLocalizationManager& Get();
 
 private:
 	/**
diff --git a/Engine/Source/Runtime/CoreUObject/Public/Internationalization/PackageLocalizationUtil.h b/Engine/Source/Runtime/CoreUObject/Public/Internationalization/PackageLocalizationUtil.h
--- a/Engine/Source/Runtime/CoreUObject/Public/Internationalization/PackageLocalizationUtil.h
+++ b/Engine/Source/Runtime/CoreUObject/Public/Internationalization/PackageLocalizationUtil.h
@@ -6,7 +6,7 @@
 #include "CoreMinimal.h"
 
 /** Utility functions for dealing with localized package names */
-struct COREUOBJECT_API FPackageLocalizationUtil
+struct FPackageLocalizationUtil
 {
 	/**
 	 * Converts a localized version of a package path to the source version (by removing /L10N/<code> from the package path, if present)
@@ -17,7 +17,7 @@ struct COREUOBJECT_API FPackageLocalizationUtil
 	 *
 	 * @returns True if the conversion happened, false otherwise
 	 */
-	static bool ConvertLocalizedToSource(const FString& InLocalized, FString& OutSource);
+	static COREUOBJECT_API bool ConvertLocalizedToSource(const FString& InLocalized, FString& OutSource);
 
 	/**
 	 * Converts a source version of a package path to the localized version for the given culture (by adding /L10N/<code> to the package path)
@@ -29,7 +29,7 @@ struct COREUOBJECT_API FPackageLocalizationUtil
 	 *
 	 * @returns True if the conversion happened, false otherwise
 	 */
-	static bool ConvertSourceToLocalized(const FString& InSource, const FString& InCulture, FString& OutLocalized);
+	static COREUOBJECT_API bool ConvertSourceToLocalized(const FString& InSource, const FString& InCulture, FString& OutLocalized);
 
 	/**
 	 * Given a package path, get the localized root package for the given culture (eg, if given "/Game/MyFolder/MyAsset" and a culture of "fr", this would return "/Game/L10N/fr")
@@ -40,5 +40,5 @@ struct COREUOBJECT_API FPackageLocalizationUtil
 	 *
 	 * @returns True if the conversion happened, false otherwise
 	 */
-	static bool GetLocalizedRoot(const FString& InPath, const FString& InCulture, FString& OutLocalized);
+	static COREUOBJECT_API bool GetLocalizedRoot(const FString& InPath, const FString& InCulture, FString& OutLocalized);
 };
diff --git a/Engine/Source/Runtime/CoreUObject/Public/Misc/AssetRegistryInterface.h b/Engine/Source/Runtime/CoreUObject/Public/Misc/AssetRegistryInterface.h
--- a/Engine/Source/Runtime/CoreUObject/Public/Misc/AssetRegistryInterface.h
+++ b/Engine/Source/Runtime/CoreUObject/Public/Misc/AssetRegistryInterface.h
@@ -204,7 +204,7 @@ namespace UE::AssetRegistry
 		 * TODO: this storage should be a class static variable on IAssetRegistry, but that type is defined in the AssetRegistry module, and many modules try to access the singleton (and call virtual functions on it) without linking against
 		 * the AssetRegistry module, so the storage for the singleton needs to be defined in a lower-level module that all of those modules do include
 		 */
-		class COREUOBJECT_API IAssetRegistrySingleton
+		class IAssetRegistrySingleton
 		{
 		public:
 			static IAssetRegistry* Get()
@@ -212,26 +212,26 @@ namespace UE::AssetRegistry
 				return Singleton;
 			}
 		private:
-			static IAssetRegistry* Singleton;
+			static COREUOBJECT_API IAssetRegistry* Singleton;
 			friend class ::UAssetRegistryImpl;
 		};
 	}
 
-	class COREUOBJECT_API FFiltering
+	class FFiltering
 	{
 	public:
 		/** Return whether to filter out assets of the given class and flags from the editor's asset registry */
-		static bool ShouldSkipAsset(const FTopLevelAssetPath& AssetClass, uint32 PackageFlags);
+		static COREUOBJECT_API bool ShouldSkipAsset(const FTopLevelAssetPath& AssetClass, uint32 PackageFlags);
 
 		/** Return whether to filter out the given object (assumed to be an asset) from the editor's asset registry */
-		static bool ShouldSkipAsset(const UObject* InAsset);
+		static COREUOBJECT_API bool ShouldSkipAsset(const UObject* InAsset);
 
 		/** Call to invalidate the list of skip assets and cause their next use to recreate them on demand */
-		static void MarkDirty();
+		static COREUOBJECT_API void MarkDirty();
 
 #if WITH_ENGINE && WITH_EDITOR
 		/** Copy the global skip classes set from the given external sets that were already populated. */
-		static void SetSkipClasses(const TSet<FTopLevelAssetPath>& InSkipUncookedClasses, const TSet<FTopLevelAssetPath>& InSkipCookedClasses);
+		static COREUOBJECT_API void SetSkipClasses(const TSet<FTopLevelAssetPath>& InSkipUncookedClasses, const TSet<FTopLevelAssetPath>& InSkipCookedClasses);
 #endif
 	};
 
@@ -275,13 +275,13 @@ namespace EAssetSetManagerFlags
 /**
  * Asset Registry module interface
  */
-class COREUOBJECT_API IAssetRegistryInterface
+class IAssetRegistryInterface
 {
 public:
 	/**
 	 * Tries to gets a pointer to the active AssetRegistryInterface implementation. 
 	 */
-	static IAssetRegistryInterface* GetPtr();
+	static COREUOBJECT_API IAssetRegistryInterface* GetPtr();
 
 	/**
 	 * Lookup dependencies for the given package name and fill OutDependencies with direct dependencies
@@ -308,6 +308,6 @@ public:
 
 protected:
 
-	static IAssetRegistryInterface* Default;
+	static COREUOBJECT_API IAssetRegistryInterface* Default;
 	friend class UAssetRegistryImpl;
 };
diff --git a/Engine/Source/Runtime/CoreUObject/Public/Misc/NotifyHook.h b/Engine/Source/Runtime/CoreUObject/Public/Misc/NotifyHook.h
--- a/Engine/Source/Runtime/CoreUObject/Public/Misc/NotifyHook.h
+++ b/Engine/Source/Runtime/CoreUObject/Public/Misc/NotifyHook.h
@@ -12,11 +12,11 @@ class FProperty;
 struct FPropertyChangedEvent;
 
 // Notification hook.
-class COREUOBJECT_API FNotifyHook
+class FNotifyHook
 {
 public:
 	virtual void NotifyPreChange( FProperty* PropertyAboutToChange ) {}
 	virtual void NotifyPostChange( const FPropertyChangedEvent& PropertyChangedEvent, FProperty* PropertyThatChanged ) {}
-	virtual void NotifyPreChange( class FEditPropertyChain* PropertyAboutToChange );
-	virtual void NotifyPostChange( const FPropertyChangedEvent& PropertyChangedEvent, class FEditPropertyChain* PropertyThatChanged );
+	COREUOBJECT_API virtual void NotifyPreChange( class FEditPropertyChain* PropertyAboutToChange );
+	COREUOBJECT_API virtual void NotifyPostChange( const FPropertyChangedEvent& PropertyChangedEvent, class FEditPropertyChain* PropertyThatChanged );
 };
diff --git a/Engine/Source/Runtime/CoreUObject/Public/Misc/PackageName.h b/Engine/Source/Runtime/CoreUObject/Public/Misc/PackageName.h
--- a/Engine/Source/Runtime/CoreUObject/Public/Misc/PackageName.h
+++ b/Engine/Source/Runtime/CoreUObject/Public/Misc/PackageName.h
@@ -27,7 +27,7 @@ struct FGuid;
 
 DECLARE_LOG_CATEGORY_EXTERN(LogPackageName, Log, All);
 
-class COREUOBJECT_API FPackageName
+class FPackageName
 {
 public:
 	/** Indicates the format the input was in in functions that take flexible input that can be a LocalPath, PackageName, or ObjectPath */
@@ -58,7 +58,7 @@ public:
 	 * @param InPath the path that was passed to the convert function
 	 * @param ErrorCode The error code returned from the convert function
 	 */
-	static FString FormatErrorAsString(FStringView InPath, EErrorCode ErrorCode);
+	static COREUOBJECT_API FString FormatErrorAsString(FStringView InPath, EErrorCode ErrorCode);
 
 	/**
 	 * Return a user-readable string for the error code returned from a FPackageName function
@@ -66,7 +66,7 @@ public:
 	 * @param InPath the path that was passed to the convert function
 	 * @param ErrorCode The error code returned from the convert function
 	 */
-	static FText FormatErrorAsText(FStringView InPath, EErrorCode ErrorCode);
+	static COREUOBJECT_API FText FormatErrorAsText(FStringView InPath, EErrorCode ErrorCode);
 
 	/**
 	 * Helper function for converting short to long script package name (InputCore -> /Script/InputCore)
@@ -74,16 +74,16 @@ public:
 	 * @param InShortName Short package name.
 	 * @return Long package name.
 	 */
-	static FString ConvertToLongScriptPackageName(const TCHAR* InShortName);
+	static COREUOBJECT_API FString ConvertToLongScriptPackageName(const TCHAR* InShortName);
 
 	/** Return the LongPackageName of module's native script package. Does not check whether module is native. */
-	static FName GetModuleScriptPackageName(FName ModuleName);
-	static FString GetModuleScriptPackageName(FStringView ModuleName);
+	static COREUOBJECT_API FName GetModuleScriptPackageName(FName ModuleName);
+	static COREUOBJECT_API FString GetModuleScriptPackageName(FStringView ModuleName);
 
 	/**
 	 * Registers all short package names found in ini files.
 	 */
-	static void RegisterShortPackageNamesForUObjectModules();
+	static COREUOBJECT_API void RegisterShortPackageNamesForUObjectModules();
 
 	/**
 	 * Finds long script package name associated with a short package name.
@@ -91,7 +91,7 @@ public:
 	 * @param InShortName Short script package name.
 	 * @return Long script package name (/Script/Package) associated with short name or NULL.
 	 */
-	static FName* FindScriptPackageName(FName InShortName);
+	static COREUOBJECT_API FName* FindScriptPackageName(FName InShortName);
 
 	/** 
 	 * Tries to convert the supplied relative or absolute filename to a long package name/path starting with a root like /game
@@ -102,7 +102,7 @@ public:
 	 * @param OutFailureReason Description of an error if the conversion failed.
 	 * @return Returns true if the supplied filename properly maps to one of the long package roots.
 	 */
-	static bool TryConvertFilenameToLongPackageName(const FString& InFilename, FString& OutPackageName, FString* OutFailureReason = nullptr);
+	static COREUOBJECT_API bool TryConvertFilenameToLongPackageName(const FString& InFilename, FString& OutPackageName, FString* OutFailureReason = nullptr);
 
 	/** 
 	 * Converts the supplied filename to long package name.
@@ -111,7 +111,7 @@ public:
 	 * @param InFilename Filename to convert.
 	 * @return Long package name.
 	 */
-	static FString FilenameToLongPackageName(const FString& InFilename);
+	static COREUOBJECT_API FString FilenameToLongPackageName(const FString& InFilename);
 
 	/** 
 	 * Tries to convert a long package name to a file name with the supplied extension.
@@ -121,7 +121,7 @@ public:
 	 * @param InExtension Package extension.
 	 * @return Package filename.
 	 */
-	static bool TryConvertLongPackageNameToFilename(const FString& InLongPackageName, FString& OutFilename, const FString& InExtension = TEXT(""));
+	static COREUOBJECT_API bool TryConvertLongPackageNameToFilename(const FString& InLongPackageName, FString& OutFilename, const FString& InExtension = TEXT(""));
 
 	/** 
 	 * Find the MountPoint for a LocalPath, LongPackageName, or ObjectPath and return its elements. Use this function instead of TryConvertFilenameToLongPackageName or
@@ -135,7 +135,7 @@ public:
 	 * @param OutFlexNameType			If non-null, will be set to the FlexNameType of InPath.
 	 * @param OutFailureReason			If non-null, will be set to the reason InPath could not be converted, or to EErrorCode::Unknown if the function was successful.
 	 */
-	static bool TryConvertToMountedPath(FStringView InPath, FString* OutLocalPathNoExtension, FString* OutPackageName, FString* OutObjectName, FString* OutSubObjectName, FString* OutExtension, EFlexNameType* OutFlexNameType = nullptr, EErrorCode* OutFailureReason = nullptr);
+	static COREUOBJECT_API bool TryConvertToMountedPath(FStringView InPath, FString* OutLocalPathNoExtension, FString* OutPackageName, FString* OutObjectName, FString* OutSubObjectName, FString* OutExtension, EFlexNameType* OutFlexNameType = nullptr, EErrorCode* OutFailureReason = nullptr);
 
 	/** 
 	 * Converts a long package name to a file name with the supplied extension.
@@ -145,7 +145,7 @@ public:
 	 * @param InExtension Package extension.
 	 * @return Package filename.
 	 */
-	static FString LongPackageNameToFilename(const FString& InLongPackageName, const FString& InExtension = TEXT(""));
+	static COREUOBJECT_API FString LongPackageNameToFilename(const FString& InLongPackageName, const FString& InExtension = TEXT(""));
 
 	/** 
 	 * Returns the path to the specified package, excluding the short package name
@@ -153,7 +153,7 @@ public:
 	 * @param InLongPackageName Long Package Name.
 	 * @return The path containing the specified package.
 	 */
-	static FString GetLongPackagePath(const FString& InLongPackageName);
+	static COREUOBJECT_API FString GetLongPackagePath(const FString& InLongPackageName);
 
 	/** 
 	 * Returns the clean asset name for the specified package, same as GetShortName
@@ -161,7 +161,7 @@ public:
 	 * @param InLongPackageName Long Package Name
 	 * @return Clean asset name.
 	 */
-	static FString GetLongPackageAssetName(const FString& InLongPackageName);
+	static COREUOBJECT_API FString GetLongPackageAssetName(const FString& InLongPackageName);
 
 	/** 
 	 * Convert a long package name into root, path, and name components
@@ -173,7 +173,7 @@ public:
 	 * @param bStripRootLeadingSlash String any leading / character from the returned root
 	 * @return True if the conversion was possible, false otherwise
 	 */
-	static bool SplitLongPackageName(const FString& InLongPackageName, FString& OutPackageRoot, FString& OutPackagePath, FString& OutPackageName, const bool bStripRootLeadingSlash = false);
+	static COREUOBJECT_API bool SplitLongPackageName(const FString& InLongPackageName, FString& OutPackageRoot, FString& OutPackagePath, FString& OutPackageName, const bool bStripRootLeadingSlash = false);
 
 	/**
 	 * Split a full object path (Class /Path/To/A/Package.Object:SubObject) into its constituent pieces
@@ -188,9 +188,9 @@ public:
 	 *                          other names. Spaces in those names is invalid, but some code ignores the
 	 *                          invalidity in ObjectName if it only cares about packageName.
 	 */
-	static void SplitFullObjectPath(const FString& InFullObjectPath, FString& OutClassName,
+	static COREUOBJECT_API void SplitFullObjectPath(const FString& InFullObjectPath, FString& OutClassName,
 		FString& OutPackageName, FString& OutObjectName, FString& OutSubObjectName, bool bDetectClassName = true);
-	static void SplitFullObjectPath(FStringView InFullObjectPath, FStringView& OutClassName,
+	static COREUOBJECT_API void SplitFullObjectPath(FStringView InFullObjectPath, FStringView& OutClassName,
 		FStringView& OutPackageName, FStringView& OutObjectName, FStringView& OutSubObjectName, bool bDetectClassName=true);
 
 	/** 
@@ -201,8 +201,8 @@ public:
 	 * @param OutReason					When returning false, this will provide a description of what was wrong with the name.
 	 * @return							true if a valid long package name
 	 */
-	static bool IsValidLongPackageName(FStringView InLongPackageName, bool bIncludeReadOnlyRoots = false, EErrorCode* OutReason = nullptr);
-	static bool IsValidLongPackageName(FStringView InLongPackageName, bool bIncludeReadOnlyRoots, FText* OutReason );
+	static COREUOBJECT_API bool IsValidLongPackageName(FStringView InLongPackageName, bool bIncludeReadOnlyRoots = false, EErrorCode* OutReason = nullptr);
+	static COREUOBJECT_API bool IsValidLongPackageName(FStringView InLongPackageName, bool bIncludeReadOnlyRoots, FText* OutReason );
 
 	/**
 	 * Report whether a given name is the proper format for a PackageName, without checking whether it is in one of the registered mount points
@@ -211,8 +211,8 @@ public:
 	 * @param OutReason					When returning false, this will provide a description of what was wrong with the name.
 	 * @return							true if valid text for a long package name
 	 */
-	static bool IsValidTextForLongPackageName(FStringView InLongPackageName, EErrorCode* OutReason = nullptr);
-	static bool IsValidTextForLongPackageName(FStringView InLongPackageName, FText* OutReason);
+	static COREUOBJECT_API bool IsValidTextForLongPackageName(FStringView InLongPackageName, EErrorCode* OutReason = nullptr);
+	static COREUOBJECT_API bool IsValidTextForLongPackageName(FStringView InLongPackageName, FText* OutReason);
 
 
 	/**
@@ -224,7 +224,7 @@ public:
 	 * @param OutReason					When returning false, this will provide a description of what was wrong with the name.
 	 * @return							true if a valid object path
 	 */
-	static bool IsValidObjectPath(const FString& InObjectPath, FText* OutReason = nullptr);
+	static COREUOBJECT_API bool IsValidObjectPath(const FString& InObjectPath, FText* OutReason = nullptr);
 
 	/**
 	 * Returns true if the path starts with a valid root (i.e. /Game/, /Engine/, etc).
@@ -232,7 +232,7 @@ public:
 	 * @param InObjectPath				The object path to test
 	 * @return							true if a valid object path
 	 */
-	static bool IsValidPath(const FString& InPath);
+	static COREUOBJECT_API bool IsValidPath(const FString& InPath);
 
 	/**
 	 * Checks if the string is a ShortPackageName. A ShortPackageName is the leaf name after the last
@@ -242,9 +242,9 @@ public:
 	 * @param PossiblyLongName Package name.
 	 * @return true if the given name is a short package name (contains no slashes), false otherwise.
 	 */
-	static bool IsShortPackageName(const FString& PossiblyLongName);
-	static bool IsShortPackageName(const FName PossiblyLongName);
-	static bool IsShortPackageName(FStringView PossiblyLongName);
+	static COREUOBJECT_API bool IsShortPackageName(const FString& PossiblyLongName);
+	static COREUOBJECT_API bool IsShortPackageName(const FName PossiblyLongName);
+	static COREUOBJECT_API bool IsShortPackageName(FStringView PossiblyLongName);
 
 	/**
 	 * Converts package name to short name.
@@ -252,7 +252,7 @@ public:
 	 * @param Package Package with name to convert.
 	 * @return Short package name.
 	 */
-	static FString GetShortName(const UPackage* Package);
+	static COREUOBJECT_API FString GetShortName(const UPackage* Package);
 
 	/**
 	 * Converts package name to short name.
@@ -260,9 +260,9 @@ public:
 	 * @param LongName Package name to convert.
 	 * @return Short package name.
 	 */
-	static FString GetShortName(const FString& LongName);
-	static FString GetShortName(const FName& LongName);
-	static FString GetShortName(const TCHAR* LongName);
+	static COREUOBJECT_API FString GetShortName(const FString& LongName);
+	static COREUOBJECT_API FString GetShortName(const FName& LongName);
+	static COREUOBJECT_API FString GetShortName(const TCHAR* LongName);
 
 	/**
 	 * Converts package name to short name.
@@ -270,9 +270,9 @@ public:
 	 * @param LongName Package name to convert.
 	 * @return Short package name.
 	 */
-	static FName GetShortFName(const FString& LongName);
-	static FName GetShortFName(const FName& LongName);
-	static FName GetShortFName(const TCHAR* LongName);
+	static COREUOBJECT_API FName GetShortFName(const FString& LongName);
+	static COREUOBJECT_API FName GetShortFName(const FName& LongName);
+	static COREUOBJECT_API FName GetShortFName(const TCHAR* LongName);
 
 	/**
 	 * Tries to convert a file or directory in game-relative package name format to the corresponding local path
@@ -284,7 +284,7 @@ public:
 	 * @param OutLocalPath The corresponding local-path file (with the extension or lack of extension from the input).
 	 * @return Whether the conversion was successful.
 	 */
-	static bool TryConvertGameRelativePackagePathToLocalPath(FStringView RelativePackagePath, FString& OutLocalPath);
+	static COREUOBJECT_API bool TryConvertGameRelativePackagePathToLocalPath(FStringView RelativePackagePath, FString& OutLocalPath);
 
 	/**
 	 * This will insert a mount point at the head of the search chain (so it can overlap an existing mount point and win).
@@ -292,7 +292,7 @@ public:
 	 * @param RootPath Logical Root Path.
 	 * @param ContentPath Content Path on disk.
 	 */
-	static void RegisterMountPoint(const FString& RootPath, const FString& ContentPath);
+	static COREUOBJECT_API void RegisterMountPoint(const FString& RootPath, const FString& ContentPath);
 
 	/**
 	 * This will remove a previously inserted mount point.
@@ -300,12 +300,12 @@ public:
 	 * @param RootPath Logical Root Path.
 	 * @param ContentPath Content Path on disk.
 	 */
-	static void UnRegisterMountPoint(const FString& RootPath, const FString& ContentPath);
+	static COREUOBJECT_API void UnRegisterMountPoint(const FString& RootPath, const FString& ContentPath);
 
 	/**
 	 * Returns whether the specific logical root path is a valid mount point.
 	 */
-	static bool MountPointExists(const FString& RootPath);
+	static COREUOBJECT_API bool MountPointExists(const FString& RootPath);
 
 	/**
 	 * Get the mount point for a given package path
@@ -314,14 +314,14 @@ public:
 	 * @param InWithoutSlashes Optional parameters that keeps the slashes around the mount point if false
 	 * @return FName corresponding to the mount point, or Empty if invalid
 	 */
-	static FName GetPackageMountPoint(const FString& InPackagePath, bool InWithoutSlashes = true);
+	static COREUOBJECT_API FName GetPackageMountPoint(const FString& InPackagePath, bool InWithoutSlashes = true);
 
 	/**
 	 * Get the path associated with the given package root
 	 * @param InPackageRoot Package root to return the path for, e.g. /Game/, /Engine/, /PluginName/
 	 * @return Filesystem path associated with the provided package root
 	 */
-	static FString GetContentPathForPackageRoot(FStringView InPackageRoot);
+	static COREUOBJECT_API FString GetContentPathForPackageRoot(FStringView InPackageRoot);
 
 	/**
 	 * Checks if the package exists on disk.
@@ -347,7 +347,7 @@ public:
 	 * @param InAllowTextFormats Detect text format packages as well as binary (priority to text)
 	 * @return true if the specified package name points to an existing package, false otherwise.
 	 **/
-	static bool DoesPackageExist(const FString& LongPackageName, FString* OutFilename = nullptr, bool InAllowTextFormats = true);
+	static COREUOBJECT_API bool DoesPackageExist(const FString& LongPackageName, FString* OutFilename = nullptr, bool InAllowTextFormats = true);
 
 	/**
 	 * Checks if the package exists on disk. PackagePath must be a mounted path, otherwise returns false
@@ -358,7 +358,7 @@ public:
 	 * @param OutPackagePath If nonnull and the package exists, set to a copy of PackagePath with the HeaderExtension set to the extension that exists on disk (and if bMatchCaseOnDisk is true, capitalization changed to match). If not found, this variable is not written
 	 * @return true if the specified package name points to an existing package, false otherwise.
 	 **/
-	static bool DoesPackageExist(const FPackagePath& PackagePath, bool bMatchCaseOnDisk = false, FPackagePath* OutPackagePath = nullptr);
+	static COREUOBJECT_API bool DoesPackageExist(const FPackagePath& PackagePath, bool bMatchCaseOnDisk = false, FPackagePath* OutPackagePath = nullptr);
 
 	/**
 	 * Checks if the package exists on disk. PackagePath must be a mounted path, otherwise returns false
@@ -367,7 +367,7 @@ public:
 	 * @param OutPackagePath If nonnull and the package exists, set to a copy of PackagePath with the HeaderExtension set to the extension that exists on disk. If not found, this variable is not written
 	 * @return true if the specified package name points to an existing package, false otherwise.
 	 **/
-	static bool DoesPackageExist(const FPackagePath& PackagePath, FPackagePath* OutPackagePath);
+	static COREUOBJECT_API bool DoesPackageExist(const FPackagePath& PackagePath, FPackagePath* OutPackagePath);
 
 	enum class EPackageLocationFilter : uint8
 	{
@@ -387,7 +387,7 @@ public:
 	 * @param OutPackagePath If nonnull and the package exists, set to a copy of PackagePath with the HeaderExtension set to the extension that exists on disk (and if bMatchCaseOnDisk is true, capitalization changed to match). If not found, this variable is not written
 	 * @return the set of locations where the package exists (IoDispatcher or FileSystem, both or neither)
 	 **/
-	static EPackageLocationFilter DoesPackageExistEx(const FPackagePath& PackagePath, EPackageLocationFilter Filterconst, bool bMatchCaseOnDisk = false, FPackagePath* OutPackagePath = nullptr);
+	static COREUOBJECT_API EPackageLocationFilter DoesPackageExistEx(const FPackagePath& PackagePath, EPackageLocationFilter Filterconst, bool bMatchCaseOnDisk = false, FPackagePath* OutPackagePath = nullptr);
 
 	/**
 	 * Attempts to find a package given its short name on disk (very slow).
@@ -396,7 +396,7 @@ public:
 	 * @param OutLongPackageName Long package name corresponding to the found file (if any).
 	 * @return true if the specified package name points to an existing package, false otherwise.
 	 **/
-	static bool SearchForPackageOnDisk(const FString& PackageName, FString* OutLongPackageName = NULL, FString* OutFilename = NULL);
+	static COREUOBJECT_API bool SearchForPackageOnDisk(const FString& PackageName, FString* OutLongPackageName = NULL, FString* OutFilename = NULL);
 
 	/**
 	 * Tries to convert object path with short package name to object path with long package name found on disk (very slow)
@@ -406,7 +406,7 @@ public:
 	 *
 	 * @return True if succeeded. False otherwise.
 	 */
-	static bool TryConvertShortPackagePathToLongInObjectPath(const FString& ObjectPath, FString& ConvertedObjectPath);
+	static COREUOBJECT_API bool TryConvertShortPackagePathToLongInObjectPath(const FString& ObjectPath, FString& ConvertedObjectPath);
 
 	/**
 	 * Gets normalized object path i.e. with long package format.
@@ -415,7 +415,7 @@ public:
 	 *
 	 * @return Normalized path (or empty path, if short object path was given and it wasn't found on the disk).
 	 */
-	static FString GetNormalizedObjectPath(const FString& ObjectPath);
+	static COREUOBJECT_API FString GetNormalizedObjectPath(const FString& ObjectPath);
 
 	/**
 	 * Gets the resolved path of a long package as determined by the delegates registered with FCoreDelegates::PackageNameResolvers.
@@ -426,7 +426,7 @@ public:
 	 *
 	 * @return Resolved package path, or the source package path if there is no resolution occurs.
 	 */
-	static FString GetDelegateResolvedPackagePath(const FString& InSourcePackagePath);
+	static COREUOBJECT_API FString GetDelegateResolvedPackagePath(const FString& InSourcePackagePath);
 
 	/**
 	 * Gets the source version of a localized long package path (it is also safe to pass non-localized paths into this function).
@@ -435,7 +435,7 @@ public:
 	 *
 	 * @return Source package path.
 	 */
-	static FString GetSourcePackagePath(const FString& InLocalizedPackagePath);
+	static COREUOBJECT_API FString GetSourcePackagePath(const FString& InLocalizedPackagePath);
 
 	/**
 	 * Gets the localized version of a long package path for the current culture, or returns the source package if there is no suitable localized package.
@@ -444,7 +444,7 @@ public:
 	 *
 	 * @return Localized package path, or the source package path if there is no suitable localized package.
 	 */
-	static FString GetLocalizedPackagePath(const FString& InSourcePackagePath);
+	static COREUOBJECT_API FString GetLocalizedPackagePath(const FString& InSourcePackagePath);
 
 	/**
 	 * Gets the localized version of a long package path for the given culture, or returns the source package if there is no suitable localized package.
@@ -454,34 +454,34 @@ public:
 	 *
 	 * @return Localized package path, or the source package path if there is no suitable localized package.
 	 */
-	static FString GetLocalizedPackagePath(const FString& InSourcePackagePath, const FString& InCultureName);
+	static COREUOBJECT_API FString GetLocalizedPackagePath(const FString& InSourcePackagePath, const FString& InCultureName);
 
 	/** 
 	 * Returns the file extension for packages containing assets.
 	 *
 	 * @return	file extension for asset packages ( dot included )
 	 */
-	static const FString& GetAssetPackageExtension();
+	static COREUOBJECT_API const FString& GetAssetPackageExtension();
 	/** 
 	 * Returns the file extension for packages containing assets.
 	 *
 	 * @return	file extension for asset packages ( dot included )
 	 */
-	static const FString& GetMapPackageExtension();
+	static COREUOBJECT_API const FString& GetMapPackageExtension();
 
 	/**
 	* Returns the file extension for packages containing text assets.
 	*
 	* @return	file extension for text asset packages ( dot included )
 	*/
-	static const FString& GetTextAssetPackageExtension();
+	static COREUOBJECT_API const FString& GetTextAssetPackageExtension();
 
 	/**
 	* Returns the file extension for packages containing text maps.
 	*
 	* @return	file extension for text map packages ( dot included )
 	*/
-	static const FString& GetTextMapPackageExtension();
+	static COREUOBJECT_API const FString& GetTextMapPackageExtension();
 
 	/**
 	 * Returns whether the passed in extension is a valid text package
@@ -490,7 +490,7 @@ public:
 	 * @param	Extension to test.
 	 * @return	True if Ext is either an text asset or a text map extension, otherwise false
 	 */
-	static bool IsTextPackageExtension(const TCHAR* Ext);
+	static COREUOBJECT_API bool IsTextPackageExtension(const TCHAR* Ext);
 
 	/**
 	 * Returns whether the passed in extension is a text header extension
@@ -498,7 +498,7 @@ public:
 	 * @param	Extension to test.
 	 * @return	True if Ext is either an text asset or a text map extension, otherwise false
 	 */
-	static bool IsTextPackageExtension(EPackageExtension Extension);
+	static COREUOBJECT_API bool IsTextPackageExtension(EPackageExtension Extension);
 
 	/**
 	 * Returns whether the passed in extension is a valid text asset package
@@ -507,7 +507,7 @@ public:
 	 * @param	Extension to test.
 	 * @return	True if Ext is a text asset extension, otherwise false
 	 */
-	static bool IsTextAssetPackageExtension(const TCHAR* Ext);
+	static COREUOBJECT_API bool IsTextAssetPackageExtension(const TCHAR* Ext);
 
 	/**
 	 * Returns whether the passed in extension is a valid text map package
@@ -516,7 +516,7 @@ public:
 	 * @param	Extension to test.
 	 * @return	True if Ext is a text map extension, otherwise false
 	 */
-	static bool IsTextMapPackageExtension(const TCHAR* Ext);
+	static COREUOBJECT_API bool IsTextMapPackageExtension(const TCHAR* Ext);
 
 	/** 
 	 * Returns whether the passed in extension is a valid binary package
@@ -525,7 +525,7 @@ public:
 	 * @param	Extension to test. 
 	 * @return	True if Ext is either a binary  asset or map extension, otherwise false
 	 */
-	static bool IsPackageExtension(const TCHAR* Ext);
+	static COREUOBJECT_API bool IsPackageExtension(const TCHAR* Ext);
 
 	/**
 	 * Returns whether the passed in extension is a valid binary package extension.
@@ -533,7 +533,7 @@ public:
 	 * @param	Extension to test.
 	 * @return	True if Ext is either a binary asset or a binary map extension, otherwise false
 	 */
-	static bool IsPackageExtension(EPackageExtension Extension);
+	static COREUOBJECT_API bool IsPackageExtension(EPackageExtension Extension);
 
 	/**
 	 * Returns whether the passed in extension is a valid binary asset package
@@ -542,7 +542,7 @@ public:
 	 * @param	Extension to test.
 	 * @return	True if Ext is a binary asset extension, otherwise false
 	 */
-	static bool IsAssetPackageExtension(const TCHAR* Ext);
+	static COREUOBJECT_API bool IsAssetPackageExtension(const TCHAR* Ext);
 
 	/**
 	 * Returns whether the passed in extension is a valid binary map package
@@ -551,7 +551,7 @@ public:
 	 * @param	Extension to test.
 	 * @return	True if Ext is a binary asset extension, otherwise false
 	 */
-	static bool IsMapPackageExtension(const TCHAR* Ext);
+	static COREUOBJECT_API bool IsMapPackageExtension(const TCHAR* Ext);
 
 	/** 
 	 * Returns whether the passed in filename ends with any of the known
@@ -574,7 +574,7 @@ public:
 	 * @param	RootDir				The root of the directory structure to recurse through
 	 * @return	Returns true if any packages have been found, otherwise false
 	 */
-	static bool FindPackagesInDirectory(TArray<FString>& OutPackages, const FString& RootDir);
+	static COREUOBJECT_API bool FindPackagesInDirectory(TArray<FString>& OutPackages, const FString& RootDir);
 
 	/**
 	 * This will recurse over the given list of directory structures looking for packages.
@@ -583,7 +583,7 @@ public:
 	 * @param	RootDirss			The roots of the directory structures to recurse through
 	 * @return	Returns true if any packages have been found, otherwise false
 	 */
-	static bool FindPackagesInDirectories(TArray<FString>& OutPackages, const TArrayView<const FString>& RootDirs);
+	static COREUOBJECT_API bool FindPackagesInDirectories(TArray<FString>& OutPackages, const TArrayView<const FString>& RootDirs);
 
 	/**
 	 * This will recurse over a directory structure looking for packages.
@@ -593,8 +593,8 @@ public:
 	 */
 	typedef TFunctionRef<bool(const TCHAR*)> FPackageNameVisitor;
 	typedef TFunctionRef<bool(const TCHAR*, const FFileStatData&)> FPackageNameStatVisitor;
-	static void IteratePackagesInDirectory(const FString& RootDir, const FPackageNameVisitor& Visitor);
-	static void IteratePackagesInDirectory(const FString& RootDir, const FPackageNameStatVisitor& Visitor);
+	static COREUOBJECT_API void IteratePackagesInDirectory(const FString& RootDir, const FPackageNameVisitor& Visitor);
+	static COREUOBJECT_API void IteratePackagesInDirectory(const FString& RootDir, const FPackageNameStatVisitor& Visitor);
 
 	/** Event that is triggered when a new content path is mounted */
 	DECLARE_MULTICAST_DELEGATE_TwoParams( FOnContentPathMountedEvent, const FString& /* Asset path */, const FString& /* ContentPath */ );
@@ -618,10 +618,10 @@ public:
 	 * @param	bWithoutLeadingSlashes	  Strip slash at start of each path to end up with "Game/"
 	 * @param	bWithoutTrailingSlashes	  Strip trailing slash at end of each path to end up with "/Game"
 	 */
-	static void QueryRootContentPaths( TArray<FString>& OutRootContentPaths, bool bIncludeReadOnlyRoots = false, bool bWithoutLeadingSlashes = false, bool bWithoutTrailingSlashes = false);
+	static COREUOBJECT_API void QueryRootContentPaths( TArray<FString>& OutRootContentPaths, bool bIncludeReadOnlyRoots = false, bool bWithoutLeadingSlashes = false, bool bWithoutTrailingSlashes = false);
 	
 	/** If the FLongPackagePathsSingleton is not created yet, this function will create it and thus allow mount points to be added */
-	static void OnCoreUObjectInitialized();
+	static COREUOBJECT_API void OnCoreUObjectInitialized();
 
 	/** 
 	 * Converts the supplied export text path to an object path and class name.
@@ -631,10 +631,10 @@ public:
 	 * @param OutObjectPath The path to the object.
 	 * @return True if the supplied export text path could be parsed
 	 */
-	static bool ParseExportTextPath(FWideStringView InExportTextPath, FWideStringView* OutClassName, FWideStringView* OutObjectPath);
-	static bool ParseExportTextPath(FAnsiStringView InExportTextPath, FAnsiStringView* OutClassName, FAnsiStringView* OutObjectPath);
-	static bool ParseExportTextPath(const FString& InExportTextPath, FString* OutClassName, FString* OutObjectPath);	
-	static bool ParseExportTextPath(const TCHAR* InExportTextPath, FStringView* OutClassName, FStringView* OutObjectPath);
+	static COREUOBJECT_API bool ParseExportTextPath(FWideStringView InExportTextPath, FWideStringView* OutClassName, FWideStringView* OutObjectPath);
+	static COREUOBJECT_API bool ParseExportTextPath(FAnsiStringView InExportTextPath, FAnsiStringView* OutClassName, FAnsiStringView* OutObjectPath);
+	static COREUOBJECT_API bool ParseExportTextPath(const FString& InExportTextPath, FString* OutClassName, FString* OutObjectPath);	
+	static COREUOBJECT_API bool ParseExportTextPath(const TCHAR* InExportTextPath, FStringView* OutClassName, FStringView* OutObjectPath);
 
 
 	/** 
@@ -643,10 +643,10 @@ public:
 	 * @param InExportTextPath The export text path for an object. Takes on the form: ClassName'ObjectPath'
 	 * @return The path to the object referred to by the supplied export path.
 	 */
-	static FWideStringView	ExportTextPathToObjectPath(FWideStringView InExportTextPath);
-	static FAnsiStringView	ExportTextPathToObjectPath(FAnsiStringView InExportTextPath);
-	static FString			ExportTextPathToObjectPath(const FString& InExportTextPath);
-	static FString			ExportTextPathToObjectPath(const TCHAR* InExportTextPath);
+	static COREUOBJECT_API FWideStringView	ExportTextPathToObjectPath(FWideStringView InExportTextPath);
+	static COREUOBJECT_API FAnsiStringView	ExportTextPathToObjectPath(FAnsiStringView InExportTextPath);
+	static COREUOBJECT_API FString			ExportTextPathToObjectPath(const FString& InExportTextPath);
+	static COREUOBJECT_API FString			ExportTextPathToObjectPath(const TCHAR* InExportTextPath);
 
 	/**
 	 * Returns the top level 'directory' in a package name. If the package is part of a
@@ -657,7 +657,7 @@ public:
 	 * 
 	 * "/PackageRoot/Path/Leaf" -> (return "PackageRoot"; OutRelativePath = "Path/Leaf")
 	 */
-	static FStringView SplitPackageNameRoot(FStringView InPackageName, FStringView* OutRelativePath);
+	static COREUOBJECT_API FStringView SplitPackageNameRoot(FStringView InPackageName, FStringView* OutRelativePath);
 	
 	/** 
 	 * Returns the name of the package referred to by the specified object path
@@ -668,9 +668,9 @@ public:
 	 *   "/Game/MyAsset.MyAsset"                         -> "/Game/MyAsset"
 	 *   "/Game/MyAsset"                                 -> "/Game/MyAsset"
 	 */
-	static FWideStringView ObjectPathToPackageName(FWideStringView InObjectPath);
-	static FAnsiStringView ObjectPathToPackageName(FAnsiStringView InObjectPath);
-	static FString ObjectPathToPackageName(const FString& InObjectPath);
+	static COREUOBJECT_API FWideStringView ObjectPathToPackageName(FWideStringView InObjectPath);
+	static COREUOBJECT_API FAnsiStringView ObjectPathToPackageName(FAnsiStringView InObjectPath);
+	static COREUOBJECT_API FString ObjectPathToPackageName(const FString& InObjectPath);
 
 	/**
 	 * Returns any remaining object path after trimming the package name from the specified object path
@@ -681,9 +681,9 @@ public:
 	 *   "/Game/MyAsset.MyAsset"                         -> "MyAsset"
 	 *   "/Game/MyAsset"                                 -> ""
 	 */
-	static FWideStringView ObjectPathToPathWithinPackage(FWideStringView InObjectPath);
-	static FAnsiStringView ObjectPathToPathWithinPackage(FAnsiStringView InObjectPath);
-	static FString ObjectPathToPathWithinPackage(const FString& InObjectPath);
+	static COREUOBJECT_API FWideStringView ObjectPathToPathWithinPackage(FWideStringView InObjectPath);
+	static COREUOBJECT_API FAnsiStringView ObjectPathToPathWithinPackage(FAnsiStringView InObjectPath);
+	static COREUOBJECT_API FString ObjectPathToPathWithinPackage(const FString& InObjectPath);
 
 	/**
 	 * Returns the path name of the outer of the leaf object referred to by the specified object path
@@ -694,9 +694,9 @@ public:
 	 *   "/Game/MyAsset.MyAsset"                         -> "/Game/MyAsset"
 	 *   "/Game/MyAsset"                                 -> ""
 	 */
-	static FWideStringView ObjectPathToOuterPath(FWideStringView InObjectPath);
-	static FAnsiStringView ObjectPathToOuterPath(FAnsiStringView InObjectPath);
-	static FString ObjectPathToOuterPath(const FString& InObjectPath);
+	static COREUOBJECT_API FWideStringView ObjectPathToOuterPath(FWideStringView InObjectPath);
+	static COREUOBJECT_API FAnsiStringView ObjectPathToOuterPath(FAnsiStringView InObjectPath);
+	static COREUOBJECT_API FString ObjectPathToOuterPath(const FString& InObjectPath);
 
 	/**
 	 * Returns the path from (and including) the subobject referred to by the specified object path
@@ -707,9 +707,9 @@ public:
 	 *   "/Game/MyAsset.MyAsset"                         -> "MyAsset"
 	 *   "/Game/MyAsset"                                 -> "/Game/MyAsset"
 	 */
-	static FWideStringView ObjectPathToSubObjectPath(FWideStringView InObjectPath);
-	static FAnsiStringView ObjectPathToSubObjectPath(FAnsiStringView InObjectPath);
-	static FString ObjectPathToSubObjectPath(const FString& InObjectPath);
+	static COREUOBJECT_API FWideStringView ObjectPathToSubObjectPath(FWideStringView InObjectPath);
+	static COREUOBJECT_API FAnsiStringView ObjectPathToSubObjectPath(FAnsiStringView InObjectPath);
+	static COREUOBJECT_API FString ObjectPathToSubObjectPath(const FString& InObjectPath);
 
 	/** 
 	 * Returns the name of the leaf object referred to by the specified object path
@@ -720,46 +720,46 @@ public:
 	 *   "/Game/MyAsset.MyAsset"                         -> "MyAsset"
 	 *   "/Game/MyAsset"                                 -> "/Game/MyAsset"
 	 */
-	static FWideStringView ObjectPathToObjectName(FWideStringView InObjectPath);
-	static FAnsiStringView ObjectPathToObjectName(FAnsiStringView InObjectPath);
-	static FString ObjectPathToObjectName(const FString& InObjectPath);
+	static COREUOBJECT_API FWideStringView ObjectPathToObjectName(FWideStringView InObjectPath);
+	static COREUOBJECT_API FAnsiStringView ObjectPathToObjectName(FAnsiStringView InObjectPath);
+	static COREUOBJECT_API FString ObjectPathToObjectName(const FString& InObjectPath);
 
 	/**
 	 * Checks the package's path to see if it's a Verse package
 	 */
-	static bool IsVersePackage(FStringView InPackageName);
+	static COREUOBJECT_API bool IsVersePackage(FStringView InPackageName);
 
 	/**
 	 * Checks the root of the package's path to see if it is a script package
 	 * @return true if the root of the path matches the script path
 	 */
-	static bool IsScriptPackage(FStringView InPackageName);
+	static COREUOBJECT_API bool IsScriptPackage(FStringView InPackageName);
 
 	/**
 	 * Checks the root of the package's path to see if it's a memory package
 	 * This should be set for packages that reside in memory and not on disk, we treat them similar to a script package
 	 * @return true if the root of the patch matches the memory path
 	 */
-	static bool IsMemoryPackage(FStringView InPackageName);
+	static COREUOBJECT_API bool IsMemoryPackage(FStringView InPackageName);
 
 	/**
 	 * Checks the root of the package's path to see if it is a temp package
 	 * Temp packages are sometimes saved to disk, and sometimes only exist in memory. They are never in source control
 	 * @return true if the root of the patch matches the temp path
 	 */
-	static bool IsTempPackage(FStringView InPackageName);
+	static COREUOBJECT_API bool IsTempPackage(FStringView InPackageName);
 
 	/**
 	 * Checks the root of the package's path to see if it is a localized package
 	 * @return true if the root of the path matches any localized root path
 	 */
-	static bool IsLocalizedPackage(FStringView InPackageName);
+	static COREUOBJECT_API bool IsLocalizedPackage(FStringView InPackageName);
 
 	/**
 	 * Checks if a package name contains characters that are invalid for package names.
 	 */
-	static bool DoesPackageNameContainInvalidCharacters(FStringView InLongPackageName, FText* OutReason);
-	static bool DoesPackageNameContainInvalidCharacters(FStringView InLongPackageName, EErrorCode* OutReason = nullptr);
+	static COREUOBJECT_API bool DoesPackageNameContainInvalidCharacters(FStringView InLongPackageName, FText* OutReason);
+	static COREUOBJECT_API bool DoesPackageNameContainInvalidCharacters(FStringView InLongPackageName, EErrorCode* OutReason = nullptr);
 	
 	/**
 	* Checks if a package can be found using known package extensions (header extensions only; files with the extensions of other segments are not returned).
@@ -769,10 +769,10 @@ public:
 	* @param InAllowTextFormats Detect text format packages as well as binary (priority to text)
 	* @return true if the package could be found on disk.
 	*/
-	static bool FindPackageFileWithoutExtension(const FString& InPackageFilename, FString& OutFilename);
+	static COREUOBJECT_API bool FindPackageFileWithoutExtension(const FString& InPackageFilename, FString& OutFilename);
 
 	UE_DEPRECATED(5.0, "Specifying AllowTextFormats is no longer supported. Text format is instead specified as a field on the result struct from IPackageResourceManager->OpenReadPackage.")
-	static bool FindPackageFileWithoutExtension(const FString& InPackageFilename, FString& OutFilename, bool InAllowTextFormats);
+	static COREUOBJECT_API bool FindPackageFileWithoutExtension(const FString& InPackageFilename, FString& OutFilename, bool InAllowTextFormats);
 
 	/**
 	 * Converts a long package name to the case it exists as on disk.
@@ -781,7 +781,7 @@ public:
 	 * @param Extension The extension for this package
 	 * @return True if the long package name was fixed up, false otherwise
 	 */
-	static bool FixPackageNameCase(FString& LongPackageName, FStringView Extension);
+	static COREUOBJECT_API bool FixPackageNameCase(FString& LongPackageName, FStringView Extension);
 
 	/** Override whether a package exist or not. */
 	DECLARE_DELEGATE_RetVal_OneParam(bool, FDoesPackageExistOverride, FName);
@@ -803,7 +803,7 @@ public:
 	 * @param OutFailureReason If non-null, will be set to the reason InPath could not be converted, or to EErrorCode::Unknown if the function was successful.
 	 * @return True if the MountPoint was found, else false
 	 */
-	static bool TryGetMountPointForPath(FStringView InFilePathOrPackageName, FStringBuilderBase& OutMountPointPackageName, FStringBuilderBase& OutMountPointFilePath, FStringBuilderBase& OutRelPath,
+	static COREUOBJECT_API bool TryGetMountPointForPath(FStringView InFilePathOrPackageName, FStringBuilderBase& OutMountPointPackageName, FStringBuilderBase& OutMountPointFilePath, FStringBuilderBase& OutRelPath,
 		EFlexNameType* OutFlexNameType = nullptr, EErrorCode* OutFailureReason = nullptr);
 
 private:
@@ -814,7 +814,7 @@ private:
 	 * @param InFilename
 	 * @param OutPackageName Long package name.
 	 */
-	static void InternalFilenameToLongPackageName(FStringView InFilename, FStringBuilderBase& OutPackageName);
+	static COREUOBJECT_API void InternalFilenameToLongPackageName(FStringView InFilename, FStringBuilderBase& OutPackageName);
 
 	/**
 	 * Internal helper to find a mount point that contains the given LocalFilePath or PackageName, and if found, return the MountPoint, RelativePath, and PackageExtension
@@ -830,17 +830,17 @@ private:
 	 * @param OutFailureReason If non-null, it is set to the failurereason if the MountPoint is not found, otherwise it is set to EErrorCode::PackageNameUnknown
 	 * @return True if the MountPoint was found, else false
 	 */
-	static bool TryConvertToMountedPathComponents(FStringView InPath, FStringBuilderBase& OutMountPointPackageName, FStringBuilderBase& OutMountPointFilePath, FStringBuilderBase& OutRelPath,
+	static COREUOBJECT_API bool TryConvertToMountedPathComponents(FStringView InPath, FStringBuilderBase& OutMountPointPackageName, FStringBuilderBase& OutMountPointFilePath, FStringBuilderBase& OutRelPath,
 		FStringBuilderBase& OutObjectName, EPackageExtension& OutExtension, FStringBuilderBase& OutCustomExtension, EFlexNameType* OutFlexNameType = nullptr, EErrorCode* OutFailureReason = nullptr);
 
 	/** Event that is triggered when a new content path is mounted */
-	static FOnContentPathMountedEvent OnContentPathMountedEvent;
+	static COREUOBJECT_API FOnContentPathMountedEvent OnContentPathMountedEvent;
 
 	/** Event that is triggered when a new content path is removed */
-	static FOnContentPathDismountedEvent OnContentPathDismountedEvent;
+	static COREUOBJECT_API FOnContentPathDismountedEvent OnContentPathDismountedEvent;
 
 	/** Delegate used to check whether a package exist without using the filesystem. */
-	static FDoesPackageExistOverride DoesPackageExistOverrideDelegate;
+	static COREUOBJECT_API FDoesPackageExistOverride DoesPackageExistOverrideDelegate;
 
 	friend class FPackagePath;
 };
diff --git a/Engine/Source/Runtime/CoreUObject/Public/Misc/RedirectCollector.h b/Engine/Source/Runtime/CoreUObject/Public/Misc/RedirectCollector.h
--- a/Engine/Source/Runtime/CoreUObject/Public/Misc/RedirectCollector.h
+++ b/Engine/Source/Runtime/CoreUObject/Public/Misc/RedirectCollector.h
@@ -24,7 +24,7 @@ class FArchive;
 
 enum class ESoftObjectPathCollectType : uint8;
 
-class COREUOBJECT_API FRedirectCollector
+class FRedirectCollector
 {
 private:
 	
@@ -81,7 +81,7 @@ public:
 	 * @param InPath The soft object path that was loaded
 	 * @Param InArchive The archive that loaded this path
 	 */
-	void OnSoftObjectPathLoaded(const struct FSoftObjectPath& InPath, FArchive* InArchive);
+	COREUOBJECT_API void OnSoftObjectPathLoaded(const struct FSoftObjectPath& InPath, FArchive* InArchive);
 
 	/**
 	 * Called at the end of Package Save to record soft package references that might have been created by save transformations
@@ -89,13 +89,13 @@ public:
 	 * @param PackageNames List of of soft package references needed by the referencing package
 	 * @param bEditorOnlyReferences if the PackageNames list are references made by editor only properties
 	 */
-	void CollectSavedSoftPackageReferences(FName ReferencingPackage, const TSet<FName>& PackageNames, bool bEditorOnlyReferences);
+	COREUOBJECT_API void CollectSavedSoftPackageReferences(FName ReferencingPackage, const TSet<FName>& PackageNames, bool bEditorOnlyReferences);
 
 	/**
 	 * Load all soft object paths to resolve them, add that to the remap table, and empty the array
 	 * @param FilterPackage If set, only load references that were created by FilterPackage. If empty, resolve  all of them
 	 */
-	void ResolveAllSoftObjectPaths(FName FilterPackage = NAME_None);
+	COREUOBJECT_API void ResolveAllSoftObjectPaths(FName FilterPackage = NAME_None);
 
 	/**
 	 * Returns the list of packages referenced by soft object paths loaded by FilterPackage, and remove them from the internal list
@@ -103,23 +103,23 @@ public:
 	 * @param bGetEditorOnly If true will return references loaded by editor only objects, if false it will not
 	 * @param OutReferencedPackages Return list of packages referenced by FilterPackage
 	 */
-	void ProcessSoftObjectPathPackageList(FName FilterPackage, bool bGetEditorOnly, TSet<FName>& OutReferencedPackages);
+	COREUOBJECT_API void ProcessSoftObjectPathPackageList(FName FilterPackage, bool bGetEditorOnly, TSet<FName>& OutReferencedPackages);
 
 	/** Adds a new mapping for redirector path to destination path, this is called from the Asset Registry to register all redirects it knows about */
 	UE_DEPRECATED(5.1, "Asset path FNames have been deprecated, use FSoftObjectPath instead.")
-	void AddAssetPathRedirection(FName OriginalPath, FName RedirectedPath);
-	void AddAssetPathRedirection(const FSoftObjectPath& OriginalPath, const FSoftObjectPath& RedirectedPath);
+	COREUOBJECT_API void AddAssetPathRedirection(FName OriginalPath, FName RedirectedPath);
+	COREUOBJECT_API void AddAssetPathRedirection(const FSoftObjectPath& OriginalPath, const FSoftObjectPath& RedirectedPath);
 
 	/** Removes an asset path redirection, call this when deleting redirectors */
 	UE_DEPRECATED(5.1, "Asset path FNames have been deprecated, use FSoftObjectPath instead.")
-	void RemoveAssetPathRedirection(FName OriginalPath);
-	void RemoveAssetPathRedirection(const FSoftObjectPath& OriginalPath);
+	COREUOBJECT_API void RemoveAssetPathRedirection(FName OriginalPath);
+	COREUOBJECT_API void RemoveAssetPathRedirection(const FSoftObjectPath& OriginalPath);
 
 	/** Returns a remapped asset path, if it returns null there is no relevant redirector */
 	UE_DEPRECATED(5.1, "Asset path FNames have been deprecated, use FSoftObjectPath instead.")
-	FName GetAssetPathRedirection(FName OriginalPath);
+	COREUOBJECT_API FName GetAssetPathRedirection(FName OriginalPath);
 	/** Returns a remapped asset path, if there is no relevant redirector, the return value reports true from IsNull() */
-	FSoftObjectPath GetAssetPathRedirection(const FSoftObjectPath& OriginalPath);
+	COREUOBJECT_API FSoftObjectPath GetAssetPathRedirection(const FSoftObjectPath& OriginalPath);
 
 	/**
 	 * Do we have any references to resolve.
@@ -138,10 +138,10 @@ public:
 	 * @param OutExcludedReferences Out set that is reset and then appended with any discovered values
 	 * @return Whether any references were found
 	 */
-	bool RemoveAndCopySoftObjectPathExclusions(FName PackageName, TSet<FName>& OutExcludedReferences);
+	COREUOBJECT_API bool RemoveAndCopySoftObjectPathExclusions(FName PackageName, TSet<FName>& OutExcludedReferences);
 
 	/** Called from the cooker to stop the tracking of exclusions. */
-	void OnStartupPackageLoadComplete();
+	COREUOBJECT_API void OnStartupPackageLoadComplete();
 
 private:
 
@@ -150,7 +150,7 @@ private:
 	typedef TMap<FName, FSoftObjectPathPropertySet> FSoftObjectPathMap;
 
 	/** Return whether SoftObjectPathExclusions are currently being tracked, based on commandline and cook phase. */
-	bool ShouldTrackPackageReferenceTypes();
+	COREUOBJECT_API bool ShouldTrackPackageReferenceTypes();
 
 	/** The discovered references that should be followed during cook */
 	FSoftObjectPathMap SoftObjectPathMap;
diff --git a/Engine/Source/Runtime/CoreUObject/Public/Misc/UObjectToken.h b/Engine/Source/Runtime/CoreUObject/Public/Misc/UObjectToken.h
--- a/Engine/Source/Runtime/CoreUObject/Public/Misc/UObjectToken.h
+++ b/Engine/Source/Runtime/CoreUObject/Public/Misc/UObjectToken.h
@@ -44,14 +44,14 @@ public:
 	}
 
 	/** Get the delegate for default token activation */
-	COREUOBJECT_API static FOnMessageTokenActivated& DefaultOnMessageTokenActivated()
+	static FOnMessageTokenActivated& DefaultOnMessageTokenActivated()
 	{
 		return DefaultMessageTokenActivated;
 	}
 
 	/** Get the delegate for displaying the object name */
 	DECLARE_DELEGATE_RetVal_TwoParams(FText, FOnGetDisplayName, const UObject*, bool);
-	COREUOBJECT_API static FOnGetDisplayName& DefaultOnGetObjectDisplayName()
+	static FOnGetDisplayName& DefaultOnGetObjectDisplayName()
 	{
 		return DefaultGetObjectDisplayName;
 	}
diff --git a/Engine/Source/Runtime/CoreUObject/Public/Misc/WarnIfAssetsLoadedInScope.h b/Engine/Source/Runtime/CoreUObject/Public/Misc/WarnIfAssetsLoadedInScope.h
--- a/Engine/Source/Runtime/CoreUObject/Public/Misc/WarnIfAssetsLoadedInScope.h
+++ b/Engine/Source/Runtime/CoreUObject/Public/Misc/WarnIfAssetsLoadedInScope.h
@@ -11,13 +11,13 @@ struct FAssetData;
 
 #if WITH_EDITOR
 
-class COREUOBJECT_API FWarnIfAssetsLoadedInScope : public FNoncopyable
+class FWarnIfAssetsLoadedInScope : public FNoncopyable
 {
 public:
-    FWarnIfAssetsLoadedInScope();
-	FWarnIfAssetsLoadedInScope(TConstArrayView<FAssetData> InSpecificAssets);
+    COREUOBJECT_API FWarnIfAssetsLoadedInScope();
+	COREUOBJECT_API FWarnIfAssetsLoadedInScope(TConstArrayView<FAssetData> InSpecificAssets);
 
-    ~FWarnIfAssetsLoadedInScope();
+    COREUOBJECT_API ~FWarnIfAssetsLoadedInScope();
 	
 private:
 	bool IsEnabled() const;
@@ -36,4 +36,4 @@ private:
     TArray<FLoadEvent> ObjectsLoaded;
 };
 
-#endif
\ No newline at end of file
+#endif
diff --git a/Engine/Source/Runtime/CoreUObject/Public/Serialization/ArchiveCrc32.h b/Engine/Source/Runtime/CoreUObject/Public/Serialization/ArchiveCrc32.h
--- a/Engine/Source/Runtime/CoreUObject/Public/Serialization/ArchiveCrc32.h
+++ b/Engine/Source/Runtime/CoreUObject/Public/Serialization/ArchiveCrc32.h
@@ -17,13 +17,13 @@ struct CStaticStructProvider;
 /**
 * Calculates a checksum from the input provided to the archive.
 */
-class COREUOBJECT_API FArchiveCrc32 : public FArchiveUObject
+class FArchiveCrc32 : public FArchiveUObject
 {
 public:
 	/**
 	* Default constructor.
 	*/
-	FArchiveCrc32(uint32 InCRC = 0, UObject* InRootObject = nullptr);
+	COREUOBJECT_API FArchiveCrc32(uint32 InCRC = 0, UObject* InRootObject = nullptr);
 
 	/**
 	 * @return The checksum computed so far.
@@ -31,9 +31,9 @@ public:
 	uint32 GetCrc() const { return CRC; }
 
 	//~ Begin FArchive Interface
-	virtual void Serialize(void* Data, int64 Num) override;
-	virtual FArchive& operator<<(class FName& Name);
-	virtual FArchive& operator<<(class UObject*& Object);
+	COREUOBJECT_API virtual void Serialize(void* Data, int64 Num) override;
+	COREUOBJECT_API virtual FArchive& operator<<(class FName& Name);
+	COREUOBJECT_API virtual FArchive& operator<<(class UObject*& Object);
 	virtual FString GetArchiveName() const { return TEXT("FArchiveCrc32"); }
 	//~ End FArchive Interface
 
diff --git a/Engine/Source/Runtime/CoreUObject/Public/Serialization/ArchiveHasReferences.h b/Engine/Source/Runtime/CoreUObject/Public/Serialization/ArchiveHasReferences.h
--- a/Engine/Source/Runtime/CoreUObject/Public/Serialization/ArchiveHasReferences.h
+++ b/Engine/Source/Runtime/CoreUObject/Public/Serialization/ArchiveHasReferences.h
@@ -10,18 +10,18 @@
 class FArchive;
 class UObject;
 
-class COREUOBJECT_API FArchiveHasReferences : private FArchiveUObject
+class FArchiveHasReferences : private FArchiveUObject
 {
 public:
-	FArchiveHasReferences(UObject* InTarget, const TSet<UObject*>& InPotentiallyReferencedObjects);
+	COREUOBJECT_API FArchiveHasReferences(UObject* InTarget, const TSet<UObject*>& InPotentiallyReferencedObjects);
 
 	bool HasReferences() const { return Result; }
 
-	static TArray<UObject*> GetAllReferencers(const TArray<UObject*>& Referencees, const TSet<UObject*>* ObjectsToIgnore );
-	static TArray<UObject*> GetAllReferencers(const TSet<UObject*>& Referencees, const TSet<UObject*>* ObjectsToIgnore );
+	static COREUOBJECT_API TArray<UObject*> GetAllReferencers(const TArray<UObject*>& Referencees, const TSet<UObject*>* ObjectsToIgnore );
+	static COREUOBJECT_API TArray<UObject*> GetAllReferencers(const TSet<UObject*>& Referencees, const TSet<UObject*>* ObjectsToIgnore );
 
 private:
-	virtual FArchive& operator<<( UObject*& Obj ) override;
+	COREUOBJECT_API virtual FArchive& operator<<( UObject*& Obj ) override;
 
 	UObject* Target;
 	const TSet<UObject*>& PotentiallyReferencedObjects;
diff --git a/Engine/Source/Runtime/CoreUObject/Public/Serialization/ArchiveObjectCrc32.h b/Engine/Source/Runtime/CoreUObject/Public/Serialization/ArchiveObjectCrc32.h
--- a/Engine/Source/Runtime/CoreUObject/Public/Serialization/ArchiveObjectCrc32.h
+++ b/Engine/Source/Runtime/CoreUObject/Public/Serialization/ArchiveObjectCrc32.h
@@ -20,27 +20,27 @@ FArchiveObjectCrc32
 /**
 * Calculates a checksum on an object's serialized data stream.
 */
-class COREUOBJECT_API FArchiveObjectCrc32 : public FArchiveUObject
+class FArchiveObjectCrc32 : public FArchiveUObject
 {
 public:
 	/**
 	* Default constructor.
 	*/
-	FArchiveObjectCrc32();
+	COREUOBJECT_API FArchiveObjectCrc32();
 
 	//~ Begin FArchive Interface
-	virtual void Serialize(void* Data, int64 Length);
-	virtual FArchive& operator<<(class FName& Name);
-	virtual FArchive& operator<<(class UObject*& Object);
-	virtual FArchive& operator<<(FObjectPtr& ObjectPtr) override;
+	COREUOBJECT_API virtual void Serialize(void* Data, int64 Length);
+	COREUOBJECT_API virtual FArchive& operator<<(class FName& Name);
+	COREUOBJECT_API virtual FArchive& operator<<(class UObject*& Object);
+	COREUOBJECT_API virtual FArchive& operator<<(FObjectPtr& ObjectPtr) override;
 	virtual FString GetArchiveName() const { return TEXT("FArchiveObjectCrc32"); }
 	//~ End FArchive Interface
 
 	/**
 	* Serialize the given object, calculate and return its checksum.
 	*/
-	uint32 Crc32(UObject* Object, uint32 CRC = 0);
-	uint32 Crc32(UObject* Object, UObject* Root, uint32 CRC);
+	COREUOBJECT_API uint32 Crc32(UObject* Object, uint32 CRC = 0);
+	COREUOBJECT_API uint32 Crc32(UObject* Object, UObject* Root, uint32 CRC);
 
 protected:
 	/** Return if object was already serialized */
diff --git a/Engine/Source/Runtime/CoreUObject/Public/Serialization/ArchiveReplaceObjectRef.h b/Engine/Source/Runtime/CoreUObject/Public/Serialization/ArchiveReplaceObjectRef.h
--- a/Engine/Source/Runtime/CoreUObject/Public/Serialization/ArchiveReplaceObjectRef.h
+++ b/Engine/Source/Runtime/CoreUObject/Public/Serialization/ArchiveReplaceObjectRef.h
@@ -8,7 +8,7 @@
 #include "UObject/SoftObjectPath.h"
 
 /** Base class for object replacement archives */ 
-class COREUOBJECT_API FArchiveReplaceObjectRefBase : public FArchiveUObject
+class FArchiveReplaceObjectRefBase : public FArchiveUObject
 {
 public:
 
@@ -25,7 +25,7 @@ public:
 	/**
 	* Returns a reference to the replaced references map
 	*/
-	const TMap<UObject*, TArray<FProperty*>>& GetReplacedReferences() const;
+	COREUOBJECT_API const TMap<UObject*, TArray<FProperty*>>& GetReplacedReferences() const;
 
 	/**
 	* Returns the name of this archive.
@@ -94,7 +94,7 @@ protected:
 	/**
 	* Serializes a single object
 	*/
-	void SerializeObject(UObject* ObjectToSerialize);
+	COREUOBJECT_API void SerializeObject(UObject* ObjectToSerialize);
 
 	/** Initial object to start the reference search from */
 	UObject* SearchObject = nullptr;
