[CoreUObject]

diff --git a/Engine/Source/Runtime/CoreUObject/Public/Serialization/ArchiveStackTrace.h b/Engine/Source/Runtime/CoreUObject/Public/Serialization/ArchiveStackTrace.h
--- a/Engine/Source/Runtime/CoreUObject/Public/Serialization/ArchiveStackTrace.h
+++ b/Engine/Source/Runtime/CoreUObject/Public/Serialization/ArchiveStackTrace.h
@@ -26,7 +26,7 @@ struct FObjectImport;
 struct FUObjectSerializeContext;
 
 /** Structure that holds stats from comparing two packages */
-struct COREUOBJECT_API FArchiveDiffStats
+struct FArchiveDiffStats
 {
 	/** Size of all of the differences between two packages */
 	int64 DiffSize;
@@ -46,26 +46,26 @@ struct COREUOBJECT_API FArchiveDiffStats
 };
 
 /** Ignores saving the stack trace when collecting serialize offsets .*/
-class COREUOBJECT_API FArchiveStackTraceIgnoreScope
+class FArchiveStackTraceIgnoreScope
 {
 	const bool bIgnore;
 public:
-	FArchiveStackTraceIgnoreScope(bool bInIgnore = true);
-	~FArchiveStackTraceIgnoreScope();
+	COREUOBJECT_API FArchiveStackTraceIgnoreScope(bool bInIgnore = true);
+	COREUOBJECT_API ~FArchiveStackTraceIgnoreScope();
 };
 
 /**
  * Disables collecting both offsets and stack traces when collecting serialize callstacks.
  * Typically used when appending data from one stack tracing archive to another.
  */
-class COREUOBJECT_API FArchiveStackTraceDisabledScope
+class FArchiveStackTraceDisabledScope
 {
 public:
-	FArchiveStackTraceDisabledScope();
-	~FArchiveStackTraceDisabledScope();
+	COREUOBJECT_API FArchiveStackTraceDisabledScope();
+	COREUOBJECT_API ~FArchiveStackTraceDisabledScope();
 };
 
-struct COREUOBJECT_API FArchiveDiffInfo
+struct FArchiveDiffInfo
 {
 	int64 Offset;
 	int64 Size;
@@ -95,7 +95,7 @@ struct COREUOBJECT_API FArchiveDiffInfo
 	}
 };
 
-class COREUOBJECT_API FArchiveDiffMap : public TArray<FArchiveDiffInfo>
+class FArchiveDiffMap : public TArray<FArchiveDiffInfo>
 {
 public:
 	bool ContainsOffset(int64 Offset) const
@@ -113,7 +113,7 @@ public:
 };
 
 /** Holds offsets to captured callstacks. */
-class COREUOBJECT_API FArchiveCallstacks
+class FArchiveCallstacks
 {
 public:
 	/** Offset and callstack pair */
@@ -154,10 +154,10 @@ public:
 		FCallstackData Clone() const;
 	};
 	
-	explicit FArchiveCallstacks(UObject* InAsset);
+	COREUOBJECT_API explicit FArchiveCallstacks(UObject* InAsset);
 
 	/** Returns the asset class name. */
-	FName GetAssetClass() const;
+	COREUOBJECT_API FName GetAssetClass() const;
 
 	/** Returns the total number of callstacks. */
 	int32 Num() const
@@ -166,7 +166,7 @@ public:
 	}
 
 	/** Capture and append the current callstack. */
-	void Add(
+	COREUOBJECT_API void Add(
 		int64 Offset,
 		int64 Length,
 		UObject* SerializedObject,
@@ -177,10 +177,10 @@ public:
 		int32 StackIgnoreCount);
 
 	/** Append other callstacks. */
-	void Append(const FArchiveCallstacks& Other, int64 Offset = 0);
+	COREUOBJECT_API void Append(const FArchiveCallstacks& Other, int64 Offset = 0);
 
 	/** Finds a callstack associated with data at the specified offset */
-	int32 GetCallstackIndexAtOffset(int64 Offset, int32 MinOffsetIndex = 0) const;
+	COREUOBJECT_API int32 GetCallstackIndexAtOffset(int64 Offset, int32 MinOffsetIndex = 0) const;
 
 	/** Finds a callstack associated with data at the specified offset */
 	const FCallstackAtOffset& GetCallstack(int32 CallstackIndex) const
@@ -214,7 +214,7 @@ public:
 
 private:
 	/** Adds a unique callstack to UniqueCallstacks map */
-	ANSICHAR* AddUniqueCallstack(bool bIsCollectingCallstacks, UObject* SerializedObject, FProperty* SerializedProperty, uint32& OutCallstackCRC);
+	COREUOBJECT_API ANSICHAR* AddUniqueCallstack(bool bIsCollectingCallstacks, UObject* SerializedObject, FProperty* SerializedProperty, uint32& OutCallstackCRC);
 
 	/** The asset being serialized */
 	UObject* Asset;
@@ -235,21 +235,21 @@ private:
 };
 
 /** Archive proxy that captures callstacks for each serialize call. */
-class COREUOBJECT_API FArchiveStackTraceWriter
+class FArchiveStackTraceWriter
 	: public FArchiveProxy
 {
 public:
-	FArchiveStackTraceWriter(
+	COREUOBJECT_API FArchiveStackTraceWriter(
 		FArchive& InInner,
 		FArchiveCallstacks& InCallstacks,
 		const FArchiveDiffMap* InDiffMap = nullptr,
 		int64 InDiffMapStartOffset = 0);
 
-	virtual ~FArchiveStackTraceWriter() override;
+	COREUOBJECT_API virtual ~FArchiveStackTraceWriter() override;
 
-	FORCENOINLINE virtual void Serialize(void* Data, int64 Length) override; // FORCENOINLINE so it can be counted during StackTrace
-	virtual void SetSerializeContext(FUObjectSerializeContext* Context) override;
-	virtual FUObjectSerializeContext* GetSerializeContext() override;
+	COREUOBJECT_API FORCENOINLINE virtual void Serialize(void* Data, int64 Length) override; // FORCENOINLINE so it can be counted during StackTrace
+	COREUOBJECT_API virtual void SetSerializeContext(FUObjectSerializeContext* Context) override;
+	COREUOBJECT_API virtual FUObjectSerializeContext* GetSerializeContext() override;
 	
 #if WITH_EDITOR
 	virtual void PushDebugDataString(const FName& DebugData) override
@@ -288,7 +288,7 @@ public:
 	using EPackageHeaderFormat = ICookedPackageWriter::EPackageHeaderFormat;
 
 	/** Compares two packages and logs the differences and calltacks. */
-	static void Compare(
+	static COREUOBJECT_API void Compare(
 		const FPackageData& SourcePackage,
 		const FPackageData& DestPackage,
 		const FArchiveCallstacks& Callstacks,
@@ -301,7 +301,7 @@ public:
 		bool bSuppressLogging = false);
 
 	/** Creates map with mismatching callstacks. */
-	static bool GenerateDiffMap(
+	static COREUOBJECT_API bool GenerateDiffMap(
 		const FPackageData& SourcePackage,
 		const FPackageData& DestPackage,
 		const FArchiveCallstacks& Callstacks,
@@ -309,7 +309,7 @@ public:
 		FArchiveDiffMap& OutDiffMap);
 
 	/** Logs any mismatching header data. */
-	static void DumpPackageHeaderDiffs(
+	static COREUOBJECT_API void DumpPackageHeaderDiffs(
 		const FPackageData& SourcePackage,
 		const FPackageData& DestPackage,
 		const FString& AssetFilename,
@@ -317,7 +317,7 @@ public:
 		const EPackageHeaderFormat PackageHeaderFormat = EPackageHeaderFormat::PackageFileSummary);
 
 	/** Returns a new linker for loading the specified package. */
-	static FLinkerLoad* CreateLinkerForPackage(
+	static COREUOBJECT_API FLinkerLoad* CreateLinkerForPackage(
 		FUObjectSerializeContext* LoadContext,
 		const FString& InPackageName,
 		const FString& InFilename,
@@ -338,11 +338,11 @@ private:
 /**
  * Memory backed stack trace writer.
  */
-class COREUOBJECT_API FArchiveStackTraceMemoryWriter final
+class FArchiveStackTraceMemoryWriter final
 	: public FLargeMemoryWriter
 {
 public:
-	FArchiveStackTraceMemoryWriter(
+	COREUOBJECT_API FArchiveStackTraceMemoryWriter(
 		FArchiveCallstacks& Callstacks,
 		const FArchiveDiffMap* DiffMap = nullptr,
 		const int64 DiffMapOffset = 0,
@@ -350,9 +350,9 @@ public:
 		bool bIsPersistent = false,
 		const TCHAR* Filename = nullptr);
 
-	FORCENOINLINE virtual void Serialize(void* Memory, int64 Length) override; // FORCENOINLINE so it can be counted during StackTrace
-	virtual void SetSerializeContext(FUObjectSerializeContext* Context) override;
-	virtual FUObjectSerializeContext* GetSerializeContext() override;
+	COREUOBJECT_API FORCENOINLINE virtual void Serialize(void* Memory, int64 Length) override; // FORCENOINLINE so it can be counted during StackTrace
+	COREUOBJECT_API virtual void SetSerializeContext(FUObjectSerializeContext* Context) override;
+	COREUOBJECT_API virtual FUObjectSerializeContext* GetSerializeContext() override;
 
 private:
 	FArchiveStackTraceWriter StackTraceWriter;
@@ -362,14 +362,14 @@ private:
  * Archive that stores a callstack for each of the Serialize calls and has the ability to compare itself to an existing
  * package on disk and dump all the differences to log.
  */
-class COREUOBJECT_API FArchiveStackTrace
+class FArchiveStackTrace
 	: public FLargeMemoryWriter
 {
 public:
 	using FPackageData = FArchiveStackTraceWriter::FPackageData;
 
-	FArchiveStackTrace(UObject* InAsset, const TCHAR* InFilename, bool bInCollectCallstacks = true, const FArchiveDiffMap* InDiffMap = nullptr);
-	virtual ~FArchiveStackTrace();
+	COREUOBJECT_API FArchiveStackTrace(UObject* InAsset, const TCHAR* InFilename, bool bInCollectCallstacks = true, const FArchiveDiffMap* InDiffMap = nullptr);
+	COREUOBJECT_API virtual ~FArchiveStackTrace();
 
 	FArchiveCallstacks& GetCallstacks()
 	{
@@ -381,28 +381,28 @@ public:
 		return Callstacks;
 	}
 
-	FORCENOINLINE virtual void Serialize(void* Memory, int64 Length) override; // FORCENOINLINE so it can be counted during StackTrace
-	virtual void SetSerializeContext(FUObjectSerializeContext* Context) override;
-	virtual FUObjectSerializeContext* GetSerializeContext() override;
+	COREUOBJECT_API FORCENOINLINE virtual void Serialize(void* Memory, int64 Length) override; // FORCENOINLINE so it can be counted during StackTrace
+	COREUOBJECT_API virtual void SetSerializeContext(FUObjectSerializeContext* Context) override;
+	COREUOBJECT_API virtual FUObjectSerializeContext* GetSerializeContext() override;
 
 	/** Compares this archive with the given bytes from disk or FPackageData. Dumps all differences to log. */
-	void CompareWith(const TCHAR* InFilename, const int64 TotalHeaderSize, const TCHAR* CallstackCutoffText,
+	COREUOBJECT_API void CompareWith(const TCHAR* InFilename, const int64 TotalHeaderSize, const TCHAR* CallstackCutoffText,
 		const int32 MaxDiffsToLog, TMap<FName, FArchiveDiffStats>& OutStats);
-	void CompareWith(const FPackageData& SourcePackage, const TCHAR* FileDisplayName, const int64 TotalHeaderSize,
+	COREUOBJECT_API void CompareWith(const FPackageData& SourcePackage, const TCHAR* FileDisplayName, const int64 TotalHeaderSize,
 		const TCHAR* CallstackCutoffText, const int32 MaxDiffsToLog, TMap<FName, FArchiveDiffStats>& OutStats,
 		const FArchiveStackTraceWriter::EPackageHeaderFormat PackageHeaderFormat = FArchiveStackTraceWriter::EPackageHeaderFormat::PackageFileSummary);
 
 	/** Generates a map of all differences between this archive and the given bytes from disk or FPackageData. */
-	bool GenerateDiffMap(const TCHAR* InFilename, int64 TotalHeaderSize, int32 MaxDiffsToFind, FArchiveDiffMap& OutDiffMap);
-	bool GenerateDiffMap(const FPackageData& SourcePackage, int64 TotalHeaderSize, int32 MaxDiffsToFind,
+	COREUOBJECT_API bool GenerateDiffMap(const TCHAR* InFilename, int64 TotalHeaderSize, int32 MaxDiffsToFind, FArchiveDiffMap& OutDiffMap);
+	COREUOBJECT_API bool GenerateDiffMap(const FPackageData& SourcePackage, int64 TotalHeaderSize, int32 MaxDiffsToFind,
 		FArchiveDiffMap& OutDiffMap);
 
 	/** Compares the provided buffer with the given bytes from disk or FPackageData. */
-	static bool IsIdentical(const TCHAR* InFilename, int64 BufferSize, const uint8* BufferData);
-	static bool IsIdentical(const FPackageData& SourcePackage, int64 BufferSize, const uint8* BufferData);
+	static COREUOBJECT_API bool IsIdentical(const TCHAR* InFilename, int64 BufferSize, const uint8* BufferData);
+	static COREUOBJECT_API bool IsIdentical(const FPackageData& SourcePackage, int64 BufferSize, const uint8* BufferData);
 
 	/** Helper function to load package contents into memory. Supports EDL packages. */
-	static bool LoadPackageIntoMemory(const TCHAR* InFilename, FPackageData& OutPackageData,
+	static COREUOBJECT_API bool LoadPackageIntoMemory(const TCHAR* InFilename, FPackageData& OutPackageData,
 		TUniquePtr<uint8>& OutLoadedBytes);
 
 private:
@@ -410,7 +410,7 @@ private:
 	FArchiveStackTraceWriter StackTraceWriter;
 };
 
-class COREUOBJECT_API FArchiveStackTraceReader : public FLargeMemoryReader
+class FArchiveStackTraceReader : public FLargeMemoryReader
 {
 public:
 	struct FSerializeData
@@ -443,12 +443,12 @@ private:
 	FUObjectThreadContext& ThreadContext;
 public:
 
-	FArchiveStackTraceReader(const TCHAR* InFilename, const uint8* InData, const int64 Num);
+	COREUOBJECT_API FArchiveStackTraceReader(const TCHAR* InFilename, const uint8* InData, const int64 Num);
 
-	virtual void Serialize(void* OutData, int64 Num) override;
+	COREUOBJECT_API virtual void Serialize(void* OutData, int64 Num) override;
 	const TArray<FSerializeData>& GetSerializeTrace() const
 	{
 		return SerializeTrace;
 	}
-	static FArchiveStackTraceReader* CreateFromFile(const TCHAR* InFilename);
+	static COREUOBJECT_API FArchiveStackTraceReader* CreateFromFile(const TCHAR* InFilename);
 };
diff --git a/Engine/Source/Runtime/CoreUObject/Public/Serialization/ArchiveUObject.h b/Engine/Source/Runtime/CoreUObject/Public/Serialization/ArchiveUObject.h
--- a/Engine/Source/Runtime/CoreUObject/Public/Serialization/ArchiveUObject.h
+++ b/Engine/Source/Runtime/CoreUObject/Public/Serialization/ArchiveUObject.h
@@ -13,15 +13,15 @@ struct FWeakObjectPtr;
 /**
  * Base FArchive for serializing UObjects. Supports FLazyObjectPtr and FSoftObjectPtr serialization.
  */
-class COREUOBJECT_API FArchiveUObject : public FArchive
+class FArchiveUObject : public FArchive
 {
 public:
 
-	static FArchive& SerializeLazyObjectPtr(FArchive& Ar, FLazyObjectPtr& Value);
-	static FArchive& SerializeObjectPtr(FArchive& Ar, FObjectPtr& Value);
-	static FArchive& SerializeSoftObjectPtr(FArchive& Ar, FSoftObjectPtr& Value);
-	static FArchive& SerializeSoftObjectPath(FArchive& Ar, FSoftObjectPath& Value);
-	static FArchive& SerializeWeakObjectPtr(FArchive& Ar, FWeakObjectPtr& Value);
+	static COREUOBJECT_API FArchive& SerializeLazyObjectPtr(FArchive& Ar, FLazyObjectPtr& Value);
+	static COREUOBJECT_API FArchive& SerializeObjectPtr(FArchive& Ar, FObjectPtr& Value);
+	static COREUOBJECT_API FArchive& SerializeSoftObjectPtr(FArchive& Ar, FSoftObjectPtr& Value);
+	static COREUOBJECT_API FArchive& SerializeSoftObjectPath(FArchive& Ar, FSoftObjectPath& Value);
+	static COREUOBJECT_API FArchive& SerializeWeakObjectPtr(FArchive& Ar, FWeakObjectPtr& Value);
 
 	using FArchive::operator<<; // For visibility of the overloads we don't override
 
@@ -36,4 +36,4 @@ public:
 
 #if UE_ENABLE_INCLUDE_ORDER_DEPRECATED_IN_5_2
 #include "CoreMinimal.h"
-#endif
\ No newline at end of file
+#endif
diff --git a/Engine/Source/Runtime/CoreUObject/Public/Serialization/ArchiveUObjectFromStructuredArchive.h b/Engine/Source/Runtime/CoreUObject/Public/Serialization/ArchiveUObjectFromStructuredArchive.h
--- a/Engine/Source/Runtime/CoreUObject/Public/Serialization/ArchiveUObjectFromStructuredArchive.h
+++ b/Engine/Source/Runtime/CoreUObject/Public/Serialization/ArchiveUObjectFromStructuredArchive.h
@@ -21,26 +21,26 @@ struct FObjectPtr;
 
 #if WITH_TEXT_ARCHIVE_SUPPORT
 
-class COREUOBJECT_API FArchiveUObjectFromStructuredArchiveImpl : public FArchiveFromStructuredArchiveImpl
+class FArchiveUObjectFromStructuredArchiveImpl : public FArchiveFromStructuredArchiveImpl
 {
 	using Super = FArchiveFromStructuredArchiveImpl;
 
 public:
 
-	FArchiveUObjectFromStructuredArchiveImpl(FStructuredArchive::FSlot Slot);
+	COREUOBJECT_API FArchiveUObjectFromStructuredArchiveImpl(FStructuredArchive::FSlot Slot);
 
 	using FArchive::operator<<; // For visibility of the overloads we don't override
 
 	//~ Begin FArchive Interface
-	virtual FArchive& operator<<(FLazyObjectPtr& Value) override;
-	virtual FArchive& operator<<(FSoftObjectPtr& Value) override;
-	virtual FArchive& operator<<(FSoftObjectPath& Value) override;
-	virtual FArchive& operator<<(FWeakObjectPtr& Value) override;
-	virtual FArchive& operator<<(FObjectPtr& Value) override;
+	COREUOBJECT_API virtual FArchive& operator<<(FLazyObjectPtr& Value) override;
+	COREUOBJECT_API virtual FArchive& operator<<(FSoftObjectPtr& Value) override;
+	COREUOBJECT_API virtual FArchive& operator<<(FSoftObjectPath& Value) override;
+	COREUOBJECT_API virtual FArchive& operator<<(FWeakObjectPtr& Value) override;
+	COREUOBJECT_API virtual FArchive& operator<<(FObjectPtr& Value) override;
 	//~ End FArchive Interface
 
-	virtual void PushFileRegionType(EFileRegionType Type) override;
-	virtual void PopFileRegionType() override;
+	COREUOBJECT_API virtual void PushFileRegionType(EFileRegionType Type) override;
+	COREUOBJECT_API virtual void PopFileRegionType() override;
 
 private:
 
@@ -57,7 +57,7 @@ private:
 	TMap<FSoftObjectPtr, int32> SoftObjectPtrToIndex;
 	TMap<FSoftObjectPath, int32> SoftObjectPathToIndex;
 
-	virtual bool Finalize(FStructuredArchive::FRecord Record) override;
+	COREUOBJECT_API virtual bool Finalize(FStructuredArchive::FRecord Record) override;
 };
 
 class FArchiveUObjectFromStructuredArchive
@@ -79,7 +79,7 @@ private:
 
 #else
 
-class COREUOBJECT_API FArchiveUObjectFromStructuredArchive : public FArchiveFromStructuredArchive
+class FArchiveUObjectFromStructuredArchive : public FArchiveFromStructuredArchive
 {
 public:
 
@@ -90,4 +90,4 @@ public:
 	}
 };
 
-#endif
\ No newline at end of file
+#endif
diff --git a/Engine/Source/Runtime/CoreUObject/Public/Serialization/AsyncLoading2.h b/Engine/Source/Runtime/CoreUObject/Public/Serialization/AsyncLoading2.h
--- a/Engine/Source/Runtime/CoreUObject/Public/Serialization/AsyncLoading2.h
+++ b/Engine/Source/Runtime/CoreUObject/Public/Serialization/AsyncLoading2.h
@@ -149,7 +149,7 @@ public:
 		return TypeAndId != Other.TypeAndId;
 	}
 
-	COREUOBJECT_API friend FArchive& operator<<(FArchive& Ar, FPackageObjectIndex& Value)
+	friend FArchive& operator<<(FArchive& Ar, FPackageObjectIndex& Value)
 	{
 		Ar << Value.TypeAndId;
 		return Ar;
diff --git a/Engine/Source/Runtime/CoreUObject/Public/Serialization/AsyncLoadingFlushContext.h b/Engine/Source/Runtime/CoreUObject/Public/Serialization/AsyncLoadingFlushContext.h
--- a/Engine/Source/Runtime/CoreUObject/Public/Serialization/AsyncLoadingFlushContext.h
+++ b/Engine/Source/Runtime/CoreUObject/Public/Serialization/AsyncLoadingFlushContext.h
@@ -16,14 +16,14 @@ DECLARE_DELEGATE(FOnAsyncLoadingFlushComplete);
 /**
  * Flush the async loader in a non-blocking manner.
  */
-class COREUOBJECT_API FAsyncLoadingFlushContext final : FNoncopyable
+class FAsyncLoadingFlushContext final : FNoncopyable
 {
 public:
-	FAsyncLoadingFlushContext(const FString & Context);
-	FAsyncLoadingFlushContext(FString&& Context);
-	~FAsyncLoadingFlushContext();
+	COREUOBJECT_API FAsyncLoadingFlushContext(const FString & Context);
+	COREUOBJECT_API FAsyncLoadingFlushContext(FString&& Context);
+	COREUOBJECT_API ~FAsyncLoadingFlushContext();
 
-	void Flush(const FOnAsyncLoadingFlushComplete& OnFlushComplete);
+	COREUOBJECT_API void Flush(const FOnAsyncLoadingFlushComplete& OnFlushComplete);
 
 	int32 GetId() const { return Id; }
 
@@ -55,5 +55,5 @@ private:
 	FTSTicker::FDelegateHandle WaitingWarnDelegateHandle;
 
 	/** Increment ids so each instance has a unique id. */
-	static int32 NextId;
+	static COREUOBJECT_API int32 NextId;
 };
diff --git a/Engine/Source/Runtime/CoreUObject/Public/Serialization/BulkData.h b/Engine/Source/Runtime/CoreUObject/Public/Serialization/BulkData.h
--- a/Engine/Source/Runtime/CoreUObject/Public/Serialization/BulkData.h
+++ b/Engine/Source/Runtime/CoreUObject/Public/Serialization/BulkData.h
@@ -84,7 +84,7 @@ COREUOBJECT_API FIoFilenameHash MakeIoFilenameHash(const FIoChunkId& ChunkID);
  * It functions pretty much the same as IAsyncReadRequest expect that it also holds
  * the file handle as well.
  */
-class COREUOBJECT_API IBulkDataIORequest
+class IBulkDataIORequest
 {
 public:
 	virtual ~IBulkDataIORequest() {}
@@ -220,7 +220,7 @@ TUniquePtr<IBulkDataIORequest> CreateBulkDataIoDispatcherRequest(
 /**
  * @documentation @todo documentation
  */
-struct COREUOBJECT_API FOwnedBulkDataPtr
+struct FOwnedBulkDataPtr
 {
 	explicit FOwnedBulkDataPtr(void* InAllocatedData)
 		: AllocatedData(InAllocatedData)
@@ -238,8 +238,8 @@ struct COREUOBJECT_API FOwnedBulkDataPtr
 		
 	}
 
-	~FOwnedBulkDataPtr();
-	const void* GetPointer();
+	COREUOBJECT_API ~FOwnedBulkDataPtr();
+	COREUOBJECT_API const void* GetPointer();
 
 	IMappedFileHandle* GetMappedHandle()
 	{
@@ -490,7 +490,7 @@ struct FBulkDataSerializationParams
 /**
  * @documentation @todo documentation
  */
-class COREUOBJECT_API FBulkData
+class FBulkData
 {
 	// This struct represents an optional allocation.
 	struct FAllocatedPtr
@@ -513,7 +513,7 @@ class COREUOBJECT_API FBulkData
 		void* GetAllocationForWrite(const FBulkData* Owner) const;
 		const void* GetAllocationReadOnly(const FBulkData* Owner) const;
 
-		FOwnedBulkDataPtr* StealFileMapping(FBulkData* Owner);
+		COREUOBJECT_API FOwnedBulkDataPtr* StealFileMapping(FBulkData* Owner);
 		void Swap(FBulkData* Owner, void** DstBuffer);
 
 	private:
@@ -554,19 +554,19 @@ public:
 	 *
 	 * @param Other the source array to copy
 	 */
-	FBulkData( const FBulkData& Other );
+	COREUOBJECT_API FBulkData( const FBulkData& Other );
 
 	/**
 	 * Virtual destructor, free'ing allocated memory.
 	 */
-	~FBulkData();
+	COREUOBJECT_API ~FBulkData();
 
 	/**
 	 * Copies the source array into this one after detaching from archive.
 	 *
 	 * @param Other the source array to copy
 	 */
-	FBulkData& operator=( const FBulkData& Other );
+	COREUOBJECT_API FBulkData& operator=( const FBulkData& Other );
 
 	/*-----------------------------------------------------------------------------
 		Static functions.
@@ -577,7 +577,7 @@ public:
 	 *
 	 * @param Log FOutputDevice to use for logging
 	 */
-	static void DumpBulkDataUsage( FOutputDevice& Log );
+	static COREUOBJECT_API void DumpBulkDataUsage( FOutputDevice& Log );
 
 	/*-----------------------------------------------------------------------------
 		Accessors
@@ -588,64 +588,64 @@ public:
 	 *
 	 * @return Size of the bulk data in bytes
 	 */
-	int64 GetBulkDataSize() const;
+	COREUOBJECT_API int64 GetBulkDataSize() const;
 	/**
 	 * Returns the size of the bulk data on disk. This can differ from GetBulkDataSize if
 	 * BULKDATA_SerializeCompressed is set.
 	 *
 	 * @return Size of the bulk data on disk or INDEX_NONE in case there's no association
 	 */
-	int64 GetBulkDataSizeOnDisk() const;
+	COREUOBJECT_API int64 GetBulkDataSizeOnDisk() const;
 	/**
 	 * Returns the offset into the file the bulk data is located at.
 	 *
 	 * @return Offset into the file or INDEX_NONE in case there is no association
 	 */
-	int64 GetBulkDataOffsetInFile() const;
+	COREUOBJECT_API int64 GetBulkDataOffsetInFile() const;
 	/**
 	 * Returns whether the bulk data is stored compressed on disk.
 	 *
 	 * @return true if data is compressed on disk, false otherwise
 	 */
-	bool IsStoredCompressedOnDisk() const;
+	COREUOBJECT_API bool IsStoredCompressedOnDisk() const;
 
 	/**
 	 * Returns true if the data can be loaded from disk.
 	 */
-	bool CanLoadFromDisk() const;
+	COREUOBJECT_API bool CanLoadFromDisk() const;
 	
 	/**
 	 * Returns true if the data references a file that currently exists and can be referenced by the file system.
 	 */
-	bool DoesExist() const;
+	COREUOBJECT_API bool DoesExist() const;
 
 	/**
 	 * Returns flags usable to decompress the bulk data
 	 * 
 	 * @return NAME_None if the data was not compressed on disk, or valid format to pass to FCompression::UncompressMemory for this data
 	 */
-	FName GetDecompressionFormat() const;
+	COREUOBJECT_API FName GetDecompressionFormat() const;
 
 	/**
 	 * Returns whether the bulk data is currently loaded and resident in memory.
 	 *
 	 * @return true if bulk data is loaded, false otherwise
 	 */
-	bool IsBulkDataLoaded() const;
+	COREUOBJECT_API bool IsBulkDataLoaded() const;
 
 	/**
 	* Returns whether the bulk data asynchronous load has completed.
 	*
 	* @return true if bulk data has been loaded or async loading was not used to load this data, false otherwise
 	*/
-	bool IsAsyncLoadingComplete() const;
+	COREUOBJECT_API bool IsAsyncLoadingComplete() const;
 
 	/**
 	* Returns whether this bulk data is used
 	* @return true if BULKDATA_Unused is not set
 	*/
 	UE_DEPRECATED(5.3, "The feature is being removed, it is assumed that all bulkdata is available for use")
-	bool IsAvailableForUse() const;
+	COREUOBJECT_API bool IsAvailableForUse() const;
 
 	/**
 	* Returns whether this bulk data represents optional data or not
@@ -722,7 +722,7 @@ public:
 	 *
 	 * @param BulkDataFlagsToSet	Bulk data flags to set
 	 */
-	void SetBulkDataFlags( uint32 BulkDataFlagsToSet );
+	COREUOBJECT_API void SetBulkDataFlags( uint32 BulkDataFlagsToSet );
 
 	/**
 	 * Enable the given flags and disable all other flags.
@@ -731,31 +731,31 @@ public:
 	 *
 	 * @param BulkDataFlagsToSet	Bulk data flags to set
 	 */
-	void ResetBulkDataFlags(uint32 BulkDataFlagsToSet);
+	COREUOBJECT_API void ResetBulkDataFlags(uint32 BulkDataFlagsToSet);
 
 	/**
 	* Gets the current bulk data flags.
 	*
 	* @return Bulk data flags currently set
 	*/
-	uint32 GetBulkDataFlags() const;
+	COREUOBJECT_API uint32 GetBulkDataFlags() const;
 
 	UE_DEPRECATED(5.1, "Bulk Data will always use default alignment")
-	void SetBulkDataAlignment(uint16 BulkDataAlignmentToSet);
+	COREUOBJECT_API void SetBulkDataAlignment(uint16 BulkDataAlignmentToSet);
 
 	/**
 	* Gets the current bulk data alignment.
 	*
 	* @return Bulk data alignment currently set
 	*/
-	uint32 GetBulkDataAlignment() const;
+	COREUOBJECT_API uint32 GetBulkDataAlignment() const;
 
 	/**
 	 * Clears the passed in bulk data flags.
 	 *
 	 * @param BulkDataFlagsToClear	Bulk data flags to clear
 	 */
-	void ClearBulkDataFlags( uint32 BulkDataFlagsToClear );
+	COREUOBJECT_API void ClearBulkDataFlags( uint32 BulkDataFlagsToClear );
 
 	/** Returns the PackagePath this bulkdata resides in */
 	UE_DEPRECATED(5.1, "Deprecated, no replacement")
@@ -780,10 +780,10 @@ public:
 	FIoFilenameHash GetIoFilenameHash() const { return MakeIoFilenameHash(BulkChunkId); }
 	
 	/** Returns a FIoChunkId for the bulkdata payload, this will be invalid if the bulkdata is not stored in the IoStore */
-	FIoChunkId CreateChunkId() const;
+	COREUOBJECT_API FIoChunkId CreateChunkId() const;
 
 	/** Returns a string representing the bulk data for debugging purposes. */
-	FString GetDebugName() const;
+	COREUOBJECT_API FString GetDebugName() const;
 
 	/*-----------------------------------------------------------------------------
 		Data retrieval and manipulation.
@@ -795,20 +795,20 @@ public:
 	 * @param Dest [in/out] Pointer to pointer going to hold copy, can point to NULL pointer in which case memory is allocated
 	 * @param bDiscardInternalCopy Whether to discard/ free the potentially internally allocated copy of the data
 	 */
-	void GetCopy( void** Dest, bool bDiscardInternalCopy = true );
+	COREUOBJECT_API void GetCopy( void** Dest, bool bDiscardInternalCopy = true );
 
 	/**
 	 * Locks the bulk data and returns a pointer to it.
 	 *
 	 * @param	LockFlags	Flags determining lock behavior
 	 */
-	void* Lock( uint32 LockFlags );
+	COREUOBJECT_API void* Lock( uint32 LockFlags );
 
 	/**
 	 * Locks the bulk data and returns a read-only pointer to it.
 	 * This variant can be called on a const bulkdata
 	 */
-	const void* LockReadOnly() const;
+	COREUOBJECT_API const void* LockReadOnly() const;
 
 	/**
 	 * Change size of locked bulk data. Only valid if locked via read-write lock.
@@ -816,12 +816,12 @@ public:
 	 * @param ElementCount	Number of elements to allocate.
 	 * @param ElementSize	Size in bytes of the individual element.
 	 */
-	void* Realloc(int64 ElementCount, int64 ElementSize);
+	COREUOBJECT_API void* Realloc(int64 ElementCount, int64 ElementSize);
 
 	/** 
 	 * Unlocks bulk data after which point the pointer returned by Lock no longer is valid.
 	 */
-	void Unlock() const;
+	COREUOBJECT_API void Unlock() const;
 
 	/** 
 	 * Checks if this bulk is locked
@@ -839,26 +839,26 @@ public:
 	 * Note that once this has been called, the bulkdata object will no longer be able to reload
 	 * it's payload from disk!
 	 */
-	void RemoveBulkData();
+	COREUOBJECT_API void RemoveBulkData();
 
 #if !USE_RUNTIME_BULKDATA
 	/**
 	 * Load the bulk data using a file reader. Works even when no archive is attached to the bulk data..
   	 * @return Whether the operation succeeded.
 	 */
-	bool LoadBulkDataWithFileReader();
+	COREUOBJECT_API bool LoadBulkDataWithFileReader();
 
 	/**
 	 * Test if it is possible to load the bulk data using a file reader, even when no archive is attached to the bulk data.
 	 * @return Whether the operation is allowed.
 	 */
-	bool CanLoadBulkDataWithFileReader() const;
+	COREUOBJECT_API bool CanLoadBulkDataWithFileReader() const;
 #endif // USE_RUNTIME_BULKDATA
 
 	/**
 	 * Forces the bulk data to be resident in memory and detaches the archive.
 	 */
-	void ForceBulkDataResident();
+	COREUOBJECT_API void ForceBulkDataResident();
 
 	/** 
 	* Initiates a new asynchronous operation to load the dulkdata from disk assuming that it is not already
@@ -869,14 +869,14 @@ public:
 	* and false if the data is already loaded or cannot be loaded from disk.
 	*/
 	UE_DEPRECATED(5.1, "Use FBulkDataRequest or CreateStreamingRequest instead")
-	bool StartAsyncLoading();
+	COREUOBJECT_API bool StartAsyncLoading();
 	
 	/**
 	 * Sets whether we should store the data compressed on disk.
 	 *
 	 * @param CompressionFlags	Flags to use for compressing the data. Use COMPRESS_NONE for no compression, or something like COMPRESS_ZLIB to compress the data
 	 */
-	void StoreCompressedOnDisk( FName CompressionFormat );
+	COREUOBJECT_API void StoreCompressedOnDisk( FName CompressionFormat );
 
 	/**
 	 * Deallocates bulk data without detaching the archive, so that further bulk data accesses require a reload.
@@ -884,7 +884,7 @@ public:
 	 *
 	 * @return Whether the operation succeeded.
 	 */
-	bool UnloadBulkData();
+	COREUOBJECT_API bool UnloadBulkData();
 
 	/*-----------------------------------------------------------------------------
 		Serialization.
@@ -901,7 +901,7 @@ public:
 	 * @param bAttemptFileMapping	If true, attempt to map this instead of loading it into malloc'ed memory
 	 * @param FileRegionType	When cooking, a hint describing the type of data, used by some platforms to improve compression ratios
 	 */
-	void Serialize(FArchive& Ar, UObject* Owner, bool bAttemptFileMapping, int32 ElementSize, EFileRegionType FileRegionType);
+	COREUOBJECT_API void Serialize(FArchive& Ar, UObject* Owner, bool bAttemptFileMapping, int32 ElementSize, EFileRegionType FileRegionType);
 	
 	/**
 	 * Serialize just the bulk data portion to/ from the passed in memory.
@@ -922,7 +922,7 @@ public:
 	}
 
 	UE_DEPRECATED(5.1, "Call GetBulkDataVersions instead.")
-	FCustomVersionContainer GetCustomVersions(FArchive& InlineArchive) const;
+	COREUOBJECT_API FCustomVersionContainer GetCustomVersions(FArchive& InlineArchive) const;
 
 	/**
 	 * Get the CustomVersions used in the file containing the BulkData payload. If !IsInSeparateFile, this will be
@@ -932,7 +932,7 @@ public:
 	 * @param InlineArchive The archive that was used to load this object
 	 *
 	 */
-	void GetBulkDataVersions(FArchive& InlineArchive, FPackageFileVersion& OutUEVersion, int32& OutLicenseeUEVersion,
+	COREUOBJECT_API void GetBulkDataVersions(FArchive& InlineArchive, FPackageFileVersion& OutUEVersion, int32& OutLicenseeUEVersion,
 		FCustomVersionContainer& OutCustomVersions) const;
 
 #if WITH_EDITOR
@@ -942,7 +942,7 @@ public:
 	 * loading from disk are not the same as the values on disk.
 	 * This function handles running the same steps that SerializeBulkData does, but skips the deserialization of the BulkData.
 	 * */
-	void SetFlagsFromDiskWrittenValues(EBulkDataFlags InBulkDataFlags, int64 InBulkDataOffsetInFile, int64 InBulkDataSizeOnDisk, int64 LinkerSummaryBulkDataStartOffset);
+	COREUOBJECT_API void SetFlagsFromDiskWrittenValues(EBulkDataFlags InBulkDataFlags, int64 InBulkDataOffsetInFile, int64 InBulkDataSizeOnDisk, int64 LinkerSummaryBulkDataStartOffset);
 #endif
 
 	/*-----------------------------------------------------------------------------
@@ -954,7 +954,7 @@ public:
 	 *
 	 * @return A valid handle if the file can be accessed, if it cannot then nullptr.
 	 */
-	IAsyncReadFileHandle* OpenAsyncReadHandle() const;
+	COREUOBJECT_API IAsyncReadFileHandle* OpenAsyncReadHandle() const;
 
 	/**
 	 * Create an async read request for the bulk data.
@@ -965,7 +965,7 @@ public:
 	 * @param UserSuppliedMemory	A pointer to memory for the IO request to be written to, it is up to the caller to make sure that it is large enough. If the pointer is null then the system will allocate memory instead.
 	 * @return						A request for the read. This is owned by the caller and must be deleted by the caller.
 	 */
-	IBulkDataIORequest* CreateStreamingRequest(EAsyncIOPriorityAndFlags Priority, FBulkDataIORequestCallBack* CompleteCallback, uint8* UserSuppliedMemory) const;
+	COREUOBJECT_API IBulkDataIORequest* CreateStreamingRequest(EAsyncIOPriorityAndFlags Priority, FBulkDataIORequestCallBack* CompleteCallback, uint8* UserSuppliedMemory) const;
 
 	/**
 	 * Create an async read request for the bulk data.
@@ -978,7 +978,7 @@ public:
 	 * @param UserSuppliedMemory	A pointer to memory for the IO request to be written to, it is up to the caller to make sure that it is large enough. If the pointer is null then the system will allocate memory instead.
 	 * @return						A request for the read. This is owned by the caller and must be deleted by the caller.
 	 */
-	IBulkDataIORequest* CreateStreamingRequest(int64 OffsetInBulkData, int64 BytesToRead, EAsyncIOPriorityAndFlags Priority, FBulkDataIORequestCallBack* CompleteCallback, uint8* UserSuppliedMemory) const;
+	COREUOBJECT_API IBulkDataIORequest* CreateStreamingRequest(int64 OffsetInBulkData, int64 BytesToRead, EAsyncIOPriorityAndFlags Priority, FBulkDataIORequestCallBack* CompleteCallback, uint8* UserSuppliedMemory) const;
 
 	/**
 	 * Create an async read request for a range of bulk data streaming tokens
@@ -993,16 +993,16 @@ public:
 	 * @return					A request for the read. This is owned by the caller and must be deleted by the caller.
 	**/
 	UE_DEPRECATED(5.1, "Use FBulkDataRequest instead")
-	static IBulkDataIORequest* CreateStreamingRequestForRange(const BulkDataRangeArray& RangeArray, EAsyncIOPriorityAndFlags Priority, FBulkDataIORequestCallBack* CompleteCallback);
+	static COREUOBJECT_API IBulkDataIORequest* CreateStreamingRequestForRange(const BulkDataRangeArray& RangeArray, EAsyncIOPriorityAndFlags Priority, FBulkDataIORequestCallBack* CompleteCallback);
 
 	/** Enable the given flags in the given accumulator variable. */
-	static void SetBulkDataFlagsOn(EBulkDataFlags& InOutAccumulator, EBulkDataFlags FlagsToSet);
+	static COREUOBJECT_API void SetBulkDataFlagsOn(EBulkDataFlags& InOutAccumulator, EBulkDataFlags FlagsToSet);
 	/** Disable the given flags in the given accumulator variable. */
-	static void ClearBulkDataFlagsOn(EBulkDataFlags& InOutAccumulator, EBulkDataFlags FlagsToClear);
+	static COREUOBJECT_API void ClearBulkDataFlagsOn(EBulkDataFlags& InOutAccumulator, EBulkDataFlags FlagsToClear);
 	/** Returns whether all of the specified flags are set. */
-	static bool HasFlags(EBulkDataFlags Flags, EBulkDataFlags Contains);
+	static COREUOBJECT_API bool HasFlags(EBulkDataFlags Flags, EBulkDataFlags Contains);
 	/** Returns decompress method flags specified by the given bulk data flags. */
-	static FName GetDecompressionFormat(EBulkDataFlags InFlags);
+	static COREUOBJECT_API FName GetDecompressionFormat(EBulkDataFlags InFlags);
 
 	/*-----------------------------------------------------------------------------
 		Class specific virtuals.
@@ -1010,11 +1010,11 @@ public:
 
 protected:
 
-	void SerializeBulkData(FArchive& Ar, void* Data, int64 DataSize, EBulkDataFlags InBulkDataFlags);
+	COREUOBJECT_API void SerializeBulkData(FArchive& Ar, void* Data, int64 DataSize, EBulkDataFlags InBulkDataFlags);
 
 private:
 
-	int64 SerializePayload(FArchive& Ar, EBulkDataFlags SerializationFlags, const TOptional<EFileRegionType>& RegionType);
+	COREUOBJECT_API int64 SerializePayload(FArchive& Ar, EBulkDataFlags SerializationFlags, const TOptional<EFileRegionType>& RegionType);
 #if WITH_EDITOR
 	/**
 	 * Detaches the bulk data from the passed in archive. Needs to match the archive we are currently
@@ -1023,7 +1023,7 @@ private:
 	 * @param Ar						Archive to detach from
 	 * @param bEnsureBulkDataIsLoaded	whether to ensure that bulk data is load before detaching from archive
 	 */
-	void DetachFromArchive( FArchive* Ar, bool bEnsureBulkDataIsLoaded );
+	COREUOBJECT_API void DetachFromArchive( FArchive* Ar, bool bEnsureBulkDataIsLoaded );
 #endif // WITH_EDITOR
 
 	/*-----------------------------------------------------------------------------
@@ -1033,7 +1033,7 @@ private:
 	/**
 	 * Loads the bulk data if it is not already loaded.
 	 */
-	void MakeSureBulkDataIsLoaded();
+	COREUOBJECT_API void MakeSureBulkDataIsLoaded();
 
 	/**
 	 * Loads the data from disk into the specified memory block. This requires us still being attached to an
@@ -1043,15 +1043,15 @@ private:
 	 *
 	 * @return Whether the load succeeded
 	 */
-	bool TryLoadDataIntoMemory(FIoBuffer Dest);
+	COREUOBJECT_API bool TryLoadDataIntoMemory(FIoBuffer Dest);
 
 	/** Flushes any pending async load of bulk data  and copies the data to Dest buffer*/
-	void FlushAsyncLoading();
+	COREUOBJECT_API void FlushAsyncLoading();
 
 	/** Returns if the offset needs fixing when serialized */
-	bool NeedsOffsetFixup() const;
+	COREUOBJECT_API bool NeedsOffsetFixup() const;
 	
-	bool CanDiscardInternalData() const;
+	COREUOBJECT_API bool CanDiscardInternalData() const;
 	
 	/** Reallocate bulk data */
 	inline void* ReallocateData(int64 SizeInBytes) { return DataAllocation.ReallocateData(this, SizeInBytes); }
@@ -1192,7 +1192,7 @@ public:
 };
 
 UE_DEPRECATED(5.1, "Use FBulkData/TBulkData");
-struct COREUOBJECT_API FUntypedBulkData : public FBulkData
+struct FUntypedBulkData : public FBulkData
 {
 	virtual ~FUntypedBulkData() = default;
 
@@ -1243,7 +1243,7 @@ protected:
 #endif // !USE_RUNTIME_BULKDATA
 	}
 
-	virtual void SerializeElements(FArchive& Ar, void* Data);
+	COREUOBJECT_API virtual void SerializeElements(FArchive& Ar, void* Data);
 	
 	virtual void SerializeElement(FArchive& Ar, void* Data, int64 ElementIndex) = 0;
 
@@ -1254,7 +1254,7 @@ protected:
 	
 private:
 
-	void SerializeBulkData(FArchive& Ar, void* Data, int64 BulkDataSize, EBulkDataFlags InBulkDataFlags);
+	COREUOBJECT_API void SerializeBulkData(FArchive& Ar, void* Data, int64 BulkDataSize, EBulkDataFlags InBulkDataFlags);
 
 	FSerializeBulkDataElements SerializeElementsCallback;
 };
@@ -1307,7 +1307,7 @@ public:
 };
 
 /** Handle to a bulk data I/O request. */
-class COREUOBJECT_API FBulkDataRequest
+class FBulkDataRequest
 {
 public:
 	class IHandle;
@@ -1334,16 +1334,16 @@ public:
 	static constexpr EAsyncIOPriorityAndFlags DefaultPriority = AIOP_BelowNormal;
 
 	/** Constructs a new handle to bulk data request. */
-	FBulkDataRequest();
+	COREUOBJECT_API FBulkDataRequest();
 
 	/** Destructor, cancels and waits for any pending requests. */
-	~FBulkDataRequest();
+	COREUOBJECT_API ~FBulkDataRequest();
 
 	/** Moves ownership from an invalid or pending request. */
-	FBulkDataRequest(FBulkDataRequest&&);
+	COREUOBJECT_API FBulkDataRequest(FBulkDataRequest&&);
 
 	/** Moves ownership from an invalid or pending request. */
-	FBulkDataRequest& operator=(FBulkDataRequest&&);
+	COREUOBJECT_API FBulkDataRequest& operator=(FBulkDataRequest&&);
 	
 	/** Not copy constructable. */
 	FBulkDataRequest(const FBulkDataRequest&) = delete;
@@ -1352,31 +1352,31 @@ public:
 	FBulkDataRequest& operator=(const FBulkDataRequest&) = delete;
 
 	/** Returns current status of the request. */
-	EStatus GetStatus() const;
+	COREUOBJECT_API EStatus GetStatus() const;
 
 	/** Returns whether the request is associated with a pending or completed request. */
-	bool IsNone() const;
+	COREUOBJECT_API bool IsNone() const;
 
 	/** Returns whether the request is pending. */
-	bool IsPending() const;
+	COREUOBJECT_API bool IsPending() const;
 
 	/** Returns whether the request completed successfully. */
-	bool IsOk() const;
+	COREUOBJECT_API bool IsOk() const;
 
 	/** Returns whether the request has been completed. */
-	bool IsCompleted() const;
+	COREUOBJECT_API bool IsCompleted() const;
 
 	/** Cancel the pending request. */
-	void Cancel();
+	COREUOBJECT_API void Cancel();
 
 	/** Reset the request handle to an invalid state. Will cancel and wait if the request is not completed. */
-	void Reset();
+	COREUOBJECT_API void Reset();
 
 protected:
 	friend class FBulkDataBatchRequest;
 	
-	FBulkDataRequest(IHandle* InHandle);
-	int32 GetRefCount() const;
+	COREUOBJECT_API FBulkDataRequest(IHandle* InHandle);
+	COREUOBJECT_API int32 GetRefCount() const;
 
 	TRefCountPtr<IHandle> Handle;
 };
@@ -1392,23 +1392,23 @@ using FBulkDataBatchReadRequest = FBulkDataRequest;
  * one handle needs to be passed in. The last handle will block until the
  * entire batch is complete before being released.
  */
-class COREUOBJECT_API FBulkDataBatchRequest : public FBulkDataRequest
+class FBulkDataBatchRequest : public FBulkDataRequest
 {
 public:
 	class FBatchHandle;
 
 private:
-	class COREUOBJECT_API FBuilder
+	class FBuilder
 	{
 	public:
-		~FBuilder();
+		COREUOBJECT_API ~FBuilder();
 		FBuilder(const FBuilder&) = delete;
 		FBuilder& operator=(const FBuilder&) = delete;
 
 	protected:
-		explicit FBuilder(int32 MaxCount);
-		FBulkDataBatchRequest::FBatchHandle& GetBatch();
-		EStatus IssueBatch(FBulkDataBatchRequest* OutRequest, FCompletionCallback&& Callback);
+		COREUOBJECT_API explicit FBuilder(int32 MaxCount);
+		COREUOBJECT_API FBulkDataBatchRequest::FBatchHandle& GetBatch();
+		COREUOBJECT_API EStatus IssueBatch(FBulkDataBatchRequest* OutRequest, FCompletionCallback&& Callback);
 
 		int32 BatchCount = 0;
 		int32 NumLoaded = 0;
@@ -1422,21 +1422,21 @@ public:
 	using FBulkDataRequest::FBulkDataRequest;
 
 	/** Blocks the calling thread until the request is completed. */
-	void Wait();
+	COREUOBJECT_API void Wait();
 
 	/** Waits the specified amount of time in milliseconds for the request to be completed. */
-	bool WaitFor(uint32 Milliseconds);
+	COREUOBJECT_API bool WaitFor(uint32 Milliseconds);
 
 	/** Waits the specified amount of time for the request to be completed. */
-	bool WaitFor(const FTimespan& WaitTime);
+	COREUOBJECT_API bool WaitFor(const FTimespan& WaitTime);
 
 	/** Issue one or more I/O request in a single batch. */
-	class COREUOBJECT_API FBatchBuilder : public FBuilder
+	class FBatchBuilder : public FBuilder
 	{
 	public:
-		FBatchBuilder(int32 MaxCount);
+		COREUOBJECT_API FBatchBuilder(int32 MaxCount);
 		/** Read the entire bulk data and copy the result to the specified instance. */
-		FBatchBuilder& Read(FBulkData& BulkData, EAsyncIOPriorityAndFlags Priority = DefaultPriority);
+		COREUOBJECT_API FBatchBuilder& Read(FBulkData& BulkData, EAsyncIOPriorityAndFlags Priority = DefaultPriority);
 		/**
 		 * Read the bulk data from the specified offset and size and copy the result into the destination buffer.
 		 * @param BulkData		The bulk data instance.
@@ -1468,26 +1468,26 @@ public:
 		 * @param OutRequest	A handle to the batch request.
 		 * @return				Status of the issue operation.
 		 */
-		EStatus Issue(FBulkDataBatchRequest& OutRequest);
+		COREUOBJECT_API EStatus Issue(FBulkDataBatchRequest& OutRequest);
 		/**
 		 * Issue the batch. 
 		 * @return				Status of the issue operation.
 		 *
 		 * @note Assumes one or more handle(s) has been passed into any of the read operations.
 		 */
-		[[nodiscard]] EStatus Issue();
+		[[nodiscard]] COREUOBJECT_API EStatus Issue();
 
 	private:
-		FBatchBuilder& Read(const FBulkData& BulkData, uint64 Offset, uint64 Size, EAsyncIOPriorityAndFlags Priority, FIoBuffer& Dst, FBulkDataBatchReadRequest* OutRequest);
+		COREUOBJECT_API FBatchBuilder& Read(const FBulkData& BulkData, uint64 Offset, uint64 Size, EAsyncIOPriorityAndFlags Priority, FIoBuffer& Dst, FBulkDataBatchReadRequest* OutRequest);
 	};
 
 	/** Reads one or more bulk data and copies the result into a single I/O buffer. */
-	class COREUOBJECT_API FScatterGatherBuilder : public FBuilder
+	class FScatterGatherBuilder : public FBuilder
 	{
 	public:
-		FScatterGatherBuilder(int32 MaxCount);
+		COREUOBJECT_API FScatterGatherBuilder(int32 MaxCount);
 		/** Read the bulk data from the specified offset and size. */
-		FScatterGatherBuilder& Read(const FBulkData& BulkData, uint64 Offset = 0, uint64 Size = MAX_uint64);
+		COREUOBJECT_API FScatterGatherBuilder& Read(const FBulkData& BulkData, uint64 Offset = 0, uint64 Size = MAX_uint64);
 		/**
 		 * Issue the batch.
 		 * @param Dst			An empty or preallocated I/O buffer.	
@@ -1498,7 +1498,7 @@ public:
 		 *
 		 * @note Assumes one or more handle(s) has been passed into any of the read operations.
 		 */
-		EStatus Issue(FIoBuffer& Dst, EAsyncIOPriorityAndFlags Priority, FCompletionCallback&& Callback, FBulkDataBatchRequest& OutRequest);
+		COREUOBJECT_API EStatus Issue(FIoBuffer& Dst, EAsyncIOPriorityAndFlags Priority, FCompletionCallback&& Callback, FBulkDataBatchRequest& OutRequest);
 
 	private:
 		struct FRequest
diff --git a/Engine/Source/Runtime/CoreUObject/Public/Serialization/BulkDataRegistry.h b/Engine/Source/Runtime/CoreUObject/Public/Serialization/BulkDataRegistry.h
--- a/Engine/Source/Runtime/CoreUObject/Public/Serialization/BulkDataRegistry.h
+++ b/Engine/Source/Runtime/CoreUObject/Public/Serialization/BulkDataRegistry.h
@@ -145,8 +145,8 @@ public:
 	COREUOBJECT_API uint64 GetBulkDataResaveSize(FName PackageName);
 
 private:
-	COREUOBJECT_API void OnEndLoadPackage(const FEndLoadPackageContext& Context);
-	COREUOBJECT_API void OnAllModuleLoadingPhasesComplete();
+	void OnEndLoadPackage(const FEndLoadPackageContext& Context);
+	void OnAllModuleLoadingPhasesComplete();
 
 	FRWLock Lock;
 	TMap<FName, uint64> PackageBulkResaveSize;
@@ -165,4 +165,4 @@ COREUOBJECT_API UE::DerivedData::IBuildInputResolver* GetGlobalBuildInputResolve
 COREUOBJECT_API void SetGlobalBuildInputResolver(UE::DerivedData::IBuildInputResolver* InResolver);
 
 
-#endif // WITH_EDITOR
\ No newline at end of file
+#endif // WITH_EDITOR
diff --git a/Engine/Source/Runtime/CoreUObject/Public/Serialization/DeferredMessageLog.h b/Engine/Source/Runtime/CoreUObject/Public/Serialization/DeferredMessageLog.h
--- a/Engine/Source/Runtime/CoreUObject/Public/Serialization/DeferredMessageLog.h
+++ b/Engine/Source/Runtime/CoreUObject/Public/Serialization/DeferredMessageLog.h
@@ -23,7 +23,7 @@ class FTokenizedMessage;
  * finished to prevent modules from being loaded outside of game thread.
  * Also makes sure the messages are added to the message queue in a thread-safe way.
 */
-class COREUOBJECT_API FDeferredMessageLog
+class FDeferredMessageLog
 {
 	FName LogCategory;
 
@@ -33,12 +33,12 @@ class COREUOBJECT_API FDeferredMessageLog
 	void AddMessage(TSharedRef<FTokenizedMessage>& Message);
 
 public:
-	FDeferredMessageLog(const FName& InLogCategory);
+	COREUOBJECT_API FDeferredMessageLog(const FName& InLogCategory);
 	
-	TSharedRef<FTokenizedMessage> Info(const FText& Message);
-	TSharedRef<FTokenizedMessage> Warning(const FText& Message);
-	TSharedRef<FTokenizedMessage> Error(const FText& Message);
+	COREUOBJECT_API TSharedRef<FTokenizedMessage> Info(const FText& Message);
+	COREUOBJECT_API TSharedRef<FTokenizedMessage> Warning(const FText& Message);
+	COREUOBJECT_API TSharedRef<FTokenizedMessage> Error(const FText& Message);
 
-	static void Flush();
-	static void Cleanup();
+	static COREUOBJECT_API void Flush();
+	static COREUOBJECT_API void Cleanup();
 };
diff --git a/Engine/Source/Runtime/CoreUObject/Public/Serialization/EditorBulkData.h b/Engine/Source/Runtime/CoreUObject/Public/Serialization/EditorBulkData.h
--- a/Engine/Source/Runtime/CoreUObject/Public/Serialization/EditorBulkData.h
+++ b/Engine/Source/Runtime/CoreUObject/Public/Serialization/EditorBulkData.h
@@ -128,17 +128,17 @@ enum class ECompressionOptions : uint8
  */
 
 /** The base class with no type */
-class COREUOBJECT_API FEditorBulkData final
+class FEditorBulkData final
 {
 public:
 	FEditorBulkData() = default;
-	~FEditorBulkData();
+	COREUOBJECT_API ~FEditorBulkData();
 
-	FEditorBulkData(FEditorBulkData&& Other);
-	FEditorBulkData& operator=(FEditorBulkData&& Other);
+	COREUOBJECT_API FEditorBulkData(FEditorBulkData&& Other);
+	COREUOBJECT_API FEditorBulkData& operator=(FEditorBulkData&& Other);
 
-	FEditorBulkData(const FEditorBulkData& Other);
-	FEditorBulkData& operator=(const FEditorBulkData& Other);
+	COREUOBJECT_API FEditorBulkData(const FEditorBulkData& Other);
+	COREUOBJECT_API FEditorBulkData& operator=(const FEditorBulkData& Other);
 
 	/** 
 	 * Convenience method to make it easier to convert from FBulkData to FEditorBulkData and sets the Guid 
@@ -147,13 +147,13 @@ public:
 	 * @param Guid		A guid associated with the bulkdata object which will be used to identify the payload.
 	 *					This MUST remain the same between sessions so that the payloads key remains consistent!
 	 */
-	void CreateFromBulkData(FBulkData& BulkData, const FGuid& Guid, UObject* Owner);
+	COREUOBJECT_API void CreateFromBulkData(FBulkData& BulkData, const FGuid& Guid, UObject* Owner);
 	
 	/** 
 	 * Fix legacy content that created the Id from non-unique Guids.
 	 * Run on objects created before FUE5MainStreamObjectVersion::VirtualizedBulkDataHaveUniqueGuids
 	 */
-	void CreateLegacyUniqueIdentifier(UObject* Owner);
+	COREUOBJECT_API void CreateLegacyUniqueIdentifier(UObject* Owner);
 
 	/** 
 	 * Used to serialize the bulkdata to/from a FArchive
@@ -165,16 +165,16 @@ public:
 	 * @param bAllowRegistry Legacy parameter to skip registration when loading BulkData we know we will need to
 	 *				modify the identifier of. Should always be true for non-legacy serialization.
 	 */
-	void Serialize(FArchive& Ar, UObject* Owner, bool bAllowRegister=true);
+	COREUOBJECT_API void Serialize(FArchive& Ar, UObject* Owner, bool bAllowRegister=true);
 
 	/** Reset to a truly empty state */
-	void Reset();
+	COREUOBJECT_API void Reset();
 
 	// TODO: Not sure if there is a scenario where we can reload the payload off disk but it is also
 	// in memory.and so this might not really do anything anymore.
 
 	/** Unloads the data (if possible) but leaves it in a state where the data can be reloaded */
-	void UnloadData();
+	COREUOBJECT_API void UnloadData();
 
 	/** 
 	 * Removes the ability for the bulkdata object to load it's payload from disk (if it was doing so) 
@@ -184,13 +184,13 @@ public:
 	 *									the payload should be loaded into memory so that it can be accessed in the future.
 	 *									If false then the payload is not important and does not need to be loaded.
 	 */
-	void DetachFromDisk(FArchive* Ar, bool bEnsurePayloadIsLoaded);
+	COREUOBJECT_API void DetachFromDisk(FArchive* Ar, bool bEnsurePayloadIsLoaded);
 
 	/** Returns a unique identifier for the object itself. */
-	FGuid GetIdentifier() const;
+	COREUOBJECT_API FGuid GetIdentifier() const;
 
 	/** Returns an unique identifier for the content of the payload. */
-	const FIoHash& GetPayloadId() const;
+	COREUOBJECT_API const FIoHash& GetPayloadId() const;
 
 	/** Returns the size of the payload in bytes. */
 	int64 GetPayloadSize() const
@@ -205,10 +205,10 @@ public:
 	}
 
 	/** Returns if the payload would require loading in order to be accessed. Returns false if the payload is already in memory or of zero length */
-	bool DoesPayloadNeedLoading() const;
+	COREUOBJECT_API bool DoesPayloadNeedLoading() const;
 
 	/** Returns an immutable FCompressedBuffer reference to the payload data. */
-	TFuture<FSharedBuffer> GetPayload() const;
+	COREUOBJECT_API TFuture<FSharedBuffer> GetPayload() const;
 
 	/**
 	 * Returns an immutable FCompressedBuffer reference to the payload data.
@@ -217,7 +217,7 @@ public:
 	 * will be handled by the FCompressedBuffer interface. Call FCompressedBuffer::Decompress() to get access to
 	 * the payload in FSharedBuffer format.
 	 */
-	TFuture<FCompressedBuffer> GetCompressedPayload() const;
+	COREUOBJECT_API TFuture<FCompressedBuffer> GetCompressedPayload() const;
 
 	/**
 	 * Replaces the existing payload (if any) with a new one. 
@@ -233,7 +233,7 @@ public:
 	 * @param InPayload	The payload that this bulkdata object should reference. @see FSharedBuffer
 	 * @param Owner The object that owns the bulkdata, or null to not associate with a UObject.
 	 */
-	void UpdatePayload(FSharedBuffer InPayload, UObject* Owner = nullptr);
+	COREUOBJECT_API void UpdatePayload(FSharedBuffer InPayload, UObject* Owner = nullptr);
 
 	/**
 	 * Replaces the existing payload (if any) with a new one.
@@ -243,14 +243,14 @@ public:
 	 * @param InPayload	The payload that this bulkdata object should reference. @see FCompressedBuffer
 	 * @param Owner The object that owns the bulkdata, or null to not associate with a UObject.
 	 */
-	void UpdatePayload(FCompressedBuffer InPayload, UObject* Owner = nullptr);
+	COREUOBJECT_API void UpdatePayload(FCompressedBuffer InPayload, UObject* Owner = nullptr);
 	
 	/**
 	 * Utility struct used to compute the Payload ID before calling UpdatePayload
 	 */
-	struct COREUOBJECT_API FSharedBufferWithID
+	struct FSharedBufferWithID
 	{
-		FSharedBufferWithID(FSharedBuffer InPayload);
+		COREUOBJECT_API FSharedBufferWithID(FSharedBuffer InPayload);
 
 		FSharedBufferWithID() = default;
 		FSharedBufferWithID(FSharedBufferWithID&&) = default;
@@ -282,7 +282,7 @@ public:
 	 * @param InPayload	The payload to update the bulkdata with
 	 * @param Owner The object that owns the bulkdata, or null to not associate with a UObject.
 	 */
-	void UpdatePayload(FSharedBufferWithID InPayload, UObject* Owner = nullptr);
+	COREUOBJECT_API void UpdatePayload(FSharedBufferWithID InPayload, UObject* Owner = nullptr);
 
 	/** 
 	 * Sets the compression options to be applied to the payload during serialization.
@@ -292,7 +292,7 @@ public:
 	 * 
 	 * @param Option	The high level option to use. @see UE::Serialization::ECompressionOptions
 	 */ 
-	void SetCompressionOptions(ECompressionOptions Option);
+	COREUOBJECT_API void SetCompressionOptions(ECompressionOptions Option);
 
 	/** 
 	 * Sets the compression options to be applied to the payload during serialization.
@@ -303,17 +303,17 @@ public:
 	 * @param Compressor		The Oodle compressor to use. @see ECompressedBufferCompressor
 	 * @param CompressionLevel	The Oodle compression level to use. @see ECompressedBufferCompressionLevel
 	 */
-	void SetCompressionOptions(ECompressedBufferCompressor Compressor, ECompressedBufferCompressionLevel CompressionLevel);
+	COREUOBJECT_API void SetCompressionOptions(ECompressedBufferCompressor Compressor, ECompressedBufferCompressionLevel CompressionLevel);
 
 	UE_DEPRECATED(5.1, "Call GetBulkDataVersions instead.")
-	FCustomVersionContainer GetCustomVersions(FArchive& InlineArchive);
+	COREUOBJECT_API FCustomVersionContainer GetCustomVersions(FArchive& InlineArchive);
 
 	/**
 	 * Get the versions used in the file containing the payload.
 	 *
 	 * @param InlineArchive The archive that was used to load this object
 	 */
-	void GetBulkDataVersions(FArchive& InlineArchive, FPackageFileVersion& OutUEVersion, int32& OutLicenseeUEVersion,
+	COREUOBJECT_API void GetBulkDataVersions(FArchive& InlineArchive, FPackageFileVersion& OutUEVersion, int32& OutLicenseeUEVersion,
 		FCustomVersionContainer& OutCustomVersions) const;
 
 	/**
@@ -321,7 +321,7 @@ public:
 	 * copied from another BulkData, and it will pass on this flag to any BulkData copied/moved from it.
 	 * Use Reset() to remove this state. Torn-off BulkDatas share the guid with the BulkData they copy from.
 	 */
-	void TearOff();
+	COREUOBJECT_API void TearOff();
 
 	/** Make a torn-off copy of this bulk data. */
 	FEditorBulkData CopyTornOff() const
@@ -332,10 +332,10 @@ public:
 	// Functions used by the BulkDataRegistry
 
 	/** Used to serialize the bulkdata to/from a limited cache system used by the BulkDataRegistry. */
-	void SerializeForRegistry(FArchive& Ar);
+	COREUOBJECT_API void SerializeForRegistry(FArchive& Ar);
 	
 	/** Return true if the bulkdata has a source location that persists between editor processes (package file or virtualization). */
-	bool CanSaveForRegistry() const;
+	COREUOBJECT_API bool CanSaveForRegistry() const;
 	
 	/** Return whether the BulkData has legacy payload id that needs to be updated from loaded payload before it can be used in DDC. */
 	bool HasPlaceholderPayloadId() const
@@ -344,29 +344,29 @@ public:
 	}
 	
 	/** Return whether the BulkData is an in-memory payload without a persistent source location. */
-	bool IsMemoryOnlyPayload() const;
+	COREUOBJECT_API bool IsMemoryOnlyPayload() const;
 	
 	/** Load the payload and set the correct payload id, if the bulkdata has a PlaceholderPayloadId. */
-	void UpdatePayloadId();
+	COREUOBJECT_API void UpdatePayloadId();
 	
 	/** Return whether *this has the same source for the bulkdata (e.g. identical file locations if from file) as Other */
-	bool LocationMatches(const FEditorBulkData& Other) const;
+	COREUOBJECT_API bool LocationMatches(const FEditorBulkData& Other) const;
 
 	/**
 	 * Update the Owner of this BulkData in the BulkDataRegistry to include the Owner information.
 	 * Has no effect if the BulkData is not valid for registration.
 	 * The Owner information is lost and this function must be called again if the BulkData's payload information is modified.
 	 */
-	void UpdateRegistrationOwner(UObject* Owner);
+	COREUOBJECT_API void UpdateRegistrationOwner(UObject* Owner);
 
 #if UE_ENABLE_VIRTUALIZATION_TOGGLE
 	UE_DEPRECATED(5.0, "SetVirtualizationOptOut is an internal feature for development and will be removed without warning!")
-	void SetVirtualizationOptOut(bool bOptOut);
+	COREUOBJECT_API void SetVirtualizationOptOut(bool bOptOut);
 #endif //UE_ENABLE_VIRTUALIZATION_TOGGLE
 
 protected:
 	enum class ETornOff {};
-	FEditorBulkData(const FEditorBulkData& Other, ETornOff);
+	COREUOBJECT_API FEditorBulkData(const FEditorBulkData& Other, ETornOff);
 
 private:
 	friend struct FTocEntry;
@@ -418,33 +418,33 @@ private:
 	};
 
 	/** Old legacy path that saved the payload to the end of the package */
-	void SerializeToLegacyPath(FLinkerSave& LinkerSave, FCompressedBuffer PayloadToSerialize, EFlags UpdatedFlags, UObject* Owner);
+	COREUOBJECT_API void SerializeToLegacyPath(FLinkerSave& LinkerSave, FCompressedBuffer PayloadToSerialize, EFlags UpdatedFlags, UObject* Owner);
 	/** The new path that saves payloads to the FPackageTrailer which is then appended to the end of the package file */
-	void SerializeToPackageTrailer(FLinkerSave& LinkerSave, FCompressedBuffer PayloadToSerialize, EFlags UpdatedFlags, UObject* Owner);
+	COREUOBJECT_API void SerializeToPackageTrailer(FLinkerSave& LinkerSave, FCompressedBuffer PayloadToSerialize, EFlags UpdatedFlags, UObject* Owner);
 
-	void UpdatePayloadImpl(FSharedBuffer&& InPayload, const FIoHash& InPayloadID, UObject* Owner);
+	COREUOBJECT_API void UpdatePayloadImpl(FSharedBuffer&& InPayload, const FIoHash& InPayloadID, UObject* Owner);
 
-	FCompressedBuffer GetDataInternal() const;
+	COREUOBJECT_API FCompressedBuffer GetDataInternal() const;
 
-	FCompressedBuffer LoadFromDisk() const;
-	FCompressedBuffer LoadFromPackageFile() const;
-	FCompressedBuffer LoadFromPackageTrailer() const;
-	FCompressedBuffer LoadFromSidecarFile() const;
-	FCompressedBuffer LoadFromSidecarFileInternal(ErrorVerbosity Verbosity) const;
+	COREUOBJECT_API FCompressedBuffer LoadFromDisk() const;
+	COREUOBJECT_API FCompressedBuffer LoadFromPackageFile() const;
+	COREUOBJECT_API FCompressedBuffer LoadFromPackageTrailer() const;
+	COREUOBJECT_API FCompressedBuffer LoadFromSidecarFile() const;
+	COREUOBJECT_API FCompressedBuffer LoadFromSidecarFileInternal(ErrorVerbosity Verbosity) const;
 
-	bool SerializeData(FArchive& Ar, FCompressedBuffer& Payload, const EFlags PayloadFlags) const;
+	COREUOBJECT_API bool SerializeData(FArchive& Ar, FCompressedBuffer& Payload, const EFlags PayloadFlags) const;
 
-	void PushData(const FPackagePath& InPackagePath);
-	FCompressedBuffer PullData() const;
+	COREUOBJECT_API void PushData(const FPackagePath& InPackagePath);
+	COREUOBJECT_API FCompressedBuffer PullData() const;
 
-	bool CanUnloadData() const;
-	bool CanLoadDataFromDisk() const;
+	COREUOBJECT_API bool CanUnloadData() const;
+	COREUOBJECT_API bool CanLoadDataFromDisk() const;
 
-	void UpdateKeyIfNeeded();
-	void UpdateKeyIfNeeded(FCompressedBuffer InPayload) const;
+	COREUOBJECT_API void UpdateKeyIfNeeded();
+	COREUOBJECT_API void UpdateKeyIfNeeded(FCompressedBuffer InPayload) const;
 
-	void RecompressForSerialization(FCompressedBuffer& InOutPayload, EFlags PayloadFlags) const;
-	EFlags BuildFlagsForSerialization(FArchive& Ar, bool bKeepFileDataByReference) const;
+	COREUOBJECT_API void RecompressForSerialization(FCompressedBuffer& InOutPayload, EFlags PayloadFlags) const;
+	COREUOBJECT_API EFlags BuildFlagsForSerialization(FArchive& Ar, bool bKeepFileDataByReference) const;
 
 	bool IsDataVirtualized() const 
 	{ 
@@ -506,17 +506,17 @@ private:
 		return AttachedAr != nullptr;
 	}
 
-	void Register(UObject* Owner, const TCHAR* LogCallerName, bool bAllowUpdateId);
+	COREUOBJECT_API void Register(UObject* Owner, const TCHAR* LogCallerName, bool bAllowUpdateId);
 	/** Used when we need to clear the registration for our bulkdataId because its ownership is transferring elsewhere. */
-	void Unregister();
+	COREUOBJECT_API void Unregister();
 	/**
 	 * Used when we want to keep registration of our BulkDataId tied to our Payload location, but this BulkData is exiting memory
 	 * and the registry needs to know so it can drop its copy of the payload and allow us to reregister later if reloaded.
 	 */
-	void OnExitMemory();
+	COREUOBJECT_API void OnExitMemory();
 	/** Check whether we should be registered, and register/unregister/update as necessary to match. */
-	void UpdateRegistrationData(UObject* Owner, const TCHAR* LogCallerName, bool bAllowUpdateId);
-	void LogRegisterError(UE::BulkDataRegistry::ERegisterResult Value, UObject* Owner, const FGuid& FailedBulkDataId,
+	COREUOBJECT_API void UpdateRegistrationData(UObject* Owner, const TCHAR* LogCallerName, bool bAllowUpdateId);
+	COREUOBJECT_API void LogRegisterError(UE::BulkDataRegistry::ERegisterResult Value, UObject* Owner, const FGuid& FailedBulkDataId,
 		const TCHAR* CallerName, bool bHandledbyCreateUniqueGuid) const;
 
 	/**
@@ -525,7 +525,7 @@ private:
 	 * This error is not reasonable to expect, if it occurs then it means something has gone very wrong, which is why there is
 	 * also an ensure to make sure that any instances of it occurring as properly recorded and can be investigated.
 	 */
-	bool TryPayloadValidationForSaving(const FCompressedBuffer& PayloadForSaving, FLinkerSave* LinkerSave) const;
+	COREUOBJECT_API bool TryPayloadValidationForSaving(const FCompressedBuffer& PayloadForSaving, FLinkerSave* LinkerSave) const;
 
 	/** 
 	 * Utility to return an apt error message if the payload is invalid when trying to save the bulkdata. It will try to provide the best info from the given options.
@@ -533,7 +533,7 @@ private:
 	 * @param Linker	The FLinkerSave being used to save the bulkdata at the point of error. This can be nullptr if the bulkdata is not being saved to a package on disk.
 	 * @return			The formatted error message.
 	 */
-	FText GetCorruptedPayloadErrorMsgForSave(FLinkerSave* Linker) const;
+	COREUOBJECT_API FText GetCorruptedPayloadErrorMsgForSave(FLinkerSave* Linker) const;
 
 	/**
 	 * A utility for validating that a package trailer builder was created correctly. If a problem is encountered we will assert
@@ -547,10 +547,10 @@ private:
 	 * @param Id			The hash of the payload we want to verify
 	 * @param PayloadFlags	The flags for the payload.
 	 */
-	static void ValidatePackageTrailerBuilder(const FLinkerSave* LinkerSave, const FIoHash& Id, EFlags PayloadFlags);
+	static COREUOBJECT_API void ValidatePackageTrailerBuilder(const FLinkerSave* LinkerSave, const FIoHash& Id, EFlags PayloadFlags);
 
 	/** Returns true if we should use legacy serialization instead of the FPackageTrailer system. This can be removed when UE_ENABLE_VIRTUALIZATION_TOGGLE is removed. */
-	bool ShouldUseLegacySerialization(const FLinkerSave* LinkerSave) const;
+	COREUOBJECT_API bool ShouldUseLegacySerialization(const FLinkerSave* LinkerSave) const;
 
 	/** Unique identifier for the bulkdata object itself */
 	FGuid BulkDataId;
diff --git a/Engine/Source/Runtime/CoreUObject/Public/Serialization/FindReferencersArchive.h b/Engine/Source/Runtime/CoreUObject/Public/Serialization/FindReferencersArchive.h
--- a/Engine/Source/Runtime/CoreUObject/Public/Serialization/FindReferencersArchive.h
+++ b/Engine/Source/Runtime/CoreUObject/Public/Serialization/FindReferencersArchive.h
@@ -75,7 +75,7 @@ public:
 	 *
 	 * This is overridden for the specific Archive Types
 	 **/
-	COREUOBJECT_API virtual FString GetArchiveName() const { return TEXT("FFindReferencersArchive"); }
+	virtual FString GetArchiveName() const { return TEXT("FFindReferencersArchive"); }
 
 	/**
 	 * Resets the reference counts.  Keeps the same target objects but sets up everything to test a new potential referencer.
diff --git a/Engine/Source/Runtime/CoreUObject/Public/Serialization/Formatters/JsonArchiveInputFormatter.h b/Engine/Source/Runtime/CoreUObject/Public/Serialization/Formatters/JsonArchiveInputFormatter.h
--- a/Engine/Source/Runtime/CoreUObject/Public/Serialization/Formatters/JsonArchiveInputFormatter.h
+++ b/Engine/Source/Runtime/CoreUObject/Public/Serialization/Formatters/JsonArchiveInputFormatter.h
@@ -28,68 +28,68 @@ struct FWeakObjectPtr;
 
 #if WITH_TEXT_ARCHIVE_SUPPORT
 
-class COREUOBJECT_API FJsonArchiveInputFormatter final : public FStructuredArchiveFormatter
+class FJsonArchiveInputFormatter final : public FStructuredArchiveFormatter
 {
 public:
-	FJsonArchiveInputFormatter(FArchive& InInner, TFunction<UObject* (const FPackageIndex)> InResolveObject = nullptr);
-	virtual ~FJsonArchiveInputFormatter();
-
-	virtual FArchive& GetUnderlyingArchive() override;
-	virtual FStructuredArchiveFormatter* CreateSubtreeReader() override;
-
-	virtual bool HasDocumentTree() const override;
-
-	virtual void EnterRecord() override;
-	virtual void LeaveRecord() override;
-	virtual void EnterField(FArchiveFieldName Name) override;
-	virtual void LeaveField() override;
-	virtual bool TryEnterField(FArchiveFieldName Name, bool bEnterWhenSaving) override;
-
-	virtual void EnterArray(int32& NumElements) override;
-	virtual void LeaveArray() override;
-	virtual void EnterArrayElement() override;
-	virtual void LeaveArrayElement() override;
-
-	virtual void EnterStream() override;
-	virtual void LeaveStream() override;
-	virtual void EnterStreamElement() override;
-	virtual void LeaveStreamElement() override;
-
-	virtual void EnterMap(int32& NumElements) override;
-	virtual void LeaveMap() override;
-	virtual void EnterMapElement(FString& Name) override;
-	virtual void LeaveMapElement() override;
-
-	virtual void EnterAttributedValue() override;
-	virtual void EnterAttribute(FArchiveFieldName AttributeName) override;
-	virtual void EnterAttributedValueValue() override;
-	virtual void LeaveAttribute() override;
-	virtual void LeaveAttributedValue() override;
-	virtual bool TryEnterAttribute(FArchiveFieldName AttributeName, bool bEnterWhenSavin) override;
-	virtual bool TryEnterAttributedValueValue() override;
-
-	virtual void Serialize(uint8& Value) override;
-	virtual void Serialize(uint16& Value) override;
-	virtual void Serialize(uint32& Value) override;
-	virtual void Serialize(uint64& Value) override;
-	virtual void Serialize(int8& Value) override;
-	virtual void Serialize(int16& Value) override;
-	virtual void Serialize(int32& Value) override;
-	virtual void Serialize(int64& Value) override;
-	virtual void Serialize(float& Value) override;
-	virtual void Serialize(double& Value) override;
-	virtual void Serialize(bool& Value) override;
-	virtual void Serialize(FString& Value) override;
-	virtual void Serialize(FName& Value) override;
-	virtual void Serialize(UObject*& Value) override;
-	virtual void Serialize(FText& Value) override;
-	virtual void Serialize(FWeakObjectPtr& Value) override;
-	virtual void Serialize(FSoftObjectPtr& Value) override;
-	virtual void Serialize(FSoftObjectPath& Value) override;
-	virtual void Serialize(FLazyObjectPtr& Value) override;
-	virtual void Serialize(FObjectPtr& Value) override;
-	virtual void Serialize(TArray<uint8>& Value) override;
-	virtual void Serialize(void* Data, uint64 DataSize) override;
+	COREUOBJECT_API FJsonArchiveInputFormatter(FArchive& InInner, TFunction<UObject* (const FPackageIndex)> InResolveObject = nullptr);
+	COREUOBJECT_API virtual ~FJsonArchiveInputFormatter();
+
+	COREUOBJECT_API virtual FArchive& GetUnderlyingArchive() override;
+	COREUOBJECT_API virtual FStructuredArchiveFormatter* CreateSubtreeReader() override;
+
+	COREUOBJECT_API virtual bool HasDocumentTree() const override;
+
+	COREUOBJECT_API virtual void EnterRecord() override;
+	COREUOBJECT_API virtual void LeaveRecord() override;
+	COREUOBJECT_API virtual void EnterField(FArchiveFieldName Name) override;
+	COREUOBJECT_API virtual void LeaveField() override;
+	COREUOBJECT_API virtual bool TryEnterField(FArchiveFieldName Name, bool bEnterWhenSaving) override;
+
+	COREUOBJECT_API virtual void EnterArray(int32& NumElements) override;
+	COREUOBJECT_API virtual void LeaveArray() override;
+	COREUOBJECT_API virtual void EnterArrayElement() override;
+	COREUOBJECT_API virtual void LeaveArrayElement() override;
+
+	COREUOBJECT_API virtual void EnterStream() override;
+	COREUOBJECT_API virtual void LeaveStream() override;
+	COREUOBJECT_API virtual void EnterStreamElement() override;
+	COREUOBJECT_API virtual void LeaveStreamElement() override;
+
+	COREUOBJECT_API virtual void EnterMap(int32& NumElements) override;
+	COREUOBJECT_API virtual void LeaveMap() override;
+	COREUOBJECT_API virtual void EnterMapElement(FString& Name) override;
+	COREUOBJECT_API virtual void LeaveMapElement() override;
+
+	COREUOBJECT_API virtual void EnterAttributedValue() override;
+	COREUOBJECT_API virtual void EnterAttribute(FArchiveFieldName AttributeName) override;
+	COREUOBJECT_API virtual void EnterAttributedValueValue() override;
+	COREUOBJECT_API virtual void LeaveAttribute() override;
+	COREUOBJECT_API virtual void LeaveAttributedValue() override;
+	COREUOBJECT_API virtual bool TryEnterAttribute(FArchiveFieldName AttributeName, bool bEnterWhenSavin) override;
+	COREUOBJECT_API virtual bool TryEnterAttributedValueValue() override;
+
+	COREUOBJECT_API virtual void Serialize(uint8& Value) override;
+	COREUOBJECT_API virtual void Serialize(uint16& Value) override;
+	COREUOBJECT_API virtual void Serialize(uint32& Value) override;
+	COREUOBJECT_API virtual void Serialize(uint64& Value) override;
+	COREUOBJECT_API virtual void Serialize(int8& Value) override;
+	COREUOBJECT_API virtual void Serialize(int16& Value) override;
+	COREUOBJECT_API virtual void Serialize(int32& Value) override;
+	COREUOBJECT_API virtual void Serialize(int64& Value) override;
+	COREUOBJECT_API virtual void Serialize(float& Value) override;
+	COREUOBJECT_API virtual void Serialize(double& Value) override;
+	COREUOBJECT_API virtual void Serialize(bool& Value) override;
+	COREUOBJECT_API virtual void Serialize(FString& Value) override;
+	COREUOBJECT_API virtual void Serialize(FName& Value) override;
+	COREUOBJECT_API virtual void Serialize(UObject*& Value) override;
+	COREUOBJECT_API virtual void Serialize(FText& Value) override;
+	COREUOBJECT_API virtual void Serialize(FWeakObjectPtr& Value) override;
+	COREUOBJECT_API virtual void Serialize(FSoftObjectPtr& Value) override;
+	COREUOBJECT_API virtual void Serialize(FSoftObjectPath& Value) override;
+	COREUOBJECT_API virtual void Serialize(FLazyObjectPtr& Value) override;
+	COREUOBJECT_API virtual void Serialize(FObjectPtr& Value) override;
+	COREUOBJECT_API virtual void Serialize(TArray<uint8>& Value) override;
+	COREUOBJECT_API virtual void Serialize(void* Data, uint64 DataSize) override;
 
 private:
 	FArchive& Inner;
diff --git a/Engine/Source/Runtime/CoreUObject/Public/Serialization/Formatters/JsonArchiveOutputFormatter.h b/Engine/Source/Runtime/CoreUObject/Public/Serialization/Formatters/JsonArchiveOutputFormatter.h
--- a/Engine/Source/Runtime/CoreUObject/Public/Serialization/Formatters/JsonArchiveOutputFormatter.h
+++ b/Engine/Source/Runtime/CoreUObject/Public/Serialization/Formatters/JsonArchiveOutputFormatter.h
@@ -24,68 +24,68 @@ struct FWeakObjectPtr;
 
 #if WITH_TEXT_ARCHIVE_SUPPORT
 
-class COREUOBJECT_API FJsonArchiveOutputFormatter final : public FStructuredArchiveFormatter
+class FJsonArchiveOutputFormatter final : public FStructuredArchiveFormatter
 {
 public:
-	FJsonArchiveOutputFormatter(FArchive& InInner);
-	virtual ~FJsonArchiveOutputFormatter();
-
-	virtual FArchive& GetUnderlyingArchive() override;
-
-	virtual bool HasDocumentTree() const override;
-
-	virtual void EnterRecord() override;
-	virtual void LeaveRecord() override;
-	virtual void EnterField(FArchiveFieldName Name) override;
-	virtual void LeaveField() override;
-	virtual bool TryEnterField(FArchiveFieldName Name, bool bEnterWhenSaving) override;
-
-	virtual void EnterArray(int32& NumElements) override;
-	virtual void LeaveArray() override;
-	virtual void EnterArrayElement() override;
-	virtual void LeaveArrayElement() override;
-
-	virtual void EnterStream() override;
-	virtual void LeaveStream() override;
-	virtual void EnterStreamElement() override;
-	virtual void LeaveStreamElement() override;
-
-	virtual void EnterMap(int32& NumElements) override;
-	virtual void LeaveMap() override;
-	virtual void EnterMapElement(FString& Name) override;
-	virtual void LeaveMapElement() override;
-
-	virtual void EnterAttributedValue() override;
-	virtual void EnterAttribute(FArchiveFieldName AttributeName) override;
-	virtual void EnterAttributedValueValue() override;
-	virtual void LeaveAttribute() override;
-	virtual void LeaveAttributedValue() override;
-	virtual bool TryEnterAttribute(FArchiveFieldName AttributeName, bool bEnterWhenSaving) override;
-
-	virtual bool TryEnterAttributedValueValue() override;
-
-	virtual void Serialize(uint8& Value) override;
-	virtual void Serialize(uint16& Value) override;
-	virtual void Serialize(uint32& Value) override;
-	virtual void Serialize(uint64& Value) override;
-	virtual void Serialize(int8& Value) override;
-	virtual void Serialize(int16& Value) override;
-	virtual void Serialize(int32& Value) override;
-	virtual void Serialize(int64& Value) override;
-	virtual void Serialize(float& Value) override;
-	virtual void Serialize(double& Value) override;
-	virtual void Serialize(bool& Value) override;
-	virtual void Serialize(FString& Value) override;
-	virtual void Serialize(FName& Value) override;
-	virtual void Serialize(UObject*& Value) override;
-	virtual void Serialize(FText& Value) override;
-	virtual void Serialize(FWeakObjectPtr& Value) override;
-	virtual void Serialize(FSoftObjectPtr& Value) override;
-	virtual void Serialize(FSoftObjectPath& Value) override;
-	virtual void Serialize(FLazyObjectPtr& Value) override;
-	virtual void Serialize(FObjectPtr& Value) override;
-	virtual void Serialize(TArray<uint8>& Value) override;
-	virtual void Serialize(void* Data, uint64 DataSize) override;
+	COREUOBJECT_API FJsonArchiveOutputFormatter(FArchive& InInner);
+	COREUOBJECT_API virtual ~FJsonArchiveOutputFormatter();
+
+	COREUOBJECT_API virtual FArchive& GetUnderlyingArchive() override;
+
+	COREUOBJECT_API virtual bool HasDocumentTree() const override;
+
+	COREUOBJECT_API virtual void EnterRecord() override;
+	COREUOBJECT_API virtual void LeaveRecord() override;
+	COREUOBJECT_API virtual void EnterField(FArchiveFieldName Name) override;
+	COREUOBJECT_API virtual void LeaveField() override;
+	COREUOBJECT_API virtual bool TryEnterField(FArchiveFieldName Name, bool bEnterWhenSaving) override;
+
+	COREUOBJECT_API virtual void EnterArray(int32& NumElements) override;
+	COREUOBJECT_API virtual void LeaveArray() override;
+	COREUOBJECT_API virtual void EnterArrayElement() override;
+	COREUOBJECT_API virtual void LeaveArrayElement() override;
+
+	COREUOBJECT_API virtual void EnterStream() override;
+	COREUOBJECT_API virtual void LeaveStream() override;
+	COREUOBJECT_API virtual void EnterStreamElement() override;
+	COREUOBJECT_API virtual void LeaveStreamElement() override;
+
+	COREUOBJECT_API virtual void EnterMap(int32& NumElements) override;
+	COREUOBJECT_API virtual void LeaveMap() override;
+	COREUOBJECT_API virtual void EnterMapElement(FString& Name) override;
+	COREUOBJECT_API virtual void LeaveMapElement() override;
+
+	COREUOBJECT_API virtual void EnterAttributedValue() override;
+	COREUOBJECT_API virtual void EnterAttribute(FArchiveFieldName AttributeName) override;
+	COREUOBJECT_API virtual void EnterAttributedValueValue() override;
+	COREUOBJECT_API virtual void LeaveAttribute() override;
+	COREUOBJECT_API virtual void LeaveAttributedValue() override;
+	COREUOBJECT_API virtual bool TryEnterAttribute(FArchiveFieldName AttributeName, bool bEnterWhenSaving) override;
+
+	COREUOBJECT_API virtual bool TryEnterAttributedValueValue() override;
+
+	COREUOBJECT_API virtual void Serialize(uint8& Value) override;
+	COREUOBJECT_API virtual void Serialize(uint16& Value) override;
+	COREUOBJECT_API virtual void Serialize(uint32& Value) override;
+	COREUOBJECT_API virtual void Serialize(uint64& Value) override;
+	COREUOBJECT_API virtual void Serialize(int8& Value) override;
+	COREUOBJECT_API virtual void Serialize(int16& Value) override;
+	COREUOBJECT_API virtual void Serialize(int32& Value) override;
+	COREUOBJECT_API virtual void Serialize(int64& Value) override;
+	COREUOBJECT_API virtual void Serialize(float& Value) override;
+	COREUOBJECT_API virtual void Serialize(double& Value) override;
+	COREUOBJECT_API virtual void Serialize(bool& Value) override;
+	COREUOBJECT_API virtual void Serialize(FString& Value) override;
+	COREUOBJECT_API virtual void Serialize(FName& Value) override;
+	COREUOBJECT_API virtual void Serialize(UObject*& Value) override;
+	COREUOBJECT_API virtual void Serialize(FText& Value) override;
+	COREUOBJECT_API virtual void Serialize(FWeakObjectPtr& Value) override;
+	COREUOBJECT_API virtual void Serialize(FSoftObjectPtr& Value) override;
+	COREUOBJECT_API virtual void Serialize(FSoftObjectPath& Value) override;
+	COREUOBJECT_API virtual void Serialize(FLazyObjectPtr& Value) override;
+	COREUOBJECT_API virtual void Serialize(FObjectPtr& Value) override;
+	COREUOBJECT_API virtual void Serialize(TArray<uint8>& Value) override;
+	COREUOBJECT_API virtual void Serialize(void* Data, uint64 DataSize) override;
 
 	void SetObjectIndicesMap(const TMap<TObjectPtr<UObject>, FPackageIndex>* InObjectIndicesMap)
 	{
@@ -104,23 +104,23 @@ private:
 	TArray<int32> NumAttributesStack;
 	TArray<int64> TextStartPosStack;
 
-	void Write(ANSICHAR Character);
+	COREUOBJECT_API void Write(ANSICHAR Character);
 
-	void Write(const ANSICHAR* Text);
-	void Write(const FString& Text);
+	COREUOBJECT_API void Write(const ANSICHAR* Text);
+	COREUOBJECT_API void Write(const FString& Text);
 
-	void WriteFieldName(const TCHAR* Name);
-	void WriteValue(const FString& Value);
+	COREUOBJECT_API void WriteFieldName(const TCHAR* Name);
+	COREUOBJECT_API void WriteValue(const FString& Value);
 
-	void WriteOptionalComma();
-	void WriteOptionalNewline();
-	void WriteOptionalAttributedBlockOpening();
-	void WriteOptionalAttributedBlockValue();
-	void WriteOptionalAttributedBlockClosing();
+	COREUOBJECT_API void WriteOptionalComma();
+	COREUOBJECT_API void WriteOptionalNewline();
+	COREUOBJECT_API void WriteOptionalAttributedBlockOpening();
+	COREUOBJECT_API void WriteOptionalAttributedBlockValue();
+	COREUOBJECT_API void WriteOptionalAttributedBlockClosing();
 
-	void SerializeStringInternal(const FString& String);
+	COREUOBJECT_API void SerializeStringInternal(const FString& String);
 
-	bool IsObjectAllowed(UObject* InObject) const;
+	COREUOBJECT_API bool IsObjectAllowed(UObject* InObject) const;
 };
 
 #endif
diff --git a/Engine/Source/Runtime/CoreUObject/Public/Serialization/PropertyLocalizationDataGathering.h b/Engine/Source/Runtime/CoreUObject/Public/Serialization/PropertyLocalizationDataGathering.h
--- a/Engine/Source/Runtime/CoreUObject/Public/Serialization/PropertyLocalizationDataGathering.h
+++ b/Engine/Source/Runtime/CoreUObject/Public/Serialization/PropertyLocalizationDataGathering.h
@@ -89,7 +89,7 @@ enum class EPropertyLocalizationGathererResultFlags : uint8
 };
 ENUM_CLASS_FLAGS(EPropertyLocalizationGathererResultFlags);
 
-class COREUOBJECT_API FPropertyLocalizationDataGatherer
+class FPropertyLocalizationDataGatherer
 {
 public:
 	typedef TFunction<void(const UObject* /*Object*/, FPropertyLocalizationDataGatherer& /*PropertyLocalizationDataGatherer*/, const EPropertyLocalizationGathererTextFlags /*GatherTextFlags*/)> FLocalizationDataObjectGatheringCallback;
@@ -114,37 +114,37 @@ public:
 		}
 	};
 
-	FPropertyLocalizationDataGatherer(TArray<FGatherableTextData>& InOutGatherableTextDataArray, const UPackage* const InPackage, EPropertyLocalizationGathererResultFlags& OutResultFlags);
+	COREUOBJECT_API FPropertyLocalizationDataGatherer(TArray<FGatherableTextData>& InOutGatherableTextDataArray, const UPackage* const InPackage, EPropertyLocalizationGathererResultFlags& OutResultFlags);
 
 	// Non-copyable
 	FPropertyLocalizationDataGatherer(const FPropertyLocalizationDataGatherer&) = delete;
 	FPropertyLocalizationDataGatherer& operator=(const FPropertyLocalizationDataGatherer&) = delete;
 
-	void GatherLocalizationDataFromObjectWithCallbacks(const UObject* Object, const EPropertyLocalizationGathererTextFlags GatherTextFlags);
-	void GatherLocalizationDataFromObject(const UObject* Object, const EPropertyLocalizationGathererTextFlags GatherTextFlags);
-	void GatherLocalizationDataFromObjectFields(const FString& PathToParent, const UObject* Object, const EPropertyLocalizationGathererTextFlags GatherTextFlags);
+	COREUOBJECT_API void GatherLocalizationDataFromObjectWithCallbacks(const UObject* Object, const EPropertyLocalizationGathererTextFlags GatherTextFlags);
+	COREUOBJECT_API void GatherLocalizationDataFromObject(const UObject* Object, const EPropertyLocalizationGathererTextFlags GatherTextFlags);
+	COREUOBJECT_API void GatherLocalizationDataFromObjectFields(const FString& PathToParent, const UObject* Object, const EPropertyLocalizationGathererTextFlags GatherTextFlags);
 	
-	void GatherLocalizationDataFromStructWithCallbacks(const FString& PathToParent, const UScriptStruct* Struct, const void* StructData, const void* DefaultStructData, const EPropertyLocalizationGathererTextFlags GatherTextFlags);
-	void GatherLocalizationDataFromStruct(const FString& PathToParent, const UScriptStruct* Struct, const void* StructData, const void* DefaultStructData, const EPropertyLocalizationGathererTextFlags GatherTextFlags);
-	void GatherLocalizationDataFromStructFields(const FString& PathToParent, const UStruct* Struct, const void* StructData, const void* DefaultStructData, const EPropertyLocalizationGathererTextFlags GatherTextFlags);
+	COREUOBJECT_API void GatherLocalizationDataFromStructWithCallbacks(const FString& PathToParent, const UScriptStruct* Struct, const void* StructData, const void* DefaultStructData, const EPropertyLocalizationGathererTextFlags GatherTextFlags);
+	COREUOBJECT_API void GatherLocalizationDataFromStruct(const FString& PathToParent, const UScriptStruct* Struct, const void* StructData, const void* DefaultStructData, const EPropertyLocalizationGathererTextFlags GatherTextFlags);
+	COREUOBJECT_API void GatherLocalizationDataFromStructFields(const FString& PathToParent, const UStruct* Struct, const void* StructData, const void* DefaultStructData, const EPropertyLocalizationGathererTextFlags GatherTextFlags);
 	
-	void GatherLocalizationDataFromChildTextProperties(const FString& PathToParent, const FProperty* const Property, const void* const ValueAddress, const void* const DefaultValueAddress, const EPropertyLocalizationGathererTextFlags GatherTextFlags);
+	COREUOBJECT_API void GatherLocalizationDataFromChildTextProperties(const FString& PathToParent, const FProperty* const Property, const void* const ValueAddress, const void* const DefaultValueAddress, const EPropertyLocalizationGathererTextFlags GatherTextFlags);
 
-	void GatherTextInstance(const FText& Text, const FString& Description, const bool bIsEditorOnly);
-	void GatherScriptBytecode(const FString& PathToScript, const TArray<uint8>& ScriptData, const bool bIsEditorOnly);
+	COREUOBJECT_API void GatherTextInstance(const FText& Text, const FString& Description, const bool bIsEditorOnly);
+	COREUOBJECT_API void GatherScriptBytecode(const FString& PathToScript, const TArray<uint8>& ScriptData, const bool bIsEditorOnly);
 
-	bool IsDefaultTextInstance(const FText& Text) const;
-	void MarkDefaultTextInstance(const FText& Text);
+	COREUOBJECT_API bool IsDefaultTextInstance(const FText& Text) const;
+	COREUOBJECT_API void MarkDefaultTextInstance(const FText& Text);
 
-	bool ShouldProcessObject(const UObject* Object, const EPropertyLocalizationGathererTextFlags GatherTextFlags) const;
-	void MarkObjectProcessed(const UObject* Object, const EPropertyLocalizationGathererTextFlags GatherTextFlags);
+	COREUOBJECT_API bool ShouldProcessObject(const UObject* Object, const EPropertyLocalizationGathererTextFlags GatherTextFlags) const;
+	COREUOBJECT_API void MarkObjectProcessed(const UObject* Object, const EPropertyLocalizationGathererTextFlags GatherTextFlags);
 
-	const FGatherableFieldsForType& GetGatherableFieldsForType(const UStruct* InType);
+	COREUOBJECT_API const FGatherableFieldsForType& GetGatherableFieldsForType(const UStruct* InType);
 
-	static bool ExtractTextIdentity(const FText& Text, FString& OutNamespace, FString& OutKey, const bool bCleanNamespace);
+	static COREUOBJECT_API bool ExtractTextIdentity(const FText& Text, FString& OutNamespace, FString& OutKey, const bool bCleanNamespace);
 
-	static FLocalizationDataObjectGatheringCallbackMap& GetTypeSpecificLocalizationDataObjectGatheringCallbacks();
-	static FLocalizationDataStructGatheringCallbackMap& GetTypeSpecificLocalizationDataStructGatheringCallbacks();
+	static COREUOBJECT_API FLocalizationDataObjectGatheringCallbackMap& GetTypeSpecificLocalizationDataObjectGatheringCallbacks();
+	static COREUOBJECT_API FLocalizationDataStructGatheringCallbackMap& GetTypeSpecificLocalizationDataStructGatheringCallbacks();
 
 	FORCEINLINE TArray<FGatherableTextData>& GetGatherableTextDataArray() const
 	{
@@ -157,8 +157,8 @@ public:
 	}
 
 private:
-	const FGatherableFieldsForType& CacheGatherableFieldsForType(const UStruct* InType);
-	bool CanGatherFromInnerProperty(const FProperty* InInnerProperty);
+	COREUOBJECT_API const FGatherableFieldsForType& CacheGatherableFieldsForType(const UStruct* InType);
+	COREUOBJECT_API bool CanGatherFromInnerProperty(const FProperty* InInnerProperty);
 
 	struct FObjectAndGatherFlags
 	{
diff --git a/Engine/Source/Runtime/CoreUObject/Public/Serialization/SerializedPropertyScope.h b/Engine/Source/Runtime/CoreUObject/Public/Serialization/SerializedPropertyScope.h
--- a/Engine/Source/Runtime/CoreUObject/Public/Serialization/SerializedPropertyScope.h
+++ b/Engine/Source/Runtime/CoreUObject/Public/Serialization/SerializedPropertyScope.h
@@ -9,12 +9,12 @@
 #include "Serialization/Archive.h"
 
 /** Helper class to set and restore serialized property on an archive */
-class COREUOBJECT_API FSerializedPropertyScope
+class FSerializedPropertyScope
 {
 	FArchive& Ar;
 	FProperty* Property;
-	void PushProperty();
-	void PopProperty();
+	COREUOBJECT_API void PushProperty();
+	COREUOBJECT_API void PopProperty();
 public:
 	FSerializedPropertyScope(FArchive& InAr, FProperty* InProperty, const FProperty* OnlyIfOldProperty = nullptr)
 		: Ar(InAr)
diff --git a/Engine/Source/Runtime/CoreUObject/Public/Serialization/TestUndeclaredScriptStructObjectReferences.h b/Engine/Source/Runtime/CoreUObject/Public/Serialization/TestUndeclaredScriptStructObjectReferences.h
--- a/Engine/Source/Runtime/CoreUObject/Public/Serialization/TestUndeclaredScriptStructObjectReferences.h
+++ b/Engine/Source/Runtime/CoreUObject/Public/Serialization/TestUndeclaredScriptStructObjectReferences.h
@@ -21,5 +21,5 @@ struct FTestUndeclaredScriptStructObjectReferencesTest
 
 template<> struct TBaseStructure<FTestUndeclaredScriptStructObjectReferencesTest>
 {
-	COREUOBJECT_API static UScriptStruct* Get();
-};
\ No newline at end of file
+	static UScriptStruct* Get();
+};
diff --git a/Engine/Source/Runtime/CoreUObject/Public/Serialization/TextReferenceCollector.h b/Engine/Source/Runtime/CoreUObject/Public/Serialization/TextReferenceCollector.h
--- a/Engine/Source/Runtime/CoreUObject/Public/Serialization/TextReferenceCollector.h
+++ b/Engine/Source/Runtime/CoreUObject/Public/Serialization/TextReferenceCollector.h
@@ -14,7 +14,7 @@ class UObject;
 class UPackage;
 
 /** Counts all persistent text references from within a package, using the specified comparison mode */
-class COREUOBJECT_API FTextReferenceCollector
+class FTextReferenceCollector
 {
 public:
 	enum class EComparisonMode : uint8
@@ -30,9 +30,9 @@ public:
 	typedef TFunction<void(UObject*, FArchive&)> FTextReferenceCollectorCallback;
 	typedef TMap<const UClass*, FTextReferenceCollectorCallback> FTextReferenceCollectorCallbackMap;
 
-	FTextReferenceCollector(const UPackage* const InPackage, const EComparisonMode InComparisonMode, const FString& InTextNamespace, const FString& InTextKey, const FString& InTextSource, int32& OutCount);
+	COREUOBJECT_API FTextReferenceCollector(const UPackage* const InPackage, const EComparisonMode InComparisonMode, const FString& InTextNamespace, const FString& InTextKey, const FString& InTextSource, int32& OutCount);
 
-	static FTextReferenceCollectorCallbackMap& GetTypeSpecificTextReferenceCollectorCallbacks();
+	static COREUOBJECT_API FTextReferenceCollectorCallbackMap& GetTypeSpecificTextReferenceCollectorCallbacks();
 };
 
 /** Struct to automatically register a callback when it's constructed */
