[CoreUObject]

diff --git a/Engine/Source/Runtime/CoreUObject/Public/UObject/Class.h b/Engine/Source/Runtime/CoreUObject/Public/UObject/Class.h
--- a/Engine/Source/Runtime/CoreUObject/Public/UObject/Class.h
+++ b/Engine/Source/Runtime/CoreUObject/Public/UObject/Class.h
@@ -152,9 +152,9 @@ struct FRepRecord
 //
 // Base class of reflection data objects.
 //
-class COREUOBJECT_API UField : public UObject
+class UField : public UObject
 {
-	DECLARE_CASTED_CLASS_INTRINSIC(UField, UObject, CLASS_Abstract, TEXT("/Script/CoreUObject"), CASTCLASS_UField)
+	DECLARE_CASTED_CLASS_INTRINSIC_WITH_API(UField, UObject, CLASS_Abstract, TEXT("/Script/CoreUObject"), CASTCLASS_UField, COREUOBJECT_API)
 
 	typedef UField BaseFieldClass;
 	typedef UClass FieldTypeClass;
@@ -166,27 +166,27 @@ class COREUOBJECT_API UField : public UObject
 	UField(EStaticConstructor, EObjectFlags InFlags);
 
 	// UObject interface.
-	virtual void Serialize( FArchive& Ar ) override;
-	virtual void PostLoad() override;
-	virtual bool NeedsLoadForClient() const override;
-	virtual bool NeedsLoadForServer() const override;
+	COREUOBJECT_API virtual void Serialize( FArchive& Ar ) override;
+	COREUOBJECT_API virtual void PostLoad() override;
+	COREUOBJECT_API virtual bool NeedsLoadForClient() const override;
+	COREUOBJECT_API virtual bool NeedsLoadForServer() const override;
 
 	// UField interface.
-	virtual void AddCppProperty(FProperty* Property);
-	virtual void Bind();
+	COREUOBJECT_API virtual void AddCppProperty(FProperty* Property);
+	COREUOBJECT_API virtual void Bind();
 
 	/** Goes up the outer chain to look for a UClass */
-	UClass* GetOwnerClass() const;
+	COREUOBJECT_API UClass* GetOwnerClass() const;
 
 	/** Goes up the outer chain to look for a UStruct */
-	UStruct* GetOwnerStruct() const;
+	COREUOBJECT_API UStruct* GetOwnerStruct() const;
 
 	/** 
 	 * Returns a human readable string that was assigned to this field at creation. 
 	 * By default this is the same as GetName() but it can be overridden if that is an internal-only name.
 	 * This name is consistent in editor/cooked builds, is not localized, and is useful for data import/export.
 	 */
-	FString GetAuthoredName() const;
+	COREUOBJECT_API FString GetAuthoredName() const;
 
 #if WITH_EDITORONLY_DATA
 	/**
@@ -194,7 +194,7 @@ class COREUOBJECT_API UField : public UObject
 	 *
 	 * @return The display name for this object.
 	 */
-	FText GetDisplayNameText() const;
+	COREUOBJECT_API FText GetDisplayNameText() const;
 
 	/**
 	 * Finds the localized tooltip or native tooltip as a fallback.
@@ -203,7 +203,7 @@ class COREUOBJECT_API UField : public UObject
 	 *
 	 * @return The tooltip for this object.
 	 */
-	FText GetToolTipText(bool bShortTooltip = false) const;
+	COREUOBJECT_API FText GetToolTipText(bool bShortTooltip = false) const;
 
 	/** 
 	 * Formats a source comment into the form we want to show in the editor, is used by GetToolTipText and anything else that will get a native tooltip 
@@ -211,7 +211,7 @@ class COREUOBJECT_API UField : public UObject
 	 * @param ToolTipString			String parsed out of C++ headers that is modified in place
 	 * @param bRemoveExtraSections	If true, cut off the comment on first line separator or 2 empty lines in a row
 	 */
-	static void FormatNativeToolTip(FString& ToolTipString, bool bRemoveExtraSections = true);
+	static COREUOBJECT_API void FormatNativeToolTip(FString& ToolTipString, bool bRemoveExtraSections = true);
 
 	/**
 	 * Determines if the property has any metadata associated with the key
@@ -228,8 +228,8 @@ class COREUOBJECT_API UField : public UObject
 	 * @param Key The key to lookup in the metadata
 	 * @return The value associated with the key if exists, null otherwise
 	 */
-	const FString* FindMetaData(const TCHAR* Key) const;
-	const FString* FindMetaData(const FName& Key) const;
+	COREUOBJECT_API const FString* FindMetaData(const TCHAR* Key) const;
+	COREUOBJECT_API const FString* FindMetaData(const FName& Key) const;
 
 	/**
 	 * Find the metadata value associated with the key
@@ -237,8 +237,8 @@ class COREUOBJECT_API UField : public UObject
 	 * @param Key The key to lookup in the metadata
 	 * @return The value associated with the key
 	 */
-	const FString& GetMetaData(const TCHAR* Key) const;
-	const FString& GetMetaData(const FName& Key) const;
+	COREUOBJECT_API const FString& GetMetaData(const TCHAR* Key) const;
+	COREUOBJECT_API const FString& GetMetaData(const FName& Key) const;
 
 	/**
 	 * Find the metadata value associated with the key and localization namespace and key
@@ -248,8 +248,8 @@ class COREUOBJECT_API UField : public UObject
 	 * @param LocalizationKey			Key to lookup in the localization manager
 	 * @return							Localized metadata if available, defaults to whatever is provided via GetMetaData
 	 */
-	FText GetMetaDataText(const TCHAR* MetaDataKey, const FString LocalizationNamespace = FString(), const FString LocalizationKey = FString()) const;
-	FText GetMetaDataText(const FName& MetaDataKey, const FString LocalizationNamespace = FString(), const FString LocalizationKey = FString()) const;
+	COREUOBJECT_API FText GetMetaDataText(const TCHAR* MetaDataKey, const FString LocalizationNamespace = FString(), const FString LocalizationKey = FString()) const;
+	COREUOBJECT_API FText GetMetaDataText(const FName& MetaDataKey, const FString LocalizationNamespace = FString(), const FString LocalizationKey = FString()) const;
 
 	/**
 	 * Sets the metadata value associated with the key
@@ -257,8 +257,8 @@ class COREUOBJECT_API UField : public UObject
 	 * @param Key The key to lookup in the metadata
 	 * @return The value associated with the key
 	 */
-	void SetMetaData(const TCHAR* Key, const TCHAR* InValue);
-	void SetMetaData(const FName& Key, const TCHAR* InValue);
+	COREUOBJECT_API void SetMetaData(const TCHAR* Key, const TCHAR* InValue);
+	COREUOBJECT_API void SetMetaData(const FName& Key, const TCHAR* InValue);
 
 	/**
 	 * Find the metadata value associated with the key
@@ -323,26 +323,26 @@ class COREUOBJECT_API UField : public UObject
 	 * @param Key The key to lookup in the metadata
 	 * @return the class value stored in the metadata.
 	 */
-	UClass* GetClassMetaData(const TCHAR* Key) const;
-	UClass* GetClassMetaData(const FName& Key) const;
+	COREUOBJECT_API UClass* GetClassMetaData(const TCHAR* Key) const;
+	COREUOBJECT_API UClass* GetClassMetaData(const FName& Key) const;
 
 	/** Clear any metadata associated with the key */
-	void RemoveMetaData(const TCHAR* Key);
-	void RemoveMetaData(const FName& Key);
+	COREUOBJECT_API void RemoveMetaData(const TCHAR* Key);
+	COREUOBJECT_API void RemoveMetaData(const FName& Key);
 #endif // WITH_EDITORONLY_DATA
 
-	bool HasAnyCastFlags(const uint64 InCastFlags) const;
-	bool HasAllCastFlags(const uint64 InCastFlags) const;
+	COREUOBJECT_API bool HasAnyCastFlags(const uint64 InCastFlags) const;
+	COREUOBJECT_API bool HasAllCastFlags(const uint64 InCastFlags) const;
 
 #if WITH_EDITORONLY_DATA
 	/**
 	 * Gets the FField object associated with this Field
 	 */
-	virtual FField* GetAssociatedFField();
+	COREUOBJECT_API virtual FField* GetAssociatedFField();
 	/**
 	 * Sets the FField object associated with this Field
 	 */
-	virtual void SetAssociatedFField(FField* InField);
+	COREUOBJECT_API virtual void SetAssociatedFField(FField* InField);
 #endif // WITH_EDITORONLY_DATA
 };
 
@@ -380,12 +380,12 @@ private:
 /**
  * Base class for all UObject types that contain fields.
  */
-class COREUOBJECT_API UStruct : public UField
+class UStruct : public UField
 #if USTRUCT_FAST_ISCHILDOF_IMPL == USTRUCT_ISCHILDOF_STRUCTARRAY
 	, private FStructBaseChain
 #endif
 {
-	DECLARE_CASTED_CLASS_INTRINSIC(UStruct, UField, CLASS_MatchedSerializers, TEXT("/Script/CoreUObject"), CASTCLASS_UStruct)
+	DECLARE_CASTED_CLASS_INTRINSIC_WITH_API(UStruct, UField, CLASS_MatchedSerializers, TEXT("/Script/CoreUObject"), CASTCLASS_UStruct, COREUOBJECT_API)
 
 	// Variables.
 protected:
@@ -438,7 +438,7 @@ public:
 	mutable const struct FUnversionedStructSchema* UnversionedEditorSchema = nullptr;
 
 	/** Get the Schema Hash for this struct - the hash of its property names and types. */
-	const FBlake3Hash& GetSchemaHash(bool bSkipEditorOnly) const;
+	COREUOBJECT_API const FBlake3Hash& GetSchemaHash(bool bSkipEditorOnly) const;
 
 protected:
 	/** True if this struct has Asset Registry searchable properties */
@@ -447,31 +447,31 @@ protected:
 
 public:
 	// Constructors.
-	UStruct( EStaticConstructor, int32 InSize, int32 InAlignment, EObjectFlags InFlags );
-	explicit UStruct(UStruct* InSuperStruct, SIZE_T ParamsSize = 0, SIZE_T Alignment = 0);
-	explicit UStruct(const FObjectInitializer& ObjectInitializer, UStruct* InSuperStruct, SIZE_T ParamsSize = 0, SIZE_T Alignment = 0 );
-	virtual ~UStruct();
+	COREUOBJECT_API UStruct( EStaticConstructor, int32 InSize, int32 InAlignment, EObjectFlags InFlags );
+	COREUOBJECT_API explicit UStruct(UStruct* InSuperStruct, SIZE_T ParamsSize = 0, SIZE_T Alignment = 0);
+	COREUOBJECT_API explicit UStruct(const FObjectInitializer& ObjectInitializer, UStruct* InSuperStruct, SIZE_T ParamsSize = 0, SIZE_T Alignment = 0 );
+	COREUOBJECT_API virtual ~UStruct();
 
 	// UObject interface.
-	virtual void Serialize(FArchive& Ar) override;
-	virtual void Serialize(FStructuredArchive::FRecord Record) override;
-	virtual void PostLoad() override;
-	virtual void FinishDestroy() override;
-	virtual void RegisterDependencies() override;
-	static void AddReferencedObjects(UObject* InThis, FReferenceCollector& Collector);
-	virtual void GetPreloadDependencies(TArray<UObject*>& OutDeps) override;
-	virtual void TagSubobjects(EObjectFlags NewFlags) override;
+	COREUOBJECT_API virtual void Serialize(FArchive& Ar) override;
+	COREUOBJECT_API virtual void Serialize(FStructuredArchive::FRecord Record) override;
+	COREUOBJECT_API virtual void PostLoad() override;
+	COREUOBJECT_API virtual void FinishDestroy() override;
+	COREUOBJECT_API virtual void RegisterDependencies() override;
+	static COREUOBJECT_API void AddReferencedObjects(UObject* InThis, FReferenceCollector& Collector);
+	COREUOBJECT_API virtual void GetPreloadDependencies(TArray<UObject*>& OutDeps) override;
+	COREUOBJECT_API virtual void TagSubobjects(EObjectFlags NewFlags) override;
 
 	// UField interface.
-	virtual void AddCppProperty(FProperty* Property) override;
+	COREUOBJECT_API virtual void AddCppProperty(FProperty* Property) override;
 
 	/**
 	 * Returns struct path name as a package + struct FName pair
 	 */
-	FTopLevelAssetPath GetStructPathName() const;
+	COREUOBJECT_API FTopLevelAssetPath GetStructPathName() const;
 
 	/** Searches property link chain for a property with the specified name */
-	FProperty* FindPropertyByName(FName InName) const;
+	COREUOBJECT_API FProperty* FindPropertyByName(FName InName) const;
 
 	/**
 	 * Creates new copies of components
@@ -482,16 +482,16 @@ public:
 	 * @param	Owner						the object that contains the component currently located at Data
 	 * @param	InstanceGraph				contains the mappings of instanced objects and components to their templates
 	 */
-	void InstanceSubobjectTemplates( void* Data, void const* DefaultData, UStruct* DefaultStruct, UObject* Owner, FObjectInstancingGraph* InstanceGraph );
+	COREUOBJECT_API void InstanceSubobjectTemplates( void* Data, void const* DefaultData, UStruct* DefaultStruct, UObject* Owner, FObjectInstancingGraph* InstanceGraph );
 
 	/** Returns the structure used for inheritance, may be changed by child types */
 	virtual UStruct* GetInheritanceSuper() const {return GetSuperStruct();}
 
 	/** Static wrapper for Link, using a dummy archive */
-	void StaticLink(bool bRelinkExistingProperties = false);
+	COREUOBJECT_API void StaticLink(bool bRelinkExistingProperties = false);
 
 	/** Creates the field/property links and gets structure ready for use at runtime */
-	virtual void Link(FArchive& Ar, bool bRelinkExistingProperties);
+	COREUOBJECT_API virtual void Link(FArchive& Ar, bool bRelinkExistingProperties);
 
 	/**
 	 * Serializes struct properties, does not handle defaults.  See SerializeBinEx for handling defaults.
@@ -520,7 +520,7 @@ public:
 	 *       types with deprecated properties, unless the ArWantBinarySerialization flag is set on the
 	 *       archive to force serialization to occur always.
 	 */
-	virtual void SerializeBin(FStructuredArchive::FSlot Slot, void* Data) const;
+	COREUOBJECT_API virtual void SerializeBin(FStructuredArchive::FSlot Slot, void* Data) const;
 
 	/**
 	 * Serializes the class properties that reside in Data if they differ from the corresponding values in DefaultData
@@ -535,7 +535,7 @@ public:
 	 *       types with deprecated properties, unless the ArWantBinarySerialization flag is set on the
 	 *       archive to force serialization to occur always.
 	 */
-	void SerializeBinEx( FStructuredArchive::FSlot Slot, void* Data, void const* DefaultData, UStruct* DefaultStruct ) const;
+	COREUOBJECT_API void SerializeBinEx( FStructuredArchive::FSlot Slot, void* Data, void const* DefaultData, UStruct* DefaultStruct ) const;
 
 	/** Serializes list of properties, using property tags to handle mismatches */
 	virtual void SerializeTaggedProperties(FArchive& Ar, uint8* Data, UStruct* DefaultsStruct, uint8* Defaults, const UObject* BreakRecursionIfFullyLoad = nullptr) const final 
@@ -544,13 +544,13 @@ public:
 	}
 
 	/** Serializes list of properties, using property tags to handle mismatches */
-	virtual void SerializeTaggedProperties(FStructuredArchive::FSlot Slot, uint8* Data, UStruct* DefaultsStruct, uint8* Defaults, const UObject* BreakRecursionIfFullyLoad = nullptr) const;
+	COREUOBJECT_API virtual void SerializeTaggedProperties(FStructuredArchive::FSlot Slot, uint8* Data, UStruct* DefaultsStruct, uint8* Defaults, const UObject* BreakRecursionIfFullyLoad = nullptr) const;
 
 	/**
 	 * Preloads all fields that belong to this struct
 	 * @param Ar Archive used for loading this struct
 	 */
-	virtual void PreloadChildren(FArchive& Ar);
+	COREUOBJECT_API virtual void PreloadChildren(FArchive& Ar);
 
 	/**
 	 * Initialize a struct over uninitialized memory. This may be done by calling the native constructor or individually initializing properties
@@ -559,7 +559,7 @@ public:
 	 * @param	ArrayDim	Number of elements in the array
 	 * @param	Stride		Stride of the array, If this default (0), then we will pull the size from the struct
 	 */
-	virtual void InitializeStruct(void* Dest, int32 ArrayDim = 1) const;
+	COREUOBJECT_API virtual void InitializeStruct(void* Dest, int32 ArrayDim = 1) const;
 	
 	/**
 	 * Destroy a struct in memory. This may be done by calling the native destructor or individually destroying properties
@@ -568,13 +568,13 @@ public:
 	 * @param	ArrayDim	Number of elements in the array
 	 * @param	Stride		Stride of the array. If this default (0), then we will pull the size from the struct
 	 */
-	virtual void DestroyStruct(void* Dest, int32 ArrayDim = 1) const;
+	COREUOBJECT_API virtual void DestroyStruct(void* Dest, int32 ArrayDim = 1) const;
 
 	/** Look up a property by an alternate name if it was not found in the first search, this is overridden for user structs */
 	virtual FProperty* CustomFindProperty(const FName InName) const { return nullptr; };
 
 	/** Serialize an expression to an archive. Returns expression token */
-	virtual EExprToken SerializeExpr(int32& iCode, FArchive& Ar);
+	COREUOBJECT_API virtual EExprToken SerializeExpr(int32& iCode, FArchive& Ar);
 
 	/**
 	 * Returns the struct/ class prefix used for the C++ declaration of this struct/ class.
@@ -616,7 +616,7 @@ public:
 
 	/** Returns true if this struct either is SomeBase, or is a child of SomeBase. This will not crash on null structs */
 #if USTRUCT_FAST_ISCHILDOF_COMPARE_WITH_OUTERWALK || USTRUCT_FAST_ISCHILDOF_IMPL == USTRUCT_ISCHILDOF_OUTERWALK
-	bool IsChildOf( const UStruct* SomeBase ) const;
+	COREUOBJECT_API bool IsChildOf( const UStruct* SomeBase ) const;
 #else
 	bool IsChildOf(const UStruct* SomeBase) const
 	{
@@ -634,16 +634,16 @@ public:
 	 * Sets the super struct pointer and updates hash information as necessary.
 	 * Note that this is not sufficient to actually reparent a struct, it simply sets a pointer.
 	 */
-	virtual void SetSuperStruct(UStruct* NewSuperStruct);
+	COREUOBJECT_API virtual void SetSuperStruct(UStruct* NewSuperStruct);
 
 	UE_DEPRECATED(4.23, "Replace with GetAuthoredNameForField or UField::GetAuthoredName")
-	virtual FString PropertyNameToDisplayName(FName InName) const;
+	COREUOBJECT_API virtual FString PropertyNameToDisplayName(FName InName) const;
 
 	/** Returns a human readable string for a given field, overridden for user defined structs */
-	virtual FString GetAuthoredNameForField(const UField* Field) const;
+	COREUOBJECT_API virtual FString GetAuthoredNameForField(const UField* Field) const;
 
 	/** Returns a human readable string for a given field, overridden for user defined structs */
-	virtual FString GetAuthoredNameForField(const FField* Field) const;
+	COREUOBJECT_API virtual FString GetAuthoredNameForField(const FField* Field) const;
 
 	/** If true, this class has been cleaned and sanitized (trashed) and should not be used */
 	virtual bool IsStructTrashed() const
@@ -652,14 +652,14 @@ public:
 	}
 
 	/** Destroys all properties owned by this struct */
-	void DestroyChildPropertiesAndResetPropertyLinks();
+	COREUOBJECT_API void DestroyChildPropertiesAndResetPropertyLinks();
 
 #if WITH_EDITORONLY_DATA
 	/** Try and find boolean metadata with the given key. If not found on this class, work up hierarchy looking for it. */
-	bool GetBoolMetaDataHierarchical(const FName& Key) const;
+	COREUOBJECT_API bool GetBoolMetaDataHierarchical(const FName& Key) const;
 
 	/** Try and find string metadata with the given key. If not found on this class, work up hierarchy looking for it. */
-	bool GetStringMetaDataHierarchical(const FName& Key, FString* OutValue = nullptr) const;
+	COREUOBJECT_API bool GetStringMetaDataHierarchical(const FName& Key, FString* OutValue = nullptr) const;
 
 	/**
 	 * Determines if the struct or any of its super structs has any metadata associated with the provided key
@@ -667,7 +667,7 @@ public:
 	 * @param Key The key to lookup in the metadata
 	 * @return pointer to the UStruct that has associated metadata, nullptr if Key is not associated with any UStruct in the hierarchy
 	 */
-	const UStruct* HasMetaDataHierarchical(const FName& Key) const;
+	COREUOBJECT_API const UStruct* HasMetaDataHierarchical(const FName& Key) const;
 
 	/* Returns true if this struct has Asset Registry searchable properties */
 	FORCEINLINE bool HasAssetRegistrySearchableProperties() const
@@ -700,20 +700,20 @@ public:
 	 * Collects UObjects referenced by bytecode
 	 * @param OutReferencedObjects buffer to store the referenced objects in (not cleared by this function)
 	 */
-	void CollectBytecodeReferencedObjects(TArray<UObject*>& OutReferencedObjects);
+	COREUOBJECT_API void CollectBytecodeReferencedObjects(TArray<UObject*>& OutReferencedObjects);
 	/**
 	 * Collects UObjects referenced by properties
 	 * @param OutReferencedObjects buffer to store the referenced objects in (not cleared by this function)
 	 */
-	void CollectPropertyReferencedObjects(TArray<UObject*>& OutReferencedObjects);
+	COREUOBJECT_API void CollectPropertyReferencedObjects(TArray<UObject*>& OutReferencedObjects);
 	/**
 	 * Collects UObjects referenced by bytecode and properties for faster GC access
 	 */
-	void CollectBytecodeAndPropertyReferencedObjects();
+	COREUOBJECT_API void CollectBytecodeAndPropertyReferencedObjects();
 	/**
 	 * Collects UObjects referenced by bytecode and properties for this class and its child fields and their children...
 	 */
-	void CollectBytecodeAndPropertyReferencedObjectsRecursively();
+	COREUOBJECT_API void CollectBytecodeAndPropertyReferencedObjectsRecursively();
 
 protected:
 
@@ -727,14 +727,14 @@ protected:
 	virtual bool ArePropertyGuidsAvailable() const { return false; }
 
 	/** Serializes properties of this struct */
-	void SerializeProperties(FArchive& Ar);
+	COREUOBJECT_API void SerializeProperties(FArchive& Ar);
 
 #if WITH_EDITORONLY_DATA
-	void ConvertUFieldsToFFields();
+	COREUOBJECT_API void ConvertUFieldsToFFields();
 #endif // WITH_EDITORONLY_DATA
 
 	/** Serializes list of properties to a te, using property tags to handle mismatches */
-	void LoadTaggedPropertiesFromText(FStructuredArchive::FSlot Slot, uint8* Data, UStruct* DefaultsStruct, uint8* Defaults, const UObject* BreakRecursionIfFullyLoad) const;
+	COREUOBJECT_API void LoadTaggedPropertiesFromText(FStructuredArchive::FSlot Slot, uint8* Data, UStruct* DefaultsStruct, uint8* Defaults, const UObject* BreakRecursionIfFullyLoad) const;
 
 private:
 #if USTRUCT_FAST_ISCHILDOF_IMPL == USTRUCT_ISCHILDOF_STRUCTARRAY
@@ -747,7 +747,7 @@ private:
 	friend class FBlueprintCompileReinstancer;
 #endif
 
-	void SerializeVersionedTaggedProperties(FStructuredArchive::FSlot Slot, uint8* Data, UStruct* DefaultsStruct, uint8* Defaults, const UObject* BreakRecursionIfFullyLoad) const;
+	COREUOBJECT_API void SerializeVersionedTaggedProperties(FStructuredArchive::FSlot Slot, uint8* Data, UStruct* DefaultsStruct, uint8* Defaults, const UObject* BreakRecursionIfFullyLoad) const;
 };
 
 /**
@@ -901,7 +901,7 @@ class UScriptStruct : public UStruct
 {
 public:
 	/** Interface to template to manage dynamic access to C++ struct construction and destruction **/
-	struct COREUOBJECT_API ICppStructOps
+	struct ICppStructOps
 	{
 		/** Filled by implementation classes to report their capabilities */
 		struct FCapabilities
@@ -1780,9 +1780,9 @@ public:
 //
 // Reflection data for a replicated or Kismet callable function.
 //
-class COREUOBJECT_API UFunction : public UStruct
+class UFunction : public UStruct
 {
-	DECLARE_CASTED_CLASS_INTRINSIC(UFunction, UStruct, 0, TEXT("/Script/CoreUObject"), CASTCLASS_UFunction)
+	DECLARE_CASTED_CLASS_INTRINSIC_WITH_API(UFunction, UStruct, 0, TEXT("/Script/CoreUObject"), CASTCLASS_UFunction, COREUOBJECT_API)
 	DECLARE_WITHIN(UClass)
 public:
 	// Persistent variables.
@@ -1851,31 +1851,31 @@ public:
 	 * @param Stack  - The parameter stack for the function call.
 	 * @param Result - The result of the function.
 	 */
-	void Invoke(UObject* Obj, FFrame& Stack, RESULT_DECL);
+	COREUOBJECT_API void Invoke(UObject* Obj, FFrame& Stack, RESULT_DECL);
 
 	// Constructors.
-	explicit UFunction(const FObjectInitializer& ObjectInitializer, UFunction* InSuperFunction, EFunctionFlags InFunctionFlags = FUNC_None, SIZE_T ParamsSize = 0 );
-	explicit UFunction(UFunction* InSuperFunction, EFunctionFlags InFunctionFlags = FUNC_None, SIZE_T ParamsSize = 0);
+	COREUOBJECT_API explicit UFunction(const FObjectInitializer& ObjectInitializer, UFunction* InSuperFunction, EFunctionFlags InFunctionFlags = FUNC_None, SIZE_T ParamsSize = 0 );
+	COREUOBJECT_API explicit UFunction(UFunction* InSuperFunction, EFunctionFlags InFunctionFlags = FUNC_None, SIZE_T ParamsSize = 0);
 
 	/** Initializes transient members like return value offset */
-	void InitializeDerivedMembers();
+	COREUOBJECT_API void InitializeDerivedMembers();
 
 	// UObject interface.
-	virtual void Serialize( FArchive& Ar ) override;
-	virtual void PostLoad() override;
+	COREUOBJECT_API virtual void Serialize( FArchive& Ar ) override;
+	COREUOBJECT_API virtual void PostLoad() override;
 
 	// UField interface.
-	virtual void Bind() override;
+	COREUOBJECT_API virtual void Bind() override;
 
 	// UStruct interface.
 	virtual UStruct* GetInheritanceSuper() const override { return nullptr;}
-	virtual void Link(FArchive& Ar, bool bRelinkExistingProperties) override;
+	COREUOBJECT_API virtual void Link(FArchive& Ar, bool bRelinkExistingProperties) override;
 
 	/** Returns parent function if there is one, or null */
-	UFunction* GetSuperFunction() const;
+	COREUOBJECT_API UFunction* GetSuperFunction() const;
 
 	/** Returns the return value property if there is one, or null */
-	FProperty* GetReturnProperty() const;
+	COREUOBJECT_API FProperty* GetReturnProperty() const;
 
 	/** Returns the owning UClass* without branching */
 	FORCEINLINE UClass* GetOuterUClassUnchecked() const
@@ -1929,7 +1929,7 @@ public:
 	 *
 	 * @return	true if function signatures are compatible.
 	 */
-	bool IsSignatureCompatibleWith(const UFunction* OtherFunction) const;
+	COREUOBJECT_API bool IsSignatureCompatibleWith(const UFunction* OtherFunction) const;
 
 	/**
 	 * Determines if two functions have an identical signature (note: currently doesn't allow
@@ -1941,33 +1941,33 @@ public:
 	 *
 	 * @return	true if function signatures are compatible.
 	 */
-	bool IsSignatureCompatibleWith(const UFunction* OtherFunction, uint64 IgnoreFlags) const;
+	COREUOBJECT_API bool IsSignatureCompatibleWith(const UFunction* OtherFunction, uint64 IgnoreFlags) const;
 };
 
 //
 // Function definition used by dynamic delegate declarations
 //
-class COREUOBJECT_API UDelegateFunction : public UFunction
+class UDelegateFunction : public UFunction
 {
-	DECLARE_CASTED_CLASS_INTRINSIC(UDelegateFunction, UFunction, 0, TEXT("/Script/CoreUObject"), CASTCLASS_UDelegateFunction)
+	DECLARE_CASTED_CLASS_INTRINSIC_WITH_API(UDelegateFunction, UFunction, 0, TEXT("/Script/CoreUObject"), CASTCLASS_UDelegateFunction, COREUOBJECT_API)
 	DECLARE_WITHIN(UObject)
 public:
-	explicit UDelegateFunction(const FObjectInitializer& ObjectInitializer, UFunction* InSuperFunction, EFunctionFlags InFunctionFlags = FUNC_None, SIZE_T ParamsSize = 0);
-	explicit UDelegateFunction(UFunction* InSuperFunction, EFunctionFlags InFunctionFlags = FUNC_None, SIZE_T ParamsSize = 0);
+	COREUOBJECT_API explicit UDelegateFunction(const FObjectInitializer& ObjectInitializer, UFunction* InSuperFunction, EFunctionFlags InFunctionFlags = FUNC_None, SIZE_T ParamsSize = 0);
+	COREUOBJECT_API explicit UDelegateFunction(UFunction* InSuperFunction, EFunctionFlags InFunctionFlags = FUNC_None, SIZE_T ParamsSize = 0);
 };
 
 //
 // Function definition used by sparse dynamic delegate declarations
 //
-class COREUOBJECT_API USparseDelegateFunction : public UDelegateFunction
+class USparseDelegateFunction : public UDelegateFunction
 {
-	DECLARE_CASTED_CLASS_INTRINSIC(USparseDelegateFunction, UDelegateFunction, 0, TEXT("/Script/CoreUObject"), CASTCLASS_USparseDelegateFunction)
+	DECLARE_CASTED_CLASS_INTRINSIC_WITH_API(USparseDelegateFunction, UDelegateFunction, 0, TEXT("/Script/CoreUObject"), CASTCLASS_USparseDelegateFunction, COREUOBJECT_API)
 	DECLARE_WITHIN(UObject)
 public:
-	explicit USparseDelegateFunction(const FObjectInitializer& ObjectInitializer, UFunction* InSuperFunction, EFunctionFlags InFunctionFlags = FUNC_None, SIZE_T ParamsSize = 0);
-	explicit USparseDelegateFunction(UFunction* InSuperFunction, EFunctionFlags InFunctionFlags = FUNC_None, SIZE_T ParamsSize = 0);
+	COREUOBJECT_API explicit USparseDelegateFunction(const FObjectInitializer& ObjectInitializer, UFunction* InSuperFunction, EFunctionFlags InFunctionFlags = FUNC_None, SIZE_T ParamsSize = 0);
+	COREUOBJECT_API explicit USparseDelegateFunction(UFunction* InSuperFunction, EFunctionFlags InFunctionFlags = FUNC_None, SIZE_T ParamsSize = 0);
 
-	virtual void Serialize(FArchive& Ar) override;
+	COREUOBJECT_API virtual void Serialize(FArchive& Ar) override;
 
 	FName OwningClassName;
 	FName DelegateName;
@@ -2659,7 +2659,7 @@ private:
 
 
 /** information about an interface a class implements */
-struct COREUOBJECT_API FImplementedInterface
+struct FImplementedInterface
 {
 	/** the interface class */
 	TObjectPtr<UClass> Class;
@@ -2726,9 +2726,9 @@ enum class EGetSparseClassDataMethod : uint8
 /**
  * An object class.
  */
-class COREUOBJECT_API UClass : public UStruct
+class UClass : public UStruct
 {
-	DECLARE_CASTED_CLASS_INTRINSIC_NO_CTOR(UClass, UStruct, 0, TEXT("/Script/CoreUObject"), CASTCLASS_UClass, NO_API)
+	DECLARE_CASTED_CLASS_INTRINSIC_NO_CTOR(UClass, UStruct, 0, TEXT("/Script/CoreUObject"), CASTCLASS_UClass, COREUOBJECT_API)
 	DECLARE_WITHIN_UPACKAGE()
 
 public:
@@ -2772,7 +2772,7 @@ public:
 	FField* PropertiesPendingDestruction;
 
 	/** Destroys properties that couldn't be destroyed in PurgeClass */
-	void DestroyPropertiesPendingDestruction();
+	COREUOBJECT_API void DestroyPropertiesPendingDestruction();
 #endif
 
 #if WITH_EDITOR
@@ -2795,16 +2795,16 @@ public:
 
 #if WITH_EDITOR
 	// Editor only properties
-	void GetHideFunctions(TArray<FString>& OutHideFunctions) const;
-	bool IsFunctionHidden(const TCHAR* InFunction) const;
-	void GetAutoExpandCategories(TArray<FString>& OutAutoExpandCategories) const;
-	bool IsAutoExpandCategory(const TCHAR* InCategory) const;
-	void GetPrioritizeCategories(TArray<FString>& OutPrioritizedCategories) const;
-	bool IsPrioritizeCategory(const TCHAR* InCategory) const;
-	void GetAutoCollapseCategories(TArray<FString>& OutAutoCollapseCategories) const;
-	bool IsAutoCollapseCategory(const TCHAR* InCategory) const;
-	void GetClassGroupNames(TArray<FString>& OutClassGroupNames) const;
-	bool IsClassGroupName(const TCHAR* InGroupName) const;
+	COREUOBJECT_API void GetHideFunctions(TArray<FString>& OutHideFunctions) const;
+	COREUOBJECT_API bool IsFunctionHidden(const TCHAR* InFunction) const;
+	COREUOBJECT_API void GetAutoExpandCategories(TArray<FString>& OutAutoExpandCategories) const;
+	COREUOBJECT_API bool IsAutoExpandCategory(const TCHAR* InCategory) const;
+	COREUOBJECT_API void GetPrioritizeCategories(TArray<FString>& OutPrioritizedCategories) const;
+	COREUOBJECT_API bool IsPrioritizeCategory(const TCHAR* InCategory) const;
+	COREUOBJECT_API void GetAutoCollapseCategories(TArray<FString>& OutAutoCollapseCategories) const;
+	COREUOBJECT_API bool IsAutoCollapseCategory(const TCHAR* InCategory) const;
+	COREUOBJECT_API void GetClassGroupNames(TArray<FString>& OutClassGroupNames) const;
+	COREUOBJECT_API bool IsClassGroupName(const TCHAR* InGroupName) const;
 #endif
 	/**
 	 * Calls AddReferencedObjects static method on the specified object.
@@ -2859,7 +2859,7 @@ public:
 	 * Returns a pointer to the sidecar data structure, based on the EGetSparseClassDataMethod.
 	 * @note It is only safe to mutate this data when using "CreateIfNull", as others may return archetype/default data; consider GetOrCreateSparseClassData for this use-case.
 	 */
-	const void* GetSparseClassData(const EGetSparseClassDataMethod GetMethod);
+	COREUOBJECT_API const void* GetSparseClassData(const EGetSparseClassDataMethod GetMethod);
 
 	/**
 	 * Returns a pointer to the sidecar data structure. This function will create an instance of the data structure if one has been specified and it has not yet been created.
@@ -2869,18 +2869,18 @@ public:
 	/**
 	 * Returns a pointer to the type of the sidecar data structure if one is specified.
 	 */
-	UScriptStruct* GetSparseClassDataStruct() const;
+	COREUOBJECT_API UScriptStruct* GetSparseClassDataStruct() const;
 
-	void SetSparseClassDataStruct(UScriptStruct* InSparseClassDataStruct);
+	COREUOBJECT_API void SetSparseClassDataStruct(UScriptStruct* InSparseClassDataStruct);
 
 	/** 
 	 * Clears the sparse class data struct for this and all child classes that directly reference it as a super-struct 
 	 * This will rename the current sparse class data struct aside into the transient package
 	 */
-	void ClearSparseClassDataStruct(bool bInRecomplingOnLoad);
+	COREUOBJECT_API void ClearSparseClassDataStruct(bool bInRecomplingOnLoad);
 
 	/** Assemble reference token streams for all classes if they haven't had it assembled already */
-	static void AssembleReferenceTokenStreams();
+	static COREUOBJECT_API void AssembleReferenceTokenStreams();
 
 #if WITH_EDITOR
 	void GenerateFunctionList(TArray<FName>& OutArray) const 
@@ -2891,9 +2891,9 @@ public:
 #endif // WITH_EDITOR
 
 protected:
-	void* CreateSparseClassData();
+	COREUOBJECT_API void* CreateSparseClassData();
 
-	void CleanupSparseClassData();
+	COREUOBJECT_API void CleanupSparseClassData();
 
 private:
 #if WITH_EDITOR
@@ -2928,9 +2928,9 @@ public:
 
 public:
 	// Constructors
-	UClass(const FObjectInitializer& ObjectInitializer = FObjectInitializer::Get());
-	explicit UClass(const FObjectInitializer& ObjectInitializer, UClass* InSuperClass);
-	UClass( EStaticConstructor, FName InName, uint32 InSize, uint32 InAlignment, EClassFlags InClassFlags, EClassCastFlags InClassCastFlags,
+	COREUOBJECT_API UClass(const FObjectInitializer& ObjectInitializer = FObjectInitializer::Get());
+	COREUOBJECT_API explicit UClass(const FObjectInitializer& ObjectInitializer, UClass* InSuperClass);
+	COREUOBJECT_API UClass( EStaticConstructor, FName InName, uint32 InSize, uint32 InAlignment, EClassFlags InClassFlags, EClassCastFlags InClassCastFlags,
 		const TCHAR* InClassConfigName, EObjectFlags InFlags, ClassConstructorType InClassConstructor,
 		ClassVTableHelperCtorCallerType InClassVTableHelperCtorCaller,
 		FUObjectCppClassStaticFunctions&& InCppClassStaticFunctions);
@@ -2946,7 +2946,7 @@ public:
 	 * @param	TClass_Super_StaticClass		Static class of the super class
 	 * @param	TClass_WithinClass_StaticClass	Static class of the WithinClass
 	 */
-	bool HotReloadPrivateStaticClass(
+	COREUOBJECT_API bool HotReloadPrivateStaticClass(
 		uint32			InSize,
 		EClassFlags		InClassFlags,
 		EClassCastFlags	InClassCastFlags,
@@ -2968,7 +2968,7 @@ public:
 	*											because dispatch is shared, so the C++ remap table does not work in this case, and this should be false
 	* @return	true if the function was found and replaced, false if it was not
 	*/
-	bool ReplaceNativeFunction(FName InName, FNativeFuncPtr InPointer, bool bAddToFunctionRemapTable);
+	COREUOBJECT_API bool ReplaceNativeFunction(FName InName, FNativeFuncPtr InPointer, bool bAddToFunctionRemapTable);
 #endif
 
 	/**
@@ -2976,7 +2976,7 @@ public:
 	 *
 	 * @return The version of this class that references should be stored to
 	 */
-	virtual UClass* GetAuthoritativeClass();
+	COREUOBJECT_API virtual UClass* GetAuthoritativeClass();
 	const UClass* GetAuthoritativeClass() const { return const_cast<UClass*>(this)->GetAuthoritativeClass(); }
 
 	/**
@@ -2984,7 +2984,7 @@ public:
 	 * @param	InName							name of the function
 	 * @param	InPointer						pointer to the function
 	 */
-	void AddNativeFunction(const ANSICHAR* InName, FNativeFuncPtr InPointer);
+	COREUOBJECT_API void AddNativeFunction(const ANSICHAR* InName, FNativeFuncPtr InPointer);
 
 	/**
 	 * Add a native function to the internal native function table, but with a unicode name. Used when generating code from blueprints, 
@@ -2992,7 +2992,7 @@ public:
 	 * @param	InName							name of the function
 	 * @param	InPointer						pointer to the function
 	 */
-	void AddNativeFunction(const WIDECHAR* InName, FNativeFuncPtr InPointer);
+	COREUOBJECT_API void AddNativeFunction(const WIDECHAR* InName, FNativeFuncPtr InPointer);
 
 	/** Add a function to the function map */
 	void AddFunctionToFunctionMap(UFunction* Function, FName FuncName)
@@ -3001,7 +3001,7 @@ public:
 		FuncMap.Add(FuncName, Function);
 	}
 
-	void CreateLinkAndAddChildFunctionsToMap(const FClassFunctionLinkInfo* Functions, uint32 NumFunctions);
+	COREUOBJECT_API void CreateLinkAndAddChildFunctionsToMap(const FClassFunctionLinkInfo* Functions, uint32 NumFunctions);
 
 	/** Remove a function from the function map */
 	void RemoveFunctionFromFunctionMap(UFunction* Function)
@@ -3011,40 +3011,40 @@ public:
 	}
 
 	/** Clears the function name caches, in case things have changed */
-	void ClearFunctionMapsCaches();
+	COREUOBJECT_API void ClearFunctionMapsCaches();
 
 	/** Looks for a given function name */
-	UFunction* FindFunctionByName(FName InName, EIncludeSuperFlag::Type IncludeSuper = EIncludeSuperFlag::IncludeSuper) const;
+	COREUOBJECT_API UFunction* FindFunctionByName(FName InName, EIncludeSuperFlag::Type IncludeSuper = EIncludeSuperFlag::IncludeSuper) const;
 
 	// UObject interface.
-	virtual void Serialize(FArchive& Ar) override;
-	virtual void PostLoad() override;
-	virtual void FinishDestroy() override;
-	virtual void DeferredRegister(UClass *UClassStaticClass,const TCHAR* PackageName,const TCHAR* InName) override;
-	virtual bool Rename(const TCHAR* NewName = nullptr, UObject* NewOuter = nullptr, ERenameFlags Flags = REN_None) override;
-	virtual void TagSubobjects(EObjectFlags NewFlags) override;
-	virtual void PostInitProperties() override;
-	static void AddReferencedObjects(UObject* InThis, FReferenceCollector& Collector);
-	virtual FRestoreForUObjectOverwrite* GetRestoreForUObjectOverwrite() override;
-	virtual FString GetDesc() override;
-	virtual void GetAssetRegistryTags(TArray<FAssetRegistryTag>& OutTags) const override;
+	COREUOBJECT_API virtual void Serialize(FArchive& Ar) override;
+	COREUOBJECT_API virtual void PostLoad() override;
+	COREUOBJECT_API virtual void FinishDestroy() override;
+	COREUOBJECT_API virtual void DeferredRegister(UClass *UClassStaticClass,const TCHAR* PackageName,const TCHAR* InName) override;
+	COREUOBJECT_API virtual bool Rename(const TCHAR* NewName = nullptr, UObject* NewOuter = nullptr, ERenameFlags Flags = REN_None) override;
+	COREUOBJECT_API virtual void TagSubobjects(EObjectFlags NewFlags) override;
+	COREUOBJECT_API virtual void PostInitProperties() override;
+	static COREUOBJECT_API void AddReferencedObjects(UObject* InThis, FReferenceCollector& Collector);
+	COREUOBJECT_API virtual FRestoreForUObjectOverwrite* GetRestoreForUObjectOverwrite() override;
+	COREUOBJECT_API virtual FString GetDesc() override;
+	COREUOBJECT_API virtual void GetAssetRegistryTags(TArray<FAssetRegistryTag>& OutTags) const override;
 #if WITH_EDITOR
-	virtual void PostLoadAssetRegistryTags(const FAssetData& InAssetData, TArray<FAssetRegistryTag>& OutTagsAndValuesToUpdate) const;
+	COREUOBJECT_API virtual void PostLoadAssetRegistryTags(const FAssetData& InAssetData, TArray<FAssetRegistryTag>& OutTagsAndValuesToUpdate) const;
 #endif // WITH_EDITOR
 	virtual bool IsAsset() const override { return false; }	
 	virtual bool IsNameStableForNetworking() const override { return true; } // For now, assume all classes have stable net names
-	virtual void GetPreloadDependencies(TArray<UObject*>& OutDeps) override;
+	COREUOBJECT_API virtual void GetPreloadDependencies(TArray<UObject*>& OutDeps) override;
 	// End of UObject interface.
 
 	// UField interface.
-	virtual void Bind() override;
-	virtual const TCHAR* GetPrefixCPP() const override;
+	COREUOBJECT_API virtual void Bind() override;
+	COREUOBJECT_API virtual const TCHAR* GetPrefixCPP() const override;
 	// End of UField interface.
 
 	// UStruct interface.
-	virtual void Link(FArchive& Ar, bool bRelinkExistingProperties) override;
-	virtual void SetSuperStruct(UStruct* NewSuperStruct) override;
-	virtual bool IsStructTrashed() const override;
+	COREUOBJECT_API virtual void Link(FArchive& Ar, bool bRelinkExistingProperties) override;
+	COREUOBJECT_API virtual void SetSuperStruct(UStruct* NewSuperStruct) override;
+	COREUOBJECT_API virtual bool IsStructTrashed() const override;
 	// End of UStruct interface.
 
 	/**
@@ -3063,7 +3063,7 @@ public:
 	 * @param InShortNameOrPathName Class name
 	 * @return Pointer to a type object if successful, nullptr otherwise
 	 */
-	static UField* TryFindTypeSlow(UClass* TypeClass, const FString& InPathNameOrShortName, EFindFirstObjectOptions InOptions = EFindFirstObjectOptions::None);
+	static COREUOBJECT_API UField* TryFindTypeSlow(UClass* TypeClass, const FString& InPathNameOrShortName, EFindFirstObjectOptions InOptions = EFindFirstObjectOptions::None);
 
 	/**
 	 * Utility function that tries to find a type (class/struct/enum) given a path name or a short name however it will throw a warning (with a callstack) if it's the latter.
@@ -3087,7 +3087,7 @@ public:
 	 * @param InShortNameOrPathName Class name
 	 * @return Pointer to a type object if successful, nullptr otherwise
 	 */
-	static UField* TryFindTypeSlowSafe(UClass* TypeClass, const FString& InPathNameOrShortName, EFindFirstObjectOptions InOptions = EFindFirstObjectOptions::None);
+	static COREUOBJECT_API UField* TryFindTypeSlowSafe(UClass* TypeClass, const FString& InPathNameOrShortName, EFindFirstObjectOptions InOptions = EFindFirstObjectOptions::None);
 
 	/**
 	 * Utility function that tries to find a type (class/struct/enum) given a path name or a short name however it will throw a warning (with a callstack) if it's the latter.
@@ -3112,7 +3112,7 @@ public:
 	 * @param AmbiguousClassMessage Additional message to log when class name is ambiguous (e.g. current operation) 
 	 * @return Class path name if successful. Empty FTopLevelAssetPath otherwise.
 	 */
-	static FTopLevelAssetPath TryConvertShortTypeNameToPathName(UClass* TypeClass, const FString& InShortTypeName, ELogVerbosity::Type AmbiguousMessageVerbosity = ELogVerbosity::NoLogging, const TCHAR* AmbiguousClassMessage = nullptr);
+	static COREUOBJECT_API FTopLevelAssetPath TryConvertShortTypeNameToPathName(UClass* TypeClass, const FString& InShortTypeName, ELogVerbosity::Type AmbiguousMessageVerbosity = ELogVerbosity::NoLogging, const TCHAR* AmbiguousClassMessage = nullptr);
 
 	/**
 	 * Tries to convert short class name to class path name. This will only work if the class exists in memory
@@ -3138,7 +3138,7 @@ public:
 	 * @return True if the short path was successfully fixed. False if the provided export path did not contain short
 	 *         class name or the short path could not be fixed.
 	 */
-	static bool TryFixShortClassNameExportPath(FString& InOutExportPathToFix,
+	static COREUOBJECT_API bool TryFixShortClassNameExportPath(FString& InOutExportPathToFix,
 		ELogVerbosity::Type AmbiguousMessageVerbosity = ELogVerbosity::NoLogging,
 		const TCHAR* AmbiguousClassMessage = nullptr, bool bClearOnError = false);
 
@@ -3146,7 +3146,7 @@ public:
 	 * Returns the ObjectName portion of a ClassPath name: "/Path/To.Object" is converted to "Object". 
 	 * Returns the full string if it is already a ShortTypeName.
 	 */
-	static FString ConvertPathNameToShortTypeName(FStringView InClassPathOrShortTypeName);
+	static COREUOBJECT_API FString ConvertPathNameToShortTypeName(FStringView InClassPathOrShortTypeName);
 
 	/**
 	 * Takes a FullName (from e.g. AssetData.GetFullName or UObject.GetFullName) in either ShortTypeFullName or PathFullName form
@@ -3154,13 +3154,13 @@ public:
 	 * PathFullName: "/ClassPath/ClassPackage.ClassObjectName /PackagePath/PackageShortName.ObjectName:SubObjectName"
 	 * Converts it to ShortTypeFullName if not already in that format and returns it
 	 */
-	static FString ConvertFullNameToShortTypeFullName(FStringView InFullName);
+	static COREUOBJECT_API FString ConvertFullNameToShortTypeFullName(FStringView InFullName);
 
 	/**
 	 * Returns whether the given stringview is in ShortTypeName form: No directory separators "/" or object separators ".", SUBOBJECT_DELIMITER
 	 * Returns true for empty string.
 	 */
-	static bool IsShortTypeName(FStringView ClassPathOrShortTypeName);
+	static COREUOBJECT_API bool IsShortTypeName(FStringView ClassPathOrShortTypeName);
 
 #if WITH_EDITOR
 	/** Provides access to C++ type info. */
@@ -3187,7 +3187,7 @@ public:
 	 *
 	 * @return	name of the class specific ini file
 	 */
-	const FString GetConfigName() const;
+	COREUOBJECT_API const FString GetConfigName() const;
 
 	/** Returns parent class, the parent of a Class is always another class */
 	UClass* GetSuperClass() const
@@ -3196,7 +3196,7 @@ public:
 	}
 
 	/** Feedback context for default property import **/
-	static class FFeedbackContext& GetDefaultPropertiesFeedbackContext();
+	static COREUOBJECT_API class FFeedbackContext& GetDefaultPropertiesFeedbackContext();
 
 	/** Returns amount of memory used by default object */
 	int32 GetDefaultsCount()
@@ -3255,7 +3255,7 @@ public:
 	 * Get the name of the CDO for the this class
 	 * @return The name of the CDO
 	 */
-	FName GetDefaultObjectName() const;
+	COREUOBJECT_API FName GetDefaultObjectName() const;
 
 	/** Returns memory used to store temporary data on an instance, used by blueprints */
 	virtual uint8* GetPersistentUberGraphFrame(UObject* Obj, UFunction* FuncToCheck) const
@@ -3286,7 +3286,7 @@ public:
 	}
 
 	/** Searches for the default instanced object (often a component) by name **/
-	UObject* GetDefaultSubobjectByName(FName ToFind);
+	COREUOBJECT_API UObject* GetDefaultSubobjectByName(FName ToFind);
 
 	/** Adds a new default instance map item **/
 	void AddDefaultSubobject(UObject* NewSubobject, const UClass* BaseClass)
@@ -3302,7 +3302,7 @@ public:
 	 *
 	 * @param OutDefaultSubobjects An array to be filled with default subobjects.
 	 */
-	void GetDefaultObjectSubobjects(TArray<UObject*>& OutDefaultSubobjects);
+	COREUOBJECT_API void GetDefaultObjectSubobjects(TArray<UObject*>& OutDefaultSubobjects);
 
 	/**
 	 * Used to safely check whether the passed in flag is set.
@@ -3355,7 +3355,7 @@ public:
 		return (ClassCastFlags&FlagsToCheck) == FlagsToCheck;
 	}
 
-	FString GetDescription() const;
+	COREUOBJECT_API FString GetDescription() const;
 
 	/**
 	 * Assembles the token stream for realtime garbage collection by combining the per class only
@@ -3363,20 +3363,20 @@ public:
 	 * work is avoided by using an object flag.
 	 * @param bForce Assemble the stream even if it has been already assembled (deletes the old one)
 	 */
-	void AssembleReferenceTokenStream(bool bForce = false);
+	COREUOBJECT_API void AssembleReferenceTokenStream(bool bForce = false);
 
 	/** 
 	 * This will return whether or not this class implements the passed in class / interface 
 	 *
 	 * @param SomeClass - the interface to check and see if this class implements it
 	 */
-	bool ImplementsInterface(const class UClass* SomeInterface) const;
+	COREUOBJECT_API bool ImplementsInterface(const class UClass* SomeInterface) const;
 
 	/** serializes the passed in object as this class's default object using the given archive slot
 	 * @param Object the object to serialize as default
 	 * @param Slot the structured archive slot to serialize from
 	 */
-	virtual void SerializeDefaultObject(UObject* Object, FStructuredArchive::FSlot Slot);
+	COREUOBJECT_API virtual void SerializeDefaultObject(UObject* Object, FStructuredArchive::FSlot Slot);
 
 	/** serializes the passed in object as this class's default object using the given archive
 	 * @param Object the object to serialize as default
@@ -3390,7 +3390,7 @@ public:
 	/** serializes the associated sparse class data for the passed in object using the given archive slot. This should only be called if the class has an associated sparse data structure.
 	 * @param Slot the structured archive slot to serialize from
 	 */
-	void SerializeSparseClassData(FStructuredArchive::FSlot Slot);
+	COREUOBJECT_API void SerializeSparseClassData(FStructuredArchive::FSlot Slot);
 
 	/** Wraps the PostLoad() call for the class default object.
 	 * @param Object the default object to call PostLoad() on
@@ -3403,7 +3403,7 @@ public:
 	 *
 	 * In editor, properties are not freed until DestroyPropertiesPendingDestruction is called.
 	 */
-	virtual void PurgeClass(bool bRecompilingOnLoad);
+	COREUOBJECT_API virtual void PurgeClass(bool bRecompilingOnLoad);
 
 	/**
 	 * Finds the common base class that parents the two classes passed in.
@@ -3412,7 +3412,7 @@ public:
 	 * @param InClassB		the second class to find the common base for
 	 * @return				the common base class or NULL
 	 */
-	static UClass* FindCommonBase(UClass* InClassA, UClass* InClassB);
+	static COREUOBJECT_API UClass* FindCommonBase(UClass* InClassA, UClass* InClassB);
 
 	/**
 	 * Finds the common base class that parents the array of classes passed in.
@@ -3420,7 +3420,7 @@ public:
 	 * @param InClasses		the array of classes to find the common base for
 	 * @return				the common base class or NULL
 	 */
-	static UClass* FindCommonBase(const TArray<UClass*>& InClasses);
+	static COREUOBJECT_API UClass* FindCommonBase(const TArray<UClass*>& InClasses);
 
 	/**
 	 * Determines if the specified function has been implemented in a Blueprint
@@ -3428,7 +3428,7 @@ public:
 	 * @param InFunctionName	The name of the function to test
 	 * @return					True if the specified function exists and is implemented in a blueprint generated class
 	 */
-	virtual bool IsFunctionImplementedInScript(FName InFunctionName) const;
+	COREUOBJECT_API virtual bool IsFunctionImplementedInScript(FName InFunctionName) const;
 
 	UE_DEPRECATED(4.23, "IsFunctionImplementedInBlueprint is deprecated, call IsFunctionImplementedInScript instead")
 	bool IsFunctionImplementedInBlueprint(FName InFunctionName) const { return IsFunctionImplementedInScript(InFunctionName); }
@@ -3438,22 +3438,22 @@ public:
 	 * @param InProperty	The property to check if it is contained in this or a parent class.
 	 * @return				True if the property exists on this or a parent class.
 	 */
-	virtual bool HasProperty(FProperty* InProperty) const;
+	COREUOBJECT_API virtual bool HasProperty(FProperty* InProperty) const;
 
 	/** Finds the object that is used as the parent object when serializing properties, overridden for blueprints */
 	virtual UObject* FindArchetype(const UClass* ArchetypeClass, const FName ArchetypeName) const { return nullptr; }
 
 	/** Returns archetype object for CDO */
-	virtual UObject* GetArchetypeForCDO() const;
+	COREUOBJECT_API virtual UObject* GetArchetypeForCDO() const;
 
 	/** Returns archetype for sparse class data */
-	const void* GetArchetypeForSparseClassData() const;
+	COREUOBJECT_API const void* GetArchetypeForSparseClassData() const;
 
 	/** Returns the struct used by the sparse class data archetype */
-	UScriptStruct* GetSparseClassDataArchetypeStruct() const;
+	COREUOBJECT_API UScriptStruct* GetSparseClassDataArchetypeStruct() const;
 
 	/** Returns whether the sparse class data on this instance overrides that of its archetype (in type or value) */
-	bool OverridesSparseClassDataArchetype() const;
+	COREUOBJECT_API bool OverridesSparseClassDataArchetype() const;
 
 	/**
 	* Returns all objects that should be preloaded before the class default object is serialized at load time. Only used by the EDL.
@@ -3470,12 +3470,12 @@ public:
 	 * to force lazy initialization next time the data is needed.
 	 * Also happens after blueprint compiliation.
 	 */
-	void SetUpRuntimeReplicationData();
+	COREUOBJECT_API void SetUpRuntimeReplicationData();
 
 	/**
 	 * Helper function for determining if the given class is compatible with structured archive serialization
 	 */
-	static bool IsSafeToSerializeToStructuredArchives(UClass* InClass);
+	static COREUOBJECT_API bool IsSafeToSerializeToStructuredArchives(UClass* InClass);
 
 private:
 	/** 
@@ -3518,15 +3518,15 @@ private:
 	/**
 	 * Tests if all properties tagged with Replicate were registered in GetLifetimeReplicatedProps
 	 */
-	void ValidateRuntimeReplicationData();
-	void AssembleReferenceTokenStreamInternal(bool bForce = false);
-	void InternalCreateDefaultObjectWrapper() const;
+	COREUOBJECT_API void ValidateRuntimeReplicationData();
+	COREUOBJECT_API void AssembleReferenceTokenStreamInternal(bool bForce = false);
+	COREUOBJECT_API void InternalCreateDefaultObjectWrapper() const;
 protected:
 	/**
 	 * Get the default object from the class, creating it if missing, if requested or under a few other circumstances
 	 * @return		the CDO for this class
 	 **/
-	virtual UObject* CreateDefaultObject();
+	COREUOBJECT_API virtual UObject* CreateDefaultObject();
 };
 
 // @todo: BP2CPP_remove
@@ -3674,7 +3674,7 @@ enum class EObjectInstancingGraphOptions
 
 ENUM_CLASS_FLAGS(EObjectInstancingGraphOptions)
 
-struct COREUOBJECT_API FObjectInstancingGraph
+struct FObjectInstancingGraph
 {
 public:
 
@@ -3682,13 +3682,13 @@ public:
 	 * Default Constructor 
 	 * @param bDisableInstancing - if true, start with component instancing disabled
 	**/
-	explicit FObjectInstancingGraph(bool bDisableInstancing = false);
+	COREUOBJECT_API explicit FObjectInstancingGraph(bool bDisableInstancing = false);
 
 	/**
 	 * Constructor with options
 	 * @param InOptions Additional options to modify the behavior of this graph
 	**/
-	explicit FObjectInstancingGraph(EObjectInstancingGraphOptions InOptions);
+	COREUOBJECT_API explicit FObjectInstancingGraph(EObjectInstancingGraphOptions InOptions);
 
 	/**
 	 * Standard constructor
@@ -3696,7 +3696,7 @@ public:
 	 * @param	DestinationSubobjectRoot	the top-level object that is being created
 	 * @param	InOptions					Additional options to modify the behavior of this graph
 	 */
-	explicit FObjectInstancingGraph( class UObject* DestinationSubobjectRoot, EObjectInstancingGraphOptions InOptions = EObjectInstancingGraphOptions::None);
+	COREUOBJECT_API explicit FObjectInstancingGraph( class UObject* DestinationSubobjectRoot, EObjectInstancingGraphOptions InOptions = EObjectInstancingGraphOptions::None);
 
 	/**
 	 * Returns whether this instancing graph has a valid destination root.
@@ -3720,14 +3720,14 @@ public:
 	 * @param	DestinationSubobjectRoot	the top-level object that is being created
 	 * @param	InSourceRoot	Archetype of DestinationSubobjectRoot
 	 */
-	void SetDestinationRoot( class UObject* DestinationSubobjectRoot, class UObject* InSourceRoot = nullptr );
+	COREUOBJECT_API void SetDestinationRoot( class UObject* DestinationSubobjectRoot, class UObject* InSourceRoot = nullptr );
 
 	/**
 	 * Finds the destination object instance corresponding to the specified source object.
 	 *
 	 * @param	SourceObject			the object to find the corresponding instance for
 	 */
-	class UObject* GetDestinationObject(class UObject* SourceObject);
+	COREUOBJECT_API class UObject* GetDestinationObject(class UObject* SourceObject);
 
 	/**
 	 * Returns the component that has SourceComponent as its archetype, instancing the component as necessary.
@@ -3772,14 +3772,14 @@ public:
 	 *
 	 * @return	As with GetInstancedSubobject, above, but also deals with archetype creation and a few other special cases
 	 */
-	class UObject* InstancePropertyValue( class UObject* SourceComponent, class UObject* CurrentValue, class UObject* CurrentObject, EInstancePropertyValueFlags Flags = EInstancePropertyValueFlags::None );
+	COREUOBJECT_API class UObject* InstancePropertyValue( class UObject* SourceComponent, class UObject* CurrentValue, class UObject* CurrentObject, EInstancePropertyValueFlags Flags = EInstancePropertyValueFlags::None );
 
 	/**
 	 * Adds a partially built object instance to the map(s) of source objects to their instances.
 	 * @param	ObjectInstance  Object that was just allocated, but has not been constructed yet
 	 * @param	InArchetype     Archetype of ObjectInstance
 	 */
-	void AddNewObject(class UObject* ObjectInstance, class UObject* InArchetype = nullptr);
+	COREUOBJECT_API void AddNewObject(class UObject* ObjectInstance, class UObject* InArchetype = nullptr);
 
 	/**
 	 * Adds an object instance to the map of source objects to their instances.  If there is already a mapping for this object, it will be replaced
@@ -3788,7 +3788,7 @@ public:
 	 * @param	ObjectInstance  the object that should be added as the corresopnding instance for ObjectSource
 	 * @param	InArchetype     Archetype of ObjectInstance
 	 */
-	void AddNewInstance(class UObject* ObjectInstance, class UObject* InArchetype = nullptr);
+	COREUOBJECT_API void AddNewInstance(class UObject* ObjectInstance, class UObject* InArchetype = nullptr);
 
 	/**
 	 * Retrieves a list of objects that have the specified Outer
@@ -3796,7 +3796,7 @@ public:
 	 * @param	SearchOuter		the object to retrieve object instances for
 	 * @param	out_Components	receives the list of objects contained by SearchOuter
 	 */
-	void RetrieveObjectInstances( class UObject* SearchOuter, TArray<class UObject*>& out_Objects );
+	COREUOBJECT_API void RetrieveObjectInstances( class UObject* SearchOuter, TArray<class UObject*>& out_Objects );
 
 	/**
 	 * Allows looping over instances that were created during this instancing.
@@ -3898,7 +3898,7 @@ private:
 	 *			if SourceComponent is not contained by SourceRoot, return INVALID_OBJECT, indicating that the that has SourceComponent as its ObjectArchetype, or NULL if SourceComponent is not contained within
 	 *			SourceRoot.
 	 */
-	class UObject* GetInstancedSubobject( class UObject* SourceSubobject, class UObject* CurrentValue, class UObject* CurrentObject, EInstancePropertyValueFlags Flags );
+	COREUOBJECT_API class UObject* GetInstancedSubobject( class UObject* SourceSubobject, class UObject* CurrentValue, class UObject* CurrentObject, EInstancePropertyValueFlags Flags );
 
 	/**
 	 * The root of the object tree that is the source used for instancing components;
@@ -4060,192 +4060,192 @@ struct TBaseStructure : TBaseStructureBase<T>
 
 template<> struct TBaseStructure<FIntPoint> 
 {
-	COREUOBJECT_API static UScriptStruct* Get(); 
+	static COREUOBJECT_API UScriptStruct* Get(); 
 };
 
 template<> struct TBaseStructure<FIntVector> 
 { 
-	COREUOBJECT_API static UScriptStruct* Get(); 
+	static COREUOBJECT_API UScriptStruct* Get(); 
 };
 
 template<> struct TBaseStructure<FIntVector4>
 {
-	COREUOBJECT_API static UScriptStruct* Get();
+	static COREUOBJECT_API UScriptStruct* Get();
 };
 
 template<> struct TBaseStructure<FLinearColor>
 {
-	COREUOBJECT_API static UScriptStruct* Get();
+	static COREUOBJECT_API UScriptStruct* Get();
 };
 
 template<> struct TBaseStructure<FColor>
 {
-	COREUOBJECT_API static UScriptStruct* Get();
+	static COREUOBJECT_API UScriptStruct* Get();
 };
 
 template<> struct TBaseStructure<FRandomStream>
 {
-	COREUOBJECT_API static UScriptStruct* Get();
+	static COREUOBJECT_API UScriptStruct* Get();
 };
 
 template<> struct TBaseStructure<FGuid>
 {
-	COREUOBJECT_API static UScriptStruct* Get();
+	static COREUOBJECT_API UScriptStruct* Get();
 };
 
 template<> struct TBaseStructure<FFallbackStruct>
 {
-	COREUOBJECT_API static UScriptStruct* Get();
+	static COREUOBJECT_API UScriptStruct* Get();
 };
 
 template<> struct TBaseStructure<FInterpCurvePointFloat>
 {
-	COREUOBJECT_API static UScriptStruct* Get();
+	static COREUOBJECT_API UScriptStruct* Get();
 };
 
 template<> struct TBaseStructure<FInterpCurvePointVector2D>
 {
-	COREUOBJECT_API static UScriptStruct* Get();
+	static COREUOBJECT_API UScriptStruct* Get();
 };
 
 template<> struct TBaseStructure<FInterpCurvePointVector>
 {
-	COREUOBJECT_API static UScriptStruct* Get();
+	static COREUOBJECT_API UScriptStruct* Get();
 };
 
 template<> struct TBaseStructure<FInterpCurvePointQuat>
 {
-	COREUOBJECT_API static UScriptStruct* Get();
+	static COREUOBJECT_API UScriptStruct* Get();
 };
 
 template<> struct TBaseStructure<FInterpCurvePointTwoVectors>
 {
-	COREUOBJECT_API static UScriptStruct* Get();
+	static COREUOBJECT_API UScriptStruct* Get();
 };
 
 template<> struct TBaseStructure<FInterpCurvePointLinearColor>
 {
-	COREUOBJECT_API static UScriptStruct* Get();
+	static COREUOBJECT_API UScriptStruct* Get();
 };
 
 template<> struct TBaseStructure<FFloatRangeBound>
 {
-	COREUOBJECT_API static UScriptStruct* Get();
+	static COREUOBJECT_API UScriptStruct* Get();
 };
 
 template<> struct TBaseStructure<FFloatRange>
 {
-	COREUOBJECT_API static UScriptStruct* Get();
+	static COREUOBJECT_API UScriptStruct* Get();
 };
 
 template<> struct TBaseStructure<FDoubleRangeBound>
 {
-	COREUOBJECT_API static UScriptStruct* Get();
+	static COREUOBJECT_API UScriptStruct* Get();
 };
 
 template<> struct TBaseStructure<FDoubleRange>
 {
-	COREUOBJECT_API static UScriptStruct* Get();
+	static COREUOBJECT_API UScriptStruct* Get();
 };
 
 template<> struct TBaseStructure<FInt32RangeBound>
 {
-	COREUOBJECT_API static UScriptStruct* Get();
+	static COREUOBJECT_API UScriptStruct* Get();
 };
 
 template<> struct TBaseStructure<FInt32Range>
 {
-	COREUOBJECT_API static UScriptStruct* Get();
+	static COREUOBJECT_API UScriptStruct* Get();
 };
 
 template<> struct TBaseStructure<FFloatInterval>
 {
-	COREUOBJECT_API static UScriptStruct* Get();
+	static COREUOBJECT_API UScriptStruct* Get();
 };
 
 template<> struct TBaseStructure<FDoubleInterval>
 {
-	COREUOBJECT_API static UScriptStruct* Get();
+	static COREUOBJECT_API UScriptStruct* Get();
 };
 
 template<> struct TBaseStructure<FInt32Interval>
 {
-	COREUOBJECT_API static UScriptStruct* Get();
+	static COREUOBJECT_API UScriptStruct* Get();
 };
 
 struct FFrameNumber;
 
 template<> struct TBaseStructure<FFrameNumber>
 {
-	COREUOBJECT_API static UScriptStruct* Get();
+	static COREUOBJECT_API UScriptStruct* Get();
 };
 
 struct FFrameTime;
 
 template<> struct TBaseStructure<FFrameTime>
 {
-	COREUOBJECT_API static UScriptStruct* Get();
+	static COREUOBJECT_API UScriptStruct* Get();
 };
 
 struct FSoftObjectPath;
 
 template<> struct TBaseStructure<FSoftObjectPath>
 {
-	COREUOBJECT_API static UScriptStruct* Get();
+	static COREUOBJECT_API UScriptStruct* Get();
 };
 
 struct FSoftClassPath;
 
 template<> struct TBaseStructure<FSoftClassPath>
 {
-	COREUOBJECT_API static UScriptStruct* Get();
+	static COREUOBJECT_API UScriptStruct* Get();
 };
 
 struct FPrimaryAssetType;
 
 template<> struct TBaseStructure<FPrimaryAssetType>
 {
-	COREUOBJECT_API static UScriptStruct* Get();
+	static COREUOBJECT_API UScriptStruct* Get();
 };
 
 struct FPrimaryAssetId;
 
 template<> struct TBaseStructure<FPrimaryAssetId>
 {
-	COREUOBJECT_API static UScriptStruct* Get();
+	static COREUOBJECT_API UScriptStruct* Get();
 };
 
 struct FDateTime;
 
 template<> struct TBaseStructure<FDateTime>
 {
-	COREUOBJECT_API static UScriptStruct* Get();
+	static COREUOBJECT_API UScriptStruct* Get();
 };
 
 struct FPolyglotTextData;
 
 template<> struct TBaseStructure<FPolyglotTextData>
 {
-	COREUOBJECT_API static UScriptStruct* Get();
+	static COREUOBJECT_API UScriptStruct* Get();
 };
 
 struct FAssetBundleData;
 
 template<> struct TBaseStructure<FAssetBundleData>
 {
-	COREUOBJECT_API static UScriptStruct* Get();
+	static COREUOBJECT_API UScriptStruct* Get();
 };
 
 template<> struct TBaseStructure<FTestUninitializedScriptStructMembersTest>
 {
-	COREUOBJECT_API static UScriptStruct* Get();
+	static COREUOBJECT_API UScriptStruct* Get();
 };
 
 struct FTopLevelAssetPath;
 
 template<> struct TBaseStructure<FTopLevelAssetPath>
 {
-	COREUOBJECT_API static UScriptStruct* Get();
+	static COREUOBJECT_API UScriptStruct* Get();
 };
 
 // TBaseStructure for explicit core variant types only. e.g. FVector3d returns "Vector3d" struct. 
diff --git a/Engine/Source/Runtime/CoreUObject/Public/UObject/ConstructorHelpers.h b/Engine/Source/Runtime/CoreUObject/Public/UObject/ConstructorHelpers.h
--- a/Engine/Source/Runtime/CoreUObject/Public/UObject/ConstructorHelpers.h
+++ b/Engine/Source/Runtime/CoreUObject/Public/UObject/ConstructorHelpers.h
@@ -74,7 +74,7 @@ namespace ConstructorHelpersInternal
 	}
 }
 
-struct COREUOBJECT_API ConstructorHelpers
+struct ConstructorHelpers
 {
 public:
 	template<class T>
@@ -187,7 +187,7 @@ public:
 
 public:
 	/** If there is an object class, strips it off. */
-	static void StripObjectClass( FString& PathName, bool bAssertOnBadPath = false );
+	static COREUOBJECT_API void StripObjectClass( FString& PathName, bool bAssertOnBadPath = false );
 
 private:
 	static void ValidateObject(UObject *Object, const FString& PathName, const TCHAR* ObjectToFind)
@@ -204,9 +204,9 @@ private:
 #endif
 	}
 
-	static void FailedToFind(const TCHAR* ObjectToFind);
-	static void CheckFoundViaRedirect(UObject *Object, const FString& PathName, const TCHAR* ObjectToFind);
-	static void CheckIfIsInConstructor(const TCHAR* ObjectToFind);
+	static COREUOBJECT_API void FailedToFind(const TCHAR* ObjectToFind);
+	static COREUOBJECT_API void CheckFoundViaRedirect(UObject *Object, const FString& PathName, const TCHAR* ObjectToFind);
+	static COREUOBJECT_API void CheckIfIsInConstructor(const TCHAR* ObjectToFind);
 };
 
 
diff --git a/Engine/Source/Runtime/CoreUObject/Public/UObject/CoreNet.h b/Engine/Source/Runtime/CoreUObject/Public/UObject/CoreNet.h
--- a/Engine/Source/Runtime/CoreUObject/Public/UObject/CoreNet.h
+++ b/Engine/Source/Runtime/CoreUObject/Public/UObject/CoreNet.h
@@ -66,7 +66,7 @@ DECLARE_DELEGATE_RetVal_OneParam( bool, FNetObjectIsDynamic, const UObject*);
 //
 // Information about a field.
 //
-class COREUOBJECT_API FFieldNetCache
+class FFieldNetCache
 {
 public:
 	FFieldVariant			Field;
@@ -84,12 +84,12 @@ public:
 //
 // Information about a class, cached for network coordination.
 //
-class COREUOBJECT_API FClassNetCache
+class FClassNetCache
 {
 	friend class FClassNetCacheMgr;
 public:
-	FClassNetCache();
-	FClassNetCache( const UClass* Class );
+	COREUOBJECT_API FClassNetCache();
+	COREUOBJECT_API FClassNetCache( const UClass* Class );
 
 	int32 GetMaxIndex() const
 	{
@@ -141,7 +141,7 @@ public:
 	const FClassNetCache* GetSuper() const { return Super; }
 	const TArray< FFieldNetCache >& GetFields() const { return Fields; }
 
-	void CountBytes(FArchive& Ar) const;
+	COREUOBJECT_API void CountBytes(FArchive& Ar) const;
 
 private:
 	int32								FieldsBase;
@@ -154,26 +154,26 @@ private:
 };
 
 
-class COREUOBJECT_API FClassNetCacheMgr
+class FClassNetCacheMgr
 {
 public:
 	FClassNetCacheMgr() : bDebugChecksum( false ), DebugChecksumIndent( 0 ) { }
 	~FClassNetCacheMgr() { ClearClassNetCache(); }
 
 	/** get the cached field to index mappings for the given class */
-	const FClassNetCache*	GetClassNetCache( UClass* Class );
-	void					ClearClassNetCache();
+	COREUOBJECT_API const FClassNetCache*	GetClassNetCache( UClass* Class );
+	COREUOBJECT_API void					ClearClassNetCache();
 
-	void				SortProperties( TArray< FProperty* >& Properties ) const;
-	uint32				SortedStructFieldsChecksum( const UStruct* Struct, uint32 Checksum ) const;
-	uint32				GetPropertyChecksum( const FProperty* Property, uint32 Checksum, const bool bIncludeChildren ) const;
-	uint32				GetFunctionChecksum( const UFunction* Function, uint32 Checksum ) const;
-	uint32				GetFieldChecksum( const UField* Field, uint32 Checksum ) const;
+	COREUOBJECT_API void				SortProperties( TArray< FProperty* >& Properties ) const;
+	COREUOBJECT_API uint32				SortedStructFieldsChecksum( const UStruct* Struct, uint32 Checksum ) const;
+	COREUOBJECT_API uint32				GetPropertyChecksum( const FProperty* Property, uint32 Checksum, const bool bIncludeChildren ) const;
+	COREUOBJECT_API uint32				GetFunctionChecksum( const UFunction* Function, uint32 Checksum ) const;
+	COREUOBJECT_API uint32				GetFieldChecksum( const UField* Field, uint32 Checksum ) const;
 
 	bool				bDebugChecksum;
 	int					DebugChecksumIndent;
 
-	void CountBytes(FArchive& Ar) const;
+	COREUOBJECT_API void CountBytes(FArchive& Ar) const;
 
 private:
 	TMap< TWeakObjectPtr< const UClass >, FClassNetCache* > ClassFieldIndices;
@@ -183,9 +183,9 @@ private:
 //
 // Maps objects and names to and from indices for network communication.
 //
-class COREUOBJECT_API UPackageMap : public UObject
+class UPackageMap : public UObject
 {
-	DECLARE_CLASS_INTRINSIC(UPackageMap, UObject, CLASS_Transient | CLASS_Abstract | 0, TEXT("/Script/CoreUObject"));
+	DECLARE_CASTED_CLASS_INTRINSIC_WITH_API(UPackageMap, UObject, CLASS_Transient | CLASS_Abstract | 0, TEXT("/Script/CoreUObject"), CASTCLASS_None, COREUOBJECT_API);
 
 	virtual bool		WriteObject( FArchive & Ar, UObject* InOuter, FNetworkGUID NetGUID, FString ObjName ) { return false; }
 
@@ -193,9 +193,9 @@ class COREUOBJECT_API UPackageMap : public UObject
 	virtual bool		SerializeObject( FArchive& Ar, UClass* InClass, UObject*& Obj, FNetworkGUID *OutNetGUID = NULL ) { return false; }
 
 	// @todo document
-	virtual bool		SerializeName( FArchive& Ar, FName& InName );
+	COREUOBJECT_API virtual bool		SerializeName( FArchive& Ar, FName& InName );
 
-	static bool			StaticSerializeName( FArchive& Ar, FName& InName );
+	static COREUOBJECT_API bool			StaticSerializeName( FArchive& Ar, FName& InName );
 
 	virtual UObject*	ResolvePathAndAssignNetGUID( const FNetworkGUID& NetGUID, const FString& PathName ) { return NULL; }
 
@@ -227,7 +227,7 @@ class COREUOBJECT_API UPackageMap : public UObject
 	virtual FNetworkGUID	GetNetGUIDFromObject( const UObject* InObject) const { return FNetworkGUID(); }
 	virtual bool			IsGUIDBroken( const FNetworkGUID& NetGUID, const bool bMustBeRegistered ) const { return false; }
 
-	virtual void Serialize(FArchive& Ar) override;
+	COREUOBJECT_API virtual void Serialize(FArchive& Ar) override;
 
 protected:
 	bool					bShouldTrackUnmappedGuids;
@@ -367,12 +367,12 @@ private:
  *	A bit writer that serializes FNames and UObject* through
  *	a network packagemap.
  */
-class COREUOBJECT_API FNetBitWriter : public FBitWriter
+class FNetBitWriter : public FBitWriter
 {
 public:
-	FNetBitWriter( UPackageMap * InPackageMap, int64 InMaxBits );
-	FNetBitWriter( int64 InMaxBits );
-	FNetBitWriter();
+	COREUOBJECT_API FNetBitWriter( UPackageMap * InPackageMap, int64 InMaxBits );
+	COREUOBJECT_API FNetBitWriter( int64 InMaxBits );
+	COREUOBJECT_API FNetBitWriter();
 
 	class UPackageMap * PackageMap;
 
@@ -380,14 +380,14 @@ public:
 	FNetTraceCollectorDoNotCopyWrapper TraceCollector;
 #endif
 
-	virtual FArchive& operator<<(FName& Name) override;
-	virtual FArchive& operator<<(UObject*& Object) override;
-	virtual FArchive& operator<<(FSoftObjectPath& Value) override;
-	virtual FArchive& operator<<(FSoftObjectPtr& Value) override;
-	virtual FArchive& operator<<(FObjectPtr& Value) override;
-	virtual FArchive& operator<<(struct FWeakObjectPtr& Value) override;
+	COREUOBJECT_API virtual FArchive& operator<<(FName& Name) override;
+	COREUOBJECT_API virtual FArchive& operator<<(UObject*& Object) override;
+	COREUOBJECT_API virtual FArchive& operator<<(FSoftObjectPath& Value) override;
+	COREUOBJECT_API virtual FArchive& operator<<(FSoftObjectPtr& Value) override;
+	COREUOBJECT_API virtual FArchive& operator<<(FObjectPtr& Value) override;
+	COREUOBJECT_API virtual FArchive& operator<<(struct FWeakObjectPtr& Value) override;
 
-	virtual void CountMemory(FArchive& Ar) const override;
+	COREUOBJECT_API virtual void CountMemory(FArchive& Ar) const override;
 };
 
 
@@ -396,7 +396,7 @@ public:
  *	A bit reader that serializes FNames and UObject* through
  *	a network packagemap.
  */
-class COREUOBJECT_API FNetBitReader : public FBitReader
+class FNetBitReader : public FBitReader
 {
 public:
 	UPackageMap*													PackageMap		= nullptr;
@@ -406,16 +406,16 @@ public:
 
 
 public:
-	FNetBitReader(UPackageMap* InPackageMap=nullptr, const uint8* Src=nullptr, int64 CountBits=0);
+	COREUOBJECT_API FNetBitReader(UPackageMap* InPackageMap=nullptr, const uint8* Src=nullptr, int64 CountBits=0);
 
-	virtual FArchive& operator<<(FName& Name) override;
-	virtual FArchive& operator<<(UObject*& Object) override;
-	virtual FArchive& operator<<(FSoftObjectPath& Value) override;
-	virtual FArchive& operator<<(FSoftObjectPtr& Value) override;
-	virtual FArchive& operator<<(FObjectPtr& Value) override;
-	virtual FArchive& operator<<(struct FWeakObjectPtr& Value) override;
+	COREUOBJECT_API virtual FArchive& operator<<(FName& Name) override;
+	COREUOBJECT_API virtual FArchive& operator<<(UObject*& Object) override;
+	COREUOBJECT_API virtual FArchive& operator<<(FSoftObjectPath& Value) override;
+	COREUOBJECT_API virtual FArchive& operator<<(FSoftObjectPtr& Value) override;
+	COREUOBJECT_API virtual FArchive& operator<<(FObjectPtr& Value) override;
+	COREUOBJECT_API virtual FArchive& operator<<(struct FWeakObjectPtr& Value) override;
 
-	virtual void CountMemory(FArchive& Ar) const override;
+	COREUOBJECT_API virtual void CountMemory(FArchive& Ar) const override;
 };
 
 bool FORCEINLINE NetworkGuidSetsAreSame( const TSet< FNetworkGUID >& A, const TSet< FNetworkGUID >& B )
@@ -478,7 +478,7 @@ struct FNetDeltaSerializeInfo;
  *
  * See notes in NetSerialization.h
  */
-class COREUOBJECT_API INetSerializeCB
+class INetSerializeCB
 {
 protected:
 
diff --git a/Engine/Source/Runtime/CoreUObject/Public/UObject/CoreRedirects.h b/Engine/Source/Runtime/CoreUObject/Public/UObject/CoreRedirects.h
--- a/Engine/Source/Runtime/CoreUObject/Public/UObject/CoreRedirects.h
+++ b/Engine/Source/Runtime/CoreUObject/Public/UObject/CoreRedirects.h
@@ -72,7 +72,7 @@ ENUM_CLASS_FLAGS(ECoreRedirectMatchFlags);
 /**
  * An object path extracted into component names for matching. TODO merge with FSoftObjectPath?
  */
-struct COREUOBJECT_API FCoreRedirectObjectName
+struct FCoreRedirectObjectName
 {
 	/** Raw name of object */
 	FName ObjectName;
@@ -93,19 +93,19 @@ struct COREUOBJECT_API FCoreRedirectObjectName
 
 	}
 
-	FCoreRedirectObjectName(const FTopLevelAssetPath& TopLevelAssetPath);
+	COREUOBJECT_API FCoreRedirectObjectName(const FTopLevelAssetPath& TopLevelAssetPath);
 
 	/** Construct from a path string, this handles full paths with packages, or partial paths without */
-	FCoreRedirectObjectName(const FString& InString);
+	COREUOBJECT_API FCoreRedirectObjectName(const FString& InString);
 
 	/** Construct from object in memory */
-	FCoreRedirectObjectName(const class UObject* Object);
+	COREUOBJECT_API FCoreRedirectObjectName(const class UObject* Object);
 
 	/** Creates FString version */
-	FString ToString() const;
+	COREUOBJECT_API FString ToString() const;
 
 	/** Sets back to invalid state */
-	void Reset();
+	COREUOBJECT_API void Reset();
 
 	/** Checks for exact equality */
 	bool operator==(const FCoreRedirectObjectName& Other) const
@@ -134,16 +134,16 @@ struct COREUOBJECT_API FCoreRedirectObjectName
 		CheckSubString = (1 << 2),
 	};
 	/** Returns true if the passed in name matches requirements. */
-	bool Matches(const FCoreRedirectObjectName& Other, EMatchFlags MatchFlags = EMatchFlags::None) const;
+	COREUOBJECT_API bool Matches(const FCoreRedirectObjectName& Other, EMatchFlags MatchFlags = EMatchFlags::None) const;
 
 	UE_DEPRECATED(5.1, "Use EMatchFlags::CheckSubString to pass in bCheckSubstring=true.")
-	bool Matches(const FCoreRedirectObjectName& Other, bool bCheckSubstring) const;
+	COREUOBJECT_API bool Matches(const FCoreRedirectObjectName& Other, bool bCheckSubstring) const;
 
 	/** Returns integer of degree of match. 0 if doesn't match at all, higher integer for better matches */
-	int32 MatchScore(const FCoreRedirectObjectName& Other) const;
+	COREUOBJECT_API int32 MatchScore(const FCoreRedirectObjectName& Other) const;
 
 	/** Fills in any empty fields on this with the corresponding fields from Other. */
-	void UnionFieldsInline(const FCoreRedirectObjectName& Other);
+	COREUOBJECT_API void UnionFieldsInline(const FCoreRedirectObjectName& Other);
 
 	/** Returns the name used as the key into the acceleration map */
 	FName GetSearchKey(ECoreRedirectFlags Type) const
@@ -171,20 +171,20 @@ struct COREUOBJECT_API FCoreRedirectObjectName
 	}
 
 	/** Returns true if all names have valid characters */
-	bool HasValidCharacters(ECoreRedirectFlags Type) const;
+	COREUOBJECT_API bool HasValidCharacters(ECoreRedirectFlags Type) const;
 
 	/** Expand OldName/NewName as needed */
-	static bool ExpandNames(const FString& FullString, FName& OutName, FName& OutOuter, FName &OutPackage);
+	static COREUOBJECT_API bool ExpandNames(const FString& FullString, FName& OutName, FName& OutOuter, FName &OutPackage);
 
 	/** Turn it back into an FString */
-	static FString CombineNames(FName NewName, FName NewOuter, FName NewPackage);
+	static COREUOBJECT_API FString CombineNames(FName NewName, FName NewOuter, FName NewPackage);
 };
 ENUM_CLASS_FLAGS(FCoreRedirectObjectName::EMatchFlags);
 
 /** 
  * A single redirection from an old name to a new name, parsed out of an ini file
  */
-struct COREUOBJECT_API FCoreRedirect
+struct FCoreRedirect
 {
 	/** Flags of this redirect */
 	ECoreRedirectFlags RedirectFlags;
@@ -216,29 +216,29 @@ struct COREUOBJECT_API FCoreRedirect
 	}
 
 	/** Normalizes NewName with data from OldName */
-	void NormalizeNewName();
+	COREUOBJECT_API void NormalizeNewName();
 
 	/** Parses a char buffer into the ValueChanges map */
-	const TCHAR* ParseValueChanges(const TCHAR* Buffer);
+	COREUOBJECT_API const TCHAR* ParseValueChanges(const TCHAR* Buffer);
 
 	/** Returns true if the passed in name and flags match requirements */
-	bool Matches(ECoreRedirectFlags InFlags, const FCoreRedirectObjectName& InName,
+	COREUOBJECT_API bool Matches(ECoreRedirectFlags InFlags, const FCoreRedirectObjectName& InName,
 		ECoreRedirectMatchFlags MatchFlags = ECoreRedirectMatchFlags::None) const;
 	/** Returns true if the passed in name matches requirements */
-	bool Matches(const FCoreRedirectObjectName& InName,
+	COREUOBJECT_API bool Matches(const FCoreRedirectObjectName& InName,
 		ECoreRedirectMatchFlags MatchFlags = ECoreRedirectMatchFlags::None) const;
 
 	/** Returns true if this has value redirects */
-	bool HasValueChanges() const;
+	COREUOBJECT_API bool HasValueChanges() const;
 
 	/** Returns true if this is a substring match */
-	bool IsSubstringMatch() const;
+	COREUOBJECT_API bool IsSubstringMatch() const;
 
 	/** Convert to new names based on mapping */
-	FCoreRedirectObjectName RedirectName(const FCoreRedirectObjectName& OldObjectName) const;
+	COREUOBJECT_API FCoreRedirectObjectName RedirectName(const FCoreRedirectObjectName& OldObjectName) const;
 
 	/** See if search criteria is identical */
-	bool IdenticalMatchRules(const FCoreRedirect& Other) const;
+	COREUOBJECT_API bool IdenticalMatchRules(const FCoreRedirect& Other) const;
 
 	/** Returns the name used as the key into the acceleration map */
 	FName GetSearchKey() const
@@ -250,26 +250,26 @@ struct COREUOBJECT_API FCoreRedirect
 /**
  * A container for all of the registered core-level redirects 
  */
-struct COREUOBJECT_API FCoreRedirects
+struct FCoreRedirects
 {
 	/** Run initialization steps that are needed before any data can be stored in FCoreRedirects. Reads can occur before this, but no redirects will exist and redirect queries will all return empty. */
-	static void Initialize();
+	static COREUOBJECT_API void Initialize();
 
 	/** Returns a redirected version of the object name. If there are no valid redirects, it will return the original name */
-	static FCoreRedirectObjectName GetRedirectedName(ECoreRedirectFlags Type,
+	static COREUOBJECT_API FCoreRedirectObjectName GetRedirectedName(ECoreRedirectFlags Type,
 		const FCoreRedirectObjectName& OldObjectName, ECoreRedirectMatchFlags MatchFlags = ECoreRedirectMatchFlags::None);
 
 	/** Returns map of String->String value redirects for the object name, or nullptr if none found */
-	static const TMap<FString, FString>* GetValueRedirects(ECoreRedirectFlags Type,
+	static COREUOBJECT_API const TMap<FString, FString>* GetValueRedirects(ECoreRedirectFlags Type,
 		const FCoreRedirectObjectName& OldObjectName, ECoreRedirectMatchFlags MatchFlags = ECoreRedirectMatchFlags::None);
 
 	/** Performs both a name redirect and gets a value redirect struct if it exists. Returns true if either redirect found */
-	static bool RedirectNameAndValues(ECoreRedirectFlags Type, const FCoreRedirectObjectName& OldObjectName,
+	static COREUOBJECT_API bool RedirectNameAndValues(ECoreRedirectFlags Type, const FCoreRedirectObjectName& OldObjectName,
 		FCoreRedirectObjectName& NewObjectName, const FCoreRedirect** FoundValueRedirect,
 		ECoreRedirectMatchFlags MatchFlags = ECoreRedirectMatchFlags::None);
 
 	/** Returns true if this name has been registered as explicitly missing */
-	static bool IsKnownMissing(ECoreRedirectFlags Type, const FCoreRedirectObjectName& ObjectName);
+	static COREUOBJECT_API bool IsKnownMissing(ECoreRedirectFlags Type, const FCoreRedirectObjectName& ObjectName);
 
 	/**
 	  * Adds the given combination of (Type, ObjectName, Channel) as a missing name; IsKnownMissing queries will now find it
@@ -278,7 +278,7 @@ struct COREUOBJECT_API FCoreRedirects
 	  * @param ObjectName The name of the object now known to be missing
 	  * @param Channel may be Option_MissingLoad or Option_None; used to distinguish between detected-at-runtime and specified-by-ini
 	  */
-	static bool AddKnownMissing(ECoreRedirectFlags Type, const FCoreRedirectObjectName& ObjectName, ECoreRedirectFlags Channel = ECoreRedirectFlags::Option_MissingLoad);
+	static COREUOBJECT_API bool AddKnownMissing(ECoreRedirectFlags Type, const FCoreRedirectObjectName& ObjectName, ECoreRedirectFlags Channel = ECoreRedirectFlags::Option_MissingLoad);
 
 	/**
 	  * Removes the given combination of (Type, ObjectName, Channel) as a missing name
@@ -287,25 +287,25 @@ struct COREUOBJECT_API FCoreRedirects
 	  * @param ObjectName The name of the object that has just been loaded.
 	  * @param Channel may be Option_MissingLoad or Option_None; used to distinguish between detected-at-runtime and specified-by-ini
 	  */
-	static bool RemoveKnownMissing(ECoreRedirectFlags Type, const FCoreRedirectObjectName& ObjectName, ECoreRedirectFlags Channel = ECoreRedirectFlags::Option_MissingLoad);
+	static COREUOBJECT_API bool RemoveKnownMissing(ECoreRedirectFlags Type, const FCoreRedirectObjectName& ObjectName, ECoreRedirectFlags Channel = ECoreRedirectFlags::Option_MissingLoad);
 
-	static void ClearKnownMissing(ECoreRedirectFlags Type, ECoreRedirectFlags Channel = ECoreRedirectFlags::Option_MissingLoad);
+	static COREUOBJECT_API void ClearKnownMissing(ECoreRedirectFlags Type, ECoreRedirectFlags Channel = ECoreRedirectFlags::Option_MissingLoad);
 
 	/** Returns list of names it may have been before */
-	static bool FindPreviousNames(ECoreRedirectFlags Type, const FCoreRedirectObjectName& NewObjectName, TArray<FCoreRedirectObjectName>& PreviousNames);
+	static COREUOBJECT_API bool FindPreviousNames(ECoreRedirectFlags Type, const FCoreRedirectObjectName& NewObjectName, TArray<FCoreRedirectObjectName>& PreviousNames);
 
 	/** Returns list of all core redirects that match requirements */
-	static bool GetMatchingRedirects(ECoreRedirectFlags Type, const FCoreRedirectObjectName& OldObjectName,
+	static COREUOBJECT_API bool GetMatchingRedirects(ECoreRedirectFlags Type, const FCoreRedirectObjectName& OldObjectName,
 		TArray<const FCoreRedirect*>& FoundRedirects, ECoreRedirectMatchFlags MatchFlags = ECoreRedirectMatchFlags::None);
 
 	/** Parse all redirects out of a given ini file */
-	static bool ReadRedirectsFromIni(const FString& IniName);
+	static COREUOBJECT_API bool ReadRedirectsFromIni(const FString& IniName);
 
 	/** Adds an array of redirects to global list */
-	static bool AddRedirectList(TArrayView<const FCoreRedirect> Redirects, const FString& SourceString);
+	static COREUOBJECT_API bool AddRedirectList(TArrayView<const FCoreRedirect> Redirects, const FString& SourceString);
 
 	/** Removes an array of redirects from global list */
-	static bool RemoveRedirectList(TArrayView<const FCoreRedirect> Redirects, const FString& SourceString);
+	static COREUOBJECT_API bool RemoveRedirectList(TArrayView<const FCoreRedirect> Redirects, const FString& SourceString);
 
 	/** Returns true if this has ever been initialized */
 	static bool IsInitialized() { return bInitialized; }
@@ -314,41 +314,41 @@ struct COREUOBJECT_API FCoreRedirects
 	static bool IsInDebugMode() { return bInDebugMode; }
 
 	/** Validate a named list of redirects */
-	static void ValidateRedirectList(TArrayView<const FCoreRedirect> Redirects, const FString& SourceString);
+	static COREUOBJECT_API void ValidateRedirectList(TArrayView<const FCoreRedirect> Redirects, const FString& SourceString);
 
 	/** Validates all known redirects and warn if they seem to point to missing things */
-	static void ValidateAllRedirects();
+	static COREUOBJECT_API void ValidateAllRedirects();
 
 	/** Gets map from config key -> Flags */
 	static const TMap<FName, ECoreRedirectFlags>& GetConfigKeyMap() { return ConfigKeyMap; }
 
 	/** Goes from the containing package and name of the type to the type flag */
-	static ECoreRedirectFlags GetFlagsForTypeName(FName PackageName, FName TypeName);
+	static COREUOBJECT_API ECoreRedirectFlags GetFlagsForTypeName(FName PackageName, FName TypeName);
 
 	/** Goes from UClass Type to the type flag */
-	static ECoreRedirectFlags GetFlagsForTypeClass(UClass *TypeClass);
+	static COREUOBJECT_API ECoreRedirectFlags GetFlagsForTypeClass(UClass *TypeClass);
 
 	/** Runs set of redirector tests, returns false on failure */
-	static bool RunTests();
+	static COREUOBJECT_API bool RunTests();
 
 private:
 	/** Static only class, never constructed */
-	FCoreRedirects();
+	COREUOBJECT_API FCoreRedirects();
 
 	/** Add a single redirect to a type map */
-	static bool AddSingleRedirect(const FCoreRedirect& NewRedirect, const FString& SourceString);
+	static COREUOBJECT_API bool AddSingleRedirect(const FCoreRedirect& NewRedirect, const FString& SourceString);
 
 	/** Remove a single redirect from a type map */
-	static bool RemoveSingleRedirect(const FCoreRedirect& OldRedirect, const FString& SourceString);
+	static COREUOBJECT_API bool RemoveSingleRedirect(const FCoreRedirect& OldRedirect, const FString& SourceString);
 
 	/** Add native redirects, called before ini is parsed for the first time */
-	static void RegisterNativeRedirects();
+	static COREUOBJECT_API void RegisterNativeRedirects();
 
 #if WITH_COREREDIRECTS_MULTITHREAD_WARNING
 	/** Mark that CoreRedirects is about to start being used from multiple threads, and writes to new types of redirects are no longer allowed.
 	  * ReadRedirectsFromIni and all other AddRedirectList calls must be called before this
 	  */
-	static void EnterMultithreadedPhase();
+	static COREUOBJECT_API void EnterMultithreadedPhase();
 #endif
 
 	/** There is one of these for each registered set of redirect flags */
@@ -359,21 +359,21 @@ private:
 	};
 
 	/** Whether this has been initialized at least once */
-	static bool bInitialized;
+	static COREUOBJECT_API bool bInitialized;
 
 	/** True if we are in debug mode that does extra validation */
-	static bool bInDebugMode;
+	static COREUOBJECT_API bool bInDebugMode;
 
 	/** True if we have done our initial validation. After initial validation, each change to redirects will validate independently */
-	static bool bValidatedOnce;
+	static COREUOBJECT_API bool bValidatedOnce;
 
 #if WITH_COREREDIRECTS_MULTITHREAD_WARNING
 	/** Whether CoreRedirects is now being used multithreaded and therefore does not support writes to RedirectTypeMap keyvalue pairs */
-	static bool bIsInMultithreadedPhase;
+	static COREUOBJECT_API bool bIsInMultithreadedPhase;
 #endif
 
 	/** Map from config name to flag */
-	static TMap<FName, ECoreRedirectFlags> ConfigKeyMap;
+	static COREUOBJECT_API TMap<FName, ECoreRedirectFlags> ConfigKeyMap;
 
 	/** Map from name of thing being mapped to full list. List must be filtered further */
 	struct FRedirectTypeMap
@@ -389,11 +389,11 @@ private:
 		TMap<ECoreRedirectFlags, FRedirectNameMap*> Map;
 		TArray<TPair<ECoreRedirectFlags, FRedirectNameMap>> FastIterable;
 	};
-	static FRedirectTypeMap RedirectTypeMap;
+	static COREUOBJECT_API FRedirectTypeMap RedirectTypeMap;
 
 	/**
 	 * Lock to protect multithreaded access to *KnownMissing functions, which can be called from the async loading threads. 
 	 * TODO: The KnownMissing functions use RedirectTypeMap, which is unguarded; there is race condition vulnerability if asyncloading thread is active before all categories are added to RedirectTypeMap.
 	 */
-	static FRWLock KnownMissingLock;
+	static COREUOBJECT_API FRWLock KnownMissingLock;
 };
diff --git a/Engine/Source/Runtime/CoreUObject/Public/UObject/DynamicallyTypedValue.h b/Engine/Source/Runtime/CoreUObject/Public/UObject/DynamicallyTypedValue.h
--- a/Engine/Source/Runtime/CoreUObject/Public/UObject/DynamicallyTypedValue.h
+++ b/Engine/Source/Runtime/CoreUObject/Public/UObject/DynamicallyTypedValue.h
@@ -56,9 +56,9 @@ namespace UE
 	};
 
 	// An value stored in some uninterpreted memory and a pointer to a type that contains methods to interpret it.
-	struct COREUOBJECT_API FDynamicallyTypedValue
+	struct FDynamicallyTypedValue
 	{
-		static FDynamicallyTypedValueType& NullType();
+		static COREUOBJECT_API FDynamicallyTypedValueType& NullType();
 
 		FDynamicallyTypedValue() { InitializeToNull(); }
 		FDynamicallyTypedValue(const FDynamicallyTypedValue& Copyee) { InitializeFromCopy(Copyee); }
@@ -185,4 +185,4 @@ namespace UE
 			}
 		}
 	};
-}
\ No newline at end of file
+}
diff --git a/Engine/Source/Runtime/CoreUObject/Public/UObject/EnumProperty.h b/Engine/Source/Runtime/CoreUObject/Public/UObject/EnumProperty.h
--- a/Engine/Source/Runtime/CoreUObject/Public/UObject/EnumProperty.h
+++ b/Engine/Source/Runtime/CoreUObject/Public/UObject/EnumProperty.h
@@ -25,57 +25,57 @@ class UStruct;
 namespace UECodeGen_Private { struct FEnumPropertyParams; }
 struct FPropertyTag;
 
-class COREUOBJECT_API FEnumProperty : public FProperty
+class FEnumProperty : public FProperty
 {
-	DECLARE_FIELD(FEnumProperty, FProperty, CASTCLASS_FEnumProperty)
+	DECLARE_FIELD_API(FEnumProperty, FProperty, CASTCLASS_FEnumProperty, COREUOBJECT_API)
 
 public:
-	FEnumProperty(FFieldVariant InOwner, const FName& InName, EObjectFlags InObjectFlags);
-	FEnumProperty(FFieldVariant InOwner, const FName& InName, EObjectFlags InObjectFlags, UEnum* InEnum);
-	FEnumProperty(FFieldVariant InOwner, const FName& InName, EObjectFlags InObjectFlags, int32 InOffset, EPropertyFlags InFlags, UEnum* InEnum);
+	COREUOBJECT_API FEnumProperty(FFieldVariant InOwner, const FName& InName, EObjectFlags InObjectFlags);
+	COREUOBJECT_API FEnumProperty(FFieldVariant InOwner, const FName& InName, EObjectFlags InObjectFlags, UEnum* InEnum);
+	COREUOBJECT_API FEnumProperty(FFieldVariant InOwner, const FName& InName, EObjectFlags InObjectFlags, int32 InOffset, EPropertyFlags InFlags, UEnum* InEnum);
 
 	/**
 	 * Constructor used for constructing compiled in properties
 	 * @param InOwner Owner of the property
 	 * @param Prop Pointer to the compiled in structure describing the property
 	 **/
-	FEnumProperty(FFieldVariant InOwner, const UECodeGen_Private::FEnumPropertyParams& Prop);
+	COREUOBJECT_API FEnumProperty(FFieldVariant InOwner, const UECodeGen_Private::FEnumPropertyParams& Prop);
 
 #if WITH_EDITORONLY_DATA
-	explicit FEnumProperty(UField* InField);
+	COREUOBJECT_API explicit FEnumProperty(UField* InField);
 #endif // WITH_EDITORONLY_DATA
-	virtual ~FEnumProperty();
+	COREUOBJECT_API virtual ~FEnumProperty();
 
 	// UObject interface
-	virtual void Serialize( FArchive& Ar ) override;
-	virtual void AddReferencedObjects(FReferenceCollector& Collector) override;
-	virtual void GetPreloadDependencies(TArray<UObject*>& OutDeps) override;
+	COREUOBJECT_API virtual void Serialize( FArchive& Ar ) override;
+	COREUOBJECT_API virtual void AddReferencedObjects(FReferenceCollector& Collector) override;
+	COREUOBJECT_API virtual void GetPreloadDependencies(TArray<UObject*>& OutDeps) override;
 	// End of UObject interface
 
 	// Field interface
-	virtual void PostDuplicate(const FField& InField) override;
-	virtual FField* GetInnerFieldByName(const FName& InName) override;
-	virtual void GetInnerFields(TArray<FField*>& OutFields) override;
+	COREUOBJECT_API virtual void PostDuplicate(const FField& InField) override;
+	COREUOBJECT_API virtual FField* GetInnerFieldByName(const FName& InName) override;
+	COREUOBJECT_API virtual void GetInnerFields(TArray<FField*>& OutFields) override;
 
 	// UField interface
-	virtual void AddCppProperty(FProperty* Property) override;
+	COREUOBJECT_API virtual void AddCppProperty(FProperty* Property) override;
 	// End of UField interface
 
 	// FProperty interface
-	virtual FString GetCPPMacroType( FString& ExtendedTypeText ) const  override;
-	virtual FString GetCPPType( FString* ExtendedTypeText, uint32 CPPExportFlags ) const override;
-	virtual FString GetCPPTypeForwardDeclaration() const override;
-	virtual void LinkInternal(FArchive& Ar) override;
-	virtual bool Identical( const void* A, const void* B, uint32 PortFlags ) const override;
-	virtual void SerializeItem(FStructuredArchive::FSlot Slot, void* Value, void const* Defaults) const override;
-	virtual bool NetSerializeItem( FArchive& Ar, UPackageMap* Map, void* Data, TArray<uint8> * MetaData = NULL ) const override;
-	virtual void ExportText_Internal( FString& ValueStr, const void* PropertyValueOrContainer, EPropertyPointerType PropertyPointerType, const void* DefaultValue, UObject* Parent, int32 PortFlags, UObject* ExportRootScope ) const override;
-	virtual const TCHAR* ImportText_Internal(const TCHAR* Buffer, void* ContainerOrPropertyPtr, EPropertyPointerType PropertyPointerType, UObject* OwnerObject, int32 PortFlags, FOutputDevice* ErrorText) const override;
-	virtual int32 GetMinAlignment() const override;
-	virtual bool SameType(const FProperty* Other) const override;
-	virtual EConvertFromTypeResult ConvertFromType(const FPropertyTag& Tag, FStructuredArchive::FSlot Slot, uint8* Data, UStruct* DefaultsStruct) override;
+	COREUOBJECT_API virtual FString GetCPPMacroType( FString& ExtendedTypeText ) const  override;
+	COREUOBJECT_API virtual FString GetCPPType( FString* ExtendedTypeText, uint32 CPPExportFlags ) const override;
+	COREUOBJECT_API virtual FString GetCPPTypeForwardDeclaration() const override;
+	COREUOBJECT_API virtual void LinkInternal(FArchive& Ar) override;
+	COREUOBJECT_API virtual bool Identical( const void* A, const void* B, uint32 PortFlags ) const override;
+	COREUOBJECT_API virtual void SerializeItem(FStructuredArchive::FSlot Slot, void* Value, void const* Defaults) const override;
+	COREUOBJECT_API virtual bool NetSerializeItem( FArchive& Ar, UPackageMap* Map, void* Data, TArray<uint8> * MetaData = NULL ) const override;
+	COREUOBJECT_API virtual void ExportText_Internal( FString& ValueStr, const void* PropertyValueOrContainer, EPropertyPointerType PropertyPointerType, const void* DefaultValue, UObject* Parent, int32 PortFlags, UObject* ExportRootScope ) const override;
+	COREUOBJECT_API virtual const TCHAR* ImportText_Internal(const TCHAR* Buffer, void* ContainerOrPropertyPtr, EPropertyPointerType PropertyPointerType, UObject* OwnerObject, int32 PortFlags, FOutputDevice* ErrorText) const override;
+	COREUOBJECT_API virtual int32 GetMinAlignment() const override;
+	COREUOBJECT_API virtual bool SameType(const FProperty* Other) const override;
+	COREUOBJECT_API virtual EConvertFromTypeResult ConvertFromType(const FPropertyTag& Tag, FStructuredArchive::FSlot Slot, uint8* Data, UStruct* DefaultsStruct) override;
 #if WITH_EDITORONLY_DATA
-	virtual void AppendSchemaHash(FBlake3& Builder, bool bSkipEditorOnly) const override;
+	COREUOBJECT_API virtual void AppendSchemaHash(FBlake3& Builder, bool bSkipEditorOnly) const override;
 #endif
 	// End of FProperty interface
 
@@ -106,10 +106,10 @@ public:
 	}
 
 	// Returns the number of bits required by NetSerializeItem to encode this enum, based on the maximum value
-	uint64 GetMaxNetSerializeBits() const;
+	COREUOBJECT_API uint64 GetMaxNetSerializeBits() const;
 
 private:
-	virtual uint32 GetValueTypeHashInternal(const void* Src) const override;
+	COREUOBJECT_API virtual uint32 GetValueTypeHashInternal(const void* Src) const override;
 
 	FNumericProperty* UnderlyingProp; // The property which represents the underlying type of the enum
 	TObjectPtr<UEnum> Enum; // The enum represented by this property
