[CoreUObject]

diff --git a/Engine/Source/Runtime/CoreUObject/Public/UObject/Field.h b/Engine/Source/Runtime/CoreUObject/Public/UObject/Field.h
--- a/Engine/Source/Runtime/CoreUObject/Public/UObject/Field.h
+++ b/Engine/Source/Runtime/CoreUObject/Public/UObject/Field.h
@@ -59,7 +59,7 @@ class UStruct;
   * Object representing a type of an FField struct. 
   * Mimics a subset of UObject reflection functions.
   */
-class COREUOBJECT_API FFieldClass
+class FFieldClass
 {
 	UE_NONCOPYABLE(FFieldClass);
 
@@ -81,17 +81,17 @@ class COREUOBJECT_API FFieldClass
 	FThreadSafeCounter UnqiueNameIndexCounter;
 
 	/** Creates a default object instance of this class */
-	FField* ConstructDefaultObject();
+	COREUOBJECT_API FField* ConstructDefaultObject();
 
 public:
 
 	/** Gets the list of all field classes in existance */
-	static TArray<FFieldClass*>& GetAllFieldClasses();
+	static COREUOBJECT_API TArray<FFieldClass*>& GetAllFieldClasses();
 	/** Gets a mapping of all field class names to the actuall class objects */
-	static TMap<FName, FFieldClass*>& GetNameToFieldClassMap();
+	static COREUOBJECT_API TMap<FName, FFieldClass*>& GetNameToFieldClassMap();
 
-	explicit FFieldClass(const TCHAR* InCPPName, uint64 InId, uint64 InCastFlags, FFieldClass* InSuperClass, FField* (*ConstructFnPtr)(const FFieldVariant&, const FName&, EObjectFlags));
-	~FFieldClass();
+	COREUOBJECT_API explicit FFieldClass(const TCHAR* InCPPName, uint64 InId, uint64 InCastFlags, FFieldClass* InSuperClass, FField* (*ConstructFnPtr)(const FFieldVariant&, const FName&, EObjectFlags));
+	COREUOBJECT_API ~FFieldClass();
 
 	inline FString GetName() const
 	{
@@ -122,8 +122,8 @@ public:
 		const uint64 OtherClassId = InClass->GetId();
 		return OtherClassId ? !!(CastFlags & OtherClassId) : IsChildOf_Walk(InClass);
 	}
-	FString GetDescription() const;
-	FText GetDisplayNameText() const;
+	COREUOBJECT_API FString GetDescription() const;
+	COREUOBJECT_API FText GetDisplayNameText() const;
 	FField* Construct(const FFieldVariant& InOwner, const FName& InName, EObjectFlags InFlags = RF_NoFlags) const
 	{
 		return ConstructFn(InOwner, InName, InFlags);
@@ -187,6 +187,9 @@ private:
 #endif
 
 #define DECLARE_FIELD(TClass, TSuperClass, TStaticFlags) \
+	DECLARE_FIELD_API(TClass, TSuperClass, TStaticFlags, NO_API)
+
+#define DECLARE_FIELD_API(TClass, TSuperClass, TStaticFlags, TRequiredAPI) \
 private: \
 	TClass& operator=(TClass&&);   \
 	TClass& operator=(const TClass&);   \
@@ -197,7 +200,7 @@ public: \
 		: Super(EC_InternalUseOnlyConstructor, InClass) \
 	{ \
 	} \
-	static FFieldClass* StaticClass(); \
+	static TRequiredAPI FFieldClass* StaticClass(); \
 	static FField* Construct(const FFieldVariant& InOwner, const FName& InName, EObjectFlags InObjectFlags); \
 	inline static constexpr uint64 StaticClassCastFlagsPrivate() \
 	{ \
@@ -258,7 +261,7 @@ class UObject;
  * Exposes common interface of FFields and UObjects for easier transition from UProperties to FProperties.
  * DO NOT ABUSE. IDEALLY THIS SHOULD ONLY BE FFIELD INTERNAL STRUCTURE FOR HOLDING A POINTER TO THE OWNER OF AN FFIELD.
  */
-class COREUOBJECT_API FFieldVariant
+class FFieldVariant
 {
 	union FFieldObjectUnion
 	{
@@ -304,13 +307,13 @@ public:
 	{
 		return !!ToUObjectUnsafe();
 	}
-	bool IsValidLowLevel() const;
+	COREUOBJECT_API bool IsValidLowLevel() const;
 	inline operator bool() const
 	{
 		return IsValid();
 	}
-	bool IsA(const UClass* InClass) const;
-	bool IsA(const FFieldClass* InClass) const;
+	COREUOBJECT_API bool IsA(const UClass* InClass) const;
+	COREUOBJECT_API bool IsA(const FFieldClass* InClass) const;
 	template <typename T>
 	bool IsA() const
 	{
@@ -373,15 +376,15 @@ public:
 	{
 		return Container.Field;
 	}
-	FFieldVariant GetOwnerVariant() const;
-	UClass* GetOwnerClass() const;
-	FString GetFullName() const;
-	FString GetPathName() const;
-	FString GetName() const;
-	FString GetClassName() const;
-	FName GetFName() const;
-	bool IsNative() const;
-	UPackage* GetOutermost() const;
+	COREUOBJECT_API FFieldVariant GetOwnerVariant() const;
+	COREUOBJECT_API UClass* GetOwnerClass() const;
+	COREUOBJECT_API FString GetFullName() const;
+	COREUOBJECT_API FString GetPathName() const;
+	COREUOBJECT_API FString GetName() const;
+	COREUOBJECT_API FString GetClassName() const;
+	COREUOBJECT_API FName GetFName() const;
+	COREUOBJECT_API bool IsNative() const;
+	COREUOBJECT_API UPackage* GetOutermost() const;
 
 	bool operator==(const FFieldVariant& Other) const
 	{
@@ -393,7 +396,7 @@ public:
 	}
 
 #if WITH_EDITORONLY_DATA
-	bool HasMetaData(const FName& Key) const;
+	COREUOBJECT_API bool HasMetaData(const FName& Key) const;
 #endif
 
 	/** Support comparison functions that make this usable as a KeyValue for a TSet<> */
@@ -406,7 +409,7 @@ public:
 /**
  * Base class of reflection data objects.
  */
-class COREUOBJECT_API FField
+class FField
 {
 	UE_NONCOPYABLE(FField);
 
@@ -419,7 +422,7 @@ public:
 	typedef FField BaseFieldClass;	
 	typedef FFieldClass FieldTypeClass;
 
-	static FFieldClass* StaticClass();
+	static COREUOBJECT_API FFieldClass* StaticClass();
 
 	inline static constexpr uint64 StaticClassCastFlagsPrivate()
 	{
@@ -443,39 +446,39 @@ public:
 	EObjectFlags FlagsPrivate;
 
 	// Constructors.
-	FField(EInternal InInernal, FFieldClass* InClass);
-	FField(FFieldVariant InOwner, const FName& InName, EObjectFlags InObjectFlags);
+	COREUOBJECT_API FField(EInternal InInernal, FFieldClass* InClass);
+	COREUOBJECT_API FField(FFieldVariant InOwner, const FName& InName, EObjectFlags InObjectFlags);
 #if WITH_EDITORONLY_DATA
-	explicit FField(UField* InField);
+	COREUOBJECT_API explicit FField(UField* InField);
 #endif // WITH_EDITORONLY_DATA
-	virtual ~FField();
+	COREUOBJECT_API virtual ~FField();
 
 	// Begin UObject interface: the following functions mimic UObject interface for easier transition from UProperties to FProperties
-	virtual void Serialize(FArchive& Ar);
-	virtual void PostLoad();
-	virtual void GetPreloadDependencies(TArray<UObject*>& OutDeps);
-	virtual void BeginDestroy();	
-	virtual void AddReferencedObjects(FReferenceCollector& Collector);
-	bool IsRooted() const;
-	bool IsNative() const;
-	bool IsValidLowLevel() const;
-	bool IsIn(const UObject* InOwner) const;
-	bool IsIn(const FField* InOwner) const;
-	FLinkerLoad* GetLinker() const;	
+	COREUOBJECT_API virtual void Serialize(FArchive& Ar);
+	COREUOBJECT_API virtual void PostLoad();
+	COREUOBJECT_API virtual void GetPreloadDependencies(TArray<UObject*>& OutDeps);
+	COREUOBJECT_API virtual void BeginDestroy();	
+	COREUOBJECT_API virtual void AddReferencedObjects(FReferenceCollector& Collector);
+	COREUOBJECT_API bool IsRooted() const;
+	COREUOBJECT_API bool IsNative() const;
+	COREUOBJECT_API bool IsValidLowLevel() const;
+	COREUOBJECT_API bool IsIn(const UObject* InOwner) const;
+	COREUOBJECT_API bool IsIn(const FField* InOwner) const;
+	COREUOBJECT_API FLinkerLoad* GetLinker() const;	
 	// End UObject interface
 
 	// Begin UField interface.
-	virtual void AddCppProperty(FProperty* Property);
-	virtual void Bind();
+	COREUOBJECT_API virtual void AddCppProperty(FProperty* Property);
+	COREUOBJECT_API virtual void Bind();
 	// End UField interface
 
 	/** Constructs a new field given its class */
-	static FField* Construct(const FFieldVariant& InOwner, const FName& InName, EObjectFlags InFlags);
+	static COREUOBJECT_API FField* Construct(const FFieldVariant& InOwner, const FName& InName, EObjectFlags InFlags);
 	/** Constructs a new field given the name of its class */
-	static FField* Construct(const FName& FieldTypeName, const FFieldVariant& InOwner, const FName& InName, EObjectFlags InFlags);
+	static COREUOBJECT_API FField* Construct(const FName& FieldTypeName, const FFieldVariant& InOwner, const FName& InName, EObjectFlags InFlags);
 
 	/** Fixups after duplicating a Field */
-	virtual void PostDuplicate(const FField& InField);
+	COREUOBJECT_API virtual void PostDuplicate(const FField& InField);
 
 protected:
 	/**
@@ -601,22 +604,22 @@ public:
 	}
 
 	/** Goes up the outer chain to look for a UClass */
-	UClass* GetOwnerClass() const;
+	COREUOBJECT_API UClass* GetOwnerClass() const;
 
 	/** Goes up the outer chain to look for a UStruct */
-	UStruct* GetOwnerStruct() const;
+	COREUOBJECT_API UStruct* GetOwnerStruct() const;
 
 	/** Goes up the outer chain to look for a UField */
-	UField* GetOwnerUField() const;
+	COREUOBJECT_API UField* GetOwnerUField() const;
 
 	/** Goes up the outer chain to look for the outermost package */
-	UPackage* GetOutermost() const;
+	COREUOBJECT_API UPackage* GetOutermost() const;
 
 	/** Goes up the outer chain to look for the outer of the specified type */
-	UObject* GetTypedOwner(UClass* Target) const;
+	COREUOBJECT_API UObject* GetTypedOwner(UClass* Target) const;
 
 	/** Goes up the outer chain to look for the outer of the specified type */
-	FField* GetTypedOwner(FFieldClass* Target) const;
+	COREUOBJECT_API FField* GetTypedOwner(FFieldClass* Target) const;
 
 	template <typename T>
 	T* GetOwner() const
@@ -679,18 +682,18 @@ public:
 		}
 	}
 
-	void Rename(const FName& NewName);
+	COREUOBJECT_API void Rename(const FName& NewName);
 
-	FString GetPathName(const UObject* StopOuter = nullptr) const;
-	void GetPathName(const UObject* StopOuter, FStringBuilderBase& ResultString) const;
-	FString GetFullName() const;
+	COREUOBJECT_API FString GetPathName(const UObject* StopOuter = nullptr) const;
+	COREUOBJECT_API void GetPathName(const UObject* StopOuter, FStringBuilderBase& ResultString) const;
+	COREUOBJECT_API FString GetFullName() const;
 
 	/**
 	 * Returns a human readable string that was assigned to this field at creation.
 	 * By default this is the same as GetName() but it can be overridden if that is an internal-only name.
 	 * This name is consistent in editor/cooked builds, is not localized, and is useful for data import/export.
 	 */
-	FString GetAuthoredName() const;
+	COREUOBJECT_API FString GetAuthoredName() const;
 
 	/** Returns an inner field by name if the field has any */
 	virtual FField* GetInnerFieldByName(const FName& InName)
@@ -717,14 +720,14 @@ public:
 	* @param	bStartWithOuter		whether to include this object's name in the returned string
 	* @return	string containing the path name for this object, minus the outermost-package's name
 	*/
-	FString GetFullGroupName(bool bStartWithOuter) const;
+	COREUOBJECT_API FString GetFullGroupName(bool bStartWithOuter) const;
 
 	/**
 	* Finds the localized display name or native display name as a fallback.
 	*
 	* @return The display name for this object.
 	*/
-	FText GetDisplayNameText() const;
+	COREUOBJECT_API FText GetDisplayNameText() const;
 
 	/**
 	* Finds the localized tooltip or native tooltip as a fallback.
@@ -733,7 +736,7 @@ public:
 	*
 	* @return The tooltip for this object.
 	*/
-	FText GetToolTipText(bool bShortTooltip = false) const;
+	COREUOBJECT_API FText GetToolTipText(bool bShortTooltip = false) const;
 
 	/**
 	* Determines if the property has any metadata associated with the key
@@ -750,8 +753,8 @@ public:
 	* @param Key The key to lookup in the metadata
 	* @return The value associated with the key if it exists, null otherwise
 	*/
-	const FString* FindMetaData(const TCHAR* Key) const;
-	const FString* FindMetaData(const FName& Key) const;
+	COREUOBJECT_API const FString* FindMetaData(const TCHAR* Key) const;
+	COREUOBJECT_API const FString* FindMetaData(const FName& Key) const;
 
 	/**
 	* Find the metadata value associated with the key
@@ -759,8 +762,8 @@ public:
 	* @param Key The key to lookup in the metadata
 	* @return The value associated with the key
 	*/
-	const FString& GetMetaData(const TCHAR* Key) const;
-	const FString& GetMetaData(const FName& Key) const;
+	COREUOBJECT_API const FString& GetMetaData(const TCHAR* Key) const;
+	COREUOBJECT_API const FString& GetMetaData(const FName& Key) const;
 
 	/**
 	* Find the metadata value associated with the key and localization namespace and key
@@ -770,8 +773,8 @@ public:
 	* @param LocalizationKey			Key to lookup in the localization manager
 	* @return							Localized metadata if available, defaults to whatever is provided via GetMetaData
 	*/
-	const FText GetMetaDataText(const TCHAR* MetaDataKey, const FString LocalizationNamespace = FString(), const FString LocalizationKey = FString()) const;
-	const FText GetMetaDataText(const FName& MetaDataKey, const FString LocalizationNamespace = FString(), const FString LocalizationKey = FString()) const;
+	COREUOBJECT_API const FText GetMetaDataText(const TCHAR* MetaDataKey, const FString LocalizationNamespace = FString(), const FString LocalizationKey = FString()) const;
+	COREUOBJECT_API const FText GetMetaDataText(const FName& MetaDataKey, const FString LocalizationNamespace = FString(), const FString LocalizationKey = FString()) const;
 
 	/**
 	* Sets the metadata value associated with the key
@@ -779,11 +782,11 @@ public:
 	* @param Key The key to lookup in the metadata
 	* @return The value associated with the key
 	*/
-	void SetMetaData(const TCHAR* Key, const TCHAR* InValue);
-	void SetMetaData(const FName& Key, const TCHAR* InValue);
+	COREUOBJECT_API void SetMetaData(const TCHAR* Key, const TCHAR* InValue);
+	COREUOBJECT_API void SetMetaData(const FName& Key, const TCHAR* InValue);
 
-	void SetMetaData(const TCHAR* Key, FString&& InValue);
-	void SetMetaData(const FName& Key, FString&& InValue);
+	COREUOBJECT_API void SetMetaData(const TCHAR* Key, FString&& InValue);
+	COREUOBJECT_API void SetMetaData(const FName& Key, FString&& InValue);
 
 	/**
 	* Find the metadata value associated with the key
@@ -871,36 +874,36 @@ public:
 	* @param Key The key to lookup in the metadata
 	* @return the class value stored in the metadata.
 	*/
-	UClass* GetClassMetaData(const TCHAR* Key) const;
-	UClass* GetClassMetaData(const FName& Key) const;
+	COREUOBJECT_API UClass* GetClassMetaData(const TCHAR* Key) const;
+	COREUOBJECT_API UClass* GetClassMetaData(const FName& Key) const;
 
 	/** Clear any metadata associated with the key */
-	void RemoveMetaData(const TCHAR* Key);
-	void RemoveMetaData(const FName& Key);
+	COREUOBJECT_API void RemoveMetaData(const TCHAR* Key);
+	COREUOBJECT_API void RemoveMetaData(const FName& Key);
 
 	/** Gets all metadata associated with this field */
-	const TMap<FName, FString>* GetMetaDataMap() const;
+	COREUOBJECT_API const TMap<FName, FString>* GetMetaDataMap() const;
 
 	/** Append the given metadata to this field */
-	void AppendMetaData(const TMap<FName, FString>& MetaDataMapToAppend);
+	COREUOBJECT_API void AppendMetaData(const TMap<FName, FString>& MetaDataMapToAppend);
 
 	/** Copies all metadata from Source Field to Dest Field */
-	static void CopyMetaData(const FField* InSourceField, FField* InDestField);
+	static COREUOBJECT_API void CopyMetaData(const FField* InSourceField, FField* InDestField);
 
 	/** Creates a new FField from existing UField */
-	static FField* CreateFromUField(UField* InField);
+	static COREUOBJECT_API FField* CreateFromUField(UField* InField);
 	
 	DECLARE_MULTICAST_DELEGATE_ThreeParams(FOnConvertCustomUFieldToFField, FFieldClass*, UField*, FField*&);
 	/** Gets a delegate to convert custom UField types to FFields */
-	static FOnConvertCustomUFieldToFField& GetConvertCustomUFieldToFFieldDelegate();
+	static COREUOBJECT_API FOnConvertCustomUFieldToFField& GetConvertCustomUFieldToFFieldDelegate();
 
 #endif // WITH_EDITORONLY_DATA
 
 	/** Duplicates an FField */
-	static FField* Duplicate(const FField* InField, FFieldVariant DestOwner, const FName DestName = NAME_None, EObjectFlags FlagMask = RF_AllFlags, EInternalObjectFlags InternalFlagsMask = EInternalObjectFlags::AllFlags);
+	static COREUOBJECT_API FField* Duplicate(const FField* InField, FFieldVariant DestOwner, const FName DestName = NAME_None, EObjectFlags FlagMask = RF_AllFlags, EInternalObjectFlags InternalFlagsMask = EInternalObjectFlags::AllFlags);
 
 	/** Generates a name for a Field of a given type. Each generated name is unique in the current runtime */
-	static FName GenerateFFieldName(FFieldVariant InOwner, FFieldClass* InClass);
+	static COREUOBJECT_API FName GenerateFFieldName(FFieldVariant InOwner, FFieldClass* InClass);
 };
 
 // Support for casting between different FFIeld types
@@ -1051,4 +1054,4 @@ inline FieldType* FindFProperty(const TCHAR* InFieldPath)
 
 #if UE_ENABLE_INCLUDE_ORDER_DEPRECATED_IN_5_2
 #include "CoreMinimal.h"
-#endif
\ No newline at end of file
+#endif
diff --git a/Engine/Source/Runtime/CoreUObject/Public/UObject/FieldPath.h b/Engine/Source/Runtime/CoreUObject/Public/UObject/FieldPath.h
--- a/Engine/Source/Runtime/CoreUObject/Public/UObject/FieldPath.h
+++ b/Engine/Source/Runtime/CoreUObject/Public/UObject/FieldPath.h
@@ -32,7 +32,7 @@ template <typename T> struct TIsPODType;
 template <typename T> struct TIsWeakPointerType;
 template <typename T> struct TIsZeroConstructType;
 
-struct COREUOBJECT_API FFieldPath
+struct FFieldPath
 {
 	friend struct FGCInternals;
 
@@ -100,9 +100,9 @@ private:
 
 #if WITH_EDITORONLY_DATA
 	/** Used to check if the serial number on the provided struct is identical to the one stored in this FFieldPath */
-	bool IsFieldPathSerialNumberIdentical(UStruct* InStruct) const;
+	COREUOBJECT_API bool IsFieldPathSerialNumberIdentical(UStruct* InStruct) const;
 	/** Gets the serial number stored on the provided struct */
-	int32 GetFieldPathSerialNumber(UStruct* InStruct) const;
+	COREUOBJECT_API int32 GetFieldPathSerialNumber(UStruct* InStruct) const;
 #endif
 
 	/** FOR INTERNAL USE ONLY: gets the pointer to the resolved field without trying to resolve it */
@@ -122,7 +122,7 @@ private:
 	 * @param InResolveType Type of the resolve operation
 	 * @return Resolved owner struct
 	 */
-	UStruct* TryToResolveOwnerFromStruct(UStruct* InCurrentStruct = nullptr, EPathResolveType InResolveType = FFieldPath::UseStructIfOuterNotFound) const;
+	COREUOBJECT_API UStruct* TryToResolveOwnerFromStruct(UStruct* InCurrentStruct = nullptr, EPathResolveType InResolveType = FFieldPath::UseStructIfOuterNotFound) const;
 	
 
 	/**
@@ -130,7 +130,7 @@ private:
 	 * @param InLinker the current linker load serializing this field path
 	 * @return Resolved owner struct
 	 */
-	UStruct* TryToResolveOwnerFromLinker(FLinkerLoad* InLinker) const;
+	COREUOBJECT_API UStruct* TryToResolveOwnerFromLinker(FLinkerLoad* InLinker) const;
 
 
 	/**
@@ -138,7 +138,7 @@ private:
 	 * @param InLinker the current linker load serializing this field path
 	 * @return Resulved owner struct
 	 */
-	UStruct* ConvertFromFullPath(FLinkerLoad* InLinker);
+	COREUOBJECT_API UStruct* ConvertFromFullPath(FLinkerLoad* InLinker);
 
 public:
 
@@ -150,17 +150,17 @@ public:
 	}
 
 #if WITH_EDITORONLY_DATA
-	FFieldPath(UField* InField, const FName& InPropertyTypeName);
+	COREUOBJECT_API FFieldPath(UField* InField, const FName& InPropertyTypeName);
 #endif
 
 	/** Generates path from the passed in field pointer */
-	void Generate(FField* InField);
+	COREUOBJECT_API void Generate(FField* InField);
 
 	/** Generates path from the passed in field pointer */
-	void Generate(const TCHAR* InFieldPathString);
+	COREUOBJECT_API void Generate(const TCHAR* InFieldPathString);
 
 #if WITH_EDITORONLY_DATA
-	void GenerateFromUField(UField* InField);
+	COREUOBJECT_API void GenerateFromUField(UField* InField);
 #endif
 
 	/**
@@ -169,7 +169,7 @@ public:
 	 * @param OutOwnerIndex ObjectIndex of the Owner UObject
 	 * @return Resolved field or null
 	 */
-	FField* TryToResolvePath(UStruct* InCurrentStruct, EPathResolveType InResolveType = FFieldPath::UseStructIfOuterNotFound) const;
+	COREUOBJECT_API FField* TryToResolvePath(UStruct* InCurrentStruct, EPathResolveType InResolveType = FFieldPath::UseStructIfOuterNotFound) const;
 
 	/**
 	 * Tries to resolve the path and caches the result
@@ -262,7 +262,7 @@ public:
 		return ResolvedOwner != InOther.ResolvedOwner || Path != InOther.Path;
 	}
 
-	FString ToString() const;
+	COREUOBJECT_API FString ToString() const;
 
 	COREUOBJECT_API friend FArchive& operator<<(FArchive& Ar, FFieldPath& InOutPropertyPath);
 
diff --git a/Engine/Source/Runtime/CoreUObject/Public/UObject/FieldPathProperty.h b/Engine/Source/Runtime/CoreUObject/Public/UObject/FieldPathProperty.h
--- a/Engine/Source/Runtime/CoreUObject/Public/UObject/FieldPathProperty.h
+++ b/Engine/Source/Runtime/CoreUObject/Public/UObject/FieldPathProperty.h
@@ -24,9 +24,9 @@ struct FPropertyTag;
 // need to break this out a different type so that the DECLARE_CASTED_CLASS_INTRINSIC macro can digest the comma
 typedef TProperty<FFieldPath, FProperty> FFieldPathProperty_Super;
 
-class COREUOBJECT_API FFieldPathProperty : public FFieldPathProperty_Super
+class FFieldPathProperty : public FFieldPathProperty_Super
 {
-	DECLARE_FIELD(FFieldPathProperty, FFieldPathProperty_Super, CASTCLASS_FFieldPathProperty)
+	DECLARE_FIELD_API(FFieldPathProperty, FFieldPathProperty_Super, CASTCLASS_FFieldPathProperty, COREUOBJECT_API)
 
 public:
 
@@ -40,7 +40,7 @@ public:
 	}
 
 	UE_DEPRECATED(5.1, "Compiled-in property constructor is deprecated, use other constructors instead.")
-	FFieldPathProperty(FFieldVariant InOwner, const FName& InName, EObjectFlags InObjectFlags, int32 InOffset, EPropertyFlags InFlags, FFieldClass* InPropertyClass)
+	COREUOBJECT_API FFieldPathProperty(FFieldVariant InOwner, const FName& InName, EObjectFlags InObjectFlags, int32 InOffset, EPropertyFlags InFlags, FFieldClass* InPropertyClass)
 PRAGMA_DISABLE_DEPRECATION_WARNINGS
 		: FFieldPathProperty_Super(InOwner, InName, InObjectFlags, InOffset, InFlags)
 PRAGMA_ENABLE_DEPRECATION_WARNINGS
@@ -56,30 +56,30 @@ PRAGMA_ENABLE_DEPRECATION_WARNINGS
 	FFieldPathProperty(FFieldVariant InOwner, const UECodeGen_Private::FFieldPathPropertyParams& Prop);
 
 #if WITH_EDITORONLY_DATA
-	explicit FFieldPathProperty(UField* InField);
+	COREUOBJECT_API explicit FFieldPathProperty(UField* InField);
 #endif // WITH_EDITORONLY_DATA
 
 	FFieldClass* PropertyClass;
 
 	// UHT interface
-	virtual FString GetCPPMacroType(FString& ExtendedTypeText) const  override;
-	virtual FString GetCPPTypeForwardDeclaration() const override;
-	virtual FString GetCPPType(FString* ExtendedTypeText = nullptr, uint32 CPPExportFlags = 0) const override;
+	COREUOBJECT_API virtual FString GetCPPMacroType(FString& ExtendedTypeText) const  override;
+	COREUOBJECT_API virtual FString GetCPPTypeForwardDeclaration() const override;
+	COREUOBJECT_API virtual FString GetCPPType(FString* ExtendedTypeText = nullptr, uint32 CPPExportFlags = 0) const override;
 	// End of UHT interface
 
 	// FProperty interface
-	virtual EConvertFromTypeResult ConvertFromType(const FPropertyTag& Tag, FStructuredArchive::FSlot Slot, uint8* Data, UStruct* DefaultsStruct) override;
-	virtual bool Identical( const void* A, const void* B, uint32 PortFlags ) const override;
-	virtual void SerializeItem(FStructuredArchive::FSlot Slot, void* Value, void const* Defaults) const override;
-	virtual void ExportText_Internal( FString& ValueStr, const void* PropertyValueOrContainer, EPropertyPointerType PropertyPointerType, const void* DefaultValue, UObject* Parent, int32 PortFlags, UObject* ExportRootScope ) const override;
-	virtual const TCHAR* ImportText_Internal(const TCHAR* Buffer, void* ContainerOrPropertyPtr, EPropertyPointerType PropertyPointerType, UObject* OwnerObject, int32 PortFlags, FOutputDevice* ErrorText) const override;
-	virtual void Serialize(FArchive& Ar) override;
-	virtual bool ContainsObjectReference(TArray<const FStructProperty*>& EncounteredStructProps, EPropertyObjectReferenceType InReferenceType = EPropertyObjectReferenceType::Strong) const override;
-	virtual void EmitReferenceInfo(UE::GC::FSchemaBuilder& Schema, int32 BaseOffset, TArray<const FStructProperty*>& EncounteredStructProps, UE::GC::FPropertyStack& DebugPath) override;
-	virtual bool SupportsNetSharedSerialization() const override;
+	COREUOBJECT_API virtual EConvertFromTypeResult ConvertFromType(const FPropertyTag& Tag, FStructuredArchive::FSlot Slot, uint8* Data, UStruct* DefaultsStruct) override;
+	COREUOBJECT_API virtual bool Identical( const void* A, const void* B, uint32 PortFlags ) const override;
+	COREUOBJECT_API virtual void SerializeItem(FStructuredArchive::FSlot Slot, void* Value, void const* Defaults) const override;
+	COREUOBJECT_API virtual void ExportText_Internal( FString& ValueStr, const void* PropertyValueOrContainer, EPropertyPointerType PropertyPointerType, const void* DefaultValue, UObject* Parent, int32 PortFlags, UObject* ExportRootScope ) const override;
+	COREUOBJECT_API virtual const TCHAR* ImportText_Internal(const TCHAR* Buffer, void* ContainerOrPropertyPtr, EPropertyPointerType PropertyPointerType, UObject* OwnerObject, int32 PortFlags, FOutputDevice* ErrorText) const override;
+	COREUOBJECT_API virtual void Serialize(FArchive& Ar) override;
+	COREUOBJECT_API virtual bool ContainsObjectReference(TArray<const FStructProperty*>& EncounteredStructProps, EPropertyObjectReferenceType InReferenceType = EPropertyObjectReferenceType::Strong) const override;
+	COREUOBJECT_API virtual void EmitReferenceInfo(UE::GC::FSchemaBuilder& Schema, int32 BaseOffset, TArray<const FStructProperty*>& EncounteredStructProps, UE::GC::FPropertyStack& DebugPath) override;
+	COREUOBJECT_API virtual bool SupportsNetSharedSerialization() const override;
 	// End of FProperty interface
 
-	static FString RedirectFieldPathName(const FString& InPathName);
+	static COREUOBJECT_API FString RedirectFieldPathName(const FString& InPathName);
 };
 
 #if UE_ENABLE_INCLUDE_ORDER_DEPRECATED_IN_5_2
diff --git a/Engine/Source/Runtime/CoreUObject/Public/UObject/GCObject.h b/Engine/Source/Runtime/CoreUObject/Public/UObject/GCObject.h
--- a/Engine/Source/Runtime/CoreUObject/Public/UObject/GCObject.h
+++ b/Engine/Source/Runtime/CoreUObject/Public/UObject/GCObject.h
@@ -109,17 +109,17 @@ public:
  * non-UObject classes. It is an abstract base class requiring you to implement
  * the AddReferencedObjects() method.
  */
-class COREUOBJECT_API FGCObject
+class FGCObject
 {
 public:
 	/**
 	 * The static object referencer object that is shared across all
 	 * garbage collectible non-UObject objects.
 	 */
-	static UGCObjectReferencer* GGCObjectReferencer;
+	static COREUOBJECT_API UGCObjectReferencer* GGCObjectReferencer;
 
 	/** Initializes the global object referencer and adds it to the root set. */
-	static void StaticInit();
+	static COREUOBJECT_API void StaticInit();
 
 	/**
 	 * Tells the global object that forwards AddReferencedObjects calls on to objects
@@ -155,7 +155,7 @@ public:
 		 */
 		AddStableNativeReferencesOnly = 1 << 1,
 	};
-	explicit FGCObject(EFlags Flags);
+	COREUOBJECT_API explicit FGCObject(EFlags Flags);
 
 	virtual ~FGCObject()
 	{
@@ -166,10 +166,10 @@ public:
 	FGCObject& operator=(FGCObject&&) {return *this;}
 
 	/** Register with GC, only needed if constructed with EFlags::RegisterLater or after unregistering */
-	void RegisterGCObject();
+	COREUOBJECT_API void RegisterGCObject();
 
 	/** Unregister ahead of destruction. Safe to call multiple times. */
-	void UnregisterGCObject();
+	COREUOBJECT_API void UnregisterGCObject();
 
 	/**
 	 * Pure virtual that must be overloaded by the inheriting class. Use this
@@ -190,7 +190,7 @@ public:
 
 private:
 	friend UGCObjectReferencer;
-	static const TCHAR* UnknownGCObjectName;
+	static COREUOBJECT_API const TCHAR* UnknownGCObjectName;
 
 	const bool bCanMakeInitialReferences = false;
 	bool bReferenceAdded = false;	
diff --git a/Engine/Source/Runtime/CoreUObject/Public/UObject/GCObjectInfo.h b/Engine/Source/Runtime/CoreUObject/Public/UObject/GCObjectInfo.h
--- a/Engine/Source/Runtime/CoreUObject/Public/UObject/GCObjectInfo.h
+++ b/Engine/Source/Runtime/CoreUObject/Public/UObject/GCObjectInfo.h
@@ -24,7 +24,7 @@
  * It's purpose is to avoid holding onto direct references to UObjects which may have already been Garbage Collected.
  * FGCObjectInfo interface mimics that of UObject.
  **/
-class COREUOBJECT_API FGCObjectInfo
+class FGCObjectInfo
 {
 public:
 
@@ -62,7 +62,7 @@ public:
 	 * Tries to find an existing object matching the stored Path.
 	 * @returns Pointer to an object this info struct represents if the object is still alive, nullptr otherwise
 	 */
-	UObject* TryResolveObject();
+	COREUOBJECT_API UObject* TryResolveObject();
 
 	FGCObjectInfo* GetClass() const
 	{
@@ -136,9 +136,9 @@ public:
 		return bDisregardForGC;
 	}
 
-	void GetPathName(FStringBuilderBase& ResultString) const;
-	FString GetPathName() const;
-	FString GetFullName() const;
+	COREUOBJECT_API void GetPathName(FStringBuilderBase& ResultString) const;
+	COREUOBJECT_API FString GetPathName() const;
+	COREUOBJECT_API FString GetFullName() const;
 
 	friend uint32 GetTypeHash(const FGCObjectInfo& Info)
 	{
@@ -146,5 +146,5 @@ public:
 	}
 
 	/** Helper function for adding info about an UObject into UObject to FGCObjectInfo map */
-	static FGCObjectInfo* FindOrAddInfoHelper(const UObject* InObject, TMap<const UObject*, FGCObjectInfo*>& InOutObjectToInfoMap);
+	static COREUOBJECT_API FGCObjectInfo* FindOrAddInfoHelper(const UObject* InObject, TMap<const UObject*, FGCObjectInfo*>& InOutObjectToInfoMap);
 };
diff --git a/Engine/Source/Runtime/CoreUObject/Public/UObject/GarbageCollection.h b/Engine/Source/Runtime/CoreUObject/Public/UObject/GarbageCollection.h
--- a/Engine/Source/Runtime/CoreUObject/Public/UObject/GarbageCollection.h
+++ b/Engine/Source/Runtime/CoreUObject/Public/UObject/GarbageCollection.h
@@ -79,11 +79,11 @@ private:
 } // namespace UE::GC
 
 /** Prevent GC from running in the current scope */
-class COREUOBJECT_API FGCScopeGuard
+class FGCScopeGuard
 {
 public:
-	FGCScopeGuard();
-	~FGCScopeGuard();
+	COREUOBJECT_API FGCScopeGuard();
+	COREUOBJECT_API ~FGCScopeGuard();
 };
 
 class FGCObject;
diff --git a/Engine/Source/Runtime/CoreUObject/Public/UObject/GarbageCollectionHistory.h b/Engine/Source/Runtime/CoreUObject/Public/UObject/GarbageCollectionHistory.h
--- a/Engine/Source/Runtime/CoreUObject/Public/UObject/GarbageCollectionHistory.h
+++ b/Engine/Source/Runtime/CoreUObject/Public/UObject/GarbageCollectionHistory.h
@@ -52,34 +52,34 @@ struct FGCSnapshot
 /**
  * Garbage Collector History. Holds snapshots of a number of previous GC runs.
  **/
-class COREUOBJECT_API FGCHistory
+class FGCHistory
 {
 	/** Captured snapshots of previous GC runs. This is essentially a preallocated ring buffer. */
 	TArray<FGCSnapshot> Snapshots;
 	/** Index of the last recorded snapshot in the Snapshots array */
 	int32 MostRecentSnapshotIndex = -1;
 
-	~FGCHistory();
+	COREUOBJECT_API ~FGCHistory();
 
-	void Cleanup(FGCSnapshot& InShapshot);	
-	void MergeArrayStructHistory(TMap<const UObject*, TArray<FGCDirectReference>*>& History, FGCSnapshot& Snapshot);
+	COREUOBJECT_API void Cleanup(FGCSnapshot& InShapshot);	
+	COREUOBJECT_API void MergeArrayStructHistory(TMap<const UObject*, TArray<FGCDirectReference>*>& History, FGCSnapshot& Snapshot);
 
 public:
 
 	/** Gets the FGCHistory singleton. */
-	static FGCHistory& Get();
+	static COREUOBJECT_API FGCHistory& Get();
 
 	/** Copies information generated by Garbage Collector into the next snapshot struct. */
-	void Update(TConstArrayView<TUniquePtr<UE::GC::FWorkerContext>> Contexts);
+	COREUOBJECT_API void Update(TConstArrayView<TUniquePtr<UE::GC::FWorkerContext>> Contexts);
 
 	/** Frees memory allocated by GC history. */
-	void Cleanup();
+	COREUOBJECT_API void Cleanup();
 
 	/** 
 	 * Sets how many GC runs the history can record. 
 	 * @param HistorySize The number of GC runs the history can record. If 0 GC history will be disabled.
 	 */
-	void SetHistorySize(int32 HistorySize);
+	COREUOBJECT_API void SetHistorySize(int32 HistorySize);
 
 	FORCEINLINE int32 GetHistorySize() const
 	{
@@ -96,7 +96,7 @@ public:
 	 * @param HistoryLevel Index of a previous run: 0 - the last one, 1 - the one before and up to (GetHistorySize() - 1). Negative values are also accepted.
 	 * @returns Pointer to a prevous GC run snapshot or null if there's not been that many recorded GC runs yet
 	 */
-	FGCSnapshot* GetSnapshot(int32 HistoryLevel);
+	COREUOBJECT_API FGCSnapshot* GetSnapshot(int32 HistoryLevel);
 
 	/**
 	 * Returns the snapshot of the last GC run. See GetSnapshot(int32)
@@ -107,7 +107,7 @@ public:
 	}
 
 	/** Returns the number of bytes allocated by GC history */
-	int64 GetAllocatedSize() const;
+	COREUOBJECT_API int64 GetAllocatedSize() const;
 };
 
-#endif // ENABLE_GC_HISTORY
\ No newline at end of file
+#endif // ENABLE_GC_HISTORY
diff --git a/Engine/Source/Runtime/CoreUObject/Public/UObject/Interface.h b/Engine/Source/Runtime/CoreUObject/Public/UObject/Interface.h
--- a/Engine/Source/Runtime/CoreUObject/Public/UObject/Interface.h
+++ b/Engine/Source/Runtime/CoreUObject/Public/UObject/Interface.h
@@ -12,12 +12,12 @@
  *
  */
 
-class COREUOBJECT_API UInterface : public UObject
+class UInterface : public UObject
 {
-	DECLARE_CLASS_INTRINSIC(UInterface, UObject, CLASS_Interface | CLASS_Abstract, TEXT("/Script/CoreUObject"))
+	DECLARE_CASTED_CLASS_INTRINSIC_WITH_API(UInterface, UObject, CLASS_Interface | CLASS_Abstract, TEXT("/Script/CoreUObject"), CASTCLASS_None, COREUOBJECT_API)
 };
 
-class COREUOBJECT_API IInterface
+class IInterface
 {
 protected:
 
diff --git a/Engine/Source/Runtime/CoreUObject/Public/UObject/LazyObjectPtr.h b/Engine/Source/Runtime/CoreUObject/Public/UObject/LazyObjectPtr.h
--- a/Engine/Source/Runtime/CoreUObject/Public/UObject/LazyObjectPtr.h
+++ b/Engine/Source/Runtime/CoreUObject/Public/UObject/LazyObjectPtr.h
@@ -24,7 +24,7 @@ template <typename T> struct TIsWeakPointerType;
 /**
  * Wrapper structure for a GUID that uniquely identifies a UObject
  */
-struct COREUOBJECT_API FUniqueObjectGuid
+struct FUniqueObjectGuid
 {
 	FUniqueObjectGuid()
 	{}
@@ -40,23 +40,23 @@ struct COREUOBJECT_API FUniqueObjectGuid
 	}
 
 	/** Construct from an existing object */
-	explicit FUniqueObjectGuid(const class UObject* InObject);
+	COREUOBJECT_API explicit FUniqueObjectGuid(const class UObject* InObject);
 
 	/** Converts into a string */
-	FString ToString() const;
+	COREUOBJECT_API FString ToString() const;
 
 	/** Converts from a string */
-	void FromString(const FString& From);
+	COREUOBJECT_API void FromString(const FString& From);
 
 	/** Fixes up this UniqueObjectID to add or remove the PIE prefix depending on what is currently active */
-	FUniqueObjectGuid FixupForPIE(int32 PlayInEditorID = GPlayInEditorID) const;
+	COREUOBJECT_API FUniqueObjectGuid FixupForPIE(int32 PlayInEditorID = GPlayInEditorID) const;
 
 	/**
 	 * Attempts to find a currently loaded object that matches this object ID
 	 *
 	 * @return Found UObject, or nullptr if not currently loaded
 	 */
-	class UObject *ResolveObject() const;
+	COREUOBJECT_API class UObject *ResolveObject() const;
 
 	/** Test if this can ever point to a live UObject */
 	FORCEINLINE bool IsValid() const
@@ -112,7 +112,7 @@ struct COREUOBJECT_API FUniqueObjectGuid
 		return 0;
 	}
 
-	static FUniqueObjectGuid GetOrCreateIDForObject(const class UObject *Object);
+	static COREUOBJECT_API FUniqueObjectGuid GetOrCreateIDForObject(const class UObject *Object);
 
 private:
 	/** Guid representing the object, should be unique */
diff --git a/Engine/Source/Runtime/CoreUObject/Public/UObject/LinkerInstancingContext.h b/Engine/Source/Runtime/CoreUObject/Public/UObject/LinkerInstancingContext.h
--- a/Engine/Source/Runtime/CoreUObject/Public/UObject/LinkerInstancingContext.h
+++ b/Engine/Source/Runtime/CoreUObject/Public/UObject/LinkerInstancingContext.h
@@ -7,7 +7,7 @@
 /**
  * Helper class to map between an original package and an instance of it (including world partition cells).
  */
-class COREUOBJECT_API FLinkerInstancedPackageMap
+class FLinkerInstancedPackageMap
 {
 public:
 	enum class EInstanceMappingDirection : uint8
@@ -45,11 +45,11 @@ public:
 	}
 
 	/** Add a mapping from a package name to a new package name. There should be no separators (. or :) in these strings. */
-	void AddPackageMapping(FName Original, FName Instanced);
+	COREUOBJECT_API void AddPackageMapping(FName Original, FName Instanced);
 
-	void BuildPackageMapping(FName Original, FName Instanced, const bool bBuildWorldPartitionCellMapping = true);
+	COREUOBJECT_API void BuildPackageMapping(FName Original, FName Instanced, const bool bBuildWorldPartitionCellMapping = true);
 
-	bool FixupSoftObjectPath(FSoftObjectPath& InOutSoftObjectPath) const;
+	COREUOBJECT_API bool FixupSoftObjectPath(FSoftObjectPath& InOutSoftObjectPath) const;
 
 private:
 	friend class FLinkerInstancingContext;
diff --git a/Engine/Source/Runtime/CoreUObject/Public/UObject/MetaData.h b/Engine/Source/Runtime/CoreUObject/Public/UObject/MetaData.h
--- a/Engine/Source/Runtime/CoreUObject/Public/UObject/MetaData.h
+++ b/Engine/Source/Runtime/CoreUObject/Public/UObject/MetaData.h
@@ -22,9 +22,9 @@ class UPackage;
 /**
  * An object that holds a map of key/value pairs. 
  */
-class COREUOBJECT_API UMetaData : public UObject
+class UMetaData : public UObject
 {
-	DECLARE_CLASS_INTRINSIC(UMetaData, UObject, CLASS_MatchedSerializers, TEXT("/Script/CoreUObject"))
+	DECLARE_CASTED_CLASS_INTRINSIC_WITH_API(UMetaData, UObject, CLASS_MatchedSerializers, TEXT("/Script/CoreUObject"), CASTCLASS_None, COREUOBJECT_API);
 
 public:
 	/**
@@ -47,7 +47,7 @@ public:
 	 * @param Key The key to lookup
 	 * @return The value if found, otherwise an empty string
 	 */
-	const FString& GetValue(const UObject* Object, const TCHAR* Key);
+	COREUOBJECT_API const FString& GetValue(const UObject* Object, const TCHAR* Key);
 	
 	/**
 	 * Return the value for the given key in the given property
@@ -55,7 +55,7 @@ public:
 	 * @param Key The key to lookup
 	 * @return The value if found, otherwise an empty string
 	 */
-	const FString& GetValue(const UObject* Object, FName Key);
+	COREUOBJECT_API const FString& GetValue(const UObject* Object, FName Key);
 
 	/**
 	 * Return whether or not the Key is in the meta data
@@ -85,7 +85,7 @@ public:
 	 * @param Key The key to query for existence
 	 * @return true if found
 	 */
-	const FString* FindValue(const UObject* Object, const TCHAR* Key);
+	COREUOBJECT_API const FString* FindValue(const UObject* Object, const TCHAR* Key);
 
 	/**
 	 * Returns the value for a the given key if it exists, null otherwise
@@ -93,28 +93,28 @@ public:
 	 * @param Key The key to query for existence
 	 * @return true if found
 	 */
-	const FString* FindValue(const UObject* Object, FName Key);
+	COREUOBJECT_API const FString* FindValue(const UObject* Object, FName Key);
 
 	/**
 	 * Is there any metadata for this property?
 	 * @param Object the object to lookup the metadata for
 	 * @return TrUE if the object has any metadata at all
 	 */
-	bool HasObjectValues(const UObject* Object);
+	COREUOBJECT_API bool HasObjectValues(const UObject* Object);
 
 	/**
 	 * Set the key/value pair in the Property's metadata
 	 * @param Object the object to set the metadata for
 	 * @Values The metadata key/value pairs
 	 */
-	void SetObjectValues(const UObject* Object, const TMap<FName, FString>& Values);
+	COREUOBJECT_API void SetObjectValues(const UObject* Object, const TMap<FName, FString>& Values);
 
 	/**
 	 * Set the key/value pair in the Property's metadata
 	 * @param Object the object to set the metadata for
 	 * @Values The metadata key/value pairs
 	 */
-	void SetObjectValues(const UObject* Object, TMap<FName, FString>&& Values);
+	COREUOBJECT_API void SetObjectValues(const UObject* Object, TMap<FName, FString>&& Values);
 
 	/**
 	 * Set the key/value pair in the Object's metadata
@@ -122,7 +122,7 @@ public:
 	 * @param Key A key to set the data for
 	 * @param Value The value to set for the key
 	 */
-	void SetValue(const UObject* Object, const TCHAR* Key, const TCHAR* Value);
+	COREUOBJECT_API void SetValue(const UObject* Object, const TCHAR* Key, const TCHAR* Value);
 
 	/**
 	 * Set the key/value pair in the Property's metadata
@@ -131,49 +131,49 @@ public:
 	 * @param Value The value to set for the key
 	 * @Values The metadata key/value pairs
 	 */
-	void SetValue(const UObject* Object, FName Key, const TCHAR* Value);
+	COREUOBJECT_API void SetValue(const UObject* Object, FName Key, const TCHAR* Value);
 
 	/** 
 	 *	Remove any entry with the supplied Key form the Property's metadata 
 	 *	@param Object the object to clear the metadata for
 	 *	@param Key A key to clear the data for
 	 */
-	void RemoveValue(const UObject* Object, const TCHAR* Key);
+	COREUOBJECT_API void RemoveValue(const UObject* Object, const TCHAR* Key);
 
 	/** 
 	 *	Remove any entry with the supplied Key form the Property's metadata 
 	 *	@param Object the object to clear the metadata for
 	 *	@param Key A key to clear the data for
 	 */
-	void RemoveValue(const UObject* Object, FName Key);
+	COREUOBJECT_API void RemoveValue(const UObject* Object, FName Key);
 
 	/** Find the name/value map for metadata for a specific object */
-	static TMap<FName, FString>* GetMapForObject(const UObject* Object);
+	static COREUOBJECT_API TMap<FName, FString>* GetMapForObject(const UObject* Object);
 
 	/** Copy all metadata from the source object to the destination object. This will add to any existing metadata entries for SourceObject. */
-	static void CopyMetadata(UObject* SourceObject, UObject* DestObject);
+	static COREUOBJECT_API void CopyMetadata(UObject* SourceObject, UObject* DestObject);
 
 	/**
 	 * Removes any metadata entries that are to objects not inside the same package as this UMetaData object.
 	 */
-	void RemoveMetaDataOutsidePackage();
+	COREUOBJECT_API void RemoveMetaDataOutsidePackage();
 
 	// UObject interface
-	virtual void Serialize(FArchive& Ar) override;
-	virtual void Serialize(FStructuredArchive::FRecord Record) override;
-	virtual bool NeedsLoadForEditorGame() const override;
+	COREUOBJECT_API virtual void Serialize(FArchive& Ar) override;
+	COREUOBJECT_API virtual void Serialize(FStructuredArchive::FRecord Record) override;
+	COREUOBJECT_API virtual bool NeedsLoadForEditorGame() const override;
 	virtual bool IsAsset() const override { return false; }
 	// End of UObject interface
 
 	// Returns the remapped key name, or NAME_None was not remapped.
-	static FName GetRemappedKeyName(FName OldKey);
+	static COREUOBJECT_API FName GetRemappedKeyName(FName OldKey);
 
 private:
-	static void InitializeRedirectMap();
+	static COREUOBJECT_API void InitializeRedirectMap();
 
 private:
 	// Redirect map from deprecated keys to current key names
-	static TMap<FName, FName> KeyRedirectMap;
+	static COREUOBJECT_API TMap<FName, FName> KeyRedirectMap;
 };
 
 
diff --git a/Engine/Source/Runtime/CoreUObject/Public/UObject/NoExportTypes.h b/Engine/Source/Runtime/CoreUObject/Public/UObject/NoExportTypes.h
--- a/Engine/Source/Runtime/CoreUObject/Public/UObject/NoExportTypes.h
+++ b/Engine/Source/Runtime/CoreUObject/Public/UObject/NoExportTypes.h
@@ -2645,7 +2645,7 @@ struct FARFilter
 };
 
 USTRUCT(noexport)
-struct COREUOBJECT_API FAssetBundleEntry
+struct FAssetBundleEntry
 {
 	/** Specific name of this bundle */
 	UPROPERTY()
@@ -2663,7 +2663,7 @@ struct COREUOBJECT_API FAssetBundleEntry
 
 /** A struct with a list of asset bundle entries. If one of these is inside a UObject it will get automatically exported as the asset registry tag AssetBundleData */
 USTRUCT(noexport, IsAlwaysAccessible, HasDefaults)
-struct COREUOBJECT_API FAssetBundleData
+struct FAssetBundleData
 {
 	/** List of bundles defined */
 	UPROPERTY()
diff --git a/Engine/Source/Runtime/CoreUObject/Public/UObject/Object.h b/Engine/Source/Runtime/CoreUObject/Public/UObject/Object.h
--- a/Engine/Source/Runtime/CoreUObject/Public/UObject/Object.h
+++ b/Engine/Source/Runtime/CoreUObject/Public/UObject/Object.h
@@ -82,10 +82,10 @@ ENUM_CLASS_FLAGS(EEditChangePropagationFlags);
  * 
  * @see https://docs.unrealengine.com/ProgrammingAndScripting/ProgrammingWithCPP/UnrealArchitecture/Objects
  */
-class COREUOBJECT_API UObject : public UObjectBaseUtility
+class UObject : public UObjectBaseUtility
 {
 	// Declarations, normally created by UnrealHeaderTool boilerplate code
-	DECLARE_CLASS(UObject,UObject,CLASS_Abstract|CLASS_Intrinsic|CLASS_MatchedSerializers,CASTCLASS_None,TEXT("/Script/CoreUObject"),NO_API)
+	DECLARE_CLASS(UObject,UObject,CLASS_Abstract|CLASS_Intrinsic|CLASS_MatchedSerializers,CASTCLASS_None,TEXT("/Script/CoreUObject"),COREUOBJECT_API)
 	DEFINE_DEFAULT_OBJECT_INITIALIZER_CONSTRUCTOR_CALL(UObject)
 	typedef UObject WithinClass;
 	static UObject* __VTableCtorCaller(FVTableHelper& Helper)
@@ -101,20 +101,20 @@ class COREUOBJECT_API UObject : public UObjectBaseUtility
 	}
 
 	/** Default constructor */
-	UObject();
+	COREUOBJECT_API UObject();
 
 	/** 
 	 *  Constructor that takes an ObjectInitializer. 
 	 *  Typically not needed, but can be useful for class hierarchies that support
 	 *  optional subobjects or subobject class overriding
 	 */
-	UObject(const FObjectInitializer& ObjectInitializer);
+	COREUOBJECT_API UObject(const FObjectInitializer& ObjectInitializer);
 
 	/** DO NOT USE. This constructor is for internal usage only for statically-created objects. */
-	UObject(EStaticConstructor, EObjectFlags InFlags);
+	COREUOBJECT_API UObject(EStaticConstructor, EObjectFlags InFlags);
 
 	/** DO NOT USE. This constructor is for internal usage only for hot-reload purposes. */
-	UObject(FVTableHelper& Helper);
+	COREUOBJECT_API UObject(FVTableHelper& Helper);
 
 	UE_DEPRECATED(4.23, "CreateDefaultSubobject no longer takes bAbstract as a parameter.")
 	UObject* CreateDefaultSubobject(FName SubobjectFName, UClass* ReturnType, UClass* ClassToCreateByDefault, bool bIsRequired, bool bAbstract, bool bIsTransient)
@@ -123,7 +123,7 @@ class COREUOBJECT_API UObject : public UObjectBaseUtility
 	}
 
 	/** Utility function for templates below */
-	UObject* CreateDefaultSubobject(FName SubobjectFName, UClass* ReturnType, UClass* ClassToCreateByDefault, bool bIsRequired, bool bIsTransient);
+	COREUOBJECT_API UObject* CreateDefaultSubobject(FName SubobjectFName, UClass* ReturnType, UClass* ClassToCreateByDefault, bool bIsRequired, bool bIsTransient);
 
 	/**
 	 * Create a component or subobject only to be used with the editor. They will be stripped out in packaged builds.
@@ -209,13 +209,13 @@ class COREUOBJECT_API UObject : public UObjectBaseUtility
 	 * Gets all default subobjects associated with this object instance.
 	 * @param	OutDefaultSubobjects	Array containing all default subobjects of this object.
 	 */
-	void GetDefaultSubobjects(TArray<UObject*>& OutDefaultSubobjects);
+	COREUOBJECT_API void GetDefaultSubobjects(TArray<UObject*>& OutDefaultSubobjects);
 
 	/**
 	 * Finds a subobject associated with this object instance by its name
 	 * @param	Name	Object name to look for
 	 */
-	UObject* GetDefaultSubobjectByName(FName ToFind);
+	COREUOBJECT_API UObject* GetDefaultSubobjectByName(FName ToFind);
 
 	/*----------------------------------
 			UObject interface
@@ -233,13 +233,13 @@ public:
 	 * Called after the C++ constructor and after the properties have been initialized, including those loaded from config.
 	 * This is called before any serialization or other setup has happened.
 	 */
-	virtual void PostInitProperties();
+	COREUOBJECT_API virtual void PostInitProperties();
 	
 	/**
 	 * Called after properties are overwritten, including after subobjects initialization from a CDO.
 	 * This could be called multiple times during an object lifetime, which is not the case for PostInitProperties which is expected to be called only once.
 	 */
-	virtual void PostReinitProperties();
+	COREUOBJECT_API virtual void PostReinitProperties();
 
 	/**
 	* Called after the C++ constructor has run on the Class Default Object (CDO) for a class. This is an obscure routine used to deal with the recursion 
@@ -272,7 +272,7 @@ public:
 #endif
 
 	UE_DEPRECATED(5.0, "Use version that takes FObjectPreSaveContext instead.")
-	virtual bool PreSaveRoot(const TCHAR* Filename);
+	COREUOBJECT_API virtual bool PreSaveRoot(const TCHAR* Filename);
 
 	/**
 	 * Called from within SavePackage on the passed in base/root object. The return value of this function will be passed to PostSaveRoot. 
@@ -281,10 +281,10 @@ public:
 	 * @param	ObjectSaveContext Context providing access to parameters of the save,
 	*							  Also allows storage of variables like bCleanupIsRequired for use in PostSaveRoot
 	 */
-	virtual void PreSaveRoot(FObjectPreSaveRootContext ObjectSaveContext);
+	COREUOBJECT_API virtual void PreSaveRoot(FObjectPreSaveRootContext ObjectSaveContext);
 
 	UE_DEPRECATED(5.0, "Use version that takes FObjectPostSaveContext instead.")
-	virtual void PostSaveRoot(bool bCleanupIsRequired);
+	COREUOBJECT_API virtual void PostSaveRoot(bool bCleanupIsRequired);
 
 	/**
 	 * Called from within SavePackage on the passed in base/root object. 
@@ -292,10 +292,10 @@ public:
 	 *
 	 * @param	ObjectSaveContext Context providing access to parameters of the save and to values from PreSaveRoot
 	 */
-	virtual void PostSaveRoot(FObjectPostSaveRootContext ObjectSaveContext);
+	COREUOBJECT_API virtual void PostSaveRoot(FObjectPostSaveRootContext ObjectSaveContext);
 
 	UE_DEPRECATED(5.0, "Use version that takes FObjectPreSaveContext instead.")
-	virtual void PreSave(const class ITargetPlatform* TargetPlatform);
+	COREUOBJECT_API virtual void PreSave(const class ITargetPlatform* TargetPlatform);
 
 	/**
 	 * Presave function. Gets called once before an object gets serialized for saving. This function is necessary
@@ -303,7 +303,7 @@ public:
 	 *
 	 * @warning: Objects created from within PreSave will NOT have PreSave called on them!!!
 	 */
-	virtual void PreSave(FObjectPreSaveContext SaveContext);
+	COREUOBJECT_API virtual void PreSave(FObjectPreSaveContext SaveContext);
 
 #if WITH_EDITOR
 	/**
@@ -315,13 +315,13 @@ public:
 	 *								currently recording an active undo/redo transaction
 	 * @return true if the object was saved to the transaction buffer
 	 */
-	virtual bool Modify( bool bAlwaysMarkDirty=true );
+	COREUOBJECT_API virtual bool Modify( bool bAlwaysMarkDirty=true );
 
 	/** Method to disable the sub object special handling in the transaction buffer for this object*/
-	virtual bool IsCapturingAsRootObjectForTransaction() const;
+	COREUOBJECT_API virtual bool IsCapturingAsRootObjectForTransaction() const;
 
 	/** Utility to allow overrides of Modify to avoid doing work if this object cannot be safely modified */
-	bool CanModify() const;
+	COREUOBJECT_API bool CanModify() const;
 #else
 	FORCEINLINE bool Modify(bool bAlwaysMarkDirty = true) { return false; }
 	FORCEINLINE bool IsCapturingAsRootObjectForTransaction() const { return false; };
@@ -357,7 +357,7 @@ public:
 	 * Do any object-specific cleanup required immediately after loading an object.
 	 * This is not called for newly-created objects, and by default will always execute on the game thread.
 	 */
-	virtual void PostLoad();
+	COREUOBJECT_API virtual void PostLoad();
 
 	/**
 	 * Instances components for objects being loaded from disk, if necessary.  Ensures that component references
@@ -366,13 +366,13 @@ public:
 	 * @param	OuterInstanceGraph	when calling this method on subobjects, specifies the instancing graph which contains all instanced
 	 *								subobjects and components for a subobject root.
 	 */
-	virtual void PostLoadSubobjects(FObjectInstancingGraph* OuterInstanceGraph);
+	COREUOBJECT_API virtual void PostLoadSubobjects(FObjectInstancingGraph* OuterInstanceGraph);
 	
 	/**
 	 * Called before destroying the object.  This is called immediately upon deciding to destroy the object, to allow the object to begin an
 	 * asynchronous cleanup process.
 	 */
-	virtual void BeginDestroy();
+	COREUOBJECT_API virtual void BeginDestroy();
 
 	/**
 	 * Called to check if the object is ready for FinishDestroy.  This is called after BeginDestroy to check the completion of the
@@ -393,27 +393,27 @@ public:
 	 *
 	 * @warning Because properties are destroyed here, Super::FinishDestroy() should always be called at the end of your child class's FinishDestroy() method, rather than at the beginning.
 	 */
-	virtual void FinishDestroy();
+	COREUOBJECT_API virtual void FinishDestroy();
 
 	/** 
 	 * Handles reading, writing, and reference collecting using FArchive.
 	 * This implementation handles all FProperty serialization, but can be overridden for native variables.
 	 */
-	virtual void Serialize(FArchive& Ar);
-	virtual void Serialize(FStructuredArchive::FRecord Record);
+	COREUOBJECT_API virtual void Serialize(FArchive& Ar);
+	COREUOBJECT_API virtual void Serialize(FStructuredArchive::FRecord Record);
 #if WITH_EDITORONLY_DATA
 	/**
 	 * Call Ar.UsingCustomVersion for every CustomVersion that might be serialized by this class when saving.
 	 * This duplicates CustomVersions declared in Serialize; Serialize still needs to declare them.
 	 * Used to track which customversions will be used by a package when it is resaved.
 	 * Not yet exhaustive; add CustomVersions as necessary to remove EditorDomain warnings about missing versions. */
-	static void DeclareCustomVersions(FArchive& Ar, const UClass* SpecificSubclass);
+	static COREUOBJECT_API void DeclareCustomVersions(FArchive& Ar, const UClass* SpecificSubclass);
 	/**
 	 * Append config values or settings that can change how instances of the class are cooked, including especially
 	 * values that determine how version upgraded are conducted. Can also append a unique guid when necessary to
 	 * invalidate previous results because serialization changed and no custom version was udpated.
 	 */
-	static void AppendToClassSchema(FAppendToClassSchemaContext& Context);
+	static COREUOBJECT_API void AppendToClassSchema(FAppendToClassSchemaContext& Context);
 	/**
 	 * Declare classes that can be constructed by this class during loading. This declaration is implicitly transitive;
 	 * the caller is responsible for following the graph of ConstructClasses to find all transitive ConstructClasses.
@@ -422,7 +422,7 @@ public:
 	 * @param SpecificSubclass The class on which DeclaredConstructClasses was called. This can differ from the current
 	 *        class because each class calls Super::DeclareConstructClasses.
 	 */
-	static void DeclareConstructClasses(TArray<FTopLevelAssetPath>& OutConstructClasses, const UClass* SpecificSubclass);
+	static COREUOBJECT_API void DeclareConstructClasses(TArray<FTopLevelAssetPath>& OutConstructClasses, const UClass* SpecificSubclass);
 #endif
 
 	/** After a critical error, perform any mission-critical cleanup, such as restoring the video mode orreleasing hardware resources. */
@@ -437,7 +437,7 @@ public:
 	 *
 	 * @param PropertyThatWillChange	Property that will be changed
 	 */
-	virtual void PreEditChange(FProperty* PropertyAboutToChange);
+	COREUOBJECT_API virtual void PreEditChange(FProperty* PropertyAboutToChange);
 
 	/**
 	 * This alternate version of PreEditChange is called when properties inside structs are modified.  The property that was actually modified
@@ -445,7 +445,7 @@ public:
 	 *
 	 * @param PropertyAboutToChange the property that is about to be modified
 	 */
-	virtual void PreEditChange( class FEditPropertyChain& PropertyAboutToChange );
+	COREUOBJECT_API virtual void PreEditChange( class FEditPropertyChain& PropertyAboutToChange );
 
 	/**
 	 * Called by the editor to query whether a property of this object is allowed to be modified.
@@ -456,7 +456,7 @@ public:
 	 *
 	 * @return	true if the property can be modified in the editor, otherwise false
 	 */
-	virtual bool CanEditChange( const FProperty* InProperty ) const;
+	COREUOBJECT_API virtual bool CanEditChange( const FProperty* InProperty ) const;
 
 	/**
 	 * Alternate version of CanEditChange that includes the full property chain leading to the property in question.
@@ -467,50 +467,50 @@ public:
 	 *
 	 * @return	true if the property can be modified in the editor, otherwise false
 	 */
-	virtual bool CanEditChange(const FEditPropertyChain& PropertyChain) const;
+	COREUOBJECT_API virtual bool CanEditChange(const FEditPropertyChain& PropertyChain) const;
 
 	/** 
 	 * Intentionally non-virtual as it calls the FPropertyChangedEvent version
 	 */
-	void PostEditChange();
+	COREUOBJECT_API void PostEditChange();
 
 	/**
 	 * Called when a property on this object has been modified externally
 	 *
 	 * @param PropertyThatChanged the property that was modified
 	 */
-	virtual void PostEditChangeProperty( struct FPropertyChangedEvent& PropertyChangedEvent);
+	COREUOBJECT_API virtual void PostEditChangeProperty( struct FPropertyChangedEvent& PropertyChangedEvent);
 
 	/**
 	 * This alternate version of PostEditChange is called when properties inside structs are modified.  The property that was actually modified
 	 * is located at the tail of the list.  The head of the list of the FStructProperty member variable that contains the property that was modified.
 	 */
-	virtual void PostEditChangeChainProperty( struct FPropertyChangedChainEvent& PropertyChangedEvent );
+	COREUOBJECT_API virtual void PostEditChangeChainProperty( struct FPropertyChangedChainEvent& PropertyChangedEvent );
 
 	/** Called to set the flags that control how editor change events get propagated to archetype instances. */
-	void SetEditChangePropagationFlags(EEditChangePropagationFlags InFlags);
+	COREUOBJECT_API void SetEditChangePropagationFlags(EEditChangePropagationFlags InFlags);
 
 	/** Called before applying a transaction to the object.  Default implementation simply calls PreEditChange. */
-	virtual void PreEditUndo();
+	COREUOBJECT_API virtual void PreEditUndo();
 
 	/** Called after applying a transaction to the object.  Default implementation simply calls PostEditChange. */
-	virtual void PostEditUndo();
+	COREUOBJECT_API virtual void PostEditUndo();
 
 	/** Called after applying a transaction to the object in cases where transaction annotation was provided. Default implementation simply calls PostEditChange. */
-	virtual void PostEditUndo(TSharedPtr<ITransactionObjectAnnotation> TransactionAnnotation);
+	COREUOBJECT_API virtual void PostEditUndo(TSharedPtr<ITransactionObjectAnnotation> TransactionAnnotation);
 
 	/**
 	 * Called after the object has been transacted in some way.
 	 * TransactionEvent describes what actually happened.
 	 * @note Unlike PostEditUndo (which is called for any object in the transaction), this is only called on objects that are actually changed by the transaction.
 	 */
-	virtual void PostTransacted(const FTransactionObjectEvent& TransactionEvent);
+	COREUOBJECT_API virtual void PostTransacted(const FTransactionObjectEvent& TransactionEvent);
 
 	/** Find or create and populate an annotation object with any external data required for applying a transaction */
-	TSharedPtr<ITransactionObjectAnnotation> FindOrCreateTransactionAnnotation() const;
+	COREUOBJECT_API TSharedPtr<ITransactionObjectAnnotation> FindOrCreateTransactionAnnotation() const;
 
 	/** Create and restore a previously serialized annotation object with any external data required for applying a transaction */
-	TSharedPtr<ITransactionObjectAnnotation> CreateAndRestoreTransactionAnnotation(FArchive& Ar) const;
+	COREUOBJECT_API TSharedPtr<ITransactionObjectAnnotation> CreateAndRestoreTransactionAnnotation(FArchive& Ar) const;
 
 protected:
 	/** Factory a new annotation object and optionally populate it with data */
@@ -523,7 +523,7 @@ protected:
 	 * @return		true if the object is selected, false otherwise.
 	 * @todo UE	 this doesn't belong here, but it doesn't belong anywhere else any better
 	 */
-	virtual bool IsSelectedInEditor() const;
+	COREUOBJECT_API virtual bool IsSelectedInEditor() const;
 
 public:
 #endif // WITH_EDITOR
@@ -556,7 +556,7 @@ public:
 	 *
 	 * @return	true if this object should be loaded on clients
 	 */
-	virtual bool NeedsLoadForClient() const;
+	COREUOBJECT_API virtual bool NeedsLoadForClient() const;
 
 	/**
 	 * Called during saving to determine the load flags to save with the object.
@@ -564,7 +564,7 @@ public:
 	 *
 	 * @return	true if this object should be loaded on servers
 	 */
-	virtual bool NeedsLoadForServer() const;
+	COREUOBJECT_API virtual bool NeedsLoadForServer() const;
 
 	/**
 	 * Called during saving to determine the load flags to save with the object.
@@ -572,7 +572,7 @@ public:
 	 *
 	 * @return	true if this object should be loaded on the target platform
 	 */
-	virtual bool NeedsLoadForTargetPlatform(const class ITargetPlatform* TargetPlatform) const;
+	COREUOBJECT_API virtual bool NeedsLoadForTargetPlatform(const class ITargetPlatform* TargetPlatform) const;
 
 	/**
 	 * Called during saving to include this object in client/servers running in editor builds, even if they wouldn't normally be.
@@ -621,14 +621,14 @@ public:
 	*
 	* @return	true if this object's destructor is thread safe
 	*/
-	virtual bool IsDestructionThreadSafe() const;
+	COREUOBJECT_API virtual bool IsDestructionThreadSafe() const;
 
 	/**
 	* Called during cooking. Must return all objects that will be Preload()ed when this is serialized at load time. Only used by the EDL.
 	*
 	* @param	OutDeps				all objects that will be preloaded when this is serialized at load time
 	*/
-	virtual void GetPreloadDependencies(TArray<UObject*>& OutDeps);
+	COREUOBJECT_API virtual void GetPreloadDependencies(TArray<UObject*>& OutDeps);
 
 	/**
 	* Called during cooking. Returns a list of objects. The packages containing those objects will be prestreamed, when the package containing this is loaded. Only used by the EDL.
@@ -642,18 +642,18 @@ public:
 	/**
 	*	Update the list of classes that we should exclude from dedicated server builds
 	*/
-	static void UpdateClassesExcludedFromDedicatedServer(const TArray<FString>& InClassNames, const TArray<FString>& InModulesNames);
+	static COREUOBJECT_API void UpdateClassesExcludedFromDedicatedServer(const TArray<FString>& InClassNames, const TArray<FString>& InModulesNames);
 
 	/**
 	*	Update the list of classes that we should exclude from dedicated client builds
 	*/
-	static void UpdateClassesExcludedFromDedicatedClient(const TArray<FString>& InClassNames, const TArray<FString>& InModulesNames);
+	static COREUOBJECT_API void UpdateClassesExcludedFromDedicatedClient(const TArray<FString>& InClassNames, const TArray<FString>& InModulesNames);
 
 	/** 
 	 *	Determines if you can create an object from the supplied template in the current context (editor, client only, dedicated server, game/listen) 
 	 *	This calls NeedsLoadForClient & NeedsLoadForServer
 	 */
-	static bool CanCreateInCurrentContext(UObject* Template);
+	static COREUOBJECT_API bool CanCreateInCurrentContext(UObject* Template);
 
 	/**
 	 * Exports the property values for the specified object as text to the output device.
@@ -697,13 +697,13 @@ public:
 	 * @param	NewOuter	New Outer this object will be placed within, if null it will use the current outer
 	 * @param	Flags		Flags to specify what happens during the rename
 	 */
-	virtual bool Rename(const TCHAR* NewName=nullptr, UObject* NewOuter=nullptr, ERenameFlags Flags=REN_None);
+	COREUOBJECT_API virtual bool Rename(const TCHAR* NewName=nullptr, UObject* NewOuter=nullptr, ERenameFlags Flags=REN_None);
 
 	/** Return a one line description of an object for viewing in the thumbnail view of the generic browser */
 	virtual FString GetDesc() { return TEXT( "" ); }
 
 	/** Return the UStruct corresponding to the sidecar data structure that stores data that is constant for all instances of this class. */
-	UScriptStruct* GetSparseClassDataStruct() const;
+	COREUOBJECT_API UScriptStruct* GetSparseClassDataStruct() const;
 
 #if WITH_EDITOR
 	virtual void MoveDataToSparseClassDataStruct() const {}
@@ -715,14 +715,14 @@ public:
 	 * Returns what UWorld this object is contained within. 
 	 * By default this will follow its Outer chain, but it should be overridden if that will not work.
 	 */
-	virtual class UWorld* GetWorld() const;
+	COREUOBJECT_API virtual class UWorld* GetWorld() const;
 
 	/** Internal function used by UEngine::GetWorldFromContextObject() */
-	class UWorld* GetWorldChecked(bool& bSupported) const;
+	COREUOBJECT_API class UWorld* GetWorldChecked(bool& bSupported) const;
 
 #if WITH_EDITOR
 	/** Checks to see if GetWorld() is implemented on a specific class. Optionally overridable in derived classes. */
-	virtual bool ImplementsGetWorld() const;
+	COREUOBJECT_API virtual bool ImplementsGetWorld() const;
 #endif // #if WITH_EDITOR
 
 #endif // #if WITH_ENGINE
@@ -750,7 +750,7 @@ public:
 	 *
 	 * @param	CumulativeResourceSize	Struct used to count up the cumulative size of the resource as to be displayed to artists/LDs in the Editor.
 	 */
-	virtual void GetResourceSizeEx(FResourceSizeEx& CumulativeResourceSize);
+	COREUOBJECT_API virtual void GetResourceSizeEx(FResourceSizeEx& CumulativeResourceSize);
 
 	/**
 	 * Get the size of the object/resource for use in memory tools or to display to artists/LDs in the Editor
@@ -782,14 +782,14 @@ public:
 	 * @param InThis Object to collect references from.
 	 * @param Collector	FReferenceCollector objects to be used to collect references.
 	 */
-	static void AddReferencedObjects(UObject* InThis, FReferenceCollector& Collector);
+	static COREUOBJECT_API void AddReferencedObjects(UObject* InThis, FReferenceCollector& Collector);
 
 	/**
 	 * Helper function to call AddReferencedObjects for this object's class.
 	 *
 	 * @param Collector	FReferenceCollector objects to be used to collect references.
 	 */
-	void CallAddReferencedObjects(FReferenceCollector& Collector);
+	COREUOBJECT_API void CallAddReferencedObjects(FReferenceCollector& Collector);
 
 	/**
 	 * Save information for StaticAllocateObject in the case of overwriting an existing object.
@@ -879,7 +879,7 @@ public:
 	 *
 	 * @param	OutTags		A list of key-value pairs associated with this object and their types
 	 */
-	virtual void GetAssetRegistryTags(TArray<FAssetRegistryTag>& OutTags) const;
+	COREUOBJECT_API virtual void GetAssetRegistryTags(TArray<FAssetRegistryTag>& OutTags) const;
 
 	UE_DEPRECATED(5.1, "Use the new GetExtendedAssetRegistryTagsForSave that takes a TargetPlatform")
 	virtual void GetExternalActorExtendedAssetRegistryTags(TArray<FAssetRegistryTag>& OutTags) const {}
@@ -891,18 +891,18 @@ public:
 	 * @param	TargetPlatform	The platform that this object is being saved for.
 	 * @param	OutTags			A list of key-value pairs associated with this object and their types
 	 */
-	virtual void GetExtendedAssetRegistryTagsForSave(const ITargetPlatform* TargetPlatform, TArray<FAssetRegistryTag>& OutTags) const;
+	COREUOBJECT_API virtual void GetExtendedAssetRegistryTagsForSave(const ITargetPlatform* TargetPlatform, TArray<FAssetRegistryTag>& OutTags) const;
 #endif // WITH_EDITOR
 
 	/** Gathers a list of asset registry tags for an FAssetData  */
-	void GetAssetRegistryTags(FAssetData& Out) const;
+	COREUOBJECT_API void GetAssetRegistryTags(FAssetData& Out) const;
 
 	/** Get the common tag name used for all asset source file import paths */
-	static const FName& SourceFileTagName();
+	static COREUOBJECT_API const FName& SourceFileTagName();
 
 #if UE_USE_VERSE_PATHS
 	/** Get the common tag name used for all asset verse paths */
-	static const FName& AssetVersePathTagName();
+	static COREUOBJECT_API const FName& AssetVersePathTagName();
 #endif
 
 #if WITH_EDITOR
@@ -915,7 +915,7 @@ public:
 	 * @param InAssetData Asset data loaded from the AssetRegistry
 	 * @return Pointer to new asset data after fixup or nullptr if no fixup was required
 	 */
-	virtual void PostLoadAssetRegistryTags(const FAssetData& InAssetData, TArray<FAssetRegistryTag>& OutTagsAndValuesToUpdate) const;
+	COREUOBJECT_API virtual void PostLoadAssetRegistryTags(const FAssetData& InAssetData, TArray<FAssetRegistryTag>& OutTagsAndValuesToUpdate) const;
 
 	/**
 	 * Additional data pertaining to asset registry tags used by the editor
@@ -957,41 +957,41 @@ public:
 	};
 
 	/** Gathers a collection of asset registry tag metadata */
-	virtual void GetAssetRegistryTagMetadata(TMap<FName, FAssetRegistryTagMetadata>& OutMetadata) const;
+	COREUOBJECT_API virtual void GetAssetRegistryTagMetadata(TMap<FName, FAssetRegistryTagMetadata>& OutMetadata) const;
 
 	/** The metadata tags to be transferred from the UMetaData to the Asset Registry */
-	static TSet<FName>& GetMetaDataTagsForAssetRegistry();
+	static COREUOBJECT_API TSet<FName>& GetMetaDataTagsForAssetRegistry();
 
 #endif
 
 	/** Returns true if this object is considered an asset. */
-	virtual bool IsAsset() const;
+	COREUOBJECT_API virtual bool IsAsset() const;
 
 	/**
 	 * Returns an Type:Name pair representing the PrimaryAssetId for this object.
 	 * Assets that need to be globally referenced at runtime should return a valid Identifier.
 	 * If this is valid, the object can be referenced by identifier using the AssetManager 
 	 */
-	virtual FPrimaryAssetId GetPrimaryAssetId() const;
+	COREUOBJECT_API virtual FPrimaryAssetId GetPrimaryAssetId() const;
 
 	/** Returns true if this object is considered a localized resource. */
-	virtual bool IsLocalizedResource() const;
+	COREUOBJECT_API virtual bool IsLocalizedResource() const;
 
 	/** Returns true if this object is safe to add to the root set. */
-	virtual bool IsSafeForRootSet() const;
+	COREUOBJECT_API virtual bool IsSafeForRootSet() const;
 
 	/** 
 	 * Tags objects that are part of the same asset with the specified object flag, used for GC checking
 	 *
 	 * @param	ObjectFlags	Object Flags to enable on the related objects
 	 */
-	virtual void TagSubobjects(EObjectFlags NewFlags);
+	COREUOBJECT_API virtual void TagSubobjects(EObjectFlags NewFlags);
 
 	/** Returns properties that are replicated for the lifetime of the actor channel */
-	virtual void GetLifetimeReplicatedProps( TArray< class FLifetimeProperty > & OutLifetimeProps ) const;
+	COREUOBJECT_API virtual void GetLifetimeReplicatedProps( TArray< class FLifetimeProperty > & OutLifetimeProps ) const;
 
 	/** Called when this object begins replicating to initialize the state of custom property conditions */
-	virtual void GetReplicatedCustomConditionState(FCustomPropertyConditionState& OutActiveState) const;
+	COREUOBJECT_API virtual void GetReplicatedCustomConditionState(FCustomPropertyConditionState& OutActiveState) const;
 #if UE_WITH_IRIS
 	/**
 	 * RegisterReplicationFragments is called when we an object is added to the ReplicationSystem, it allows an object to register new or existing ReplicationFragments describing data to be replicated and how it should be accessed
@@ -1000,32 +1000,32 @@ public:
 	 * @param Context Context FFragmentRegistrationContext in which FReplicationFragments could be registered
 	 * @param RegistrationFlags Flags specifying what should be registered in the call
 	 */
-	virtual void RegisterReplicationFragments(UE::Net::FFragmentRegistrationContext& Context, UE::Net::EFragmentRegistrationFlags RegistrationFlags);
+	COREUOBJECT_API virtual void RegisterReplicationFragments(UE::Net::FFragmentRegistrationContext& Context, UE::Net::EFragmentRegistrationFlags RegistrationFlags);
 #endif // UE_WITH_IRIS
 
 	/** IsNameStableForNetworking means an object can be referred to its path name (relative to outer) over the network */
-	virtual bool IsNameStableForNetworking() const;
+	COREUOBJECT_API virtual bool IsNameStableForNetworking() const;
 
 	/** IsFullNameStableForNetworking means an object can be referred to its full path name over the network */
-	virtual bool IsFullNameStableForNetworking() const;
+	COREUOBJECT_API virtual bool IsFullNameStableForNetworking() const;
 
 	/** IsSupportedForNetworking means an object can be referenced over the network */
-	virtual bool IsSupportedForNetworking() const;
+	COREUOBJECT_API virtual bool IsSupportedForNetworking() const;
 
 	/** Returns a list of sub-objects that have stable names for networking */
 	virtual void GetSubobjectsWithStableNamesForNetworking(TArray<UObject*> &ObjList) {}
 
 	/** Called right before receiving a bunch */
-	virtual void PreNetReceive();
+	COREUOBJECT_API virtual void PreNetReceive();
 
 	/** Called right after receiving a bunch */
-	virtual void PostNetReceive();
+	COREUOBJECT_API virtual void PostNetReceive();
 
 	/** Called right after calling all OnRep notifies (called even when there are no notifies) */
 	virtual void PostRepNotifies() {}
 
 	/** Called right before being marked for destruction due to network replication */
-	virtual void PreDestroyFromReplication();
+	COREUOBJECT_API virtual void PreDestroyFromReplication();
 
 #if WITH_EDITOR
 	/**
@@ -1035,10 +1035,10 @@ public:
 	 * @return Valid if this object has data validation rules set up for it and the data for this object is valid. Returns Invalid if it does not pass 
 	 *         the rules. Returns NotValidated if no rules are set for this object.
 	 */
-	virtual EDataValidationResult IsDataValid(class FDataValidationContext& Context) const;
+	COREUOBJECT_API virtual EDataValidationResult IsDataValid(class FDataValidationContext& Context) const;
 
 	UE_DEPRECATED(5.3, "Please use the the const version of IsDataValid(FDataValidationContext&)")
-	virtual EDataValidationResult IsDataValid(class FDataValidationContext& Context);
+	COREUOBJECT_API virtual EDataValidationResult IsDataValid(class FDataValidationContext& Context);
 
 	/**
 	 * Generic function to validate objects during changelist validations, etc.
@@ -1048,7 +1048,7 @@ public:
 	 *         the rules. Returns NotValidated if no rules are set for this object.
 	 */
 	UE_DEPRECATED(5.3, "Please use IsDataValid(FDataValidationContext&) const as the FDataValidationContext will allow distinguishing between Warnings and Errors.")
-	virtual EDataValidationResult IsDataValid(TArray<FText>& ValidationErrors);
+	COREUOBJECT_API virtual EDataValidationResult IsDataValid(TArray<FText>& ValidationErrors);
 #endif // WITH_EDITOR
 
 	/*----------------------------------------------------------
@@ -1061,7 +1061,7 @@ public:
 	 * @return		true if the object is selected, false otherwise.
 	 * @todo UE this doesn't belong here, but it doesn't belong anywhere else any better
 	 */
-	bool IsSelected() const;
+	COREUOBJECT_API bool IsSelected() const;
 
 #if WITH_EDITOR
 	/**
@@ -1070,7 +1070,7 @@ public:
 	 * @param	AffectedObjects		the array of objects which have this object in their ObjectArchetype chain and will be affected by the change.
 	 *								Objects which have this object as their direct ObjectArchetype are removed from the list once they're processed.
 	 */
-	void PropagatePreEditChange( TArray<UObject*>& AffectedObjects, FEditPropertyChain& PropertyAboutToChange );
+	COREUOBJECT_API void PropagatePreEditChange( TArray<UObject*>& AffectedObjects, FEditPropertyChain& PropertyAboutToChange );
 
 	/**
 	 * Calls PostEditChange on all instances based on an archetype in AffectedObjects. Recurses on any instances.
@@ -1078,7 +1078,7 @@ public:
 	 * @param	AffectedObjects		the array of objects which have this object in their ObjectArchetype chain and will be affected by the change.
 	 *								Objects which have this object as their direct ObjectArchetype are removed from the list once they're processed.
 	 */
-	void PropagatePostEditChange( TArray<UObject*>& AffectedObjects, FPropertyChangedChainEvent& PropertyChangedEvent );
+	COREUOBJECT_API void PropagatePostEditChange( TArray<UObject*>& AffectedObjects, FPropertyChangedChainEvent& PropertyChangedEvent );
 #endif // WITH_EDITOR
 
 	/**
@@ -1087,7 +1087,7 @@ public:
 	 *
 	 * @param	Ar				the archive to use for serialization
 	 */
-	void SerializeScriptProperties( FArchive& Ar ) const;
+	COREUOBJECT_API void SerializeScriptProperties( FArchive& Ar ) const;
 
 	/**
 	 * Serializes the script property data located at Data.  When saving, only saves those properties which differ from the corresponding
@@ -1095,7 +1095,7 @@ public:
 	 *
 	 * @param	Slot				the archive slot to serialize to
 	 */
-	void SerializeScriptProperties( FStructuredArchive::FSlot Slot ) const;
+	COREUOBJECT_API void SerializeScriptProperties( FStructuredArchive::FSlot Slot ) const;
 
 	/**
 	 * Wrapper function for InitProperties() which handles safely tearing down this object before re-initializing it
@@ -1104,7 +1104,7 @@ public:
 	 * @param	SourceObject	the object to use for initializing property values in this object.  If not specified, uses this object's archetype.
 	 * @param	InstanceGraph	contains the mappings of instanced objects and components to their templates
 	 */
-	void ReinitializeProperties( UObject* SourceObject=NULL, struct FObjectInstancingGraph* InstanceGraph=NULL );
+	COREUOBJECT_API void ReinitializeProperties( UObject* SourceObject=NULL, struct FObjectInstancingGraph* InstanceGraph=NULL );
 
 	/**
 	 * This will return detail info about this specific object. (e.g. AudioComponent will return the name of the cue,
@@ -1114,19 +1114,19 @@ public:
 	 *
 	 * @note	safe to call on NULL object pointers!
 	 */
-	FString GetDetailedInfo() const;
+	COREUOBJECT_API FString GetDetailedInfo() const;
 
 	/**
 	 * Called before destroying the object.  This is called immediately upon deciding to destroy the object, to allow the object to begin an
 	 * asynchronous cleanup process.
 	 */
-	bool ConditionalBeginDestroy();
+	COREUOBJECT_API bool ConditionalBeginDestroy();
 
 	/** Called when an object is actually destroyed, memory should never be accessed again */
-	bool ConditionalFinishDestroy();
+	COREUOBJECT_API bool ConditionalFinishDestroy();
 	
 	/** PostLoad if needed. */
-	void ConditionalPostLoad();
+	COREUOBJECT_API void ConditionalPostLoad();
 
 	/**
 	 * Instances subobjects and components for objects being loaded from disk, if necessary.  Ensures that references
@@ -1135,7 +1135,7 @@ public:
 	 * @param	OuterInstanceGraph	when calling this method on subobjects, specifies the instancing graph which contains all instanced
 	 *								subobjects and components for a subobject root.
 	 */
-	void ConditionalPostLoadSubobjects( struct FObjectInstancingGraph* OuterInstanceGraph=NULL );
+	COREUOBJECT_API void ConditionalPostLoadSubobjects( struct FObjectInstancingGraph* OuterInstanceGraph=NULL );
 
 #if WITH_EDITOR
 	/**
@@ -1224,10 +1224,10 @@ public:
 	inline bool IsBasedOnArchetype( const UObject* const SomeObject ) const;
 
 	/** Returns a UFunction with the specified name, wrapper for UClass::FindFunctionByName() */
-	UFunction* FindFunction( FName InName ) const;
+	COREUOBJECT_API UFunction* FindFunction( FName InName ) const;
 	
 	/** Version of FindFunction() that will assert if the function was not found */
-	UFunction* FindFunctionChecked( FName InName ) const;
+	COREUOBJECT_API UFunction* FindFunctionChecked( FName InName ) const;
 
 	/**
 	 * Given OtherObject (which will be the same type as 'this'), recursively find any matching sub-objects from 'this' that also exist within OtherObject, and add the mappings to ObjectMapping.
@@ -1235,7 +1235,7 @@ public:
 	 * @param	OtherObject		The to find matching sub-objects within.
 	 * @param	ObjectMapping	The complete mapping between this object hierarchy and the other object hierarchy.
 	 */
-	virtual void BuildSubobjectMapping(UObject* OtherObject, TMap<UObject*, UObject*>& ObjectMapping) const;
+	COREUOBJECT_API virtual void BuildSubobjectMapping(UObject* OtherObject, TMap<UObject*, UObject*>& ObjectMapping) const;
 
 	/**
 	 * Uses the TArchiveObjectReferenceCollector to build a list of all components referenced by this object which have this object as the outer
@@ -1244,7 +1244,7 @@ public:
 	 * @param	bIncludeNestedSubobjects	controls whether subobjects which are contained by this object, but do not have this object
 	 *										as its direct Outer should be included
 	 */
-	void CollectDefaultSubobjects( TArray<UObject*>& OutDefaultSubobjects, bool bIncludeNestedSubobjects=false ) const;
+	COREUOBJECT_API void CollectDefaultSubobjects( TArray<UObject*>& OutDefaultSubobjects, bool bIncludeNestedSubobjects=false ) const;
 
 	/**
 	 * Checks default sub-object assumptions.
@@ -1252,19 +1252,19 @@ public:
 	 * @param bForceCheck Force checks even if not enabled globally.
 	 * @return true if the assumptions are met, false otherwise.
 	 */
-	bool CheckDefaultSubobjects(bool bForceCheck = false) const;
+	COREUOBJECT_API bool CheckDefaultSubobjects(bool bForceCheck = false) const;
 
 	/**
 	 * Save configuration out to ini files
 	 * @warning Must be safe to call on class-default object
 	 */
-	void SaveConfig( uint64 Flags=CPF_Config, const TCHAR* Filename=NULL, FConfigCacheIni* Config=GConfig, bool bAllowCopyToDefaultObject=true );
+	COREUOBJECT_API void SaveConfig( uint64 Flags=CPF_Config, const TCHAR* Filename=NULL, FConfigCacheIni* Config=GConfig, bool bAllowCopyToDefaultObject=true );
 
 	/**
 	 * Saves just the section(s) for this class into the default ini file for the class (with just the changes from base)
 	 */
 	UE_DEPRECATED(5.0, "TryUpdateDefaultConfigFile replaces UpdateDefaultConfigFile")
-	void UpdateDefaultConfigFile(const FString& SpecificFileLocation = "");
+	COREUOBJECT_API void UpdateDefaultConfigFile(const FString& SpecificFileLocation = "");
 
 	/**
 	 * Try to Saves just the section(s) for this class into the default ini file for the class (with just the changes from base)
@@ -1273,50 +1273,50 @@ public:
 	 *  @param bWarnIfFail If true and unable to update the Ini file due to being read only log a warning
 	 *  @return true if it was able to write false otherwise
 	 */
-	bool TryUpdateDefaultConfigFile(const FString& SpecificFileLocation = "", bool bWarnIfFail = true);
+	COREUOBJECT_API bool TryUpdateDefaultConfigFile(const FString& SpecificFileLocation = "", bool bWarnIfFail = true);
 
 	/**
 	 * Saves just the section(s) for this class into the global user ini file for the class (with just the changes from base)
 	 */
-	void UpdateGlobalUserConfigFile();
+	COREUOBJECT_API void UpdateGlobalUserConfigFile();
 
 	/**
 	 * Saves just the section(s) for this class into the project user ini file for the class (with just the changes from base)
 	 */
-	void UpdateProjectUserConfigFile();
+	COREUOBJECT_API void UpdateProjectUserConfigFile();
 
 	/**
 	 * Saves just the property into the global user ini file for the class (with just the changes from base)
 	 */
-	void UpdateSinglePropertyInConfigFile(const FProperty* InProperty, const FString& InConfigIniName);
+	COREUOBJECT_API void UpdateSinglePropertyInConfigFile(const FProperty* InProperty, const FString& InConfigIniName);
 
 private:
 	/**
 	 * Saves just the section(s) for this class into the given ini file for the class (with just the changes from base)
 	 */
-	void UpdateSingleSectionOfConfigFile(const FString& ConfigIniName);
+	COREUOBJECT_API void UpdateSingleSectionOfConfigFile(const FString& ConfigIniName);
 
 	/**
 	 * Ensures that current thread is NOT during vtable ptr retrieval process
 	 * of some UClass.
 	 */
-	void EnsureNotRetrievingVTablePtr() const;
+	COREUOBJECT_API void EnsureNotRetrievingVTablePtr() const;
 
 public:
 	/**
 	 * Get the default config filename for the specified UObject
 	 */
-	FString GetDefaultConfigFilename() const;
+	COREUOBJECT_API FString GetDefaultConfigFilename() const;
 
 	/**
 	 * Get the global user override config filename for the specified UObject
 	 */
-	FString GetGlobalUserConfigFilename() const;
+	COREUOBJECT_API FString GetGlobalUserConfigFilename() const;
 
 	/**
 	 * Get the project user override config filename for the specified UObject
 	 */
-	FString GetProjectUserConfigFilename() const;
+	COREUOBJECT_API FString GetProjectUserConfigFilename() const;
 
 	/** Returns the override config hierarchy platform (if NDAd platforms need defaults to not be in Base*.ini but still want editor to load them) */
 	virtual const TCHAR* GetConfigOverridePlatform() const { return nullptr; }
@@ -1343,7 +1343,7 @@ public:
 	 * @param	PropagationFlags	indicates how this call to LoadConfig should be propagated; expects a bitmask of UE::ELoadConfigPropagationFlags values.
 	 * @param	PropertyToLoad		if specified, only the ini value for the specified property will be imported.
 	 */
-	void LoadConfig( UClass* ConfigClass=NULL, const TCHAR* Filename=NULL, uint32 PropagationFlags=UE::LCPF_None, class FProperty* PropertyToLoad=NULL );
+	COREUOBJECT_API void LoadConfig( UClass* ConfigClass=NULL, const TCHAR* Filename=NULL, uint32 PropagationFlags=UE::LCPF_None, class FProperty* PropertyToLoad=NULL );
 
 	/**
 	 * Wrapper method for LoadConfig that is used when reloading the config data for objects at runtime which have already loaded their config data at least once.
@@ -1354,10 +1354,10 @@ public:
 	 * @param	PropagationFlags	indicates how this call to LoadConfig should be propagated; expects a bitmask of UE::ELoadConfigPropagationFlags values.
 	 * @param	PropertyToLoad		if specified, only the ini value for the specified property will be imported
 	 */
-	void ReloadConfig( UClass* ConfigClass=NULL, const TCHAR* Filename=NULL, uint32 PropagationFlags=UE::LCPF_None, class FProperty* PropertyToLoad=NULL );
+	COREUOBJECT_API void ReloadConfig( UClass* ConfigClass=NULL, const TCHAR* Filename=NULL, uint32 PropagationFlags=UE::LCPF_None, class FProperty* PropertyToLoad=NULL );
 
 	/** Import an object from a file. */
-	void ParseParms( const TCHAR* Parms );
+	COREUOBJECT_API void ParseParms( const TCHAR* Parms );
 
 	/**
 	 * Outputs a string to an arbitrary output device, describing the list of objects which are holding references to this one.
@@ -1365,10 +1365,10 @@ public:
 	 * @param	Ar						the output device to send output to
 	 * @param	Referencers				optionally allows the caller to specify the list of references to output.
 	 */
-	void OutputReferencers( FOutputDevice& Ar, FReferencerInformationList* Referencers=NULL );
+	COREUOBJECT_API void OutputReferencers( FOutputDevice& Ar, FReferencerInformationList* Referencers=NULL );
 	
 	/** Called by OutputReferencers() to get the internal list of referencers to write */
-	void RetrieveReferencers( TArray<FReferencerInformation>* OutInternalReferencers, TArray<FReferencerInformation>* OutExternalReferencers);
+	COREUOBJECT_API void RetrieveReferencers( TArray<FReferencerInformation>* OutInternalReferencers, TArray<FReferencerInformation>* OutExternalReferencers);
 
 	/**
 	 * Changes the linker and linker index to the passed in one. A linker of NULL and linker index of INDEX_NONE
@@ -1378,33 +1378,33 @@ public:
 	 * @param LinkerIndex				New LinkerIndex to set
 	 * @param bShouldDetachExisting		If true, detach existing linker and call PostLinkerChange
 	 */
-	void SetLinker( FLinkerLoad* LinkerLoad, int32 LinkerIndex, bool bShouldDetachExisting=true );
+	COREUOBJECT_API void SetLinker( FLinkerLoad* LinkerLoad, int32 LinkerIndex, bool bShouldDetachExisting=true );
 
 	/**
 	 * Return the template that an object with this class, outer and name would be
 	 * 
 	 * @return the archetype for this object
 	 */
-	static UObject* GetArchetypeFromRequiredInfo(const UClass* Class, const UObject* Outer, FName Name, EObjectFlags ObjectFlags);
+	static COREUOBJECT_API UObject* GetArchetypeFromRequiredInfo(const UClass* Class, const UObject* Outer, FName Name, EObjectFlags ObjectFlags);
 
 	/**
 	 * Return the template this object is based on. 
 	 * 
 	 * @return the archetype for this object
 	 */
-	UObject* GetArchetype() const;
+	COREUOBJECT_API UObject* GetArchetype() const;
 
 	/**
 	 * Builds a list of objects which have this object in their archetype chain.
 	 *
 	 * @param	Instances	receives the list of objects which have this one in their archetype chain
 	 */
-	void GetArchetypeInstances( TArray<UObject*>& Instances );
+	COREUOBJECT_API void GetArchetypeInstances( TArray<UObject*>& Instances );
 
 	/**
 	 * Wrapper for calling UClass::InstanceSubobjectTemplates() for this object.
 	 */
-	void InstanceSubobjectTemplates( struct FObjectInstancingGraph* InstanceGraph = NULL );
+	COREUOBJECT_API void InstanceSubobjectTemplates( struct FObjectInstancingGraph* InstanceGraph = NULL );
 
 	/**
 	 * Returns true if this object implements the interface T, false otherwise.
@@ -1418,7 +1418,7 @@ public:
 	-----------------------------*/
 
 	/** Called by VM to execute a UFunction with a filled in UStruct of parameters */
-	virtual void ProcessEvent( UFunction* Function, void* Parms );
+	COREUOBJECT_API virtual void ProcessEvent( UFunction* Function, void* Parms );
 
 	/**
 	 * Return the space this function should be called.   Checks to see if this function should
@@ -1446,10 +1446,10 @@ public:
 	}
 
 	/** Handle calling a function by name when executed from the console or a command line */
-	bool CallFunctionByNameWithArguments( const TCHAR* Cmd, FOutputDevice& Ar, UObject* Executor, bool bForceCallWithNonExec = false );
+	COREUOBJECT_API bool CallFunctionByNameWithArguments( const TCHAR* Cmd, FOutputDevice& Ar, UObject* Executor, bool bForceCallWithNonExec = false );
 
 	/** Internal VM method for executing a function */
-	void CallFunction( FFrame& Stack, RESULT_DECL, UFunction* Function );
+	COREUOBJECT_API void CallFunction( FFrame& Stack, RESULT_DECL, UFunction* Function );
 
 	/**
 	 * Internal function call processing.
@@ -1471,7 +1471,7 @@ public:
 	 * @param Result pointer to where the return value should be written
 	 * @param Function the function being called
 	 */
-	void SkipFunction(FFrame& Stack, RESULT_DECL, UFunction* Function);
+	COREUOBJECT_API void SkipFunction(FFrame& Stack, RESULT_DECL, UFunction* Function);
 
 	/**
 	 * Called on the target when a class is loaded with ClassGeneratedBy is loaded.  Should regenerate the class if needed, and return the updated class
@@ -1482,12 +1482,12 @@ public:
 	/** 
 	 * Returns whether this object is contained in or part of a blueprint object
 	 */
-	bool IsInBlueprint() const;
+	COREUOBJECT_API bool IsInBlueprint() const;
 
 	/** 
 	 *  Destroy properties that won't be destroyed by the native destructor
 	 */
-	void DestroyNonNativeProperties();
+	COREUOBJECT_API void DestroyNonNativeProperties();
 
 	/** Called during subobject creation to mark this component as editor only, which causes it to get stripped in packaged builds */
 	virtual void MarkAsEditorOnlySubobject() { }
@@ -1495,7 +1495,7 @@ public:
 	/**
 	 * Abort with a member function call at the top of the callstack, helping to ensure that most platforms will stuff this object's memory into the resulting minidump.
 	 */
-	void AbortInsideMemberFunction() const;
+	COREUOBJECT_API void AbortInsideMemberFunction() const;
 
 	// UnrealScript intrinsics, do not call directly
 
@@ -1700,17 +1700,17 @@ public:
 
 protected: 
 	/** Checks it's ok to perform subobjects check at this time. */
-	bool CanCheckDefaultSubObjects(bool bForceCheck, bool& bResult) const;
+	COREUOBJECT_API bool CanCheckDefaultSubObjects(bool bForceCheck, bool& bResult) const;
 
 	/**
 	* Checks default sub-object assumptions.
 	*
 	* @return true if the assumptions are met, false otherwise.
 	*/
-	virtual bool CheckDefaultSubobjectsInternal() const;
+	COREUOBJECT_API virtual bool CheckDefaultSubobjectsInternal() const;
 
 private:
-	void ProcessContextOpcode(FFrame& Stack, RESULT_DECL, bool bCanFailSilent);
+	COREUOBJECT_API void ProcessContextOpcode(FFrame& Stack, RESULT_DECL, bool bCanFailSilent);
 
 	/**
 	* Create a component or subobject only to be used with the editor.
@@ -1719,7 +1719,7 @@ private:
 	* @param	SubobjectName				name of the new component
 	* @param	bTransient					true if the component is being assigned to a transient property
 	*/
-	UObject* CreateEditorOnlyDefaultSubobjectImpl(FName SubobjectName, UClass* ReturnType, bool bTransient = false);
+	COREUOBJECT_API UObject* CreateEditorOnlyDefaultSubobjectImpl(FName SubobjectName, UClass* ReturnType, bool bTransient = false);
 
 public:
 
diff --git a/Engine/Source/Runtime/CoreUObject/Public/UObject/ObjectMacros.h b/Engine/Source/Runtime/CoreUObject/Public/UObject/ObjectMacros.h
--- a/Engine/Source/Runtime/CoreUObject/Public/UObject/ObjectMacros.h
+++ b/Engine/Source/Runtime/CoreUObject/Public/UObject/ObjectMacros.h
@@ -168,7 +168,7 @@ class FVTableHelper
 {
 public:
 	/** DO NOT USE. This constructor is for internal usage only for hot-reload purposes. */
-	COREUOBJECT_API FVTableHelper()
+	FVTableHelper()
 	{
 		EnsureRetrievingVTablePtrDuringCtor(TEXT("FVTableHelper()"));
 	}
@@ -644,7 +644,7 @@ class UObject;
 class FProperty;
 class FObjectInitializer; 
 
-struct COREUOBJECT_API FReferencerInformation 
+struct FReferencerInformation 
 {
 	/** the object that is referencing the target */
 	UObject*				Referencer;
@@ -655,17 +655,17 @@ struct COREUOBJECT_API FReferencerInformation
 	/** the array of UProperties in Referencer which hold references to target */
 	TArray<const FProperty*>		ReferencingProperties;
 
-	FReferencerInformation( UObject* inReferencer );
-	FReferencerInformation( UObject* inReferencer, int32 InReferences, const TArray<const FProperty*>& InProperties );
+	COREUOBJECT_API FReferencerInformation( UObject* inReferencer );
+	COREUOBJECT_API FReferencerInformation( UObject* inReferencer, int32 InReferences, const TArray<const FProperty*>& InProperties );
 };
 
-struct COREUOBJECT_API FReferencerInformationList
+struct FReferencerInformationList
 {
 	TArray<FReferencerInformation>		InternalReferences;
 	TArray<FReferencerInformation>		ExternalReferences;
 
-	FReferencerInformationList();
-	FReferencerInformationList( const TArray<FReferencerInformation>& InternalRefs, const TArray<FReferencerInformation>& ExternalRefs );
+	COREUOBJECT_API FReferencerInformationList();
+	COREUOBJECT_API FReferencerInformationList( const TArray<FReferencerInformation>& InternalRefs, const TArray<FReferencerInformation>& ExternalRefs );
 };
 
 /*----------------------------------------------------------------------------
@@ -2186,4 +2186,4 @@ public:
 
 #if UE_ENABLE_INCLUDE_ORDER_DEPRECATED_IN_5_2
 #include "CoreMinimal.h"
-#endif
\ No newline at end of file
+#endif
diff --git a/Engine/Source/Runtime/CoreUObject/Public/UObject/ObjectMemoryAnalyzer.h b/Engine/Source/Runtime/CoreUObject/Public/UObject/ObjectMemoryAnalyzer.h
--- a/Engine/Source/Runtime/CoreUObject/Public/UObject/ObjectMemoryAnalyzer.h
+++ b/Engine/Source/Runtime/CoreUObject/Public/UObject/ObjectMemoryAnalyzer.h
@@ -79,7 +79,7 @@ struct FObjectMemoryUsage
 /** 
  * Analyzes memory usage of UObjects
  **/
-struct COREUOBJECT_API FObjectMemoryAnalyzer
+struct FObjectMemoryAnalyzer
 {
 	/** Flags to modify the memory counting behavior */
 	struct EAnalyzeFlags
@@ -129,22 +129,22 @@ struct COREUOBJECT_API FObjectMemoryAnalyzer
 		ESortKey::Type SortKey;
 	};
 
-	FObjectMemoryAnalyzer(uint32 Flags = 0);
-	FObjectMemoryAnalyzer(class UClass* BaseClass, uint32 Flags = 0);
-	FObjectMemoryAnalyzer(class UObject* Object, uint32 Flags = 0);
-	FObjectMemoryAnalyzer(const TArray<class UObject*>& ObjectList, uint32 Flags = 0);
+	COREUOBJECT_API FObjectMemoryAnalyzer(uint32 Flags = 0);
+	COREUOBJECT_API FObjectMemoryAnalyzer(class UClass* BaseClass, uint32 Flags = 0);
+	COREUOBJECT_API FObjectMemoryAnalyzer(class UObject* Object, uint32 Flags = 0);
+	COREUOBJECT_API FObjectMemoryAnalyzer(const TArray<class UObject*>& ObjectList, uint32 Flags = 0);
 
 	/** Analyzes the memory usage of the specified object */
-	void AnalyzeObject(class UObject* Object);
+	COREUOBJECT_API void AnalyzeObject(class UObject* Object);
 	/** Analyzes the memory usage of the specified object list */
-	void AnalyzeObjects(const TArray<class UObject*>& ObjectList);
+	COREUOBJECT_API void AnalyzeObjects(const TArray<class UObject*>& ObjectList);
 	/** Analyzes the memory usage of all objects with the specified class */
-	void AnalyzeObjects(class UClass* BaseClass);
+	COREUOBJECT_API void AnalyzeObjects(class UClass* BaseClass);
 
 	/** Returns the results */
-	int32 GetResults(TArray<FObjectMemoryUsage>& Results);
+	COREUOBJECT_API int32 GetResults(TArray<FObjectMemoryUsage>& Results);
 	/** Returns the memory usage of an object not in the result set ( ex. referenced by a result ) */
-	const FObjectMemoryUsage& GetObjectMemoryUsage(class UObject* Obj);
+	COREUOBJECT_API const FObjectMemoryUsage& GetObjectMemoryUsage(class UObject* Obj);
 
 	struct EPrintFlags
 	{
@@ -157,7 +157,7 @@ struct COREUOBJECT_API FObjectMemoryAnalyzer
 		};
 	};
 
-	void PrintResults(FOutputDevice& Ar, uint32 PrintFlags = 0);
+	COREUOBJECT_API void PrintResults(FOutputDevice& Ar, uint32 PrintFlags = 0);
 
 private:
 	int32 GetReferencedObjects(UObject* Obj, TArray<UObject*>& ReferencedObjects);
