[CoreUObject]

diff --git a/Engine/Source/Runtime/CoreUObject/Public/UObject/PackageTrailer.h b/Engine/Source/Runtime/CoreUObject/Public/UObject/PackageTrailer.h
--- a/Engine/Source/Runtime/CoreUObject/Public/UObject/PackageTrailer.h
+++ b/Engine/Source/Runtime/CoreUObject/Public/UObject/PackageTrailer.h
@@ -196,7 +196,7 @@ struct FPayloadInfo
  * While saving a package, payloads should be added to a FPackageTrailer via ::AddPayload then once
  * the package has been saved to disk ::BuildAndAppendTrailer should be called. 
  */
-class COREUOBJECT_API FPackageTrailerBuilder
+class FPackageTrailerBuilder
 {
 public:
 	using AdditionalDataCallback = TFunction<void(FLinkerSave& LinkerSave, const class FPackageTrailer& Trailer)>;
@@ -210,7 +210,7 @@ public:
 	 * @param Ar			An archive that the trailer can use to load payloads from 
 	 * @param DebugContext	The name or path of the of the file that owns the trailer. Used for error messages.
 	 */
-	[[nodiscard]] static FPackageTrailerBuilder CreateFromTrailer(const class FPackageTrailer& Trailer, FArchive& Ar, FString DebugContext);
+	[[nodiscard]] static COREUOBJECT_API FPackageTrailerBuilder CreateFromTrailer(const class FPackageTrailer& Trailer, FArchive& Ar, FString DebugContext);
 
 	UE_DEPRECATED(5.1, "Use the overload that takes a FString instead of an FName for the last parameter")
 	static FPackageTrailerBuilder CreateFromTrailer(const class FPackageTrailer& Trailer, FArchive& Ar, const FName& PackageName)
@@ -226,7 +226,7 @@ public:
 	 * @param Trailer		The trailer to create the reference from.
 	 * @param DebugContext	The name or path of the of the file that owns the trailer. Used for error messages.
 	 */
-	[[nodiscard]] static TUniquePtr<UE::FPackageTrailerBuilder> CreateReferenceToTrailer(const class FPackageTrailer& Trailer, FString DebugContext);
+	[[nodiscard]] static COREUOBJECT_API TUniquePtr<UE::FPackageTrailerBuilder> CreateReferenceToTrailer(const class FPackageTrailer& Trailer, FString DebugContext);
 	
 	UE_DEPRECATED(5.1, "Use the overload that takes a FString instead of an FName for the last parameter")
 	static TUniquePtr<UE::FPackageTrailerBuilder> CreateReferenceToTrailer(const class FPackageTrailer& Trailer, const FName& PackageName)
@@ -236,8 +236,8 @@ public:
 
 	FPackageTrailerBuilder() = default;
 	UE_DEPRECATED(5.1, "Use the overload that takes a FString instead of an FName")
-	FPackageTrailerBuilder(const FName& InPackageName);
-	FPackageTrailerBuilder(FString&& DebugContext);
+	COREUOBJECT_API FPackageTrailerBuilder(const FName& InPackageName);
+	COREUOBJECT_API FPackageTrailerBuilder(FString&& DebugContext);
 	~FPackageTrailerBuilder() = default;
 
 	// Methods that can be called while building the trailer
@@ -251,8 +251,8 @@ public:
 	 * @param Flags			The custom flags to be applied to the payload
 	 * @param Callback		This callback will be invoked once the FPackageTrailer has been built and appended to disk.
 	 */
-	void AddPayload(const FIoHash& Identifier, FCompressedBuffer Payload, UE::Virtualization::EPayloadFilterReason Filter, AdditionalDataCallback&& Callback);
-	void AddPayload(const FIoHash& Identifier, FCompressedBuffer Payload, UE::Virtualization::EPayloadFilterReason Filter);
+	COREUOBJECT_API void AddPayload(const FIoHash& Identifier, FCompressedBuffer Payload, UE::Virtualization::EPayloadFilterReason Filter, AdditionalDataCallback&& Callback);
+	COREUOBJECT_API void AddPayload(const FIoHash& Identifier, FCompressedBuffer Payload, UE::Virtualization::EPayloadFilterReason Filter);
 	/**
 	 * Adds an already virtualized payload to the builder to be written to the trailer. When the trailer is written
 	 * the payload will have EPayloadAccessMode::Virtualized set as it's access mode. It is assumed that the payload
@@ -262,7 +262,7 @@ public:
 	 * @param Identifier	The identifier of the payload
 	 * @param RawSize		The size of the payload (in bytes) when uncompressed
 	 */
-	void AddVirtualizedPayload(const FIoHash& Identifier, int64 RawSize);
+	COREUOBJECT_API void AddVirtualizedPayload(const FIoHash& Identifier, int64 RawSize);
 
 	/** 
 	 * Allows the caller to replace a payload in the builder that is already marked as virtualized and replace it
@@ -273,14 +273,14 @@ public:
 	 * 
 	 * @return true if a virtualized payload was replaced, false if the payload was not in the builder at all
 	 */
-	bool UpdatePayloadAsLocal(const FIoHash& Identifier, FCompressedBuffer Payload);
+	COREUOBJECT_API bool UpdatePayloadAsLocal(const FIoHash& Identifier, FCompressedBuffer Payload);
 	
 	/**
 	 * @param ExportsArchive	The linker associated with the package being written to disk.
 	 * @param DataArchive		The archive where the package data has been written to. This is where the FPackageTrailer will be written to
 	 * @return True if the builder was created and appended successfully and false if any error was encountered
 	 */
-	[[nodiscard]] bool BuildAndAppendTrailer(FLinkerSave* Linker, FArchive& DataArchive);
+	[[nodiscard]] COREUOBJECT_API bool BuildAndAppendTrailer(FLinkerSave* Linker, FArchive& DataArchive);
 
 	/**
 	 * @param ExportsArchive	The linker associated with the package being written to disk.
@@ -291,14 +291,14 @@ public:
 	 *        bytes are written to DataArchive.
 	 * @return True if the builder was created and appended successfully and false if any error was encountered
 	 */
-	[[nodiscard]] bool BuildAndAppendTrailer(FLinkerSave* Linker, FArchive& DataArchive, int64& InOutPackageFileOffset);
+	[[nodiscard]] COREUOBJECT_API bool BuildAndAppendTrailer(FLinkerSave* Linker, FArchive& DataArchive, int64& InOutPackageFileOffset);
 
 	/** Returns if the builder has any payload entries or not */
-	[[nodiscard]] bool IsEmpty() const;
+	[[nodiscard]] COREUOBJECT_API bool IsEmpty() const;
 
-	[[nodiscard]] bool IsLocalPayloadEntry(const FIoHash& Identifier) const;
-	[[nodiscard]] bool IsReferencedPayloadEntry(const FIoHash& Identifier) const;
-	[[nodiscard]] bool IsVirtualizedPayloadEntry(const FIoHash& Identifier) const;
+	[[nodiscard]] COREUOBJECT_API bool IsLocalPayloadEntry(const FIoHash& Identifier) const;
+	[[nodiscard]] COREUOBJECT_API bool IsReferencedPayloadEntry(const FIoHash& Identifier) const;
+	[[nodiscard]] COREUOBJECT_API bool IsVirtualizedPayloadEntry(const FIoHash& Identifier) const;
 
 	/** 
 	 * Returns the length of the trailer (in bytes) that the builder would currently create.
@@ -306,17 +306,17 @@ public:
 	 * NOTE: At the moment this is not const as we need to check for and remove duplicate
 	 * payload entries as we do this before building the trailer not when gathering the entry info.
 	 */
-	[[nodiscard]] uint64 CalculateTrailerLength();
+	[[nodiscard]] COREUOBJECT_API uint64 CalculateTrailerLength();
 
 	/** Returns the total number of payload entries in the builder */
-	[[nodiscard]] int32 GetNumPayloads() const;
+	[[nodiscard]] COREUOBJECT_API int32 GetNumPayloads() const;
 	
 	/** Returns the number of payload entries in the builder with the access mode EPayloadAccessMode::Local */
-	[[nodiscard]] int32 GetNumLocalPayloads() const;
+	[[nodiscard]] COREUOBJECT_API int32 GetNumLocalPayloads() const;
 	/** Returns the number of payload entries in the builder with the access mode EPayloadAccessMode::Referenced */
-	[[nodiscard]] int32 GetNumReferencedPayloads() const;
+	[[nodiscard]] COREUOBJECT_API int32 GetNumReferencedPayloads() const;
 	/** Returns the number of payload entries in the builder with the access mode EPayloadAccessMode::Virtualized */
-	[[nodiscard]] int32 GetNumVirtualizedPayloads() const;
+	[[nodiscard]] COREUOBJECT_API int32 GetNumVirtualizedPayloads() const;
 
 	/** Returns the debug context associated with the builder, used for adding further description to error messages */
 	[[nodiscard]] const FString& GetDebugContext() const
@@ -376,11 +376,11 @@ private:
 	};
 
 	/** Returns the total length of the header if we were to build a trailer right now */
-	uint32 CalculatePotentialHeaderSize() const;
+	COREUOBJECT_API uint32 CalculatePotentialHeaderSize() const;
 	/** Returns the total length of all payloads combined if we were to build a trailer right now */
-	uint64 CalculatePotentialPayloadSize() const;
+	COREUOBJECT_API uint64 CalculatePotentialPayloadSize() const;
 
-	void RemoveDuplicateEntries();
+	COREUOBJECT_API void RemoveDuplicateEntries();
 
 	// Members used when building the trailer
 
@@ -403,7 +403,7 @@ private:
  * domain then it's values should be valid, but when loading non-virtualized payloads they need to come from the workspace 
  * domain package.
  */
-class COREUOBJECT_API FPackageTrailer
+class FPackageTrailer
 {
 public:
 	/** 
@@ -418,13 +418,13 @@ public:
 	}
 
 	/** Try to load a trailer from a given package path. Note that it will always try to load the trailer from the workspace domain */
-	[[nodiscard]] static bool TryLoadFromPackage(const FPackagePath& PackagePath, FPackageTrailer& OutTrailer);
+	[[nodiscard]] static COREUOBJECT_API bool TryLoadFromPackage(const FPackagePath& PackagePath, FPackageTrailer& OutTrailer);
 
 	/** Try to load a trailer from a given file path. */
-	[[nodiscard]] static bool TryLoadFromFile(const FString& Path, FPackageTrailer& OutTrailer);
+	[[nodiscard]] static COREUOBJECT_API bool TryLoadFromFile(const FString& Path, FPackageTrailer& OutTrailer);
 
 	/** Try to load a trailer from a given archive. Assumes that the trailer is at the end of the archive */
-	[[nodiscard]] static bool TryLoadFromArchive(FArchive& Ar, FPackageTrailer& OutTrailer);
+	[[nodiscard]] static COREUOBJECT_API bool TryLoadFromArchive(FArchive& Ar, FPackageTrailer& OutTrailer);
 
 	FPackageTrailer() = default;
 	~FPackageTrailer() = default;
@@ -452,7 +452,7 @@ public:
 	 * @return		True if a valid trailer was found and was able to be loaded, otherwise false. If the trailer was found but failed
 	 *				to load then the archive will be set to the error state.
 	 */
-	[[nodiscard]] bool TryLoad(FArchive& Ar);
+	[[nodiscard]] COREUOBJECT_API bool TryLoad(FArchive& Ar);
 
 	/** 
 	 * Serializes the trailer from the given archive BUT assumes that the seek position of the archive is at the end of the trailer
@@ -462,7 +462,7 @@ public:
 	 * @return		True if a valid trailer was found and was able to be loaded, otherwise false. If the trailer was found but failed
 	 *				to load then the archive will be set to the error state.
 	 */
-	[[nodiscard]] bool TryLoadBackwards(FArchive& Ar);
+	[[nodiscard]] COREUOBJECT_API bool TryLoadBackwards(FArchive& Ar);
 
 	/** 
 	 * Loads a payload that is stored locally within the package trailer. Payloads stored externally (either referenced
@@ -474,7 +474,7 @@ public:
 	 * @return	The payload in the form of a FCompressedBuffer. If the payload does not exist in the trailer or is not
 	 *			stored locally in the trailer then the FCompressedBuffer will be null.
 	 */
-	[[nodiscard]] FCompressedBuffer LoadLocalPayload(const FIoHash& Id, FArchive& Ar) const;
+	[[nodiscard]] COREUOBJECT_API FCompressedBuffer LoadLocalPayload(const FIoHash& Id, FArchive& Ar) const;
 
 	/** 
 	 * Calling this indicates that the payload has been virtualized and will no longer be stored on disk. 
@@ -482,38 +482,38 @@ public:
 	 * @param Identifier The payload that has been virtualized
 	 * @return True if the payload was in the trailer, otherwise false
 	 */
-	[[nodiscard]] bool UpdatePayloadAsVirtualized(const FIoHash& Identifier);
+	[[nodiscard]] COREUOBJECT_API bool UpdatePayloadAsVirtualized(const FIoHash& Identifier);
 
 	/**
 	 * Iterates over all payloads in the trailer and invokes the provoided callback on them
 	 */
-	void ForEachPayload(TFunctionRef<void(const FIoHash&, uint64, uint64, EPayloadAccessMode, UE::Virtualization::EPayloadFilterReason)> Callback) const;
+	COREUOBJECT_API void ForEachPayload(TFunctionRef<void(const FIoHash&, uint64, uint64, EPayloadAccessMode, UE::Virtualization::EPayloadFilterReason)> Callback) const;
 
 	/** Attempt to find the status of the given payload. @See EPayloadStatus */
-	[[nodiscard]] EPayloadStatus FindPayloadStatus(const FIoHash& Id) const;
+	[[nodiscard]] COREUOBJECT_API EPayloadStatus FindPayloadStatus(const FIoHash& Id) const;
 
 	/** Returns the absolute offset of the payload in the package file, invalid and virtualized payloads will return INDEX_NONE */
-	[[nodiscard]] int64 FindPayloadOffsetInFile(const FIoHash& Id) const;
+	[[nodiscard]] COREUOBJECT_API int64 FindPayloadOffsetInFile(const FIoHash& Id) const;
 
 	/** Returns the size of the payload on as stored on disk, invalid and virtualized payloads will return INDEX_NONE */
-	[[nodiscard]] int64 FindPayloadSizeOnDisk(const FIoHash& Id) const;
+	[[nodiscard]] COREUOBJECT_API int64 FindPayloadSizeOnDisk(const FIoHash& Id) const;
 
 	/** Returns the total size of the of the trailer on disk in bytes */
-	[[nodiscard]] int64 GetTrailerLength() const;
+	[[nodiscard]] COREUOBJECT_API int64 GetTrailerLength() const;
 
-	[[nodiscard]] FPayloadInfo GetPayloadInfo(const FIoHash& Id) const;
+	[[nodiscard]] COREUOBJECT_API FPayloadInfo GetPayloadInfo(const FIoHash& Id) const;
 
 	/** Returns an array of the payloads with the given storage type. @See EPayloadStoragetype */
-	[[nodiscard]] TArray<FIoHash> GetPayloads(EPayloadStorageType StorageType) const;
+	[[nodiscard]] COREUOBJECT_API TArray<FIoHash> GetPayloads(EPayloadStorageType StorageType) const;
 
 	/** Returns the number of payloads that the trailer owns with the given storage type. @See EPayloadStoragetype */
-	[[nodiscard]] int32 GetNumPayloads(EPayloadStorageType Type) const;
+	[[nodiscard]] COREUOBJECT_API int32 GetNumPayloads(EPayloadStorageType Type) const;
 
 	/** Returns an array of the payloads that match the given filter type. @See EPayloadFilter */
-	[[nodiscard]] TArray<FIoHash> GetPayloads(EPayloadFilter Filter) const;
+	[[nodiscard]] COREUOBJECT_API TArray<FIoHash> GetPayloads(EPayloadFilter Filter) const;
 
 	/** Returns the number of payloads that the trailer owns that match the given filter type. @See EPayloadFilter */
-	[[nodiscard]] int32 GetNumPayloads(EPayloadFilter Filter) const;
+	[[nodiscard]] COREUOBJECT_API int32 GetNumPayloads(EPayloadFilter Filter) const;
 
 	struct FHeader
 	{
@@ -570,7 +570,7 @@ private:
 	friend class FPackageTrailerBuilder;
 
 	/** Create a valid footer for the current trailer */
-	FFooter CreateFooter() const;
+	COREUOBJECT_API FFooter CreateFooter() const;
 
 	/** Where in the workspace domain package file the trailer is located */
 	int64 TrailerPositionInFile = INDEX_NONE;
diff --git a/Engine/Source/Runtime/CoreUObject/Public/UObject/PropertyTempVal.h b/Engine/Source/Runtime/CoreUObject/Public/UObject/PropertyTempVal.h
--- a/Engine/Source/Runtime/CoreUObject/Public/UObject/PropertyTempVal.h
+++ b/Engine/Source/Runtime/CoreUObject/Public/UObject/PropertyTempVal.h
@@ -13,20 +13,20 @@ class UObject;
 // A helper struct which owns a single instance of the type pointed to by a property.
 // The instance is properly constructed, destructed and can be serialized and have other
 // functions called on it.
-struct COREUOBJECT_API FPropertyTempVal
+struct FPropertyTempVal
 {
-	explicit FPropertyTempVal(FProperty* InProp);
-	~FPropertyTempVal();
+	COREUOBJECT_API explicit FPropertyTempVal(FProperty* InProp);
+	COREUOBJECT_API ~FPropertyTempVal();
 
 	// Uncopyable
 	FPropertyTempVal(const FPropertyTempVal&) = delete;
 	FPropertyTempVal& operator=(const FPropertyTempVal&) = delete;
 
 	// Serializes the instance
-	void Serialize(FArchive& Ar, const void* Defaults = nullptr);
+	COREUOBJECT_API void Serialize(FArchive& Ar, const void* Defaults = nullptr);
 
 	// Exports the text of the instance
-	void ExportText(FString& ValueStr, const void* Defaults = nullptr, UObject* Parent = nullptr, int32 PortFlags = 0, UObject* ExportRootScope = nullptr);
+	COREUOBJECT_API void ExportText(FString& ValueStr, const void* Defaults = nullptr, UObject* Parent = nullptr, int32 PortFlags = 0, UObject* ExportRootScope = nullptr);
 
 	// Returns a pointer to the internal instance
 	FORCEINLINE void* Get()
diff --git a/Engine/Source/Runtime/CoreUObject/Public/UObject/ReferenceChainSearch.h b/Engine/Source/Runtime/CoreUObject/Public/UObject/ReferenceChainSearch.h
--- a/Engine/Source/Runtime/CoreUObject/Public/UObject/ReferenceChainSearch.h
+++ b/Engine/Source/Runtime/CoreUObject/Public/UObject/ReferenceChainSearch.h
@@ -326,7 +326,7 @@ public:
 	COREUOBJECT_API FString GetRootPath(TFunctionRef<bool(FCallbackParams& Params)> ReferenceCallback, UObject* TargetObject = nullptr) const;
 
 	/** Returns all reference chains */
-	COREUOBJECT_API const TArray<FReferenceChain*>& GetReferenceChains() const
+	const TArray<FReferenceChain*>& GetReferenceChains() const
 	{
 		return ReferenceChains;
 	}
diff --git a/Engine/Source/Runtime/CoreUObject/Public/UObject/ReferencerFinder.h b/Engine/Source/Runtime/CoreUObject/Public/UObject/ReferencerFinder.h
--- a/Engine/Source/Runtime/CoreUObject/Public/UObject/ReferencerFinder.h
+++ b/Engine/Source/Runtime/CoreUObject/Public/UObject/ReferencerFinder.h
@@ -23,12 +23,12 @@ ENUM_CLASS_FLAGS(EReferencerFinderFlags);
 /**
  * Helper class for finding all objects referencing any of the objects in Referencees list
  */
-class COREUOBJECT_API FReferencerFinder
+class FReferencerFinder
 {
 public:
-	static TArray<UObject*> GetAllReferencers(const TArray<UObject*>& Referencees, const TSet<UObject*>* ObjectsToIgnore, EReferencerFinderFlags Flags = EReferencerFinderFlags::None);
-	static TArray<UObject*> GetAllReferencers(const TSet<UObject*>& Referencees, const TSet<UObject*>* ObjectsToIgnore, EReferencerFinderFlags Flags = EReferencerFinderFlags::None);
+	static COREUOBJECT_API TArray<UObject*> GetAllReferencers(const TArray<UObject*>& Referencees, const TSet<UObject*>* ObjectsToIgnore, EReferencerFinderFlags Flags = EReferencerFinderFlags::None);
+	static COREUOBJECT_API TArray<UObject*> GetAllReferencers(const TSet<UObject*>& Referencees, const TSet<UObject*>* ObjectsToIgnore, EReferencerFinderFlags Flags = EReferencerFinderFlags::None);
 
 	/** Called when the initial load phase is complete and we're done registering native UObject classes */
-	static void NotifyRegistrationComplete();
-};
\ No newline at end of file
+	static COREUOBJECT_API void NotifyRegistrationComplete();
+};
diff --git a/Engine/Source/Runtime/CoreUObject/Public/UObject/Reload.h b/Engine/Source/Runtime/CoreUObject/Public/UObject/Reload.h
--- a/Engine/Source/Runtime/CoreUObject/Public/UObject/Reload.h
+++ b/Engine/Source/Runtime/CoreUObject/Public/UObject/Reload.h
@@ -15,7 +15,7 @@
  * the module manager while a reload is in progress.
  */
 
-class COREUOBJECT_API IReload
+class IReload
 {
 public:
 
@@ -79,4 +79,4 @@ FORCEINLINE bool ReloadNotifyFunctionRemap(FNativeFuncPtr NewFunctionPointer, FN
 	}
 #endif
 	return false;
-}
\ No newline at end of file
+}
diff --git a/Engine/Source/Runtime/CoreUObject/Public/UObject/SavePackage.h b/Engine/Source/Runtime/CoreUObject/Public/UObject/SavePackage.h
--- a/Engine/Source/Runtime/CoreUObject/Public/UObject/SavePackage.h
+++ b/Engine/Source/Runtime/CoreUObject/Public/UObject/SavePackage.h
@@ -191,25 +191,25 @@ public:
 	/** Get the default settings by save when none are specified. */
 	COREUOBJECT_API static FSavePackageSettings& GetDefaultSettings();
 
-	COREUOBJECT_API bool IsDefault() const
+	bool IsDefault() const
 	{
 		return ExternalImportValidations.Num() == 0 && ExternalExportValidations.Num() == 0;
 	}
 
-	COREUOBJECT_API const TArray<TFunction<ExternalImportValidationFunc>>& GetExternalImportValidations() const
+	const TArray<TFunction<ExternalImportValidationFunc>>& GetExternalImportValidations() const
 	{
 		return ExternalImportValidations;
 	}
-	COREUOBJECT_API const TArray<TFunction<ExternalExportValidationFunc>>& GetExternalExportValidations() const
+	const TArray<TFunction<ExternalExportValidationFunc>>& GetExternalExportValidations() const
 	{
 		return ExternalExportValidations;
 	}
 
-	COREUOBJECT_API void AddExternalImportValidation(TFunction<ExternalImportValidationFunc> InValidation)
+	void AddExternalImportValidation(TFunction<ExternalImportValidationFunc> InValidation)
 	{
 		ExternalImportValidations.Add(MoveTemp(InValidation));
 	}
-	COREUOBJECT_API void AddExternalExportValidation(TFunction<ExternalExportValidationFunc> InValidation)
+	void AddExternalExportValidation(TFunction<ExternalExportValidationFunc> InValidation)
 	{
 		ExternalExportValidations.Add(MoveTemp(InValidation));
 	}
@@ -247,18 +247,18 @@ public:
 	{
 		return Validator.Get();	
 	}
-	COREUOBJECT_API void SetValidator(TUniquePtr<ISavePackageValidator>&& InValidator)
+	void SetValidator(TUniquePtr<ISavePackageValidator>&& InValidator)
 	{
 		Validator = MoveTemp(InValidator);
 	}
 	PRAGMA_ENABLE_DEPRECATION_WARNINGS
 
-	COREUOBJECT_API const TArray<TFunction<FSavePackageSettings::ExternalImportValidationFunc>>& GetExternalImportValidations() const
+	const TArray<TFunction<FSavePackageSettings::ExternalImportValidationFunc>>& GetExternalImportValidations() const
 	{
 		return SavePackageSettings.GetExternalImportValidations();
 	}
 
-	COREUOBJECT_API const TArray<TFunction<FSavePackageSettings::ExternalExportValidationFunc>>& GetExternalExportValidations() const
+	const TArray<TFunction<FSavePackageSettings::ExternalExportValidationFunc>>& GetExternalExportValidations() const
 	{
 		return SavePackageSettings.GetExternalExportValidations();
 	}
diff --git a/Engine/Source/Runtime/CoreUObject/Public/UObject/Script.h b/Engine/Source/Runtime/CoreUObject/Public/UObject/Script.h
--- a/Engine/Source/Runtime/CoreUObject/Public/UObject/Script.h
+++ b/Engine/Source/Runtime/CoreUObject/Public/UObject/Script.h
@@ -406,11 +406,11 @@ protected:
 };
 
 // Information about a blueprint instrumentation signal
-struct COREUOBJECT_API FScriptInstrumentationSignal
+struct FScriptInstrumentationSignal
 {
 public:
 
-	FScriptInstrumentationSignal(EScriptInstrumentation::Type InEventType, const UObject* InContextObject, const struct FFrame& InStackFrame, const FName EventNameIn = NAME_None);
+	COREUOBJECT_API FScriptInstrumentationSignal(EScriptInstrumentation::Type InEventType, const UObject* InContextObject, const struct FFrame& InStackFrame, const FName EventNameIn = NAME_None);
 
 	FScriptInstrumentationSignal(EScriptInstrumentation::Type InEventType, const UObject* InContextObject, UFunction* InFunction, const int32 LinkId = INDEX_NONE)
 		: EventType(InEventType)
@@ -441,16 +441,16 @@ public:
 	const FFrame& GetStackFrame() const { return *StackFramePtr; }
 
 	/** Returns the owner class name of the active instance */
-	const UClass* GetClass() const;
+	COREUOBJECT_API const UClass* GetClass() const;
 
 	/** Returns the function scope class */
-	const UClass* GetFunctionClassScope() const;
+	COREUOBJECT_API const UClass* GetFunctionClassScope() const;
 
 	/** Returns the name of the active function */
-	FName GetFunctionName() const;
+	COREUOBJECT_API FName GetFunctionName() const;
 
 	/** Returns the script code offset */
-	int32 GetScriptCodeOffset() const;
+	COREUOBJECT_API int32 GetScriptCodeOffset() const;
 
 	/** Returns the latent link id for latent events */
 	int32 GetLatentLinkId() const { return LatentLinkId; }
@@ -471,7 +471,7 @@ protected:
 };
 
 // Blueprint core runtime delegates
-class COREUOBJECT_API FBlueprintCoreDelegates
+class FBlueprintCoreDelegates
 {
 public:
 	// Callback for debugging events such as a breakpoint (Object that triggered event, active stack frame, Info)
@@ -486,20 +486,20 @@ public:
 
 public:
 	// Called when a script exception occurs
-	static FOnScriptDebuggingEvent OnScriptException;
+	static COREUOBJECT_API FOnScriptDebuggingEvent OnScriptException;
 	// Called when a script profiling event is fired
-	static FOnScriptInstrumentEvent OnScriptProfilingEvent;
+	static COREUOBJECT_API FOnScriptInstrumentEvent OnScriptProfilingEvent;
 	// Called when a script profiler is enabled/disabled
-	static FOnToggleScriptProfiler OnToggleScriptProfiler;
+	static COREUOBJECT_API FOnToggleScriptProfiler OnToggleScriptProfiler;
 
 	UE_DEPRECATED(4.26, "OnScriptExecutionEnd is deprecated, bind to delegate inside FBlueprintContextTracker instead")
-	static FOnScriptExecutionEnd OnScriptExecutionEnd;
+	static COREUOBJECT_API FOnScriptExecutionEnd OnScriptExecutionEnd;
 
 public:
-	static void ThrowScriptException(const UObject* ActiveObject, struct FFrame& StackFrame, const FBlueprintExceptionInfo& Info);
-	static void InstrumentScriptEvent(const FScriptInstrumentationSignal& Info);
-	static void SetScriptMaximumLoopIterations( const int32 MaximumLoopIterations );
-	static bool IsDebuggingEnabled();
+	static COREUOBJECT_API void ThrowScriptException(const UObject* ActiveObject, struct FFrame& StackFrame, const FBlueprintExceptionInfo& Info);
+	static COREUOBJECT_API void InstrumentScriptEvent(const FScriptInstrumentationSignal& Info);
+	static COREUOBJECT_API void SetScriptMaximumLoopIterations( const int32 MaximumLoopIterations );
+	static COREUOBJECT_API bool IsDebuggingEnabled();
 };
 
 #if DO_BLUEPRINT_GUARD
@@ -597,11 +597,11 @@ private:
 
 
 // Scoped struct to allow execution of script in editor, while resetting the runaway loop counts
-struct COREUOBJECT_API FEditorScriptExecutionGuard
+struct FEditorScriptExecutionGuard
 {
 public:
-	FEditorScriptExecutionGuard();
-	~FEditorScriptExecutionGuard();
+	COREUOBJECT_API FEditorScriptExecutionGuard();
+	COREUOBJECT_API ~FEditorScriptExecutionGuard();
 
 private:
 	bool bOldGAllowScriptExecutionInEditor;
diff --git a/Engine/Source/Runtime/CoreUObject/Public/UObject/SoftObjectPath.h b/Engine/Source/Runtime/CoreUObject/Public/UObject/SoftObjectPath.h
--- a/Engine/Source/Runtime/CoreUObject/Public/UObject/SoftObjectPath.h
+++ b/Engine/Source/Runtime/CoreUObject/Public/UObject/SoftObjectPath.h
@@ -38,7 +38,7 @@ struct FUObjectSerializeContext;
  * This is stored internally as an FTopLevelAssetPath pointing to the top level asset (/package/path.assetname) and an optional string subobject path.
  * If the MetaClass metadata is applied to a FProperty with this the UI will restrict to that type of asset.
  */
-struct COREUOBJECT_API FSoftObjectPath
+struct FSoftObjectPath
 {
 	FSoftObjectPath() = default;
 	FSoftObjectPath(const FSoftObjectPath& Other) = default;
@@ -75,7 +75,7 @@ PRAGMA_ENABLE_DEPRECATION_WARNINGS;
 
 	/** Construct from an asset FName and subobject pair */
 	UE_DEPRECATED(5.1, "Asset path FNames have been deprecated. This constructor should be used only temporarily to fix up old codepaths that produce an FName.")
-	FSoftObjectPath(FName InAssetPathName, FString InSubPathString);
+	COREUOBJECT_API FSoftObjectPath(FName InAssetPathName, FString InSubPathString);
 	
 	template <typename T>
 	FSoftObjectPath(const TObjectPtr<T>& InObject)
@@ -115,17 +115,17 @@ PRAGMA_ENABLE_DEPRECATION_WARNINGS;
 	}
 
 	/** Returns string representation of reference, in form /package/path.assetname[:subpath] */
-	FString ToString() const;
+	COREUOBJECT_API FString ToString() const;
 
 	UE_DEPRECATED(5.1, "Asset path FNames have been deprecated. This function should be used only temporarily to interface with old APIs that require an FName.")
 	FName ToFName() const { return *ToString(); }
 
 	/** Append string representation of reference, in form /package/path.assetname[:subpath] */
-	void ToString(FStringBuilderBase& Builder) const;
+	COREUOBJECT_API void ToString(FStringBuilderBase& Builder) const;
 
 	/** Append string representation of reference, in form /package/path.assetname[:subpath] */
-	void AppendString(FString& Builder) const;
-	void AppendString(FStringBuilderBase& Builder) const;
+	COREUOBJECT_API void AppendString(FString& Builder) const;
+	COREUOBJECT_API void AppendString(FStringBuilderBase& Builder) const;
 
 	/** Returns the top-level asset part of this path, without the subobject path. */
 	FTopLevelAssetPath GetAssetPath() const
@@ -141,10 +141,10 @@ PRAGMA_ENABLE_DEPRECATION_WARNINGS;
 	
 	/** Returns the entire asset path as an FName, including both package and asset but not sub object */
 	UE_DEPRECATED(5.1, "Asset path FNames have been deprecated. Use GetAssetPath instead.")
-	FName GetAssetPathName() const;
+	COREUOBJECT_API FName GetAssetPathName() const;
 	
 	UE_DEPRECATED(5.1, "Asset path FNames have been deprecated. Use SetAssetPath instead.")
-	void SetAssetPathName(FName InAssetPathName);
+	COREUOBJECT_API void SetAssetPathName(FName InAssetPathName);
 
 	/** Returns string version of asset path, including both package and asset but not sub object */
 	FORCEINLINE FString GetAssetPathString() const
@@ -195,12 +195,12 @@ PRAGMA_ENABLE_DEPRECATION_WARNINGS;
 	}
 
 	/** Sets asset path of this reference based on a string path */
-	void SetPath(const FTopLevelAssetPath& InAssetPath, FString InSubPathString = FString());
-	void SetPath(FWideStringView Path);
-	void SetPath(FAnsiStringView Path);
-	void SetPath(FUtf8StringView Path);
+	COREUOBJECT_API void SetPath(const FTopLevelAssetPath& InAssetPath, FString InSubPathString = FString());
+	COREUOBJECT_API void SetPath(FWideStringView Path);
+	COREUOBJECT_API void SetPath(FAnsiStringView Path);
+	COREUOBJECT_API void SetPath(FUtf8StringView Path);
 	UE_DEPRECATED(5.1, "Asset path FNames have been deprecated. This function should be used only temporarily to fix up old codepaths that produce an FName.")
-	void SetPath(FName Path);
+	COREUOBJECT_API void SetPath(FName Path);
 	void SetPath(const WIDECHAR* Path)			{ SetPath(FWideStringView(Path)); }
 	void SetPath(const ANSICHAR* Path)			{ SetPath(FAnsiStringView(Path)); }
 	void SetPath(const FString& Path)			{ SetPath(FStringView(Path)); }
@@ -210,14 +210,14 @@ PRAGMA_ENABLE_DEPRECATION_WARNINGS;
 	 * @param InLoadContext Optional load context when called from nested load callstack
 	 * @return Loaded UObject, or nullptr if the reference is null or the asset fails to load
 	 */
-	UObject* TryLoad(FUObjectSerializeContext* InLoadContext = nullptr) const;
+	COREUOBJECT_API UObject* TryLoad(FUObjectSerializeContext* InLoadContext = nullptr) const;
 
 	/**
 	 * Attempts to find a currently loaded object that matches this path
 	 *
 	 * @return Found UObject, or nullptr if not currently in memory
 	 */
-	UObject* ResolveObject() const;
+	COREUOBJECT_API UObject* ResolveObject() const;
 
 	/** Resets reference to point to null */
 	void Reset()
@@ -273,41 +273,41 @@ PRAGMA_ENABLE_DEPRECATION_WARNINGS;
 	}
 
 	/** Struct overrides */
-	bool Serialize(FArchive& Ar);
-	bool Serialize(FStructuredArchive::FSlot Slot);
-	bool operator==(FSoftObjectPath const& Other) const;
+	COREUOBJECT_API bool Serialize(FArchive& Ar);
+	COREUOBJECT_API bool Serialize(FStructuredArchive::FSlot Slot);
+	COREUOBJECT_API bool operator==(FSoftObjectPath const& Other) const;
 	bool operator!=(FSoftObjectPath const& Other) const
 	{
 		return !(*this == Other);
 	}
 
-	bool ExportTextItem(FString& ValueStr, FSoftObjectPath const& DefaultValue, UObject* Parent, int32 PortFlags, UObject* ExportRootScope) const;
-	bool ImportTextItem( const TCHAR*& Buffer, int32 PortFlags, UObject* Parent, FOutputDevice* ErrorText, FArchive* InSerializingArchive = nullptr );
-	bool SerializeFromMismatchedTag(struct FPropertyTag const& Tag, FStructuredArchive::FSlot Slot);
+	COREUOBJECT_API bool ExportTextItem(FString& ValueStr, FSoftObjectPath const& DefaultValue, UObject* Parent, int32 PortFlags, UObject* ExportRootScope) const;
+	COREUOBJECT_API bool ImportTextItem( const TCHAR*& Buffer, int32 PortFlags, UObject* Parent, FOutputDevice* ErrorText, FArchive* InSerializingArchive = nullptr );
+	COREUOBJECT_API bool SerializeFromMismatchedTag(struct FPropertyTag const& Tag, FStructuredArchive::FSlot Slot);
 
 	/** Serializes the internal path and also handles save/PIE fixups. Call this from the archiver overrides */
-	void SerializePath(FArchive& Ar);
+	COREUOBJECT_API void SerializePath(FArchive& Ar);
 
 	/** Serializes the internal path without any save/PIE fixups. Only call this directly if you know what you are doing */
-	void SerializePathWithoutFixup(FArchive& Ar);
+	COREUOBJECT_API void SerializePathWithoutFixup(FArchive& Ar);
 
 	/** Fixes up path for saving, call if saving with a method that skips SerializePath. This can modify the path, it will return true if it was modified */
-	bool PreSavePath(bool* bReportSoftObjectPathRedirects = nullptr);
+	COREUOBJECT_API bool PreSavePath(bool* bReportSoftObjectPathRedirects = nullptr);
 
 	/** 
 	 * Handles when a path has been loaded, call if loading with a method that skips SerializePath. This does not modify path but might call callbacks
 	 * @param InArchive The archive that loaded this path
 	 */
-	void PostLoadPath(FArchive* InArchive) const;
+	COREUOBJECT_API void PostLoadPath(FArchive* InArchive) const;
 
 	/** Fixes up this SoftObjectPath to add the PIE prefix depending on what is currently active, returns true if it was modified. The overload that takes an explicit PIE instance is preferred, if it's available. */
-	bool FixupForPIE(TFunctionRef<void(int32, FSoftObjectPath&)> InPreFixupForPIECustomFunction = [](int32, FSoftObjectPath&) {});
+	COREUOBJECT_API bool FixupForPIE(TFunctionRef<void(int32, FSoftObjectPath&)> InPreFixupForPIECustomFunction = [](int32, FSoftObjectPath&) {});
 
 	/** Fixes up this SoftObjectPath to add the PIE prefix for the given PIEInstance index, returns true if it was modified */
-	bool FixupForPIE(int32 PIEInstance, TFunctionRef<void(int32, FSoftObjectPath&)> InPreFixupForPIECustomFunction = [](int32, FSoftObjectPath&) {});
+	COREUOBJECT_API bool FixupForPIE(int32 PIEInstance, TFunctionRef<void(int32, FSoftObjectPath&)> InPreFixupForPIECustomFunction = [](int32, FSoftObjectPath&) {});
 
 	/** Fixes soft object path for CoreRedirects to handle renamed native objects, returns true if it was modified */
-	bool FixupCoreRedirects();
+	COREUOBJECT_API bool FixupCoreRedirects();
 
 	FORCEINLINE friend uint32 GetTypeHash(FSoftObjectPath const& This)
 	{
@@ -327,13 +327,13 @@ PRAGMA_ENABLE_DEPRECATION_WARNINGS;
 	{
 		return 0;
 	}
-	static FSoftObjectPath GetOrCreateIDForObject(const UObject* Object);
+	static COREUOBJECT_API FSoftObjectPath GetOrCreateIDForObject(const UObject* Object);
 	
 	/** Adds list of packages names that have been created specifically for PIE, this is used for editor fixup */
-	static void AddPIEPackageName(FName NewPIEPackageName);
+	static COREUOBJECT_API void AddPIEPackageName(FName NewPIEPackageName);
 	
 	/** Disables special PIE path handling, call when PIE finishes to clear list */
-	static void ClearPIEPackageNames();
+	static COREUOBJECT_API void ClearPIEPackageNames();
 
 private:
 	/** Asset path, patch to a top level object in a package. This is /package/path.assetname */
@@ -343,10 +343,10 @@ private:
 	FString SubPathString;
 
 	/** Package names currently being duplicated, needed by FixupForPIE */
-	static TSet<FName> PIEPackageNames;
+	static COREUOBJECT_API TSet<FName> PIEPackageNames;
 
-	UObject* ResolveObjectInternal() const;
-	UObject* ResolveObjectInternal(const TCHAR* PathString) const;
+	COREUOBJECT_API UObject* ResolveObjectInternal() const;
+	COREUOBJECT_API UObject* ResolveObjectInternal(const TCHAR* PathString) const;
 
 	friend struct Z_Construct_UScriptStruct_FSoftObjectPath_Statics;
 };
@@ -378,7 +378,7 @@ inline FStringBuilderBase& operator<<(FStringBuilderBase& Builder, const FSoftOb
 /**
  * A struct that contains a string reference to a class, can be used to make soft references to classes
  */
-struct COREUOBJECT_API FSoftClassPath : public FSoftObjectPath
+struct FSoftClassPath : public FSoftObjectPath
 {
 	FSoftClassPath() = default;
 	FSoftClassPath(const FSoftClassPath& Other) = default;
@@ -420,18 +420,18 @@ struct COREUOBJECT_API FSoftClassPath : public FSoftObjectPath
 	 * Attempts to find a currently loaded object that matches this object ID
 	 * @return Found UClass, or NULL if not currently loaded
 	 */
-	UClass* ResolveClass() const;
+	COREUOBJECT_API UClass* ResolveClass() const;
 
-	bool SerializeFromMismatchedTag(const FPropertyTag& Tag, FStructuredArchive::FSlot Slot);
+	COREUOBJECT_API bool SerializeFromMismatchedTag(const FPropertyTag& Tag, FStructuredArchive::FSlot Slot);
 
-	static FSoftClassPath GetOrCreateIDForClass(const UClass *InClass);
+	static COREUOBJECT_API FSoftClassPath GetOrCreateIDForClass(const UClass *InClass);
 
 private:
 	/** Forbid creation for UObject. This class is for UClass only. Use FSoftObjectPath instead. */
 	FSoftClassPath(const UObject* InObject) { }
 
 	/** Forbidden. This class is for UClass only. Use FSoftObjectPath instead. */
-	static FSoftObjectPath GetOrCreateIDForObject(const UObject *Object);
+	static COREUOBJECT_API FSoftObjectPath GetOrCreateIDForObject(const UObject *Object);
 };
 
 UE_DEPRECATED(5.0, "FStringAssetReference was renamed to FSoftObjectPath as it is now not always a string and can also refer to a subobject")
diff --git a/Engine/Source/Runtime/CoreUObject/Public/UObject/SparseDelegate.h b/Engine/Source/Runtime/CoreUObject/Public/UObject/SparseDelegate.h
--- a/Engine/Source/Runtime/CoreUObject/Public/UObject/SparseDelegate.h
+++ b/Engine/Source/Runtime/CoreUObject/Public/UObject/SparseDelegate.h
@@ -35,50 +35,50 @@ struct FSparseDelegate;
 
 
 /** Helper class for handling sparse delegate bindings */
-struct COREUOBJECT_API FSparseDelegateStorage
+struct FSparseDelegateStorage
 {
 public:
 
 	/** Registers the sparse delegate so that the offset can be determined. */
-	static void RegisterDelegateOffset(const UObject* OwningObject, FName DelegateName, size_t OffsetToOwner);
+	static COREUOBJECT_API void RegisterDelegateOffset(const UObject* OwningObject, FName DelegateName, size_t OffsetToOwner);
 
 	/** Binds a sparse delegate to the owner. Returns if the delegate was successfully bound. */
-	static bool Add(const UObject* DelegateOwner, const FName DelegateName, FScriptDelegate Delegate);
+	static COREUOBJECT_API bool Add(const UObject* DelegateOwner, const FName DelegateName, FScriptDelegate Delegate);
 
 	/** Binds a sparse delegate to the owner, verifying first that the delegate is not already bound. Returns if the delegate was successfully bound. */
-	static bool AddUnique(const UObject* DelegateOwner, const FName DelegateName, FScriptDelegate Delegate);
+	static COREUOBJECT_API bool AddUnique(const UObject* DelegateOwner, const FName DelegateName, FScriptDelegate Delegate);
 
 	/** Returns whether a sparse delegate is bound to the owner. */
-	static bool Contains(const UObject* DelegateOwner, const FName DelegateName, const FScriptDelegate& Delegate);
-	static bool Contains(const UObject* DelegateOwner, const FName DelegateName, const UObject* InObject, FName InFunctionName);
+	static COREUOBJECT_API bool Contains(const UObject* DelegateOwner, const FName DelegateName, const FScriptDelegate& Delegate);
+	static COREUOBJECT_API bool Contains(const UObject* DelegateOwner, const FName DelegateName, const UObject* InObject, FName InFunctionName);
 
 	/** Removes a delegate binding from the owner's sparse delegate storage. Returns true if there are still bindings to the delegate. */
-	static bool Remove(const UObject* DelegateOwner, const FName DelegateName, const FScriptDelegate& Delegate);
-	static bool Remove(const UObject* DelegateOwner, const FName DelegateName, const UObject* InObject, FName InFunctionName);
+	static COREUOBJECT_API bool Remove(const UObject* DelegateOwner, const FName DelegateName, const FScriptDelegate& Delegate);
+	static COREUOBJECT_API bool Remove(const UObject* DelegateOwner, const FName DelegateName, const UObject* InObject, FName InFunctionName);
 
 	/** Removes all sparse delegate binding from the owner for a given object. Returns true if there are still bindings to the delegate. */
-	static bool RemoveAll(const UObject* DelegateOwner, const FName DelegateName, const UObject* UserObject);
+	static COREUOBJECT_API bool RemoveAll(const UObject* DelegateOwner, const FName DelegateName, const UObject* UserObject);
 
 	/** Clear all of the named sparse delegate bindings from the owner. */
-	static void Clear(const UObject* DelegateOwner, const FName DelegateName);
+	static COREUOBJECT_API void Clear(const UObject* DelegateOwner, const FName DelegateName);
 
 	/** Acquires the actual Multicast Delegate from the annotation if any delegates are bound to it. Will be null if no entry exists in the annotation for this object/delegatename. */
-	static FMulticastScriptDelegate* GetMulticastDelegate(const UObject* DelegateOwner, const FName DelegateName);
+	static COREUOBJECT_API FMulticastScriptDelegate* GetMulticastDelegate(const UObject* DelegateOwner, const FName DelegateName);
 
 	/** Acquires the actual Multicast Delegate from the annotation if any delegates are bound to it as a shared pointer. Will be null if no entry exists in the annotation for this object/delegatename. */
-	static TSharedPtr<FMulticastScriptDelegate> GetSharedMulticastDelegate(const UObject* DelegateOwner, const FName DelegateName);
+	static COREUOBJECT_API TSharedPtr<FMulticastScriptDelegate> GetSharedMulticastDelegate(const UObject* DelegateOwner, const FName DelegateName);
 
 	/** Directly sets the Multicast Delegate for this object/delegatename pair. If the delegate is unbound it will be assigned/inserted anyways. */
-	static void SetMulticastDelegate(const UObject* DelegateOwner, const FName DelegateName, FMulticastScriptDelegate Delegate);
+	static COREUOBJECT_API void SetMulticastDelegate(const UObject* DelegateOwner, const FName DelegateName, FMulticastScriptDelegate Delegate);
 
 	/** Using the registry of sparse delegates recover the FSparseDelegate address from the UObject and name. */
-	static FSparseDelegate* ResolveSparseDelegate(const UObject* OwningObject, FName DelegateName);
+	static COREUOBJECT_API FSparseDelegate* ResolveSparseDelegate(const UObject* OwningObject, FName DelegateName);
 
 	/** Using the registry of sparse delegates recover the UObject owner from the FSparseDelegate address owning class and delegate names. */
-	static UObject* ResolveSparseOwner(const FSparseDelegate& SparseDelegate, const FName OwningClassName, const FName DelegateName);
+	static COREUOBJECT_API UObject* ResolveSparseOwner(const FSparseDelegate& SparseDelegate, const FName OwningClassName, const FName DelegateName);
 
 	/** Outputs a report about which delegates are bound. */
-	static void SparseDelegateReport(const TArray<FString>&, UWorld*, FOutputDevice&);
+	static COREUOBJECT_API void SparseDelegateReport(const TArray<FString>&, UWorld*, FOutputDevice&);
 
 private:
 
@@ -95,19 +95,19 @@ private:
 	friend struct FObjectListener;
 
 	/** A listener to get notified when objects have been deleted and remove them from the map */
-	static FObjectListener SparseDelegateObjectListener;
+	static COREUOBJECT_API FObjectListener SparseDelegateObjectListener;
 
 	/** Critical Section for locking access to the sparse delegate map */
-	static FCriticalSection SparseDelegateMapCritical;
+	static COREUOBJECT_API FCriticalSection SparseDelegateMapCritical;
 
 	/** Delegate map is a map of Delegate names to a shared pointer of the multicast script delegate */
 	typedef TMap<FName, TSharedPtr<FMulticastScriptDelegate>> FSparseDelegateMap;
 
 	/** Map of objects to the map of delegates that are bound to that object */
-	static TMap<const UObjectBase*, FSparseDelegateMap> SparseDelegates;
+	static COREUOBJECT_API TMap<const UObjectBase*, FSparseDelegateMap> SparseDelegates;
 	
 	/** Sparse delegate offsets are indexed by ActorClass/DelegateName pair */
-	static TMap<TPair<FName, FName>, size_t> SparseDelegateObjectOffsets;
+	static COREUOBJECT_API TMap<TPair<FName, FName>, size_t> SparseDelegateObjectOffsets;
 };
 
 /** Base implementation for all sparse delegate types */
@@ -463,4 +463,4 @@ private:
 
 #if UE_ENABLE_INCLUDE_ORDER_DEPRECATED_IN_5_2
 #include "CoreMinimal.h"
-#endif
\ No newline at end of file
+#endif
diff --git a/Engine/Source/Runtime/CoreUObject/Public/UObject/TextProperty.h b/Engine/Source/Runtime/CoreUObject/Public/UObject/TextProperty.h
--- a/Engine/Source/Runtime/CoreUObject/Public/UObject/TextProperty.h
+++ b/Engine/Source/Runtime/CoreUObject/Public/UObject/TextProperty.h
@@ -20,9 +20,9 @@ struct FPropertyTag;
 // need to break this out a different type so that the DECLARE_CASTED_CLASS_INTRINSIC macro can digest the comma
 typedef TProperty<FText, FProperty> FTextProperty_Super;
 
-class COREUOBJECT_API FTextProperty : public FTextProperty_Super
+class FTextProperty : public FTextProperty_Super
 {
-	DECLARE_FIELD(FTextProperty, FTextProperty_Super, CASTCLASS_FTextProperty)
+	DECLARE_FIELD_API(FTextProperty, FTextProperty_Super, CASTCLASS_FTextProperty, COREUOBJECT_API)
 
 public:
 
@@ -35,7 +35,7 @@ public:
 	}
 
 	UE_DEPRECATED(5.1, "Compiled-in property constructor is deprecated, use other constructors instead.")
-	FTextProperty(FFieldVariant InOwner, const FName& InName, EObjectFlags InObjectFlags, int32 InOffset, EPropertyFlags InFlags)
+	COREUOBJECT_API FTextProperty(FFieldVariant InOwner, const FName& InName, EObjectFlags InObjectFlags, int32 InOffset, EPropertyFlags InFlags)
 PRAGMA_DISABLE_DEPRECATION_WARNINGS
 		: FTextProperty_Super(InOwner, InName, InObjectFlags, InOffset, InFlags)
 PRAGMA_ENABLE_DEPRECATION_WARNINGS
@@ -57,14 +57,14 @@ PRAGMA_ENABLE_DEPRECATION_WARNINGS
 #endif // WITH_EDITORONLY_DATA
 
 	// FProperty interface
-	virtual EConvertFromTypeResult ConvertFromType(const FPropertyTag& Tag, FStructuredArchive::FSlot Slot, uint8* Data, UStruct* DefaultsStruct) override;
-	virtual bool Identical( const void* A, const void* B, uint32 PortFlags ) const override;
-	virtual void SerializeItem(FStructuredArchive::FSlot Slot, void* Value, void const* Defaults) const override;
+	COREUOBJECT_API virtual EConvertFromTypeResult ConvertFromType(const FPropertyTag& Tag, FStructuredArchive::FSlot Slot, uint8* Data, UStruct* DefaultsStruct) override;
+	COREUOBJECT_API virtual bool Identical( const void* A, const void* B, uint32 PortFlags ) const override;
+	COREUOBJECT_API virtual void SerializeItem(FStructuredArchive::FSlot Slot, void* Value, void const* Defaults) const override;
 protected:
-	virtual void ExportText_Internal( FString& ValueStr, const void* PropertyValueOrContainer, EPropertyPointerType PropertyPointerType, const void* DefaultValue, UObject* Parent, int32 PortFlags, UObject* ExportRootScope ) const override;
-	virtual const TCHAR* ImportText_Internal(const TCHAR* Buffer, void* ContainerOrPropertyPtr, EPropertyPointerType PropertyPointerType, UObject* OwnerObject, int32 PortFlags, FOutputDevice* ErrorText) const override;
+	COREUOBJECT_API virtual void ExportText_Internal( FString& ValueStr, const void* PropertyValueOrContainer, EPropertyPointerType PropertyPointerType, const void* DefaultValue, UObject* Parent, int32 PortFlags, UObject* ExportRootScope ) const override;
+	COREUOBJECT_API virtual const TCHAR* ImportText_Internal(const TCHAR* Buffer, void* ContainerOrPropertyPtr, EPropertyPointerType PropertyPointerType, UObject* OwnerObject, int32 PortFlags, FOutputDevice* ErrorText) const override;
 public:
-	virtual FString GetCPPTypeForwardDeclaration() const override;
+	COREUOBJECT_API virtual FString GetCPPTypeForwardDeclaration() const override;
 	// End of FProperty interface
 
 	enum class EIdenticalLexicalCompareMethod : uint8
@@ -73,8 +73,8 @@ public:
 		SourceString,
 		DisplayString
 	};
-	static bool Identical_Implementation(const FText& A, const FText& B, uint32 PortFlags);
-	static bool Identical_Implementation(const FText& A, const FText& B, uint32 PortFlags, EIdenticalLexicalCompareMethod LexicalCompareMethod);
+	static COREUOBJECT_API bool Identical_Implementation(const FText& A, const FText& B, uint32 PortFlags);
+	static COREUOBJECT_API bool Identical_Implementation(const FText& A, const FText& B, uint32 PortFlags, EIdenticalLexicalCompareMethod LexicalCompareMethod);
 };
 
 #if UE_ENABLE_INCLUDE_ORDER_DEPRECATED_IN_5_2
diff --git a/Engine/Source/Runtime/CoreUObject/Public/UObject/TopLevelAssetPath.h b/Engine/Source/Runtime/CoreUObject/Public/UObject/TopLevelAssetPath.h
--- a/Engine/Source/Runtime/CoreUObject/Public/UObject/TopLevelAssetPath.h
+++ b/Engine/Source/Runtime/CoreUObject/Public/UObject/TopLevelAssetPath.h
@@ -29,7 +29,7 @@ class UObject;
  * and the object resolves to the string `/Path/To/Package`
  * This struct is mirrored and exposed to the UE reflection system in NoExportTypes.h
 */
-struct COREUOBJECT_API FTopLevelAssetPath
+struct FTopLevelAssetPath
 {
 	FTopLevelAssetPath() { }
 	FTopLevelAssetPath(TYPE_OF_NULLPTR) { }
@@ -67,13 +67,13 @@ struct COREUOBJECT_API FTopLevelAssetPath
 	FTopLevelAssetPath& operator=(TYPE_OF_NULLPTR) { Reset(); return *this; }
 
 	/** Sets asset path of this reference based on components. */
-	bool TrySetPath(FName InPackageName, FName InAssetName);
+	COREUOBJECT_API bool TrySetPath(FName InPackageName, FName InAssetName);
 	/** Sets asset path to path of existing object. Sets an empty path and returns false if the object is not a package or the direct inner of a package. */
-	bool TrySetPath(const UObject* InObject);
+	COREUOBJECT_API bool TrySetPath(const UObject* InObject);
 	/** Sets asset path of this reference based on a string path. Resets this object and returns false if the string is empty or does not represent a top level asset path. */
-	bool TrySetPath(FWideStringView Path);
-	bool TrySetPath(FUtf8StringView Path);
-	bool TrySetPath(FAnsiStringView Path);
+	COREUOBJECT_API bool TrySetPath(FWideStringView Path);
+	COREUOBJECT_API bool TrySetPath(FUtf8StringView Path);
+	COREUOBJECT_API bool TrySetPath(FAnsiStringView Path);
 	template<typename CharType,	typename = typename std::enable_if<TIsCharType<CharType>::Value>::type>
 	bool TrySetPath(const CharType* Path) { return TrySetPath(TStringView<CharType>(Path)); }
 	bool TrySetPath(const FString& Path) { return TrySetPath(FStringView(Path)); }
@@ -85,14 +85,14 @@ struct COREUOBJECT_API FTopLevelAssetPath
 	FName GetAssetName() const { return AssetName; }
 
 	/** Append the full asset path (e.g. '/Path/To/Package.AssetName') to the string builder. */
-	void AppendString(FStringBuilderBase& Builder) const;
+	COREUOBJECT_API void AppendString(FStringBuilderBase& Builder) const;
 	/** Append the full asset path (e.g. '/Path/To/Package.AssetName') to the string. */
-	void AppendString(FString& OutString) const;
+	COREUOBJECT_API void AppendString(FString& OutString) const;
 
 	/** Return the full asset path (e.g. '/Path/To/Package.AssetName') as a string. */
-	FString ToString() const;
+	COREUOBJECT_API FString ToString() const;
 	/** Copy the full asset path (e.g. '/Path/To/Package.AssetName') into the provided string. */
-	void ToString(FString& OutString) const;
+	COREUOBJECT_API void ToString(FString& OutString) const;
 
 	// Return the full asset path (e.g. '/Path/To/Package.AssetName') as an FName.
 	UE_DEPRECATED(5.1, "FNames containing full asset paths have been replaced by FTopLevelAssetPath/FSoftLevelObjectPath."
@@ -167,9 +167,9 @@ struct COREUOBJECT_API FTopLevelAssetPath
 		return HashCombineFast(GetTypeHash(This.PackageName), GetTypeHash(This.AssetName));
 	}
 
-	bool ExportTextItem(FString& ValueStr, FTopLevelAssetPath const& DefaultValue, UObject* Parent, int32 PortFlags, UObject* ExportRootScope) const;
-	bool ImportTextItem( const TCHAR*& Buffer, int32 PortFlags, UObject* Parent, FOutputDevice* ErrorText, FArchive* InSerializingArchive = nullptr );
-	bool SerializeFromMismatchedTag(const FPropertyTag& Tag, FStructuredArchive::FSlot Slot);
+	COREUOBJECT_API bool ExportTextItem(FString& ValueStr, FTopLevelAssetPath const& DefaultValue, UObject* Parent, int32 PortFlags, UObject* ExportRootScope) const;
+	COREUOBJECT_API bool ImportTextItem( const TCHAR*& Buffer, int32 PortFlags, UObject* Parent, FOutputDevice* ErrorText, FArchive* InSerializingArchive = nullptr );
+	COREUOBJECT_API bool SerializeFromMismatchedTag(const FPropertyTag& Tag, FStructuredArchive::FSlot Slot);
 
 private:
 	/** Name of the package containing the asset e.g. /Path/To/Package */
diff --git a/Engine/Source/Runtime/CoreUObject/Public/UObject/UObjectAllocator.h b/Engine/Source/Runtime/CoreUObject/Public/UObject/UObjectAllocator.h
--- a/Engine/Source/Runtime/CoreUObject/Public/UObject/UObjectAllocator.h
+++ b/Engine/Source/Runtime/CoreUObject/Public/UObject/UObjectAllocator.h
@@ -11,7 +11,7 @@
 
 class UObjectBase;
 
-class COREUOBJECT_API FUObjectAllocator
+class FUObjectAllocator
 {
 public:
 
@@ -31,12 +31,12 @@ public:
 	 *
 	 * @param InPermanentObjectPoolSize size of permanent object pool
 	 */
-	void AllocatePermanentObjectPool(int32 InPermanentObjectPoolSize);
+	COREUOBJECT_API void AllocatePermanentObjectPool(int32 InPermanentObjectPoolSize);
 
 	/**
 	 * Prints a debugf message to allow tuning
 	 */
-	void BootMessage();
+	COREUOBJECT_API void BootMessage();
 
 	UE_DEPRECATED(5.1, "Use the more efficient FPermanentObjectPoolExtents instead")
 	FORCEINLINE bool ResidesInPermanentPool(const UObjectBase *Object) const
@@ -52,14 +52,14 @@ public:
 	 * @param bAllowPermanent if true, allow allocation in the permanent object pool, if it fits
 	 * @return newly allocated UObjectBase (not really a UObjectBase yet, no constructor like thing has been called).
 	 */
-	UObjectBase* AllocateUObject(int32 Size, int32 Alignment, bool bAllowPermanent);
+	COREUOBJECT_API UObjectBase* AllocateUObject(int32 Size, int32 Alignment, bool bAllowPermanent);
 
 	/**
 	 * Returns a UObjectBase to the free store, unless it is in the permanent object pool
 	 *
 	 * @param Object object to free
 	 */
-	void FreeUObject(UObjectBase *Object) const;
+	COREUOBJECT_API void FreeUObject(UObjectBase *Object) const;
 
 private:
 	friend class FPermanentObjectPoolExtents;
diff --git a/Engine/Source/Runtime/CoreUObject/Public/UObject/UObjectArray.h b/Engine/Source/Runtime/CoreUObject/Public/UObject/UObjectArray.h
--- a/Engine/Source/Runtime/CoreUObject/Public/UObject/UObjectArray.h
+++ b/Engine/Source/Runtime/CoreUObject/Public/UObject/UObjectArray.h
@@ -595,7 +595,7 @@ public:
 * that non-GC objects come before GC ones during iteration.
 *
 **/
-class COREUOBJECT_API FUObjectArray
+class FUObjectArray
 {
 	friend class UObject;
 	friend COREUOBJECT_API UObject* StaticAllocateObject(const UClass*, UObject*, FName, EObjectFlags, EInternalObjectFlags, bool, bool*, UPackage*);
@@ -606,7 +606,7 @@ private:
 	 *
 	 * @param Object to reset
 	 */
-	void ResetSerialNumber(UObjectBase* Object);
+	COREUOBJECT_API void ResetSerialNumber(UObjectBase* Object);
 
 public:
 
@@ -661,7 +661,7 @@ public:
 	/**
 	 * Constructor, initializes to no permanent object pool
 	 */
-	FUObjectArray();
+	COREUOBJECT_API FUObjectArray();
 
 	/**
 	 * Allocates and initializes the permanent object pool
@@ -669,23 +669,23 @@ public:
 	 * @param MaxUObjects maximum number of UObjects that can ever exist in the array
 	 * @param MaxObjectsNotConsideredByGC number of objects in the permanent object pool
 	 */
-	void AllocateObjectPool(int32 MaxUObjects, int32 MaxObjectsNotConsideredByGC, bool bPreAllocateObjectArray);
+	COREUOBJECT_API void AllocateObjectPool(int32 MaxUObjects, int32 MaxObjectsNotConsideredByGC, bool bPreAllocateObjectArray);
 
 	/**
 	 * Disables the disregard for GC optimization.
 	 *
 	 */
-	void DisableDisregardForGC();
+	COREUOBJECT_API void DisableDisregardForGC();
 
 	/**
 	* If there's enough slack in the disregard pool, we can re-open it and keep adding objects to it
 	*/
-	void OpenDisregardForGC();
+	COREUOBJECT_API void OpenDisregardForGC();
 
 	/**
 	 * After the initial load, this closes the disregard pool so that new object are GC-able
 	 */
-	void CloseDisregardForGC();
+	COREUOBJECT_API void CloseDisregardForGC();
 
 	/** Returns true if the disregard for GC pool is open */
 	bool IsOpenForDisregardForGC() const
@@ -711,14 +711,14 @@ public:
 	 * @param	AlreadyAllocatedIndex already allocated internal index to use, negative value means allocate a new index
 	 * @param	SerialNumber serial number to use
 	 */
-	void AllocateUObjectIndex(class UObjectBase* Object, EInternalObjectFlags InitialFlags, int32 AlreadyAllocatedIndex = -1, int32 SerialNumber = 0);
+	COREUOBJECT_API void AllocateUObjectIndex(class UObjectBase* Object, EInternalObjectFlags InitialFlags, int32 AlreadyAllocatedIndex = -1, int32 SerialNumber = 0);
 
 	/**
 	 * Returns a UObject index top to the global uobject array
 	 *
 	 * @param Object object to free
 	 */
-	void FreeUObjectIndex(class UObjectBase* Object);
+	COREUOBJECT_API void FreeUObjectIndex(class UObjectBase* Object);
 
 	/**
 	 * Returns the index of a UObject. Be advised this is only for very low level use.
@@ -827,35 +827,35 @@ public:
 	 *
 	 * @param Listener listener to notify when an object is deleted
 	 */
-	void AddUObjectCreateListener(FUObjectCreateListener* Listener);
+	COREUOBJECT_API void AddUObjectCreateListener(FUObjectCreateListener* Listener);
 
 	/**
 	 * Removes a listener for object creation
 	 *
 	 * @param Listener listener to remove
 	 */
-	void RemoveUObjectCreateListener(FUObjectCreateListener* Listener);
+	COREUOBJECT_API void RemoveUObjectCreateListener(FUObjectCreateListener* Listener);
 
 	/**
 	 * Adds a new listener for object deletion
 	 *
 	 * @param Listener listener to notify when an object is deleted
 	 */
-	void AddUObjectDeleteListener(FUObjectDeleteListener* Listener);
+	COREUOBJECT_API void AddUObjectDeleteListener(FUObjectDeleteListener* Listener);
 
 	/**
 	 * Removes a listener for object deletion
 	 *
 	 * @param Listener listener to remove
 	 */
-	void RemoveUObjectDeleteListener(FUObjectDeleteListener* Listener);
+	COREUOBJECT_API void RemoveUObjectDeleteListener(FUObjectDeleteListener* Listener);
 
 	/**
 	 * Removes an object from delete listeners
 	 *
 	 * @param Object to remove from delete listeners
 	 */
-	void RemoveObjectFromDeleteListeners(UObjectBase* Object);
+	COREUOBJECT_API void RemoveObjectFromDeleteListeners(UObjectBase* Object);
 
 	/**
 	 * Checks if a UObject pointer is valid
@@ -863,7 +863,7 @@ public:
 	 * @param	Object object to test for validity
 	 * @return	true if this index is valid
 	 */
-	bool IsValid(const UObjectBase* Object) const;
+	COREUOBJECT_API bool IsValid(const UObjectBase* Object) const;
 
 	/** Checks if the object index is valid. */
 	FORCEINLINE bool IsValidIndex(const UObjectBase* Object) const 
@@ -941,14 +941,14 @@ public:
 	/**
 	 * Clears some internal arrays to get rid of false memory leaks
 	 */
-	void ShutdownUObjectArray();
+	COREUOBJECT_API void ShutdownUObjectArray();
 
 	/**
 	* Given a UObject index return the serial number. If it doesn't have a serial number, give it one. Threadsafe.
 	* @param Index - UObject Index
 	* @return - the serial number for this UObject
 	*/
-	int32 AllocateSerialNumber(int32 Index);
+	COREUOBJECT_API int32 AllocateSerialNumber(int32 Index);
 
 	/**
 	* Given a UObject index return the serial number. If it doesn't have a serial number, return 0. Threadsafe.
@@ -1154,7 +1154,7 @@ public:
         return ObjObjects.GetAllocatedSize();
     }
 
-	void DumpUObjectCountsToLog() const;
+	COREUOBJECT_API void DumpUObjectCountsToLog() const;
 };
 
 /** UObject cluster. Groups UObjects into a single unit for GC. */
@@ -1180,7 +1180,7 @@ struct FUObjectCluster
 	bool bNeedsDissolving;
 };
 
-class COREUOBJECT_API FUObjectClusterContainer
+class FUObjectClusterContainer
 {
 	/** List of all clusters */
 	TArray<FUObjectCluster> Clusters;
@@ -1192,11 +1192,11 @@ class COREUOBJECT_API FUObjectClusterContainer
 	bool bClustersNeedDissolving;
 
 	/** Dissolves a cluster */
-	void DissolveCluster(FUObjectCluster& Cluster);
+	COREUOBJECT_API void DissolveCluster(FUObjectCluster& Cluster);
 
 public:
 
-	FUObjectClusterContainer();
+	COREUOBJECT_API FUObjectClusterContainer();
 
 	FORCEINLINE FUObjectCluster& operator[](int32 Index)
 	{
@@ -1205,35 +1205,35 @@ public:
 	}
 
 	/** Returns an index to a new cluster */
-	int32 AllocateCluster(int32 InRootObjectIndex);
+	COREUOBJECT_API int32 AllocateCluster(int32 InRootObjectIndex);
 
 	/** Frees the cluster at the specified index */
-	void FreeCluster(int32 InClusterIndex);
+	COREUOBJECT_API void FreeCluster(int32 InClusterIndex);
 
 	/**
 	* Gets the cluster the specified object is a root of or belongs to.
 	* @Param ClusterRootOrObjectFromCluster Root cluster object or object that belongs to a cluster
 	*/
-	FUObjectCluster* GetObjectCluster(UObjectBaseUtility* ClusterRootOrObjectFromCluster);
+	COREUOBJECT_API FUObjectCluster* GetObjectCluster(UObjectBaseUtility* ClusterRootOrObjectFromCluster);
 
 
 	/** 
 	 * Dissolves a cluster and all clusters that reference it 
 	 * @Param ClusterRootOrObjectFromCluster Root cluster object or object that belongs to a cluster
 	 */
-	void DissolveCluster(UObjectBaseUtility* ClusterRootOrObjectFromCluster);
+	COREUOBJECT_API void DissolveCluster(UObjectBaseUtility* ClusterRootOrObjectFromCluster);
 
 	/** 
 	 * Dissolve all clusters marked for dissolving 
 	 * @param bForceDissolveAllClusters if true, dissolves all clusters even if they're not marked for dissolving
 	 */
-	void DissolveClusters(bool bForceDissolveAllClusters = false);
+	COREUOBJECT_API void DissolveClusters(bool bForceDissolveAllClusters = false);
 
 	/** Dissolve the specified cluster and all clusters that reference it */
-	void DissolveClusterAndMarkObjectsAsUnreachable(FUObjectItem* RootObjectItem);
+	COREUOBJECT_API void DissolveClusterAndMarkObjectsAsUnreachable(FUObjectItem* RootObjectItem);
 
 	/*** Returns the minimum cluster size as specified in ini settings */
-	int32 GetMinClusterSize() const;
+	COREUOBJECT_API int32 GetMinClusterSize() const;
 
 	/** Gets the clusters array (for internal use only!) */
 	TArray<FUObjectCluster>& GetClustersUnsafe() 
diff --git a/Engine/Source/Runtime/CoreUObject/Public/UObject/UObjectBase.h b/Engine/Source/Runtime/CoreUObject/Public/UObject/UObjectBase.h
--- a/Engine/Source/Runtime/CoreUObject/Public/UObject/UObjectBase.h
+++ b/Engine/Source/Runtime/CoreUObject/Public/UObject/UObjectBase.h
@@ -30,7 +30,7 @@ DECLARE_DWORD_COUNTER_STAT_EXTERN(TEXT("STAT_UObjectsStatGroupTester"), STAT_UOb
 /** 
  * Low level implementation of UObject, should not be used directly in game code 
  */
-class COREUOBJECT_API UObjectBase
+class UObjectBase
 {
 	friend class UObjectBaseUtility;
 	friend struct Z_Construct_UClass_UObject_Statics;
@@ -55,7 +55,7 @@ protected:
 	 * Constructor used for bootstrapping
 	 * @param	InFlags			RF_Flags to assign
 	 */
-	UObjectBase( EObjectFlags InFlags );
+	COREUOBJECT_API UObjectBase( EObjectFlags InFlags );
 public:
 
 	/**
@@ -68,7 +68,7 @@ public:
 	 * @param	InInternalIndex		internal index to use (if already allocated), negative value means allocate a new index
 	 * @param	InSerialNumber		serial number to re-use (if already allocated)
 	 */
-	UObjectBase(UClass* InClass,
+	COREUOBJECT_API UObjectBase(UClass* InClass,
 			EObjectFlags InFlags,
 			EInternalObjectFlags InInternalFlags,
 			UObject *InOuter,
@@ -79,7 +79,7 @@ public:
 	/**
 	 * Final destructor, removes the object from the object array, and indirectly, from any annotations
 	 **/
-	virtual ~UObjectBase();
+	COREUOBJECT_API virtual ~UObjectBase();
 
 protected:
 	/**
@@ -88,20 +88,20 @@ protected:
 	 * @param NewName	new name for this object
 	 * @param NewOuter	new outer for this object, if NULL, outer will be unchanged
 	 */
-	void LowLevelRename(FName NewName,UObject *NewOuter = NULL);
+	COREUOBJECT_API void LowLevelRename(FName NewName,UObject *NewOuter = NULL);
 
 	/** Force any base classes to be registered first */
 	virtual void RegisterDependencies() {}
 
 	/** Enqueue the registration for this object. */
-	void Register(const TCHAR* PackageName,const TCHAR* Name);
+	COREUOBJECT_API void Register(const TCHAR* PackageName,const TCHAR* Name);
 	
 	/**
 	 * Convert a boot-strap registered class into a real one, add to uobject array, etc
 	 *
 	 * @param UClassStaticClass Now that it is known, fill in UClass::StaticClass() as the class
 	 */
-	virtual void DeferredRegister(UClass *UClassStaticClass,const TCHAR* PackageName,const TCHAR* Name);
+	COREUOBJECT_API virtual void DeferredRegister(UClass *UClassStaticClass,const TCHAR* PackageName,const TCHAR* Name);
 
 private:
 	/**
@@ -112,14 +112,14 @@ private:
 	 * @param InInternalIndex already allocated internal index to use, negative value means allocate a new index
 	 * @param InSerialNumber already allocated serial number to re-use
 	 */
-	void AddObject(FName Name, EInternalObjectFlags InSetInternalFlags, int32 InInternalIndex = -1, int32 InSerialNumber = 0);
+	COREUOBJECT_API void AddObject(FName Name, EInternalObjectFlags InSetInternalFlags, int32 InInternalIndex = -1, int32 InSerialNumber = 0);
 
 public:
 	/**
 	 * Checks to see if the object appears to be valid
 	 * @return true if this appears to be a valid object
 	 */
-	bool IsValidLowLevel() const;
+	COREUOBJECT_API bool IsValidLowLevel() const;
 
 	/**
 	 * Faster version of IsValidLowLevel.
@@ -128,7 +128,7 @@ public:
 	 * @param bRecursive true if the Class pointer should be checked with IsValidLowLevelFast
 	 * @return true if this appears to be a valid object
 	 */
-	bool IsValidLowLevelFast(bool bRecursive = true) const;
+	COREUOBJECT_API bool IsValidLowLevelFast(bool bRecursive = true) const;
 
 	/** 
 	 * Returns the unique ID of the object...these are reused so it is only unique while the object is alive.
@@ -158,19 +158,19 @@ public:
 	}
 
 	/** Overridable method to return a logical name for identification in stats. */
-	virtual FName GetFNameForStatID() const;
+	COREUOBJECT_API virtual FName GetFNameForStatID() const;
 
 	/** Removes the class prefix from the given string */
-	static FString RemoveClassPrefix(const TCHAR* ClassName);
+	static COREUOBJECT_API FString RemoveClassPrefix(const TCHAR* ClassName);
 
 	/** Returns the external UPackage associated with this object, if any */
-	UPackage* GetExternalPackage() const;
+	COREUOBJECT_API UPackage* GetExternalPackage() const;
 	
 	/** Associate an external package directly to this object. */
-	void SetExternalPackage(UPackage* InPackage);
+	COREUOBJECT_API void SetExternalPackage(UPackage* InPackage);
 
 	/** Returns the external UPackage for this object, if any, NOT THREAD SAFE, used by internal gc reference collecting. */
-	UPackage* GetExternalPackageInternal() const;
+	COREUOBJECT_API UPackage* GetExternalPackageInternal() const;
 
 protected:
 	/**
@@ -263,7 +263,7 @@ private:
 
 #if WITH_EDITOR
 	/** This is used by the reinstancer to re-class and re-archetype the current instances of a class before recompiling */
-	void SetClass(UClass* NewClass);
+	COREUOBJECT_API void SetClass(UClass* NewClass);
 #endif
 };
 
diff --git a/Engine/Source/Runtime/CoreUObject/Public/UObject/UObjectBaseUtility.h b/Engine/Source/Runtime/CoreUObject/Public/UObject/UObjectBaseUtility.h
--- a/Engine/Source/Runtime/CoreUObject/Public/UObject/UObjectBaseUtility.h
+++ b/Engine/Source/Runtime/CoreUObject/Public/UObject/UObjectBaseUtility.h
@@ -54,7 +54,7 @@ ENUM_CLASS_FLAGS(EObjectFullNameFlags);
 /**
  * Provides utility functions for UObject, this class should not be used directly
  */
-class COREUOBJECT_API UObjectBaseUtility : public UObjectBase
+class UObjectBaseUtility : public UObjectBase
 {
 	FORCEINLINE void MarkPendingKillOnlyInternal()
 	{
@@ -78,7 +78,7 @@ class COREUOBJECT_API UObjectBaseUtility : public UObjectBase
 	}
 
 	/** If true, objects will never be marked as PendingKill so references to them will not be nulled automatically by the garbage collector */
-	static bool bPendingKillDisabled;
+	static COREUOBJECT_API bool bPendingKillDisabled;
 
 	friend void InitNoPendingKill();
 	friend struct FInternalUObjectBaseUtilityIsValidFlagsChecker;
@@ -426,12 +426,12 @@ public:
 	 *
 	 * @note	safe to call on NULL object pointers!
 	 */
-	FString GetFullName( const UObject* StopOuter=NULL, EObjectFullNameFlags Flags = EObjectFullNameFlags::None ) const;
+	COREUOBJECT_API FString GetFullName( const UObject* StopOuter=NULL, EObjectFullNameFlags Flags = EObjectFullNameFlags::None ) const;
 
 	/**
 	 * Version of GetFullName() that eliminates unnecessary copies.
 	 */
-	void GetFullName( const UObject* StopOuter, FString& ResultString, EObjectFullNameFlags Flags = EObjectFullNameFlags::None ) const;
+	COREUOBJECT_API void GetFullName( const UObject* StopOuter, FString& ResultString, EObjectFullNameFlags Flags = EObjectFullNameFlags::None ) const;
 
 	/**
 	 * Returns the fully qualified pathname for this object as well as the name of the class, in the format:
@@ -444,7 +444,7 @@ public:
 	 *
 	 * @note	safe to call on NULL object pointers!
 	 */
-	void GetFullName(FStringBuilderBase& ResultString, const UObject* StopOuter = NULL, EObjectFullNameFlags Flags = EObjectFullNameFlags::None) const;
+	COREUOBJECT_API void GetFullName(FStringBuilderBase& ResultString, const UObject* StopOuter = NULL, EObjectFullNameFlags Flags = EObjectFullNameFlags::None) const;
 
 	/**
 	 * Returns the fully qualified pathname for this object, in the format:
@@ -455,13 +455,13 @@ public:
 	 *
 	 * @note	safe to call on NULL object pointers!
 	 */
-	FString GetPathName( const UObject* StopOuter=NULL ) const;
+	COREUOBJECT_API FString GetPathName( const UObject* StopOuter=NULL ) const;
 
 	/**
 	 * Versions of GetPathName() that eliminates unnecessary copies and allocations.
 	 */
-	void GetPathName(const UObject* StopOuter, FString& ResultString) const;
-	void GetPathName(const UObject* StopOuter, FStringBuilderBase& ResultString) const;
+	COREUOBJECT_API void GetPathName(const UObject* StopOuter, FString& ResultString) const;
+	COREUOBJECT_API void GetPathName(const UObject* StopOuter, FStringBuilderBase& ResultString) const;
 
 	/** Helper function to access the private bPendingKillDisabled variable */
 	static inline bool IsPendingKillEnabled()
@@ -503,12 +503,12 @@ public:
 	*
 	* @return	true if this object can be inside of a cluster
 	*/
-	virtual bool CanBeInCluster() const;
+	COREUOBJECT_API virtual bool CanBeInCluster() const;
 
 	/**
 	* Called after PostLoad to create UObject cluster
 	*/
-	virtual void CreateCluster();
+	COREUOBJECT_API virtual void CreateCluster();
 
 	/**
 	* Called during Garbage Collection to perform additional cleanup when the cluster is about to be destroyed due to PendingKill flag being set on it.
@@ -520,11 +520,11 @@ public:
 	* @param ClusterRootOrObjectFromCluster Object that belongs to the cluster we want to add this object to.
 	* @param Add this object to the target cluster as a mutable object without adding this object's references.
 	*/
-	void AddToCluster(UObjectBaseUtility* ClusterRootOrObjectFromCluster, bool bAddAsMutableObject = false);
+	COREUOBJECT_API void AddToCluster(UObjectBaseUtility* ClusterRootOrObjectFromCluster, bool bAddAsMutableObject = false);
 
 protected:
 	/** Helper function to create a cluster from UObject */
-	static void CreateClusterFromObject(UObjectBaseUtility* ClusterRootObject, UObjectBaseUtility* ReferencingObject);
+	static COREUOBJECT_API void CreateClusterFromObject(UObjectBaseUtility* ClusterRootObject, UObjectBaseUtility* ReferencingObject);
 
 public:
 	/**
@@ -533,7 +533,7 @@ public:
 	 * @param	bStartWithOuter		whether to include this object's name in the returned string
 	 * @return	string containing the path name for this object, minus the outermost-package's name
 	 */
-	FString GetFullGroupName( bool bStartWithOuter ) const;
+	COREUOBJECT_API FString GetFullGroupName( bool bStartWithOuter ) const;
 
 	/**
 	 * Returns the name of this object (with no path information)
@@ -565,33 +565,33 @@ public:
 	 * Get the object packaging mode.
 	 * @return true if object has a different package than its outer's package
 	 */
-	bool IsPackageExternal() const;
+	COREUOBJECT_API bool IsPackageExternal() const;
 
 	/**
 	 * Utility function to temporarily detach the object external package, if any
 	 * GetPackage will report the outer's package once detached
 	 */
-	void DetachExternalPackage();
+	COREUOBJECT_API void DetachExternalPackage();
 
 	/**
 	 * Utility function to reattach the object external package, if any
 	 * GetPackage will report the object external package if set after this call
 	 */
-	void ReattachExternalPackage();
+	COREUOBJECT_API void ReattachExternalPackage();
 
 	/**
 	 * Walks up the list of outers until it finds the top-level one that isn't a package.
 	 * Will return null if called on a package
 	 * @return outermost non package Outer.
 	 */
-	UObject* GetOutermostObject() const;
+	COREUOBJECT_API UObject* GetOutermostObject() const;
 
 	/**
 	 * Walks up the list of outers until it finds a package directly associated with the object.
 	 *
 	 * @return the package the object is in.
 	 */
-	UPackage* GetPackage() const;
+	COREUOBJECT_API UPackage* GetPackage() const;
 
 #if UE_USE_VERSE_PATHS
 	/**
@@ -599,7 +599,7 @@ public:
 	 *
 	 * @return The VersePath of the object
 	 */
-	UE::Core::FVersePath GetVersePath() const;
+	COREUOBJECT_API UE::Core::FVersePath GetVersePath() const;
 #endif
 
 	/** 
@@ -608,7 +608,7 @@ public:
 	 * @return the package the object is in.
 	 * @see GetPackage
 	 */
-	UPackage* GetOutermost() const;
+	COREUOBJECT_API UPackage* GetOutermost() const;
 
 	/** 
 	 * Finds the outermost package and marks it dirty. 
@@ -616,14 +616,14 @@ public:
 	 *
 	 * @return false if the request to mark the package dirty was suppressed by the editor and true otherwise.
 	 */
-	bool MarkPackageDirty() const;
+	COREUOBJECT_API bool MarkPackageDirty() const;
 
 	/**
 	* Determines whether this object is a template object
 	*
 	* @return	true if this object is a template object (owned by a UClass)
 	*/
-	bool IsTemplate( EObjectFlags TemplateTypes = RF_ArchetypeObject|RF_ClassDefaultObject ) const;
+	COREUOBJECT_API bool IsTemplate( EObjectFlags TemplateTypes = RF_ArchetypeObject|RF_ClassDefaultObject ) const;
 
 	/**
 	 * Traverses the outer chain searching for the next object of a certain type.  (T must be derived from UObject)
@@ -631,7 +631,7 @@ public:
 	 * @param	Target class to search for
 	 * @return	a pointer to the first object in this object's Outer chain which is of the correct type.
 	 */
-	UObject* GetTypedOuter(UClass* Target) const;
+	COREUOBJECT_API UObject* GetTypedOuter(UClass* Target) const;
 
 	/**
 	 * Traverses the outer chain searching for the next object of a certain type.  (T must be derived from UObject)
@@ -666,32 +666,32 @@ public:
 	 * @param	InInterfaceClass	Target interface to search for
 	 * @return	a pointer to the first object in this object's Outer chain which implements the specified interface.
 	 */
-	UObjectBaseUtility* GetImplementingOuterObject(const UClass* InInterfaceClass) const;
+	COREUOBJECT_API UObjectBaseUtility* GetImplementingOuterObject(const UClass* InInterfaceClass) const;
 
 	/** 
 	 * Return the dispatch to `IsInOuter` or `IsInPackage` depending on SomeOuter's class. 
 	 * Legacy function, preferably use IsInOuter or IsInPackage depending on use case.
 	 */
-	bool IsIn( const UObject* SomeOuter ) const;
+	COREUOBJECT_API bool IsIn( const UObject* SomeOuter ) const;
 
 	/** 
 	 * Overload to determine if an object is in the specified package which can now be different than its outer chain.
 	 * Calls IsInPackage.
 	 */
-	bool IsIn(const UPackage* SomePackage) const;
+	COREUOBJECT_API bool IsIn(const UPackage* SomePackage) const;
 
 	/** Returns true if the object is contained in the specified outer. */
-	bool IsInOuter(const UObject* SomeOuter) const;
+	COREUOBJECT_API bool IsInOuter(const UObject* SomeOuter) const;
 
 	/** Returns true if the object is contained in the specified package. */
-	bool IsInPackage(const UPackage* SomePackage) const;
+	COREUOBJECT_API bool IsInPackage(const UPackage* SomePackage) const;
 
 	/**
 	 * Find out if this object is inside (has an outer) that is of the specified class
 	 * @param SomeBaseClass	The base class to compare against
 	 * @return True if this object is in an object of the given type.
 	 */
-	bool IsInA( const UClass* SomeBaseClass ) const;
+	COREUOBJECT_API bool IsInA( const UClass* SomeBaseClass ) const;
 
 	/**
 	 * Checks whether this object's top-most package has any of the specified flags
@@ -700,7 +700,7 @@ public:
 	 *
 	 * @return	true if the PackageFlags member of this object's top-package has any bits from the mask set.
 	 */
-	bool RootPackageHasAnyFlags( uint32 CheckFlagMask ) const;
+	COREUOBJECT_API bool RootPackageHasAnyFlags( uint32 CheckFlagMask ) const;
 
 
 	/*-------------------
@@ -748,7 +748,7 @@ public:
 	 *
 	 * @param	TestClass	the class to find the common base for
 	 */
-	const UClass* FindNearestCommonBaseClass( const UClass* TestClass ) const;
+	COREUOBJECT_API const UClass* FindNearestCommonBaseClass( const UClass* TestClass ) const;
 
 	/**
 	 * Returns a pointer to this object safely converted to a pointer of the specified interface class.
@@ -759,13 +759,13 @@ public:
 	 *			class doesn't implement the interface indicated.  Will be the same value as 'this' if the interface class
 	 *			isn't native.
 	 */
-	void* GetInterfaceAddress( UClass* InterfaceClass );
+	COREUOBJECT_API void* GetInterfaceAddress( UClass* InterfaceClass );
 
 	/** 
 	 *	Returns a pointer to the I* native interface object that this object implements.
 	 *	Returns NULL if this object does not implement InterfaceClass, or does not do so natively.
 	 */
-	void* GetNativeInterfaceAddress(UClass* InterfaceClass);
+	COREUOBJECT_API void* GetNativeInterfaceAddress(UClass* InterfaceClass);
 
 	/** 
 	 *	Returns a pointer to the const I* native interface object that this object implements.
@@ -781,7 +781,7 @@ public:
 	 *
 	 * @return	true if this component was instanced from a template.  false if this component was created manually at runtime.
 	 */
-	bool IsDefaultSubobject() const;
+	COREUOBJECT_API bool IsDefaultSubobject() const;
 	
 
 	/*--------------------------------------------------
@@ -793,7 +793,7 @@ public:
 	 *
 	 * @return	a pointer to the linker for this object, or NULL if this object has no linker
 	 */
-	class FLinkerLoad* GetLinker() const;
+	COREUOBJECT_API class FLinkerLoad* GetLinker() const;
 
 	/**
 	 * Returns this object's LinkerIndex.
@@ -801,7 +801,7 @@ public:
 	 * @return	the index into my linker's ExportMap for the FObjectExport
 	 *			corresponding to this object.
 	 */
-	int32 GetLinkerIndex() const;
+	COREUOBJECT_API int32 GetLinkerIndex() const;
 
 	/**
 	 * Returns the UE version of the linker for this object.
@@ -812,7 +812,7 @@ public:
 	 *			a) this object is a native only class, or
 	 *			b) this object's linker has been detached, in which case it is already fully loaded
 	 */
-	FPackageFileVersion GetLinkerUEVersion() const;
+	COREUOBJECT_API FPackageFileVersion GetLinkerUEVersion() const;
 		
 	UE_DEPRECATED(5.0, "Use GetLinkerUEVersion instead which returns the version as a FPackageFileVersion. See the @FPackageFileVersion documentation for further details")
 	inline int32 GetLinkerUE4Version() const 
@@ -832,7 +832,7 @@ public:
 	 *			a) this object is a native only class, or
 	 *			b) this object's linker has been detached, in which case it is already fully loaded
 	 */
-	int32 GetLinkerLicenseeUEVersion() const;
+	COREUOBJECT_API int32 GetLinkerLicenseeUEVersion() const;
 
 	UE_DEPRECATED(5.0, "Use GetLinkerLicenseeUEVersion instead")
 	inline int32 GetLinkerLicenseeUE4Version() const { return GetLinkerLicenseeUEVersion(); }
@@ -847,7 +847,7 @@ public:
 	 *			a) this object is a native only class, or
 	 *			b) this object's linker has been detached, in which case it is already fully loaded
 	 */
-	int32 GetLinkerCustomVersion(FGuid CustomVersionKey) const;
+	COREUOBJECT_API int32 GetLinkerCustomVersion(FGuid CustomVersionKey) const;
 
 	/** 
 	 * Overloaded < operator. Compares objects by name.
@@ -1174,4 +1174,4 @@ public:
 
 #if UE_ENABLE_INCLUDE_ORDER_DEPRECATED_IN_5_2
 #include "CoreMinimal.h"
-#endif
\ No newline at end of file
+#endif
