[CoreUObject]

diff --git a/Engine/Source/Runtime/CoreUObject/Public/UObject/UObjectGlobals.h b/Engine/Source/Runtime/CoreUObject/Public/UObject/UObjectGlobals.h
--- a/Engine/Source/Runtime/CoreUObject/Public/UObject/UObjectGlobals.h
+++ b/Engine/Source/Runtime/CoreUObject/Public/UObject/UObjectGlobals.h
@@ -1048,13 +1048,13 @@ class FGCObject;
 /**
  * Internal class to finalize UObject creation (initialize properties) after the real C++ constructor is called.
  **/
-class COREUOBJECT_API FObjectInitializer
+class FObjectInitializer
 {
 public:
 	/**
 	 * Default Constructor, used when you are using the C++ "new" syntax. UObject::UObject will set the object pointer
 	 **/
-	FObjectInitializer();
+	COREUOBJECT_API FObjectInitializer();
 
 	/**
 	 * Constructor
@@ -1063,7 +1063,7 @@ public:
 	 * @param	InOptions initialization options, see EObjectInitializerOptions
 	 * @param	InInstanceGraph passed instance graph
 	 */
-	FObjectInitializer(UObject* InObj, UObject* InObjectArchetype, EObjectInitializerOptions InOptions, struct FObjectInstancingGraph* InInstanceGraph = nullptr);
+	COREUOBJECT_API FObjectInitializer(UObject* InObj, UObject* InObjectArchetype, EObjectInitializerOptions InOptions, struct FObjectInstancingGraph* InInstanceGraph = nullptr);
 
 	UE_DEPRECATED(5.0, "Use version that takes EObjectInitializerOptions")
 	FObjectInitializer(UObject* InObj, UObject* InObjectArchetype, bool bInCopyTransientsFromClassDefaults, bool bInShouldInitializeProps, struct FObjectInstancingGraph* InInstanceGraph = nullptr)
@@ -1075,14 +1075,14 @@ public:
 	}
 
 	/** Special constructor for static construct object internal that passes along the params block directly */
-	FObjectInitializer(UObject* InObj, const FStaticConstructObjectParameters& StaticConstructParams);
+	COREUOBJECT_API FObjectInitializer(UObject* InObj, const FStaticConstructObjectParameters& StaticConstructParams);
 
 private:
 	/** Helper for the common behaviors in the constructors */
-	void Construct_Internal();
+	COREUOBJECT_API void Construct_Internal();
 
 public:
-	~FObjectInitializer();
+	COREUOBJECT_API ~FObjectInitializer();
 
 	/** 
 	 * Return the archetype that this object will copy properties from later
@@ -1103,7 +1103,7 @@ public:
 	/**
 	* Return the class of the object that is being constructed
 	**/
-	UClass* GetClass() const;
+	COREUOBJECT_API UClass* GetClass() const;
 
 	/**
 	 * Create a component or subobject
@@ -1169,7 +1169,7 @@ public:
 	* @param	ReturnType					type of the new component
 	* @param	bTransient					true if the component is being assigned to a transient property
 	*/
-	UObject* CreateEditorOnlyDefaultSubobject(UObject* Outer, FName SubobjectName, const UClass* ReturnType, bool bTransient = false) const;
+	COREUOBJECT_API UObject* CreateEditorOnlyDefaultSubobject(UObject* Outer, FName SubobjectName, const UClass* ReturnType, bool bTransient = false) const;
 
 	/**
 	 * Create a component or subobject
@@ -1180,7 +1180,7 @@ public:
 	 * @param	bIsRequired                 true if the component is required and will always be created even if DoNotCreateDefaultSubobject was specified.
 	 * @param	bIsTransient                true if the component is being assigned to a transient property
 	 */
-	UObject* CreateDefaultSubobject(UObject* Outer, FName SubobjectFName, const UClass* ReturnType, const UClass* ClassToCreateByDefault, bool bIsRequired = true, bool bIsTransient = false) const;
+	COREUOBJECT_API UObject* CreateDefaultSubobject(UObject* Outer, FName SubobjectFName, const UClass* ReturnType, const UClass* ClassToCreateByDefault, bool bIsRequired = true, bool bIsTransient = false) const;
 
 	/**
 	 * Sets the class to use for a subobject defined in a base class, the class must be a subclass of the class used by the base class.
@@ -1284,7 +1284,7 @@ public:
 	/**
 	 * Asserts with the specified message if code is executed inside UObject constructor
 	 **/
-	static void AssertIfInConstructor(UObject* Outer, const TCHAR* ErrorMessage);
+	static COREUOBJECT_API void AssertIfInConstructor(UObject* Outer, const TCHAR* ErrorMessage);
 
 	FORCEINLINE void FinalizeSubobjectClassInitialization()
 	{
@@ -1292,7 +1292,7 @@ public:
 	}
 
 	/** Gets ObjectInitializer for the currently constructed object. Can only be used inside of a constructor of UObject-derived class. */
-	static FObjectInitializer& Get();
+	static COREUOBJECT_API FObjectInitializer& Get();
 
 private:
 
@@ -1310,16 +1310,16 @@ private:
 	 * @param	DefaultData			the buffer containing the source data for the initialization
 	 * @param	bCopyTransientsFromClassDefaults if true, copy the transients from the DefaultsClass defaults, otherwise copy the transients from DefaultData
 	 */
-	static void InitProperties(UObject* Obj, UClass* DefaultsClass, UObject* DefaultData, bool bCopyTransientsFromClassDefaults);
+	static COREUOBJECT_API void InitProperties(UObject* Obj, UClass* DefaultsClass, UObject* DefaultData, bool bCopyTransientsFromClassDefaults);
 
-	bool IsInstancingAllowed() const;
+	COREUOBJECT_API bool IsInstancingAllowed() const;
 
 	/**
 	 * Calls InitProperties for any default subobjects created through this ObjectInitializer.
 	 * @param bAllowInstancing	Indicates whether the object's components may be copied from their templates.
 	 * @return true if there are any subobjects which require instancing.
 	*/
-	bool InitSubobjectProperties(bool bAllowInstancing) const;
+	COREUOBJECT_API bool InitSubobjectProperties(bool bAllowInstancing) const;
 
 	/**
 	 * Create copies of the object's components from their templates.
@@ -1327,7 +1327,7 @@ private:
 	 * @param bNeedInstancing			Indicates whether the object's components need to be instanced
 	 * @param bNeedSubobjectInstancing	Indicates whether subobjects of the object's components need to be instanced
 	 */
-	void InstanceSubobjects(UClass* Class, bool bNeedInstancing, bool bNeedSubobjectInstancing) const;
+	COREUOBJECT_API void InstanceSubobjects(UClass* Class, bool bNeedInstancing, bool bNeedSubobjectInstancing) const;
 
 	/** 
 	 * Initializes a non-native property, according to the initialization rules. If the property is non-native
@@ -1336,13 +1336,13 @@ private:
 	 * @param	Data				Default data
 	 * @return	Returns true if that property was a non-native one, otherwise false
 	 */
-	static bool InitNonNativeProperty(FProperty* Property, UObject* Data);
+	static COREUOBJECT_API bool InitNonNativeProperty(FProperty* Property, UObject* Data);
 	
 	/**
 	 * Finalizes a constructed UObject by initializing properties, 
 	 * instancing/initializing sub-objects, etc.
 	 */
-	void PostConstructInit();
+	COREUOBJECT_API void PostConstructInit();
 
 private:
 
@@ -1350,13 +1350,13 @@ private:
 	struct FOverrides
 	{
 		/**  Add an override, make sure it is legal **/
-		void Add(FName InComponentName, const UClass* InComponentClass, const TArrayView<const FName>* FullPath = nullptr);
+		COREUOBJECT_API void Add(FName InComponentName, const UClass* InComponentClass, const TArrayView<const FName>* FullPath = nullptr);
 
 		/**  Add a potentially nested override, make sure it is legal **/
-		void Add(FStringView InComponentPath, const UClass* InComponentClass);
+		COREUOBJECT_API void Add(FStringView InComponentPath, const UClass* InComponentClass);
 
 		/**  Add a potentially nested override, make sure it is legal **/
-		void Add(TArrayView<const FName> InComponentPath, const UClass* InComponentClass, const TArrayView<const FName>* FullPath = nullptr);
+		COREUOBJECT_API void Add(TArrayView<const FName> InComponentPath, const UClass* InComponentClass, const TArrayView<const FName>* FullPath = nullptr);
 
 		struct FOverrideDetails
 		{
@@ -1442,13 +1442,13 @@ private:
 	};
 
 	/** Asserts if SetDefaultSubobjectClass or DoNotCreateOptionalDefaultSuobject are called inside of the constructor body */
-	void AssertIfSubobjectSetupIsNotAllowed(const FName SubobjectName) const;
+	COREUOBJECT_API void AssertIfSubobjectSetupIsNotAllowed(const FName SubobjectName) const;
 
 	/** Asserts if SetDefaultSubobjectClass or DoNotCreateOptionalDefaultSuobject are called inside of the constructor body */
-	void AssertIfSubobjectSetupIsNotAllowed(FStringView SubobjectName) const;
+	COREUOBJECT_API void AssertIfSubobjectSetupIsNotAllowed(FStringView SubobjectName) const;
 
 	/** Asserts if SetDefaultSubobjectClass or DoNotCreateOptionalDefaultSuobject are called inside of the constructor body */
-	void AssertIfSubobjectSetupIsNotAllowed(TArrayView<const FName> SubobjectNames) const;
+	COREUOBJECT_API void AssertIfSubobjectSetupIsNotAllowed(TArrayView<const FName> SubobjectNames) const;
 
 	/**  object to initialize, from static allocate object, after construction **/
 	UObject* Obj;
@@ -1484,7 +1484,7 @@ private:
 
 #if WITH_EDITORONLY_DATA
 	/** Detects when a new GC object was created */
-	void OnGCObjectCreated(FGCObject* InObject);
+	COREUOBJECT_API void OnGCObjectCreated(FGCObject* InObject);
 
 	/** Delegate handle for OnGCObjectCreated callback */
 	FDelegateHandle OnGCObjectCreatedHandle;
@@ -2141,7 +2141,7 @@ public:
 };
 
 /** Helper class for setting and resetting attributes on the FReferenceCollectorArchive */
-class COREUOBJECT_API FVerySlowReferenceCollectorArchiveScope
+class FVerySlowReferenceCollectorArchiveScope
 {	
 	FReferenceCollectorArchive& Archive;
 	const UObject* OldSerializingObject;
@@ -2172,19 +2172,19 @@ public:
 };
 
 /** Used by garbage collector to collect references via virtual AddReferencedObjects calls */
-class COREUOBJECT_API FReferenceCollector
+class FReferenceCollector
 {
 public:
 	virtual ~FReferenceCollector() {}
 	
 	/** Preferred way to add a reference that allows batching. Object must outlive GC tracing, can't be used for temporary/stack references. */
-	virtual void AddStableReference(TObjectPtr<UObject>* Object);
+	COREUOBJECT_API virtual void AddStableReference(TObjectPtr<UObject>* Object);
 	
 	/** Preferred way to add a reference array that allows batching. Can't be used for temporary/stack array. */
-	virtual void AddStableReferenceArray(TArray<TObjectPtr<UObject>>* Objects);
+	COREUOBJECT_API virtual void AddStableReferenceArray(TArray<TObjectPtr<UObject>>* Objects);
 
 	/** Preferred way to add a reference set that allows batching. Can't be used for temporary/stack set. */
-	virtual void AddStableReferenceSet(TSet<TObjectPtr<UObject>>* Objects);
+	COREUOBJECT_API virtual void AddStableReferenceSet(TSet<TObjectPtr<UObject>>* Objects);
 
 	template <typename KeyType, typename ValueType, typename Allocator, typename KeyFuncs>
 	FORCEINLINE_DEBUGGABLE void AddStableReferenceMap(TMapBase<KeyType, ValueType, Allocator, KeyFuncs>& Map)
@@ -2216,13 +2216,13 @@ public:
 
 #if !UE_REFERENCE_COLLECTOR_REQUIRE_OBJECTPTR
 	/** Preferred way to add a reference that allows batching. Object must outlive GC tracing, can't be used for temporary/stack references. */
-	virtual void AddStableReference(UObject** Object);
+	COREUOBJECT_API virtual void AddStableReference(UObject** Object);
 	
 	/** Preferred way to add a reference array that allows batching. Can't be used for temporary/stack array. */
-	virtual void AddStableReferenceArray(TArray<UObject*>* Objects);
+	COREUOBJECT_API virtual void AddStableReferenceArray(TArray<UObject*>* Objects);
 
 	/** Preferred way to add a reference set that allows batching. Can't be used for temporary/stack set. */
-	virtual void AddStableReferenceSet(TSet<UObject*>* Objects);
+	COREUOBJECT_API virtual void AddStableReferenceSet(TSet<UObject*>* Objects);
 
 	template<class UObjectType>
 	FORCEINLINE void AddStableReference(UObjectType** Object)
@@ -2515,7 +2515,7 @@ public:
 												ReferencingProperty);
 	}
 
-	void AddReferencedObject(FWeakObjectPtr& P,
+	COREUOBJECT_API void AddReferencedObject(FWeakObjectPtr& P,
 													 const UObject* ReferencingObject = nullptr,
 													 const FProperty* ReferencingProperty = nullptr);
 
@@ -2531,23 +2531,23 @@ public:
 	 * SerializeBin/TPropertyValueIterator/GetVerySlowReferenceCollectorArchive paths.
 	 */
 #if !UE_REFERENCE_COLLECTOR_REQUIRE_OBJECTPTR
-	void AddReferencedObjects(const UScriptStruct*& ScriptStruct, void* Instance, const UObject* ReferencingObject = nullptr, const FProperty* ReferencingProperty = nullptr);
+	COREUOBJECT_API void AddReferencedObjects(const UScriptStruct*& ScriptStruct, void* Instance, const UObject* ReferencingObject = nullptr, const FProperty* ReferencingProperty = nullptr);
 #endif
 
-	void AddReferencedObjects(TObjectPtr<const UScriptStruct>& ScriptStruct, void* Instance, const UObject* ReferencingObject = nullptr, const FProperty* ReferencingProperty = nullptr);
-	void AddReferencedObjects(TWeakObjectPtr<const UScriptStruct>& ScriptStruct, void* Instance, const UObject* ReferencingObject = nullptr, const FProperty* ReferencingProperty = nullptr);
+	COREUOBJECT_API void AddReferencedObjects(TObjectPtr<const UScriptStruct>& ScriptStruct, void* Instance, const UObject* ReferencingObject = nullptr, const FProperty* ReferencingProperty = nullptr);
+	COREUOBJECT_API void AddReferencedObjects(TWeakObjectPtr<const UScriptStruct>& ScriptStruct, void* Instance, const UObject* ReferencingObject = nullptr, const FProperty* ReferencingProperty = nullptr);
 	
 	/** Adds all strong property references from a struct instance, but not the struct itself. Skips AddStructReferencedObjects. */
-	void AddPropertyReferences(const UStruct* Struct, void* Instance, const UObject* ReferencingObject = nullptr);
+	COREUOBJECT_API void AddPropertyReferences(const UStruct* Struct, void* Instance, const UObject* ReferencingObject = nullptr);
 	
 	/** Same as AddPropertyReferences but also calls AddStructReferencedObjects on Struct and all nested structs */
-	void AddPropertyReferencesWithStructARO(const UScriptStruct* Struct, void* Instance, const UObject* ReferencingObject = nullptr);
+	COREUOBJECT_API void AddPropertyReferencesWithStructARO(const UScriptStruct* Struct, void* Instance, const UObject* ReferencingObject = nullptr);
 
 	/** Same as AddPropertyReferences but also calls AddStructReferencedObjects on all nested structs */
-	void AddPropertyReferencesWithStructARO(const UClass* Class, void* Instance, const UObject* ReferencingObject = nullptr);
+	COREUOBJECT_API void AddPropertyReferencesWithStructARO(const UClass* Class, void* Instance, const UObject* ReferencingObject = nullptr);
 
 	/** Internal use only. Same as AddPropertyReferences but skips field path and interface properties. Might get removed. */
-	void AddPropertyReferencesLimitedToObjectProperties(const UStruct* Struct, void* Instance, const UObject* ReferencingObject = nullptr);
+	COREUOBJECT_API void AddPropertyReferencesLimitedToObjectProperties(const UStruct* Struct, void* Instance, const UObject* ReferencingObject = nullptr);
 
 	/**
 	 * Make Add[OnlyObject]PropertyReference/AddReferencedObjects(UScriptStruct) use AddReferencedObjects(UObject*&) callbacks
@@ -2760,25 +2760,25 @@ protected:
 	* @param ReferencingObject Referencing object (if available).
 	* @param ReferencingProperty Referencing property (if available).
 	*/
-	virtual void HandleObjectReferences(FObjectPtr* InObjects, const int32 ObjectNum, const UObject* InReferencingObject, const FProperty* InReferencingProperty);
+	COREUOBJECT_API virtual void HandleObjectReferences(FObjectPtr* InObjects, const int32 ObjectNum, const UObject* InReferencingObject, const FProperty* InReferencingProperty);
 
 private:
 	/** Creates the proxy archive that uses serialization to add objects to this collector */
-	void CreateVerySlowReferenceCollectorArchive();
+	COREUOBJECT_API void CreateVerySlowReferenceCollectorArchive();
 
 	/** Default proxy archive that uses serialization to add objects to this collector */
 	TUniquePtr<FReferenceCollectorArchive> DefaultReferenceCollectorArchive;
 
-	void AddStableReferenceSetFwd(TSet<FObjectPtr>* Objects);	 
-	void AddStableReferenceArrayFwd(TArray<FObjectPtr>* Objects);		
-	void AddStableReferenceFwd(FObjectPtr* Object);	 
+	COREUOBJECT_API void AddStableReferenceSetFwd(TSet<FObjectPtr>* Objects);	 
+	COREUOBJECT_API void AddStableReferenceArrayFwd(TArray<FObjectPtr>* Objects);		
+	COREUOBJECT_API void AddStableReferenceFwd(FObjectPtr* Object);	 
 };
 
 /**
  * FReferenceFinder.
  * Helper class used to collect object references.
  */
-class COREUOBJECT_API FReferenceFinder : public FReferenceCollector
+class FReferenceFinder : public FReferenceCollector
 {
 public:
 
@@ -2794,7 +2794,7 @@ public:
 	 *									objects that have LimitOuter for their Outer (i.e. nested subobjects/components)
 	 * @param	bShouldIgnoreTransient	true to skip serialization of transient properties
 	 */
-	FReferenceFinder(TArray<UObject*>& InObjectArray, UObject* InOuter = nullptr, bool bInRequireDirectOuter = true, bool bInShouldIgnoreArchetype = false, bool bInSerializeRecursively = false, bool bInShouldIgnoreTransient = false);
+	COREUOBJECT_API FReferenceFinder(TArray<UObject*>& InObjectArray, UObject* InOuter = nullptr, bool bInRequireDirectOuter = true, bool bInShouldIgnoreArchetype = false, bool bInSerializeRecursively = false, bool bInShouldIgnoreTransient = false);
 
 	/**
 	 * Finds all objects referenced by Object.
@@ -2803,10 +2803,10 @@ public:
 	 * @param ReferencingObject object that's referencing the current object.
 	 * @param ReferencingProperty property the current object is being referenced through.
 	 */
-	virtual void FindReferences(UObject* Object, UObject* ReferencingObject = nullptr, FProperty* ReferencingProperty = nullptr);
+	COREUOBJECT_API virtual void FindReferences(UObject* Object, UObject* ReferencingObject = nullptr, FProperty* ReferencingProperty = nullptr);
 
 	// FReferenceCollector interface.
-	virtual void HandleObjectReference(UObject*& Object, const UObject* ReferencingObject, const FProperty* InReferencingProperty) override;
+	COREUOBJECT_API virtual void HandleObjectReference(UObject*& Object, const UObject* ReferencingObject, const FProperty* InReferencingProperty) override;
 	virtual bool IsIgnoringArchetypeRef() const override { return bShouldIgnoreArchetype; }
 	virtual bool IsIgnoringTransient() const override { return bShouldIgnoreTransient; }
 	virtual void SetSerializedProperty(class FProperty* Inproperty) override
@@ -2869,23 +2869,23 @@ struct FEndLoadPackageContext
 /**
  * Global CoreUObject delegates
  */
-struct COREUOBJECT_API FCoreUObjectDelegates
+struct FCoreUObjectDelegates
 {
 #if WITH_EDITOR
 	/** Callback for object property modifications, called by UObject::PostEditChangeProperty with a single property event */
 	DECLARE_MULTICAST_DELEGATE_TwoParams(FOnObjectPropertyChanged, UObject*, struct FPropertyChangedEvent&);
-	static FOnObjectPropertyChanged OnObjectPropertyChanged;
+	static COREUOBJECT_API FOnObjectPropertyChanged OnObjectPropertyChanged;
 
 	/** Callback for object property modifications, called by UObject::PreEditChange with a full property chain */
 	DECLARE_MULTICAST_DELEGATE_TwoParams(FOnPreObjectPropertyChanged, UObject*, const class FEditPropertyChain&);
-	static FOnPreObjectPropertyChanged OnPreObjectPropertyChanged;
+	static COREUOBJECT_API FOnPreObjectPropertyChanged OnPreObjectPropertyChanged;
 
 	/** Called when an object is registered for change with UObject::Modify. This gets called in both the editor and standalone game editor builds, for every object modified */
 	DECLARE_MULTICAST_DELEGATE_OneParam(FOnObjectModified, UObject*);
-	static FOnObjectModified OnObjectModified;
+	static COREUOBJECT_API FOnObjectModified OnObjectModified;
 
 	/** Set of objects modified this frame, to prevent multiple triggerings of the OnObjectModified delegate */
-	static TSet<UObject*> ObjectsModifiedThisFrame;
+	static COREUOBJECT_API TSet<UObject*> ObjectsModifiedThisFrame;
 
 	/** Broadcast OnObjectModified if the broadcast hasn't occurred for this object in this frame */
 	static void BroadcastOnObjectModified(UObject* Object)
@@ -2899,7 +2899,7 @@ struct COREUOBJECT_API FCoreUObjectDelegates
 
 	/** Callback for an object being transacted */
 	DECLARE_MULTICAST_DELEGATE_TwoParams(FOnObjectTransacted, UObject*, const class FTransactionObjectEvent&);
-	static FOnObjectTransacted OnObjectTransacted;
+	static COREUOBJECT_API FOnObjectTransacted OnObjectTransacted;
 
 	/**
 	 * Called when UObjects have been replaced to allow others a chance to fix their references
@@ -2908,7 +2908,7 @@ struct COREUOBJECT_API FCoreUObjectDelegates
 	 */
 	using FReplacementObjectMap = TMap<UObject*, UObject*>; // Alias for use in the macro
 	DECLARE_MULTICAST_DELEGATE_OneParam(FOnObjectsReplaced, const FReplacementObjectMap&);
-	static FOnObjectsReplaced OnObjectsReplaced;
+	static COREUOBJECT_API FOnObjectsReplaced OnObjectsReplaced;
 
 	/**
 	 * Called when UObjects have been re-instanced to allow others a chance to fix their references
@@ -2916,7 +2916,7 @@ struct COREUOBJECT_API FCoreUObjectDelegates
 	 * references should be self-consistent).
 	 */
 	DECLARE_MULTICAST_DELEGATE_OneParam(FOnObjectsReinstanced, const FReplacementObjectMap&);
-	static FOnObjectsReinstanced OnObjectsReinstanced;
+	static COREUOBJECT_API FOnObjectsReinstanced OnObjectsReinstanced;
 
 	/**
 	 * Called after the Blueprint compiler has finished generating the Class Default Object (CDO) for a class. This can only happen in the editor.
@@ -2924,141 +2924,141 @@ struct COREUOBJECT_API FCoreUObjectDelegates
 	 * eg) caching the name/count of properties of a certain type, or inspecting the properties on the class and using their meta-data and CDO default values to derive game data.
 	 */
 	DECLARE_MULTICAST_DELEGATE_TwoParams(FOnObjectPostCDOCompiled, UObject*, const FObjectPostCDOCompiledContext&);
-	static FOnObjectPostCDOCompiled OnObjectPostCDOCompiled;
+	static COREUOBJECT_API FOnObjectPostCDOCompiled OnObjectPostCDOCompiled;
 
 	DECLARE_MULTICAST_DELEGATE_OneParam(FOnObjectSaved, UObject*);
 	UE_DEPRECATED(5.0, "Use OnObjectPreSave instead.")
-	static FOnObjectSaved OnObjectSaved;
+	static COREUOBJECT_API FOnObjectSaved OnObjectSaved;
 
 	/** Callback for when an asset is saved. This is called from UObject::PreSave before it is actually written to disk, for every object saved */
 	DECLARE_MULTICAST_DELEGATE_TwoParams(FOnObjectPreSave, UObject*, FObjectPreSaveContext);
-	static FOnObjectPreSave OnObjectPreSave;
+	static COREUOBJECT_API FOnObjectPreSave OnObjectPreSave;
 
 	/** Callback for when an asset is loaded. This gets called in both the editor and standalone game editor builds, but only for objects that return true for IsAsset() */
 	DECLARE_MULTICAST_DELEGATE_OneParam(FOnAssetLoaded, UObject*);
-	static FOnAssetLoaded OnAssetLoaded;
+	static COREUOBJECT_API FOnAssetLoaded OnAssetLoaded;
 
 	DECLARE_TS_MULTICAST_DELEGATE_OneParam(FOnObjectConstructed, UObject*);
-	static FOnObjectConstructed OnObjectConstructed;
+	static COREUOBJECT_API FOnObjectConstructed OnObjectConstructed;
 
 	/** Callback when packages end loading in LoadPackage or AsyncLoadPackage. All packages loaded recursively due to imports are included in the single call of the explicitly-loaded package. */
 	DECLARE_MULTICAST_DELEGATE_OneParam(FOnEndLoadPackage, const FEndLoadPackageContext&);
-	static FOnEndLoadPackage OnEndLoadPackage;
+	static COREUOBJECT_API FOnEndLoadPackage OnEndLoadPackage;
 
 	/** Delegate used by SavePackage() to create the package backup */
 	DECLARE_DELEGATE_RetVal_OneParam(bool, FAutoPackageBackupDelegate, const UPackage&);
 	UE_DEPRECATED(5.1, "Backups are no longer used in Unreal Package Saves")
-	static FAutoPackageBackupDelegate AutoPackageBackupDelegate;
+	static COREUOBJECT_API FAutoPackageBackupDelegate AutoPackageBackupDelegate;
 #endif // WITH_EDITOR
 
 	/** Called when new sparse class data has been created (and the base data initialized) for the given class */
 	DECLARE_MULTICAST_DELEGATE_ThreeParams(FOnPostInitSparseClassData, UClass*, UScriptStruct*, void*);
-	static FOnPostInitSparseClassData OnPostInitSparseClassData;
+	static COREUOBJECT_API FOnPostInitSparseClassData OnPostInitSparseClassData;
 
 	/** Called by ReloadPackage during package reloading. It will be called several times for different phases of fix-up to allow custom code to handle updating objects as needed */
 	DECLARE_MULTICAST_DELEGATE_TwoParams(FOnPackageReloaded, EPackageReloadPhase, FPackageReloadedEvent*);
-	static FOnPackageReloaded OnPackageReloaded;
+	static COREUOBJECT_API FOnPackageReloaded OnPackageReloaded;
 
 	/** Called when a package reload request is received from a network file server */
 	DECLARE_DELEGATE_OneParam(FNetworkFileRequestPackageReload, const TArray<FString>& /*PackageNames*/);
-	static FNetworkFileRequestPackageReload NetworkFileRequestPackageReload;
+	static COREUOBJECT_API FNetworkFileRequestPackageReload NetworkFileRequestPackageReload;
 
 	/** Delegate used by SavePackage() to check whether a package should be saved */
 	DECLARE_DELEGATE_RetVal_ThreeParams(bool, FIsPackageOKToSaveDelegate, UPackage*, const FString&, FOutputDevice*);
-	static FIsPackageOKToSaveDelegate IsPackageOKToSaveDelegate;
+	static COREUOBJECT_API FIsPackageOKToSaveDelegate IsPackageOKToSaveDelegate;
 
 	/** Delegate for reloaded classes that have been added. */
 	DECLARE_MULTICAST_DELEGATE_OneParam(FReloadAddedClassesDelegate, const TArray<UClass*>&);
-	static FReloadAddedClassesDelegate ReloadAddedClassesDelegate;
+	static COREUOBJECT_API FReloadAddedClassesDelegate ReloadAddedClassesDelegate;
 
 	/** Delegate for reload re-instancing complete */
 	DECLARE_MULTICAST_DELEGATE(FReloadReinstancingCompleteDelegate);
-	static FReloadReinstancingCompleteDelegate ReloadReinstancingCompleteDelegate;
+	static COREUOBJECT_API FReloadReinstancingCompleteDelegate ReloadReinstancingCompleteDelegate;
 
 	/** Delegate for reload re-instancing complete */
 	DECLARE_MULTICAST_DELEGATE_OneParam(FReloadCompleteDelegate, EReloadCompleteReason);
-	static FReloadCompleteDelegate ReloadCompleteDelegate;
+	static COREUOBJECT_API FReloadCompleteDelegate ReloadCompleteDelegate;
 
 	/** Delegate for registering hot-reloaded classes that have been added  */
 	DECLARE_MULTICAST_DELEGATE_OneParam(FRegisterHotReloadAddedClassesDelegate, const TArray<UClass*>&);
 	UE_DEPRECATED(5.0, "RegisterHotReloadAddedClassesDelegate has been deprecated, use ReloadAddedClassesDelegate.")
-	static FRegisterHotReloadAddedClassesDelegate RegisterHotReloadAddedClassesDelegate;
+	static COREUOBJECT_API FRegisterHotReloadAddedClassesDelegate RegisterHotReloadAddedClassesDelegate;
 
 	/** Delegate for registering hot-reloaded classes that changed after hot-reload for reinstancing */
 	DECLARE_MULTICAST_DELEGATE_ThreeParams(FRegisterClassForHotReloadReinstancingDelegate, UClass*, UClass*, EHotReloadedClassFlags);
 	UE_DEPRECATED(5.0, "RegisterClassForHotReloadReinstancingDelegate has been deprecated, use FReload for class re-instancing.")
-	static FRegisterClassForHotReloadReinstancingDelegate RegisterClassForHotReloadReinstancingDelegate;
+	static COREUOBJECT_API FRegisterClassForHotReloadReinstancingDelegate RegisterClassForHotReloadReinstancingDelegate;
 
 	/** Delegate for reinstancing hot-reloaded classes */
 	DECLARE_MULTICAST_DELEGATE(FReinstanceHotReloadedClassesDelegate);
 	UE_DEPRECATED(5.0, "ReinstanceHotReloadedClassesDelegate has been deprecated, use FReload for class re-instancing or ReloadReinstancingCompleteDelegate for notification")
-	static FReinstanceHotReloadedClassesDelegate ReinstanceHotReloadedClassesDelegate;
+	static COREUOBJECT_API FReinstanceHotReloadedClassesDelegate ReinstanceHotReloadedClassesDelegate;
 
 	/** Delegate for catching when UClasses/UStructs/UEnums would be available via FindObject<>(), but before their CDOs would be constructed. */
 	DECLARE_MULTICAST_DELEGATE_OneParam(FCompiledInUObjectsRegisteredDelegate, FName /*Package*/);
-	static FCompiledInUObjectsRegisteredDelegate CompiledInUObjectsRegisteredDelegate;
+	static COREUOBJECT_API FCompiledInUObjectsRegisteredDelegate CompiledInUObjectsRegisteredDelegate;
 
 	/** Sent at the very beginning of LoadMap */
 	DECLARE_MULTICAST_DELEGATE_OneParam(FPreLoadMapDelegate, const FString& /* MapName */);
-	static FPreLoadMapDelegate PreLoadMap;
+	static COREUOBJECT_API FPreLoadMapDelegate PreLoadMap;
 
 	/** Sent at the very beginning of LoadMap */
 	DECLARE_MULTICAST_DELEGATE_TwoParams(FPreLoadMapWithContextDelegate, const FWorldContext& /*WorldContext*/, const FString& /* MapName */);
-	static FPreLoadMapWithContextDelegate PreLoadMapWithContext;
+	static COREUOBJECT_API FPreLoadMapWithContextDelegate PreLoadMapWithContext;
 	
 	/** Sent at the end of LoadMap */
 	DECLARE_MULTICAST_DELEGATE_OneParam(FPostLoadMapDelegate, UWorld* /* LoadedWorld */);
-	static FPostLoadMapDelegate PostLoadMapWithWorld;
+	static COREUOBJECT_API FPostLoadMapDelegate PostLoadMapWithWorld;
 
 	/** Sent when a network replay has started */
-	static FSimpleMulticastDelegate PostDemoPlay;
+	static COREUOBJECT_API FSimpleMulticastDelegate PostDemoPlay;
 
 	/** Called before garbage collection */
-	static FSimpleMulticastDelegate& GetPreGarbageCollectDelegate();
+	static COREUOBJECT_API FSimpleMulticastDelegate& GetPreGarbageCollectDelegate();
 
 	/** Delegate type for reachability analysis external roots callback. First parameter is FGarbageCollectionTracer to use for tracing, second is flags with which objects should be kept alive regardless, third is whether to force single threading */
 	DECLARE_MULTICAST_DELEGATE_ThreeParams(FTraceExternalRootsForReachabilityAnalysisDelegate, FGarbageCollectionTracer&, EObjectFlags, bool);
 
 	/** Called as last phase of reachability analysis. Allow external systems to add UObject roots *after* first reachability pass has been done */
 	UE_DEPRECATED(5.2, "Ability to inject extra root objects will be removed to reduce GC complexity")
-	static FTraceExternalRootsForReachabilityAnalysisDelegate TraceExternalRootsForReachabilityAnalysis;
+	static COREUOBJECT_API FTraceExternalRootsForReachabilityAnalysisDelegate TraceExternalRootsForReachabilityAnalysis;
 
 	/** Called after reachability analysis, before any purging */
-	static FSimpleMulticastDelegate PostReachabilityAnalysis;
+	static COREUOBJECT_API FSimpleMulticastDelegate PostReachabilityAnalysis;
 
 	/** Called after garbage collection */
-	static FSimpleMulticastDelegate& GetPostGarbageCollect();
+	static COREUOBJECT_API FSimpleMulticastDelegate& GetPostGarbageCollect();
 
 	/** Called after purging unreachable objects during garbage collection */
-	static FSimpleMulticastDelegate& GetPostPurgeGarbageDelegate();
+	static COREUOBJECT_API FSimpleMulticastDelegate& GetPostPurgeGarbageDelegate();
 
 #if !UE_BUILD_SHIPPING
 	/** Called when garbage collection detects references to objects that are marked for explicit destruction by MarkAsGarbage */
 	DECLARE_MULTICAST_DELEGATE_OneParam(FOnReportGarbageReferencers, TConstArrayView<struct FGarbageReferenceInfo>);
-	static FOnReportGarbageReferencers& GetGarbageCollectReportGarbageReferencers();
+	static COREUOBJECT_API FOnReportGarbageReferencers& GetGarbageCollectReportGarbageReferencers();
 #endif
 
 	/** Called before ConditionalBeginDestroy phase of garbage collection */
-	static FSimpleMulticastDelegate PreGarbageCollectConditionalBeginDestroy;
+	static COREUOBJECT_API FSimpleMulticastDelegate PreGarbageCollectConditionalBeginDestroy;
 
 	/** Called after ConditionalBeginDestroy phase of garbage collection */
-	static FSimpleMulticastDelegate PostGarbageCollectConditionalBeginDestroy;
+	static COREUOBJECT_API FSimpleMulticastDelegate PostGarbageCollectConditionalBeginDestroy;
 
 	/** Queries whether an object should be loaded on top ( replace ) an already existing one */
 	DECLARE_DELEGATE_RetVal_OneParam(bool, FOnLoadObjectsOnTop, const FString&);
-	static FOnLoadObjectsOnTop ShouldLoadOnTop;
+	static COREUOBJECT_API FOnLoadObjectsOnTop ShouldLoadOnTop;
 
 	/** Called when path to world root is changed */
 	DECLARE_MULTICAST_DELEGATE_OneParam(FPackageCreatedForLoad, class UPackage*);
-	static FPackageCreatedForLoad PackageCreatedForLoad;
+	static COREUOBJECT_API FPackageCreatedForLoad PackageCreatedForLoad;
 
 	/** Called when trying to figure out if a UObject is a primary asset, if it doesn't implement GetPrimaryAssetId itself */
 	DECLARE_DELEGATE_RetVal_OneParam(FPrimaryAssetId, FGetPrimaryAssetIdForObject, const UObject*);
-	static FGetPrimaryAssetIdForObject GetPrimaryAssetIdForObject;
+	static COREUOBJECT_API FGetPrimaryAssetIdForObject GetPrimaryAssetIdForObject;
 
 	/** Called during cooking to see if a specific package should be cooked for a given target platform */
 	DECLARE_DELEGATE_RetVal_TwoParams(bool, FShouldCookPackageForPlatform, const UPackage*, const ITargetPlatform*);
-	static FShouldCookPackageForPlatform ShouldCookPackageForPlatform;
+	static COREUOBJECT_API FShouldCookPackageForPlatform ShouldCookPackageForPlatform;
 };
 
 /** Allows release builds to override not verifying GC assumptions. Useful for profiling as it's hitchy. */
diff --git a/Engine/Source/Runtime/CoreUObject/Public/UObject/UObjectThreadContext.h b/Engine/Source/Runtime/CoreUObject/Public/UObject/UObjectThreadContext.h
--- a/Engine/Source/Runtime/CoreUObject/Public/UObject/UObjectThreadContext.h
+++ b/Engine/Source/Runtime/CoreUObject/Public/UObject/UObjectThreadContext.h
@@ -125,17 +125,17 @@ private:
 };
 
 /** Structure that holds the current serialization state of UObjects */
-struct COREUOBJECT_API FUObjectSerializeContext
+struct FUObjectSerializeContext
 {
 	friend class FUObjectThreadContext;
 
 private:
 
 	/** Constructor */
-	FUObjectSerializeContext();
+	COREUOBJECT_API FUObjectSerializeContext();
 
 	/** Destructor */
-	~FUObjectSerializeContext();
+	COREUOBJECT_API ~FUObjectSerializeContext();
 
 	/** Reference count of this context */
 	int32 RefCount;
@@ -169,8 +169,8 @@ public:
 	FLinkerLoad* SerializedExportLinker;
 
 	/** Adds a new loaded object */
-	void AddLoadedObject(UObject* InObject);
-	void AddUniqueLoadedObjects(const TArray<UObject*>& InObjects);
+	COREUOBJECT_API void AddLoadedObject(UObject* InObject);
+	COREUOBJECT_API void AddUniqueLoadedObjects(const TArray<UObject*>& InObjects);
 
 	/** Checks if object loading has started */
 	bool HasStartedLoading() const
@@ -182,8 +182,8 @@ public:
 		return ObjBeginLoadCount;
 	}
 
-	int32 IncrementBeginLoadCount();
-	int32 DecrementBeginLoadCount();
+	COREUOBJECT_API int32 IncrementBeginLoadCount();
+	COREUOBJECT_API int32 DecrementBeginLoadCount();
 
 	int32 IncrementImportCount()
 	{
@@ -213,7 +213,7 @@ public:
 		return !!ObjectsLoaded.Num();
 	}
 
-	bool PRIVATE_PatchNewObjectIntoExport(UObject* OldObject, UObject* NewObject);
+	COREUOBJECT_API bool PRIVATE_PatchNewObjectIntoExport(UObject* OldObject, UObject* NewObject);
 
 	/** This is only meant to be used by FAsyncPackage for performance reasons. The ObjectsLoaded array should not be manipulated directly! */
 	TArray<UObject*>& PRIVATE_GetObjectsLoadedInternalUseOnly()
@@ -253,13 +253,13 @@ public:
 	}
 
 	/** Attaches a linker to this context */
-	void AttachLinker(FLinkerLoad* InLinker);
+	COREUOBJECT_API void AttachLinker(FLinkerLoad* InLinker);
 	
 	/** Detaches a linker from this context */
-	void DetachLinker(FLinkerLoad* InLinker);
+	COREUOBJECT_API void DetachLinker(FLinkerLoad* InLinker);
 
 	/** Detaches all linkers from this context */
-	void DetachFromLinkers();
+	COREUOBJECT_API void DetachFromLinkers();
 
 	//~ TRefCountPtr interface
 	int32 AddRef()
diff --git a/Engine/Source/Runtime/CoreUObject/Public/UObject/UnrealType.h b/Engine/Source/Runtime/CoreUObject/Public/UObject/UnrealType.h
--- a/Engine/Source/Runtime/CoreUObject/Public/UObject/UnrealType.h
+++ b/Engine/Source/Runtime/CoreUObject/Public/UObject/UnrealType.h
@@ -150,9 +150,9 @@ enum class EPropertyPointerType
 //
 // An UnrealScript variable.
 //
-class COREUOBJECT_API FProperty : public FField
+class FProperty : public FField
 {
-	DECLARE_FIELD(FProperty, FField, CASTCLASS_FProperty)
+	DECLARE_FIELD_API(FProperty, FField, CASTCLASS_FProperty, COREUOBJECT_API)
 
 	// Persistent variables.
 	int32			ArrayDim;
@@ -180,35 +180,35 @@ public:
 
 public:
 	// Constructors.
-	FProperty(FFieldVariant InOwner, const FName& InName, EObjectFlags InObjectFlags);
+	COREUOBJECT_API FProperty(FFieldVariant InOwner, const FName& InName, EObjectFlags InObjectFlags);
 
 	UE_DEPRECATED(5.1, "Compiled-in property constructor is deprecated, use other constructors instead.")
-	FProperty(FFieldVariant InOwner, const FName& InName, EObjectFlags InObjectFlags, int32 InOffset, EPropertyFlags InFlags);
+	COREUOBJECT_API FProperty(FFieldVariant InOwner, const FName& InName, EObjectFlags InObjectFlags, int32 InOffset, EPropertyFlags InFlags);
 
 	/**
 	 * Constructor used for constructing compiled-in properties
 	 * @param InOwner Owner of the property
 	 * @param PropBase Pointer to the compiled in structure describing the property
 	 **/
-	FProperty(FFieldVariant InOwner, const UECodeGen_Private::FPropertyParamsBaseWithOffset& Prop, EPropertyFlags AdditionalPropertyFlags = CPF_None);
+	COREUOBJECT_API FProperty(FFieldVariant InOwner, const UECodeGen_Private::FPropertyParamsBaseWithOffset& Prop, EPropertyFlags AdditionalPropertyFlags = CPF_None);
 
 	/**
 	 * Constructor used for constructing compiled-in properties
 	 * @param InOwner Owner of the property
 	 * @param PropBase Pointer to the compiled in structure describing the property
 	 **/
-	FProperty(FFieldVariant InOwner, const UECodeGen_Private::FPropertyParamsBaseWithoutOffset& Prop, EPropertyFlags AdditionalPropertyFlags = CPF_None);
+	COREUOBJECT_API FProperty(FFieldVariant InOwner, const UECodeGen_Private::FPropertyParamsBaseWithoutOffset& Prop, EPropertyFlags AdditionalPropertyFlags = CPF_None);
 
 #if WITH_EDITORONLY_DATA
-	explicit FProperty(UField* InField);
+	COREUOBJECT_API explicit FProperty(UField* InField);
 #endif // WITH_EDITORONLY_DATA
 
 	// UObject interface
-	virtual void Serialize( FArchive& Ar ) override;
+	COREUOBJECT_API virtual void Serialize( FArchive& Ar ) override;
 	// End of UObject interface
 
 	// FField interface
-	virtual void PostDuplicate(const FField& InField) override;
+	COREUOBJECT_API virtual void PostDuplicate(const FField& InField) override;
 
 	/** parses and imports a text definition of a single property's value (if array, may be an individual element)
 	 * also includes parsing of special operations for array properties (Add/Remove/RemoveIndex/Empty)
@@ -221,17 +221,17 @@ public:
 	 * @param DefinedProperties (out)	list of properties/indices that have been parsed by previous calls, so duplicate definitions cause an error
 	 * @return pointer to remaining text in the stream (even on failure, but on failure it may not be advanced past the entire key/value pair)
 	 */
-	static const TCHAR* ImportSingleProperty( const TCHAR* Str, void* DestData, class UStruct* ObjectStruct, UObject* SubobjectOuter, int32 PortFlags,
+	static COREUOBJECT_API const TCHAR* ImportSingleProperty( const TCHAR* Str, void* DestData, class UStruct* ObjectStruct, UObject* SubobjectOuter, int32 PortFlags,
 											FOutputDevice* Warn, TArray<struct FDefinedProperty>& DefinedProperties );
 
 	/** Gets a redirected property name, will return NAME_None if no redirection was found */
-	static FName FindRedirectedPropertyName(UStruct* ObjectStruct, FName OldName);
+	static COREUOBJECT_API FName FindRedirectedPropertyName(UStruct* ObjectStruct, FName OldName);
 
 	// UHT interface
-	void ExportCppDeclaration(FOutputDevice& Out, EExportedDeclaration::Type DeclarationType, const TCHAR* ArrayDimOverride = NULL, uint32 AdditionalExportCPPFlags = 0
+	COREUOBJECT_API void ExportCppDeclaration(FOutputDevice& Out, EExportedDeclaration::Type DeclarationType, const TCHAR* ArrayDimOverride = NULL, uint32 AdditionalExportCPPFlags = 0
 		, bool bSkipParameterName = false, const FString* ActualCppType = nullptr, const FString* ActualExtendedType = nullptr, const FString* ActualParameterName = nullptr) const;
-	virtual FString GetCPPMacroType( FString& ExtendedTypeText ) const;
-	virtual bool PassCPPArgsByRef() const;
+	COREUOBJECT_API virtual FString GetCPPMacroType( FString& ExtendedTypeText ) const;
+	COREUOBJECT_API virtual bool PassCPPArgsByRef() const;
 
 	/**
 	 * Returns the C++ name of the property, including the _DEPRECATED suffix if the 
@@ -239,7 +239,7 @@ public:
 	 *
 	 * @return C++ name of property
 	 */
-	FString GetNameCPP() const;
+	COREUOBJECT_API FString GetNameCPP() const;
 
 	/**
 	 * Returns the text to use for exporting this property to header file.
@@ -247,14 +247,14 @@ public:
 	 * @param	ExtendedTypeText	for property types which use templates, will be filled in with the type
 	 * @param	CPPExportFlags		flags for modifying the behavior of the export
 	 */
-	virtual FString GetCPPType( FString* ExtendedTypeText=NULL, uint32 CPPExportFlags=0 ) const PURE_VIRTUAL(FProperty::GetCPPType,return TEXT(""););
+	COREUOBJECT_API virtual FString GetCPPType( FString* ExtendedTypeText=NULL, uint32 CPPExportFlags=0 ) const PURE_VIRTUAL(FProperty::GetCPPType,return TEXT(""););
 
-	virtual FString GetCPPTypeForwardDeclaration() const PURE_VIRTUAL(FProperty::GetCPPTypeForwardDeclaration, return TEXT(""););
+	COREUOBJECT_API virtual FString GetCPPTypeForwardDeclaration() const PURE_VIRTUAL(FProperty::GetCPPTypeForwardDeclaration, return TEXT(""););
 	// End of UHT interface
 
 #if WITH_EDITORONLY_DATA
 	/** Gets the wrapper object for this property or creates one if it doesn't exist yet */
-	UPropertyWrapper* GetUPropertyWrapper();
+	COREUOBJECT_API UPropertyWrapper* GetUPropertyWrapper();
 #endif
 
 	/** Checks if this property as a native setter function */
@@ -295,19 +295,19 @@ private:
 	/** Set the alignment offset for this property 
 	 * @return the size of the structure including this newly added property
 	*/
-	int32 SetupOffset();
+	COREUOBJECT_API int32 SetupOffset();
 
 protected:
 	friend class FMapProperty;
 	friend class UEProperty_Private::FProperty_DoNotUse;
 
 	/** Set the alignment offset for this property - added for FMapProperty */
-	void SetOffset_Internal(int32 NewOffset);
+	COREUOBJECT_API void SetOffset_Internal(int32 NewOffset);
 
 	/**
 	 * Initializes internal state.
 	 */
-	void Init();
+	COREUOBJECT_API void Init();
 
 public:
 	/** Return offset of property from container base. */
@@ -348,7 +348,7 @@ public:
 	}
 
 protected:
-	virtual void LinkInternal(FArchive& Ar);
+	COREUOBJECT_API virtual void LinkInternal(FArchive& Ar);
 public:
 
 	/**
@@ -364,7 +364,7 @@ public:
 	*			CannotConvert:    the tag is not something that the property can convert.
 	*			UseSerializeItem: no conversion was done on the property - this can mean that the tag is correct and normal serialization applies or that the tag is incompatible.
 	*/
-	virtual EConvertFromTypeResult ConvertFromType(const FPropertyTag& Tag, FStructuredArchive::FSlot Slot, uint8* Data, UStruct* DefaultsStruct);
+	COREUOBJECT_API virtual EConvertFromTypeResult ConvertFromType(const FPropertyTag& Tag, FStructuredArchive::FSlot Slot, uint8* Data, UStruct* DefaultsStruct);
 
 	/**
 	 * Determines whether the property values are identical.
@@ -375,7 +375,7 @@ public:
 	 *
 	 * @return	true if the property values are identical
 	 */
-	virtual bool Identical( const void* A, const void* B, uint32 PortFlags=0 ) const PURE_VIRTUAL(FProperty::Identical,return false;);
+	COREUOBJECT_API virtual bool Identical( const void* A, const void* B, uint32 PortFlags=0 ) const PURE_VIRTUAL(FProperty::Identical,return false;);
 
 	/**
 	 * Determines whether the property values are identical.
@@ -441,9 +441,9 @@ public:
 		}
 	}
 
-	virtual void SerializeItem(FStructuredArchive::FSlot Slot, void* Value, void const* Defaults = NULL) const PURE_VIRTUAL(FProperty::SerializeItem, );
-	virtual bool NetSerializeItem( FArchive& Ar, UPackageMap* Map, void* Data, TArray<uint8> * MetaData = NULL ) const;
-	virtual bool SupportsNetSharedSerialization() const;
+	COREUOBJECT_API virtual void SerializeItem(FStructuredArchive::FSlot Slot, void* Value, void const* Defaults = NULL) const PURE_VIRTUAL(FProperty::SerializeItem, );
+	COREUOBJECT_API virtual bool NetSerializeItem( FArchive& Ar, UPackageMap* Map, void* Data, TArray<uint8> * MetaData = NULL ) const;
+	COREUOBJECT_API virtual bool SupportsNetSharedSerialization() const;
 
 	UE_DEPRECATED(5.1, "Please use ExportTextItem_InContainer or ExportTextItem_Direct instead.")
 	virtual void ExportTextItem(FString& ValueStr, const void* PropertyValue, const void* DefaultValue, UObject* Parent, int32 PortFlags, UObject* ExportRootScope = nullptr) const
@@ -534,7 +534,7 @@ public:
 	* @param InValue Pointer to the memory that the value will be copied from. Must be at least ElementSize big
 	* @param ArrayIndex Index into the static array to copy the value from. If the property is not a static array it should be 0
 	*/
-	void SetSingleValue_InContainer(void* OutContainer, const void* InValue, int32 ArrayIndex) const;
+	COREUOBJECT_API void SetSingleValue_InContainer(void* OutContainer, const void* InValue, int32 ArrayIndex) const;
 
 	/**
 	* Copies a single value to OutValue even if the property represents a static array of values
@@ -542,13 +542,13 @@ public:
 	* @param OutValue Pointer to the memory that the value will be copied to. Must be at least ElementSize big
 	* @param ArrayIndex Index into the static array to copy the value from. If the property is not a static array it should be 0
 	*/
-	void GetSingleValue_InContainer(const void* InContainer, void* OutValue, int32 ArrayIndex) const;
+	COREUOBJECT_API void GetSingleValue_InContainer(const void* InContainer, void* OutValue, int32 ArrayIndex) const;
 
 	/** Allocates and initializes memory to hold a value this property represents */
-	void* AllocateAndInitializeValue() const;
+	COREUOBJECT_API void* AllocateAndInitializeValue() const;
 
 	/** Destroys and frees memory with a value this property represents */
-	void DestroyAndFreeValue(void* InMemory) const;
+	COREUOBJECT_API void DestroyAndFreeValue(void* InMemory) const;
 
 	/**
 	 * Helper function for setting container / struct property value and performing operation directly on the value memory
@@ -556,7 +556,7 @@ public:
 	 * @param DirectPropertyAddress Direct property value address. Can be null only if OutContainer is a valid pointer.
 	 * @param DirectValueAccessFunc Function that manipulates directly on property value address. The value address can be different than the passed in DirectPropertyAddress if setters and getters are present and OutContainer pointer is valid.
 	 */
-	void PerformOperationWithSetter(void* OutContainer, void* DirectPropertyAddress, TFunctionRef<void(void*)> DirectValueAccessFunc) const;
+	COREUOBJECT_API void PerformOperationWithSetter(void* OutContainer, void* DirectPropertyAddress, TFunctionRef<void(void*)> DirectValueAccessFunc) const;
 
 	/**
 	 * Helper function for getting container / struct property value and performing operation directly on the value memory
@@ -564,7 +564,7 @@ public:
 	 * @param DirectPropertyAddress Direct property value address. Can be null only if OutContainer is a valid pointer.
 	 * @param DirectValueAccessFunc Function that manipulates directly on property value address. The value address can be different than the passed in DirectPropertyAddress if setters and getters are present and OutContainer pointer is valid.
 	 */
-	void PerformOperationWithGetter(void* OutContainer, const void* DirectPropertyAddress, TFunctionRef<void(const void*)> DirectValueAccessFunc) const;
+	COREUOBJECT_API void PerformOperationWithGetter(void* OutContainer, const void* DirectPropertyAddress, TFunctionRef<void(const void*)> DirectValueAccessFunc) const;
 
 	/** 
 	 * Gets value address at given index inside of a static array or container
@@ -572,7 +572,7 @@ public:
 	 * @param Index into the static array or container
 	 * @returns address of the value at given index
 	 */
-	virtual void* GetValueAddressAtIndex_Direct(const FProperty* Inner, void* InValueAddress, int32 Index) const;
+	COREUOBJECT_API virtual void* GetValueAddressAtIndex_Direct(const FProperty* Inner, void* InValueAddress, int32 Index) const;
 
 #if WITH_EDITORONLY_DATA
 	/**
@@ -586,16 +586,16 @@ public:
 	 *                        If true, sub- properties that are editor-only should not be appended to the hash.
 	 *                        This property's base data is appended without regard for bSkipEditorOnly.
 	 */
-	virtual void AppendSchemaHash(FBlake3& Builder, bool bSkipEditorOnly) const;
+	COREUOBJECT_API virtual void AppendSchemaHash(FBlake3& Builder, bool bSkipEditorOnly) const;
 #endif
 protected:
 
-	virtual void ExportText_Internal(FString& ValueStr, const void* PropertyValueOrContainer, EPropertyPointerType PointerType, const void* DefaultValue, UObject* Parent, int32 PortFlags, UObject* ExportRootScope = nullptr) const PURE_VIRTUAL(FProperty::ExportText, );
-	virtual const TCHAR* ImportText_Internal(const TCHAR* Buffer, void* ContainerOrPropertyPtr, EPropertyPointerType PointerType, UObject* OwnerObject, int32 PortFlags, FOutputDevice* ErrorText) const PURE_VIRTUAL(FProperty::ImportText, return nullptr;);
+	COREUOBJECT_API virtual void ExportText_Internal(FString& ValueStr, const void* PropertyValueOrContainer, EPropertyPointerType PointerType, const void* DefaultValue, UObject* Parent, int32 PortFlags, UObject* ExportRootScope = nullptr) const PURE_VIRTUAL(FProperty::ExportText, );
+	COREUOBJECT_API virtual const TCHAR* ImportText_Internal(const TCHAR* Buffer, void* ContainerOrPropertyPtr, EPropertyPointerType PointerType, UObject* OwnerObject, int32 PortFlags, FOutputDevice* ErrorText) const PURE_VIRTUAL(FProperty::ImportText, return nullptr;);
 
 public:
 	
-	bool ExportText_Direct(FString& ValueStr, const void* Data, const void* Delta, UObject* Parent, int32 PortFlags, UObject* ExportRootScope = nullptr) const;
+	COREUOBJECT_API bool ExportText_Direct(FString& ValueStr, const void* Data, const void* Delta, UObject* Parent, int32 PortFlags, UObject* ExportRootScope = nullptr) const;
 	FORCEINLINE bool ExportText_InContainer(int32 Index, FString& ValueStr, const void* Data, const void* Delta, UObject* Parent, int32 PortFlags, UObject* ExportRootScope = nullptr) const
 	{
 		return ExportText_Direct(ValueStr, ContainerPtrToValuePtr<void>(Data, Index), ContainerPtrToValuePtrForDefaults<void>(NULL, Delta, Index), Parent, PortFlags, ExportRootScope);
@@ -766,11 +766,11 @@ public:
 	/**
 	 * Returns the hash value for an element of this property.
 	 */
-	uint32 GetValueTypeHash(const void* Src) const;
+	COREUOBJECT_API uint32 GetValueTypeHash(const void* Src) const;
 
 protected:
-	virtual void CopyValuesInternal( void* Dest, void const* Src, int32 Count  ) const;
-	virtual uint32 GetValueTypeHashInternal(const void* Src) const;
+	COREUOBJECT_API virtual void CopyValuesInternal( void* Dest, void const* Src, int32 Count  ) const;
+	COREUOBJECT_API virtual uint32 GetValueTypeHashInternal(const void* Src) const;
 
 public:
 	/**
@@ -811,7 +811,7 @@ public:
 	 *									SIZE = the ElementSize of this FProperty
 	 * @param	Src					the address of the value to copy from. should be evaluated the same way as Dest
 	 */
-	virtual void CopySingleValueToScriptVM( void* Dest, void const* Src ) const;
+	COREUOBJECT_API virtual void CopySingleValueToScriptVM( void* Dest, void const* Src ) const;
 
 	/**
 	 * Copy the value for all elements of this property. To the script VM.
@@ -821,13 +821,13 @@ public:
 	 *									OFFSET = the Offset of this FProperty
 	 * @param	Src					the address of the value to copy from. should be evaluated the same way as Dest
 	 */
-	virtual void CopyCompleteValueToScriptVM( void* Dest, void const* Src ) const;
+	COREUOBJECT_API virtual void CopyCompleteValueToScriptVM( void* Dest, void const* Src ) const;
 
 	/**
 	 * Equivalent to the above functions, but using the container and aware of getters/setters when the container has them.
 	 */
-	virtual void CopyCompleteValueToScriptVM_InContainer( void* OutValue, void const* InContainer ) const;
-	virtual void CopyCompleteValueFromScriptVM_InContainer( void* OutContainer, void const* InValue ) const;
+	COREUOBJECT_API virtual void CopyCompleteValueToScriptVM_InContainer( void* OutValue, void const* InContainer ) const;
+	COREUOBJECT_API virtual void CopyCompleteValueFromScriptVM_InContainer( void* OutContainer, void const* InValue ) const;
 
 	/**
 	 * Copy the value for a single element of this property. From the script VM.
@@ -839,7 +839,7 @@ public:
 	 *									SIZE = the ElementSize of this FProperty
 	 * @param	Src					the address of the value to copy from. should be evaluated the same way as Dest
 	 */
-	virtual void CopySingleValueFromScriptVM( void* Dest, void const* Src ) const;
+	COREUOBJECT_API virtual void CopySingleValueFromScriptVM( void* Dest, void const* Src ) const;
 
 	/**
 	 * Copy the value for all elements of this property. From the script VM.
@@ -849,7 +849,7 @@ public:
 	 *									OFFSET = the Offset of this FProperty
 	 * @param	Src					the address of the value to copy from. should be evaluated the same way as Dest
 	 */
-	virtual void CopyCompleteValueFromScriptVM( void* Dest, void const* Src ) const;
+	COREUOBJECT_API virtual void CopyCompleteValueFromScriptVM( void* Dest, void const* Src ) const;
 
 	/**
 	 * Zeros the value for this property. The existing data is assumed valid (so for example this calls FString::Empty)
@@ -886,7 +886,7 @@ public:
 		}
 	}
 protected:
-	virtual void ClearValueInternal( void* Data ) const;
+	COREUOBJECT_API virtual void ClearValueInternal( void* Data ) const;
 public:
 	/**
 	 * Destroys the value for this property. The existing data is assumed valid (so for example this calls FString::Empty)
@@ -915,7 +915,7 @@ public:
 		}
 	}
 protected:
-	virtual void DestroyValueInternal( void* Dest ) const;
+	COREUOBJECT_API virtual void DestroyValueInternal( void* Dest ) const;
 public:
 
 	/**
@@ -955,7 +955,7 @@ public:
 		}
 	}
 protected:
-	virtual void InitializeValueInternal( void* Dest ) const;
+	COREUOBJECT_API virtual void InitializeValueInternal( void* Dest ) const;
 public:
 
 	/**
@@ -966,9 +966,9 @@ public:
 	 *
 	 * @return	true if ImportText should be allowed
 	 */
-	bool ValidateImportFlags( uint32 PortFlags, FOutputDevice* ErrorText = NULL ) const;
-	bool ShouldPort( uint32 PortFlags=0 ) const;
-	virtual FName GetID() const;
+	COREUOBJECT_API bool ValidateImportFlags( uint32 PortFlags, FOutputDevice* ErrorText = NULL ) const;
+	COREUOBJECT_API bool ShouldPort( uint32 PortFlags=0 ) const;
+	COREUOBJECT_API virtual FName GetID() const;
 
 	/**
 	 * Creates new copies of components
@@ -978,9 +978,9 @@ public:
 	 * @param	Owner				the object that contains this property's data
 	 * @param	InstanceGraph		contains the mappings of instanced objects and components to their templates
 	 */
-	virtual void InstanceSubobjects( void* Data, void const* DefaultData, UObject* Owner, struct FObjectInstancingGraph* InstanceGraph );
+	COREUOBJECT_API virtual void InstanceSubobjects( void* Data, void const* DefaultData, UObject* Owner, struct FObjectInstancingGraph* InstanceGraph );
 
-	virtual int32 GetMinAlignment() const;
+	COREUOBJECT_API virtual int32 GetMinAlignment() const;
 
 	/**
 	 * Returns true if this property, or in the case of e.g. array or struct properties any sub- property, contains a
@@ -990,7 +990,7 @@ public:
 	 *
 	 * @return true if property (or sub- properties) contains the specified type of UObject reference, false otherwise
 	 */
-	virtual bool ContainsObjectReference(TArray<const FStructProperty*>& EncounteredStructProps, EPropertyObjectReferenceType InReferenceType = EPropertyObjectReferenceType::Strong) const;
+	COREUOBJECT_API virtual bool ContainsObjectReference(TArray<const FStructProperty*>& EncounteredStructProps, EPropertyObjectReferenceType InReferenceType = EPropertyObjectReferenceType::Strong) const;
 
 	/**
 	 * Returns true if this property, or in the case of e.g. array or struct properties any sub- property, contains a
@@ -1019,7 +1019,7 @@ public:
 	 * Emits tokens used by realtime garbage collection code to passed in ReferenceTokenStream. The offset emitted is relative
 	 * to the passed in BaseOffset which is used by e.g. arrays of structs.
 	 */
-	virtual void EmitReferenceInfo(UE::GC::FSchemaBuilder& Schema, int32 BaseOffset, TArray<const FStructProperty*>& EncounteredStructProps, UE::GC::FPropertyStack& DebugPath);
+	COREUOBJECT_API virtual void EmitReferenceInfo(UE::GC::FSchemaBuilder& Schema, int32 BaseOffset, TArray<const FStructProperty*>& EncounteredStructProps, UE::GC::FPropertyStack& DebugPath);
 
     // @TODO: Surely this can have an int32 overflow. This should probably return size_t. Just
     // need to audit all callers to make such a change.
@@ -1027,7 +1027,7 @@ public:
 	{
 		return ArrayDim * ElementSize;
 	}
-	bool ShouldSerializeValue( FArchive& Ar ) const;
+	COREUOBJECT_API bool ShouldSerializeValue( FArchive& Ar ) const;
 
 	/**
 	 * Determines whether this property value is eligible for copying when duplicating an object
@@ -1128,20 +1128,20 @@ public:
 	}
 
 	/** returns true, if Other is property of exactly the same type */
-	virtual bool SameType(const FProperty* Other) const;
+	COREUOBJECT_API virtual bool SameType(const FProperty* Other) const;
 
 	ELifetimeCondition GetBlueprintReplicationCondition() const { return BlueprintReplicationCondition; }
 	void SetBlueprintReplicationCondition(ELifetimeCondition InBlueprintReplicationCondition) { BlueprintReplicationCondition = InBlueprintReplicationCondition; }
 };
 
 
-class COREUOBJECT_API FPropertyHelpers
+class FPropertyHelpers
 {
 public:
-	static const TCHAR* ReadToken( const TCHAR* Buffer, FString& Out, bool DottedNames = false);
+	static COREUOBJECT_API const TCHAR* ReadToken( const TCHAR* Buffer, FString& Out, bool DottedNames = false);
 
 	// @param Out Appended to
-	static const TCHAR* ReadToken( const TCHAR* Buffer, FStringBuilderBase& Out, bool DottedNames = false);
+	static COREUOBJECT_API const TCHAR* ReadToken( const TCHAR* Buffer, FStringBuilderBase& Out, bool DottedNames = false);
 };
 
 namespace UEProperty_Private
@@ -1158,7 +1158,7 @@ namespace UEProperty_Private
 		 * 
 		 * Explicitly exposed for this singular case -- DO NOT USE otherwise.
 		 */
-		static COREUOBJECT_API void Unsafe_AlterOffset(FProperty& Property, const int32 OffsetOverride)
+		static void Unsafe_AlterOffset(FProperty& Property, const int32 OffsetOverride)
 		{
 			Property.SetOffset_Internal(OffsetOverride);
 		}
@@ -1166,7 +1166,7 @@ namespace UEProperty_Private
 }
 
 /** reference to a property and optional array index used in property text import to detect duplicate references */
-struct COREUOBJECT_API FDefinedProperty
+struct FDefinedProperty
 {
     FProperty* Property;
     int32 Index;
@@ -1179,7 +1179,7 @@ struct COREUOBJECT_API FDefinedProperty
 /**
  * Creates a temporary object that represents the default constructed value of a FProperty
  */
-class COREUOBJECT_API FDefaultConstructedPropertyElement
+class FDefaultConstructedPropertyElement
 {
 public:
 	FDefaultConstructedPropertyElement() = default;
diff --git a/Engine/Source/Runtime/CoreUObject/Public/UObject/UnrealTypePrivate.h b/Engine/Source/Runtime/CoreUObject/Public/UObject/UnrealTypePrivate.h
--- a/Engine/Source/Runtime/CoreUObject/Public/UObject/UnrealTypePrivate.h
+++ b/Engine/Source/Runtime/CoreUObject/Public/UObject/UnrealTypePrivate.h
@@ -11,7 +11,7 @@
 
 #define USE_UPROPERTY_LOAD_DEFERRING (USE_CIRCULAR_DEPENDENCY_LOAD_DEFERRING && WITH_EDITORONLY_DATA)
 
-class COREUOBJECT_API UProperty : public UField
+class UProperty : public UField
 {
 	DECLARE_CASTED_CLASS_INTRINSIC_NO_CTOR(UProperty, UField, CLASS_Abstract, TEXT("/Script/CoreUObject"), CASTCLASS_FProperty, NO_API)
 	DECLARE_WITHIN(UObject)
@@ -42,12 +42,12 @@ public:
 
 
 	// Constructors.
-	UProperty(const FObjectInitializer& ObjectInitializer = FObjectInitializer::Get());
-	UProperty(ECppProperty, int32 InOffset, EPropertyFlags InFlags);
-	UProperty( const FObjectInitializer& ObjectInitializer, ECppProperty, int32 InOffset, EPropertyFlags InFlags );
+	COREUOBJECT_API UProperty(const FObjectInitializer& ObjectInitializer = FObjectInitializer::Get());
+	COREUOBJECT_API UProperty(ECppProperty, int32 InOffset, EPropertyFlags InFlags);
+	COREUOBJECT_API UProperty( const FObjectInitializer& ObjectInitializer, ECppProperty, int32 InOffset, EPropertyFlags InFlags );
 
 	// UObject interface
-	virtual void Serialize( FArchive& Ar ) override;
+	COREUOBJECT_API virtual void Serialize( FArchive& Ar ) override;
 	// End of UObject interface
 
 	/**
@@ -174,12 +174,12 @@ public:
 
 	FField* AssociatedField;
 
-	virtual FField* GetAssociatedFField() override;
-	virtual void SetAssociatedFField(FField* InField) override;
+	COREUOBJECT_API virtual FField* GetAssociatedFField() override;
+	COREUOBJECT_API virtual void SetAssociatedFField(FField* InField) override;
 #endif // WITH_EDITORONLY_DATA
 };
 
-class COREUOBJECT_API UNumericProperty : public UProperty
+class UNumericProperty : public UProperty
 {
 	DECLARE_CASTED_CLASS_INTRINSIC(UNumericProperty, UProperty, CLASS_Abstract, TEXT("/Script/CoreUObject"), CASTCLASS_FNumericProperty)
 
@@ -192,7 +192,7 @@ class COREUOBJECT_API UNumericProperty : public UProperty
 	{}
 };
 
-class COREUOBJECT_API UByteProperty : public UNumericProperty
+class UByteProperty : public UNumericProperty
 {
 	DECLARE_CASTED_CLASS_INTRINSIC(UByteProperty, UNumericProperty, 0, TEXT("/Script/CoreUObject"), CASTCLASS_FByteProperty)
 
@@ -214,12 +214,12 @@ public:
 	}
 
 	// UObject interface.
-	virtual void Serialize( FArchive& Ar ) override;
-	static void AddReferencedObjects(UObject* InThis, FReferenceCollector& Collector);
+	COREUOBJECT_API virtual void Serialize( FArchive& Ar ) override;
+	static COREUOBJECT_API void AddReferencedObjects(UObject* InThis, FReferenceCollector& Collector);
 	// End of UObject interface
 };
 
-class COREUOBJECT_API UInt8Property : public UNumericProperty
+class UInt8Property : public UNumericProperty
 {
 	DECLARE_CASTED_CLASS_INTRINSIC(UInt8Property, UNumericProperty, 0, TEXT("/Script/CoreUObject"), CASTCLASS_FInt8Property)
 
@@ -234,7 +234,7 @@ class COREUOBJECT_API UInt8Property : public UNumericProperty
 	}
 };
 
-class COREUOBJECT_API UInt16Property : public UNumericProperty
+class UInt16Property : public UNumericProperty
 {
 	DECLARE_CASTED_CLASS_INTRINSIC(UInt16Property, UNumericProperty, 0, TEXT("/Script/CoreUObject"), CASTCLASS_FInt16Property)
 
@@ -249,7 +249,7 @@ class COREUOBJECT_API UInt16Property : public UNumericProperty
 	}
 };
 
-class COREUOBJECT_API UIntProperty : public UNumericProperty
+class UIntProperty : public UNumericProperty
 {
 	DECLARE_CASTED_CLASS_INTRINSIC(UIntProperty, UNumericProperty, 0, TEXT("/Script/CoreUObject"), CASTCLASS_FIntProperty)
 
@@ -264,7 +264,7 @@ class COREUOBJECT_API UIntProperty : public UNumericProperty
 	}
 };
 
-class COREUOBJECT_API UInt64Property : public UNumericProperty
+class UInt64Property : public UNumericProperty
 {
 	DECLARE_CASTED_CLASS_INTRINSIC(UInt64Property, UNumericProperty, 0, TEXT("/Script/CoreUObject"), CASTCLASS_FInt64Property)
 
@@ -279,7 +279,7 @@ class COREUOBJECT_API UInt64Property : public UNumericProperty
 	}
 };
 
-class COREUOBJECT_API UUInt16Property : public UNumericProperty
+class UUInt16Property : public UNumericProperty
 {
 	DECLARE_CASTED_CLASS_INTRINSIC(UUInt16Property, UNumericProperty, 0, TEXT("/Script/CoreUObject"), CASTCLASS_FUInt16Property)
 
@@ -294,7 +294,7 @@ class COREUOBJECT_API UUInt16Property : public UNumericProperty
 	}
 };
 
-class COREUOBJECT_API UUInt32Property : public UNumericProperty
+class UUInt32Property : public UNumericProperty
 {
 	DECLARE_CASTED_CLASS_INTRINSIC(UUInt32Property, UNumericProperty, 0, TEXT("/Script/CoreUObject"), CASTCLASS_FUInt32Property)
 
@@ -304,7 +304,7 @@ class COREUOBJECT_API UUInt32Property : public UNumericProperty
 	}
 };
 
-class COREUOBJECT_API UUInt64Property : public UNumericProperty
+class UUInt64Property : public UNumericProperty
 {
 	DECLARE_CASTED_CLASS_INTRINSIC(UUInt64Property, UNumericProperty, 0, TEXT("/Script/CoreUObject"), CASTCLASS_FUInt64Property)
 
@@ -319,7 +319,7 @@ class COREUOBJECT_API UUInt64Property : public UNumericProperty
 	}
 };
 
-class COREUOBJECT_API UFloatProperty : public UNumericProperty
+class UFloatProperty : public UNumericProperty
 {
 	DECLARE_CASTED_CLASS_INTRINSIC(UFloatProperty, UNumericProperty, 0, TEXT("/Script/CoreUObject"), CASTCLASS_FFloatProperty)
 
@@ -334,7 +334,7 @@ class COREUOBJECT_API UFloatProperty : public UNumericProperty
 	}
 };
 
-class COREUOBJECT_API UDoubleProperty : public UNumericProperty
+class UDoubleProperty : public UNumericProperty
 {
 	DECLARE_CASTED_CLASS_INTRINSIC(UDoubleProperty, UNumericProperty, 0, TEXT("/Script/CoreUObject"), CASTCLASS_FDoubleProperty)
 
@@ -349,7 +349,7 @@ class COREUOBJECT_API UDoubleProperty : public UNumericProperty
 	}
 };
 
-class COREUOBJECT_API UBoolProperty : public UProperty
+class UBoolProperty : public UProperty
 {
 	DECLARE_CASTED_CLASS_INTRINSIC_NO_CTOR(UBoolProperty, UProperty, 0, TEXT("/Script/CoreUObject"), CASTCLASS_FBoolProperty, NO_API)
 
@@ -365,7 +365,7 @@ public:
 	/** Mask of the field with the property value. Either equal to ByteMask or 255 in case of 'bool' type. */
 	uint8 FieldMask;
 
-	UBoolProperty(const FObjectInitializer& ObjectInitializer = FObjectInitializer::Get());
+	COREUOBJECT_API UBoolProperty(const FObjectInitializer& ObjectInitializer = FObjectInitializer::Get());
 	
 	/**
 	 * Constructor.
@@ -378,7 +378,7 @@ public:
 	 * @param InElementSize Sizeof of the boolean type this property represents.
 	 * @param bIsNativeBool true if this property represents C++ bool type.
 	 */
-	UBoolProperty(ECppProperty, int32 InOffset, EPropertyFlags InFlags, uint32 InBitMask, uint32 InElementSize, bool bIsNativeBool);
+	COREUOBJECT_API UBoolProperty(ECppProperty, int32 InOffset, EPropertyFlags InFlags, uint32 InBitMask, uint32 InElementSize, bool bIsNativeBool);
 
 	/**
 	 * Constructor.
@@ -392,10 +392,10 @@ public:
 	 * @param InElementSize Sizeof of the boolean type this property represents.
 	 * @param bIsNativeBool true if this property represents C++ bool type.
 	 */
-	UBoolProperty( const FObjectInitializer& ObjectInitializer, ECppProperty, int32 InOffset, EPropertyFlags InFlags, uint32 InBitMask, uint32 InElementSize, bool bIsNativeBool );
+	COREUOBJECT_API UBoolProperty( const FObjectInitializer& ObjectInitializer, ECppProperty, int32 InOffset, EPropertyFlags InFlags, uint32 InBitMask, uint32 InElementSize, bool bIsNativeBool );
 
 	// UObject interface.
-	virtual void Serialize( FArchive& Ar ) override;
+	COREUOBJECT_API virtual void Serialize( FArchive& Ar ) override;
 	// End of UObject interface
 
 	/** 
@@ -405,7 +405,7 @@ public:
 	 * @param InSize size of the bitfield/bool type.
 	 * @param bIsNativeBool true if this property represents C++ bool type.
 	 */
-	void SetBoolSize( const uint32 InSize, const bool bIsNativeBool = false, const uint32 InBitMask = 0 );
+	COREUOBJECT_API void SetBoolSize( const uint32 InSize, const bool bIsNativeBool = false, const uint32 InBitMask = 0 );
 
 	/**
 	 * If the return value is true this UBoolProperty represents C++ bool type.
@@ -416,7 +416,7 @@ public:
 	}
 };
 
-class COREUOBJECT_API UObjectPropertyBase : public UProperty
+class UObjectPropertyBase : public UProperty
 {
 	DECLARE_CASTED_CLASS_INTRINSIC(UObjectPropertyBase, UProperty, CLASS_Abstract, TEXT("/Script/CoreUObject"), CASTCLASS_FObjectPropertyBase)
 
@@ -436,19 +436,19 @@ public:
 	{}
 
 	// UObject interface
-	virtual void Serialize( FArchive& Ar ) override;
-	static void AddReferencedObjects(UObject* InThis, FReferenceCollector& Collector);
-	virtual void BeginDestroy() override;
+	COREUOBJECT_API virtual void Serialize( FArchive& Ar ) override;
+	static COREUOBJECT_API void AddReferencedObjects(UObject* InThis, FReferenceCollector& Collector);
+	COREUOBJECT_API virtual void BeginDestroy() override;
 	// End of UObject interface
 
 #if USE_UPROPERTY_LOAD_DEFERRING
-	void SetPropertyClass(UClass* NewPropertyClass);
+	COREUOBJECT_API void SetPropertyClass(UClass* NewPropertyClass);
 #else
 	FORCEINLINE void SetPropertyClass(UClass* NewPropertyClass) { PropertyClass = NewPropertyClass; }
 #endif // USE_UPROPERTY_LOAD_DEFERRING
 };
 
-class COREUOBJECT_API UObjectProperty : public UObjectPropertyBase
+class UObjectProperty : public UObjectPropertyBase
 {
 	DECLARE_CASTED_CLASS_INTRINSIC(UObjectProperty, UObjectPropertyBase, 0, TEXT("/Script/CoreUObject"), CASTCLASS_FObjectProperty)
 
@@ -463,7 +463,7 @@ class COREUOBJECT_API UObjectProperty : public UObjectPropertyBase
 	}
 };
 
-class COREUOBJECT_API UWeakObjectProperty : public UObjectPropertyBase
+class UWeakObjectProperty : public UObjectPropertyBase
 {
 	DECLARE_CASTED_CLASS_INTRINSIC(UWeakObjectProperty, UObjectPropertyBase, 0, TEXT("/Script/CoreUObject"), CASTCLASS_FWeakObjectProperty)
 
@@ -478,7 +478,7 @@ class COREUOBJECT_API UWeakObjectProperty : public UObjectPropertyBase
 	}
 };
 
-class COREUOBJECT_API ULazyObjectProperty : public UObjectPropertyBase
+class ULazyObjectProperty : public UObjectPropertyBase
 {
 	DECLARE_CASTED_CLASS_INTRINSIC(ULazyObjectProperty, UObjectPropertyBase, 0, TEXT("/Script/CoreUObject"), CASTCLASS_FLazyObjectProperty)
 
@@ -493,7 +493,7 @@ class COREUOBJECT_API ULazyObjectProperty : public UObjectPropertyBase
 	}
 };
 
-class COREUOBJECT_API USoftObjectProperty : public UObjectPropertyBase
+class USoftObjectProperty : public UObjectPropertyBase
 {
 	DECLARE_CASTED_CLASS_INTRINSIC(USoftObjectProperty, UObjectPropertyBase, 0, TEXT("/Script/CoreUObject"), CASTCLASS_FSoftObjectProperty)
 
@@ -506,7 +506,7 @@ class COREUOBJECT_API USoftObjectProperty : public UObjectPropertyBase
 	{}
 };
 
-class COREUOBJECT_API UClassProperty : public UObjectProperty
+class UClassProperty : public UObjectProperty
 {
 	DECLARE_CASTED_CLASS_INTRINSIC(UClassProperty, UObjectProperty, 0, TEXT("/Script/CoreUObject"), CASTCLASS_FClassProperty)
 
@@ -528,9 +528,9 @@ public:
 	}
 
 	// UObject interface
-	virtual void Serialize( FArchive& Ar ) override;
-	static void AddReferencedObjects(UObject* InThis, FReferenceCollector& Collector);
-	virtual void BeginDestroy() override;
+	COREUOBJECT_API virtual void Serialize( FArchive& Ar ) override;
+	static COREUOBJECT_API void AddReferencedObjects(UObject* InThis, FReferenceCollector& Collector);
+	COREUOBJECT_API virtual void BeginDestroy() override;
 	// End of UObject interface
 
 	/**
@@ -543,13 +543,13 @@ public:
 	 * @param  NewMetaClass    The MetaClass you want this property set with.
 	 */
 #if USE_UPROPERTY_LOAD_DEFERRING
-	void SetMetaClass(UClass* NewMetaClass);
+	COREUOBJECT_API void SetMetaClass(UClass* NewMetaClass);
 #else
 	FORCEINLINE void SetMetaClass(UClass* NewMetaClass) { MetaClass = NewMetaClass; }
 #endif // USE_UPROPERTY_LOAD_DEFERRING
 };
 
-class COREUOBJECT_API USoftClassProperty : public USoftObjectProperty
+class USoftClassProperty : public USoftObjectProperty
 {
 	DECLARE_CASTED_CLASS_INTRINSIC(USoftClassProperty, USoftObjectProperty, 0, TEXT("/Script/CoreUObject"), CASTCLASS_FSoftClassProperty)
 
@@ -569,9 +569,9 @@ public:
 	{}
 
 	// UObject interface
-	virtual void Serialize( FArchive& Ar ) override;
-	static void AddReferencedObjects(UObject* InThis, FReferenceCollector& Collector);
-	virtual void BeginDestroy() override;
+	COREUOBJECT_API virtual void Serialize( FArchive& Ar ) override;
+	static COREUOBJECT_API void AddReferencedObjects(UObject* InThis, FReferenceCollector& Collector);
+	COREUOBJECT_API virtual void BeginDestroy() override;
 	// End of UObject interface
 
 	/**
@@ -584,13 +584,13 @@ public:
 	 * @param  NewMetaClass    The MetaClass you want this property set with.
 	 */
 #if USE_UPROPERTY_LOAD_DEFERRING
-	void SetMetaClass(UClass* NewMetaClass);
+	COREUOBJECT_API void SetMetaClass(UClass* NewMetaClass);
 #else
 	FORCEINLINE void SetMetaClass(UClass* NewMetaClass) { MetaClass = NewMetaClass; }
 #endif // USE_UPROPERTY_LOAD_DEFERRING
 };
 
-class COREUOBJECT_API UInterfaceProperty : public UProperty
+class UInterfaceProperty : public UProperty
 {
 	DECLARE_CASTED_CLASS_INTRINSIC(UInterfaceProperty, UProperty, 0, TEXT("/Script/CoreUObject"), CASTCLASS_FInterfaceProperty)
 
@@ -612,9 +612,9 @@ public:
 	}
 
 	// UObject interface
-	virtual void Serialize( FArchive& Ar ) override;
-	virtual void BeginDestroy() override;
-	static void AddReferencedObjects(UObject* InThis, FReferenceCollector& Collector);
+	COREUOBJECT_API virtual void Serialize( FArchive& Ar ) override;
+	COREUOBJECT_API virtual void BeginDestroy() override;
+	static COREUOBJECT_API void AddReferencedObjects(UObject* InThis, FReferenceCollector& Collector);
 	// End of UObject interface
 
 	/**
@@ -627,13 +627,13 @@ public:
 	 * @param  NewInterfaceClass    The InterfaceClass you want this property set with.
 	 */
 #if USE_UPROPERTY_LOAD_DEFERRING
-	void SetInterfaceClass(UClass* NewInterfaceClass);
+	COREUOBJECT_API void SetInterfaceClass(UClass* NewInterfaceClass);
 #else
 	FORCEINLINE void SetInterfaceClass(UClass* NewInterfaceClass) { InterfaceClass = NewInterfaceClass; }
 #endif // USE_UPROPERTY_LOAD_DEFERRING
 };
 
-class COREUOBJECT_API UNameProperty : public UProperty
+class UNameProperty : public UProperty
 {
 	DECLARE_CASTED_CLASS_INTRINSIC(UNameProperty, UProperty, 0, TEXT("/Script/CoreUObject"), CASTCLASS_FNameProperty)
 public:
@@ -649,7 +649,7 @@ public:
 	}
 };
 
-class COREUOBJECT_API UStrProperty : public UProperty
+class UStrProperty : public UProperty
 {
 	DECLARE_CASTED_CLASS_INTRINSIC(UStrProperty, UProperty, 0, TEXT("/Script/CoreUObject"), CASTCLASS_FStrProperty)
 public:
@@ -665,7 +665,7 @@ public:
 	}
 };
 
-class COREUOBJECT_API UArrayProperty : public UProperty
+class UArrayProperty : public UProperty
 {
 	DECLARE_CASTED_CLASS_INTRINSIC(UArrayProperty, UProperty, 0, TEXT("/Script/CoreUObject"), CASTCLASS_FArrayProperty)
 
@@ -685,12 +685,12 @@ public:
 	}
 
 	// UObject interface
-	virtual void Serialize( FArchive& Ar ) override;
-	static void AddReferencedObjects(UObject* InThis, FReferenceCollector& Collector);
+	COREUOBJECT_API virtual void Serialize( FArchive& Ar ) override;
+	static COREUOBJECT_API void AddReferencedObjects(UObject* InThis, FReferenceCollector& Collector);
 	// End of UObject interface
 };
 
-class COREUOBJECT_API UMapProperty : public UProperty
+class UMapProperty : public UProperty
 {
 	DECLARE_CASTED_CLASS_INTRINSIC(UMapProperty, UProperty, 0, TEXT("/Script/CoreUObject"), CASTCLASS_FMapProperty)
 
@@ -701,15 +701,15 @@ public:
 	TObjectPtr<UProperty>       ValueProp;
 	FScriptMapLayout MapLayout;
 
-	UMapProperty(const FObjectInitializer& ObjectInitializer, ECppProperty, int32 InOffset, EPropertyFlags InFlags);
+	COREUOBJECT_API UMapProperty(const FObjectInitializer& ObjectInitializer, ECppProperty, int32 InOffset, EPropertyFlags InFlags);
 
 	// UObject interface
-	virtual void Serialize(FArchive& Ar) override;
-	static void AddReferencedObjects(UObject* InThis, FReferenceCollector& Collector);
+	COREUOBJECT_API virtual void Serialize(FArchive& Ar) override;
+	static COREUOBJECT_API void AddReferencedObjects(UObject* InThis, FReferenceCollector& Collector);
 	// End of UObject interface
 };
 
-class COREUOBJECT_API USetProperty : public UProperty
+class USetProperty : public UProperty
 {
 	DECLARE_CASTED_CLASS_INTRINSIC(USetProperty, UProperty, 0, TEXT("/Script/CoreUObject"), CASTCLASS_FSetProperty)
 
@@ -719,15 +719,15 @@ public:
 	TObjectPtr<UProperty>       ElementProp;
 	FScriptSetLayout SetLayout;
 
-	USetProperty(const FObjectInitializer& ObjectInitializer, ECppProperty, int32 InOffset, EPropertyFlags InFlags);
+	COREUOBJECT_API USetProperty(const FObjectInitializer& ObjectInitializer, ECppProperty, int32 InOffset, EPropertyFlags InFlags);
 
 	// UObject interface
-	virtual void Serialize(FArchive& Ar) override;
-	static void AddReferencedObjects(UObject* InThis, FReferenceCollector& Collector);
+	COREUOBJECT_API virtual void Serialize(FArchive& Ar) override;
+	static COREUOBJECT_API void AddReferencedObjects(UObject* InThis, FReferenceCollector& Collector);
 	// End of UObject interface
 };
 
-class COREUOBJECT_API UStructProperty : public UProperty
+class UStructProperty : public UProperty
 {
 	DECLARE_CASTED_CLASS_INTRINSIC(UStructProperty, UProperty, 0, TEXT("/Script/CoreUObject"), CASTCLASS_FStructProperty)
 
@@ -736,16 +736,16 @@ public:
 	// Variables.
 	TObjectPtr<class UScriptStruct> Struct;
 
-	UStructProperty(ECppProperty, int32 InOffset, EPropertyFlags InFlags, UScriptStruct* InStruct);
-	UStructProperty( const FObjectInitializer& ObjectInitializer, ECppProperty, int32 InOffset, EPropertyFlags InFlags, UScriptStruct* InStruct );
+	COREUOBJECT_API UStructProperty(ECppProperty, int32 InOffset, EPropertyFlags InFlags, UScriptStruct* InStruct);
+	COREUOBJECT_API UStructProperty( const FObjectInitializer& ObjectInitializer, ECppProperty, int32 InOffset, EPropertyFlags InFlags, UScriptStruct* InStruct );
 
 	// UObject interface
-	virtual void Serialize( FArchive& Ar ) override;
-	static void AddReferencedObjects(UObject* InThis, FReferenceCollector& Collector);
+	COREUOBJECT_API virtual void Serialize( FArchive& Ar ) override;
+	static COREUOBJECT_API void AddReferencedObjects(UObject* InThis, FReferenceCollector& Collector);
 	// End of UObject interface
 };
 
-class COREUOBJECT_API UDelegateProperty : public UProperty
+class UDelegateProperty : public UProperty
 {
 	DECLARE_CASTED_CLASS_INTRINSIC(UDelegateProperty, UProperty, 0, TEXT("/Script/CoreUObject"), CASTCLASS_FDelegateProperty)
 
@@ -767,12 +767,12 @@ public:
 	}
 
 	// UObject interface
-	virtual void Serialize( FArchive& Ar ) override;
-	virtual void BeginDestroy() override;
+	COREUOBJECT_API virtual void Serialize( FArchive& Ar ) override;
+	COREUOBJECT_API virtual void BeginDestroy() override;
 	// End of UObject interface
 };
 
-class COREUOBJECT_API UMulticastDelegateProperty : public UProperty
+class UMulticastDelegateProperty : public UProperty
 {
 	DECLARE_CASTED_CLASS_INTRINSIC(UMulticastDelegateProperty, UProperty, CLASS_Abstract, TEXT("/Script/CoreUObject"), CASTCLASS_FMulticastDelegateProperty)
 
@@ -788,12 +788,12 @@ public:
 	}
 
 	// UObject interface
-	virtual void Serialize( FArchive& Ar ) override;
-	virtual void BeginDestroy() override;
+	COREUOBJECT_API virtual void Serialize( FArchive& Ar ) override;
+	COREUOBJECT_API virtual void BeginDestroy() override;
 	// End of UObject interface
 };
 
-class COREUOBJECT_API UMulticastInlineDelegateProperty : public UMulticastDelegateProperty
+class UMulticastInlineDelegateProperty : public UMulticastDelegateProperty
 {
 	DECLARE_CASTED_CLASS_INTRINSIC(UMulticastInlineDelegateProperty, UMulticastDelegateProperty, 0, TEXT("/Script/CoreUObject"), CASTCLASS_FMulticastInlineDelegateProperty)
 
@@ -810,7 +810,7 @@ public:
 	}
 };
 
-class COREUOBJECT_API UMulticastSparseDelegateProperty : public UMulticastDelegateProperty
+class UMulticastSparseDelegateProperty : public UMulticastDelegateProperty
 {
 	DECLARE_CASTED_CLASS_INTRINSIC(UMulticastSparseDelegateProperty, UMulticastDelegateProperty, 0, TEXT("/Script/CoreUObject"), CASTCLASS_FMulticastSparseDelegateProperty)
 
@@ -827,24 +827,24 @@ public:
 	}
 };
 
-class COREUOBJECT_API UEnumProperty : public UProperty
+class UEnumProperty : public UProperty
 {
 	DECLARE_CASTED_CLASS_INTRINSIC(UEnumProperty, UProperty, 0, TEXT("/Script/CoreUObject"), CASTCLASS_FEnumProperty)
 
 public:
-	UEnumProperty(const FObjectInitializer& ObjectInitializer, UEnum* InEnum);
-	UEnumProperty(const FObjectInitializer& ObjectInitializer, ECppProperty, int32 InOffset, EPropertyFlags InFlags, UEnum* InEnum);
+	COREUOBJECT_API UEnumProperty(const FObjectInitializer& ObjectInitializer, UEnum* InEnum);
+	COREUOBJECT_API UEnumProperty(const FObjectInitializer& ObjectInitializer, ECppProperty, int32 InOffset, EPropertyFlags InFlags, UEnum* InEnum);
 
 	// UObject interface
-	virtual void Serialize(FArchive& Ar) override;
-	static void AddReferencedObjects(UObject* InThis, FReferenceCollector& Collector);
+	COREUOBJECT_API virtual void Serialize(FArchive& Ar) override;
+	static COREUOBJECT_API void AddReferencedObjects(UObject* InThis, FReferenceCollector& Collector);
 	// End of UObject interface
 
 	TObjectPtr<UNumericProperty> UnderlyingProp; // The property which represents the underlying type of the enum
 	TObjectPtr<UEnum> Enum; // The enum represented by this property
 };
 
-class COREUOBJECT_API UTextProperty : public UProperty
+class UTextProperty : public UProperty
 {
 	DECLARE_CASTED_CLASS_INTRINSIC(UTextProperty, UProperty, 0, TEXT("/Script/CoreUObject"), CASTCLASS_FTextProperty)
 
