Fix nested default subobject type overrides on Blueprint class assets that derive from a non-native parent class.

diff --git a/Engine/Source/Runtime/CoreUObject/Private/UObject/UObjectGlobals.cpp b/Engine/Source/Runtime/CoreUObject/Private/UObject/UObjectGlobals.cpp
--- a/Engine/Source/Runtime/CoreUObject/Private/UObject/UObjectGlobals.cpp
+++ b/Engine/Source/Runtime/CoreUObject/Private/UObject/UObjectGlobals.cpp
@@ -5391,7 +5391,9 @@ UObject* FObjectInitializer::CreateDefaultSubobject(UObject* Outer, FName Subobj
 			}
 			else if (!bIsTransient && Outer->GetArchetype()->IsInBlueprint())
 			{
-				UObject* MaybeTemplate = Result->GetArchetype();
+				// The archetype of the outer is not native, so we need to copy properties to the subobjects after
+				// the C++ constructor chain for the outer has run (because those sets properties on the subobjects)
+				UObject* MaybeTemplate = Outer->GetArchetype()->GetDefaultSubobjectByName(SubobjectFName);
 				if (MaybeTemplate && Template != MaybeTemplate && MaybeTemplate->IsA(ReturnType))
 				{
 					ComponentInits.Add(Result, MaybeTemplate);
