Updating AddReferencedObject(s) callsites ahead of FReferenceCollector API change.

diff --git a/Engine/Source/Runtime/CoreUObject/Public/Templates/Casts.h b/Engine/Source/Runtime/CoreUObject/Public/Templates/Casts.h
--- a/Engine/Source/Runtime/CoreUObject/Public/Templates/Casts.h
+++ b/Engine/Source/Runtime/CoreUObject/Public/Templates/Casts.h
@@ -248,12 +248,6 @@ template< class T, class U > FORCEINLINE T* Cast       ( const TWeakObjectPtr<U>
 template< class T, class U > FORCEINLINE T* ExactCast  ( const TWeakObjectPtr<U>& Src                                                                   ) { return ExactCast  <T>(Src.Get()); }
 template< class T, class U > FORCEINLINE T* CastChecked( const TWeakObjectPtr<U>& Src, ECastCheckedType::Type CheckType = ECastCheckedType::NullChecked ) { return CastChecked<T>(Src.Get(), CheckType); }
 
-template< class T, class U > FORCEINLINE TWeakObjectPtr<T>* CastCheckedPtr( TWeakObjectPtr<U>& Src, ECastCheckedType::Type CheckType = ECastCheckedType::NullChecked )
-{
-	return CastChecked<T>(Src.Get(), CheckType) ? reinterpret_cast<TWeakObjectPtr<T>*>(&Src) : nullptr;
-}
-
-
 // object ptr versions
 template <typename To, typename From>
 FORCEINLINE To* Cast(const TObjectPtr<From>& InSrc)
