Allow cross level references to private objects if the referencing object is transient.

diff --git a/Engine/Source/Runtime/CoreUObject/Private/UObject/PropertyBaseObject.cpp b/Engine/Source/Runtime/CoreUObject/Private/UObject/PropertyBaseObject.cpp
--- a/Engine/Source/Runtime/CoreUObject/Private/UObject/PropertyBaseObject.cpp
+++ b/Engine/Source/Runtime/CoreUObject/Private/UObject/PropertyBaseObject.cpp
@@ -562,7 +562,7 @@ TObjectPtr<UObject> FObjectPropertyBase::FindImportedObject( const FProperty* Pr
 	}
 
 	// if we found an object, and we have a parent, make sure we are in the same package or share an outer if the found object is private, unless it's a cross level property
-	if (Result && !Result->HasAnyFlags(RF_Public) && OwnerObject 
+	if (Result && !Result->HasAnyFlags(RF_Public) && OwnerObject && !OwnerObject->HasAnyFlags(RF_Transient)
 		&& Result->GetOutermostObject() != OwnerObject->GetOutermostObject()
 		&& Result->GetPackage() != OwnerObject->GetPackage())
 	{
