[Backout] CL 25929189 (needs more investigation)

diff --git a/Engine/Source/Runtime/CoreUObject/Private/UObject/UObjectGlobals.cpp b/Engine/Source/Runtime/CoreUObject/Private/UObject/UObjectGlobals.cpp
--- a/Engine/Source/Runtime/CoreUObject/Private/UObject/UObjectGlobals.cpp
+++ b/Engine/Source/Runtime/CoreUObject/Private/UObject/UObjectGlobals.cpp
@@ -5391,9 +5391,7 @@ UObject* FObjectInitializer::CreateDefaultSubobject(UObject* Outer, FName Subobj
 			}
 			else if (!bIsTransient && Outer->GetArchetype()->IsInBlueprint())
 			{
-				// The archetype of the outer is not native, so we need to copy properties to the subobjects after
-				// the C++ constructor chain for the outer has run (because those sets properties on the subobjects)
-				UObject* MaybeTemplate = Outer->GetArchetype()->GetDefaultSubobjectByName(SubobjectFName);
+				UObject* MaybeTemplate = Result->GetArchetype();
 				if (MaybeTemplate && Template != MaybeTemplate && MaybeTemplate->IsA(ReturnType))
 				{
 					ComponentInits.Add(Result, MaybeTemplate);
