SavePackage: Improve clarity of error message when an ExternalActor still has its ExternalPackage set. The ExternalPackage is not saveable so the import to the package (and its name) will fail to be harvested and this will trigger the "Missing import name" error message in SavePackage's ValidateImports. That error message is arcane, so preempt it with a new error message that is more specific.
#rnx

diff --git a/Engine/Source/Runtime/CoreUObject/Private/UObject/SavePackage/PackageHarvester.cpp b/Engine/Source/Runtime/CoreUObject/Private/UObject/SavePackage/PackageHarvester.cpp
--- a/Engine/Source/Runtime/CoreUObject/Private/UObject/SavePackage/PackageHarvester.cpp
+++ b/Engine/Source/Runtime/CoreUObject/Private/UObject/SavePackage/PackageHarvester.cpp
@@ -3,11 +3,11 @@
 
 #include "UObject/SavePackage/PackageHarvester.h"
 
+#include "Interfaces/ITargetPlatform.h"
 #include "UObject/SavePackage/SaveContext.h"
 #include "UObject/SavePackage/SavePackageUtilities.h"
 #include "UObject/UObjectGlobals.h"
 #include "UObject/UObjectHash.h"
-#include "Interfaces/ITargetPlatform.h"
 
 EObjectMark GenerateMarksForObject(const UObject* InObject, const ITargetPlatform* TargetPlatform)
 {
@@ -674,6 +674,16 @@ void FPackageHarvester::ProcessImport(TObjectPtr<UObject> InObject)
 	if (Package && Package != InObject)
 	{
 		*this << Package;
+
+		// The package needs to be included in the SaveContext's HarvestingRealm, or ValidateImports will log an error:
+		// "Missing import package name...". Log it here instead with an explanation for why it's not included.
+		// It's possible the package will be marked as editoronly by not including it in the game realm, so don't check
+		// all CurrentExportHarvestingRealms, just check the SaveContext's HarvestingRealm.
+		if (SaveContext.IsIncluded(InObject) && !SaveContext.IsIncluded(Package))
+		{
+			SaveContext.RecordIllegalReference(CurrentExportDependencies.CurrentExport, InObject,
+				EIllegalRefReason::ExternalPackage, Package->GetName());
+		}
 	}
 	else
 	{
diff --git a/Engine/Source/Runtime/CoreUObject/Private/UObject/SavePackage/SaveContext.h b/Engine/Source/Runtime/CoreUObject/Private/UObject/SavePackage/SaveContext.h
--- a/Engine/Source/Runtime/CoreUObject/Private/UObject/SavePackage/SaveContext.h
+++ b/Engine/Source/Runtime/CoreUObject/Private/UObject/SavePackage/SaveContext.h
@@ -90,6 +90,7 @@ enum class EIllegalRefReason : uint8
 	ReferenceFromOptionalToMissingGameExport,
 	UnsaveableClass,
 	UnsaveableOuter,
+	ExternalPackage,
 };
 
 /** Small struct to store illegal references harvested during save */
diff --git a/Engine/Source/Runtime/CoreUObject/Private/UObject/SavePackage2.cpp b/Engine/Source/Runtime/CoreUObject/Private/UObject/SavePackage2.cpp
--- a/Engine/Source/Runtime/CoreUObject/Private/UObject/SavePackage2.cpp
+++ b/Engine/Source/Runtime/CoreUObject/Private/UObject/SavePackage2.cpp
@@ -493,6 +493,24 @@ ESavePackageResult ValidateRealms(FSaveContext& SaveContext)
 					Reference.To ? *Reference.To->GetPathName() : TEXT("Unknown"),
 					*Reference.FormatStringArg);
 				break;
+			case EIllegalRefReason::ExternalPackage:
+				if (Reference.From->GetOutermostObject() == Reference.To->GetOutermostObject() && SaveContext.IsCooking())
+				{
+					ErrorMessage = FString::Printf(TEXT("Can't save %s: export (%s) has a reference to export (%s) which still has its external package set to (%s)."),
+						SaveContext.GetFilename(),
+						Reference.From ? *Reference.From->GetPathName() : TEXT("Unknown"),
+						Reference.To ? *Reference.To->GetPathName() : TEXT("Unknown"),
+						!Reference.FormatStringArg.IsEmpty() ? *Reference.FormatStringArg : TEXT("Unknown"));
+				}
+				else
+				{
+					ErrorMessage = FString::Printf(TEXT("Can't save %s: export (%s) has a reference to import (%s), but the import is in ExternalPackage (%s) which was marked unsaveable."),
+						SaveContext.GetFilename(),
+						Reference.From ? *Reference.From->GetPathName() : TEXT("Unknown"),
+						Reference.To ? *Reference.To->GetPathName() : TEXT("Unknown"),
+						!Reference.FormatStringArg.IsEmpty() ? *Reference.FormatStringArg : TEXT("Unknown"));
+				}
+				break;
 			default:
 				ErrorMessage = FString::Printf(TEXT("Can't save %s: Unknown Illegal reference from object (%s) to object (%s)"),
 					SaveContext.GetFilename(),
