Removing EGCOptions::WithClusters to reduce the number of template instance permutations in Garbage Collector code

diff --git a/Engine/Source/Runtime/CoreUObject/Private/UObject/GarbageCollection.cpp b/Engine/Source/Runtime/CoreUObject/Private/UObject/GarbageCollection.cpp
--- a/Engine/Source/Runtime/CoreUObject/Private/UObject/GarbageCollection.cpp
+++ b/Engine/Source/Runtime/CoreUObject/Private/UObject/GarbageCollection.cpp
@@ -839,8 +839,6 @@ void FContextPoolScope::Cleanup()
 }
 
 //////////////////////////////////////////////////////////////////////////
-	
-FORCEINLINE constexpr bool IsWithClusters(EGCOptions Options) {	return !!(Options & EGCOptions::WithClusters); }
 
 template<EGCOptions Options, class ContainerType>
 static void MarkReferencedClustersAsReachable(int32 ClusterIndex, ContainerType& ObjectsToSerialize);
@@ -962,8 +960,6 @@ static bool MarkClusterMutableObjectsAsReachable(FUObjectCluster& Cluster, Conta
 template<EGCOptions Options, class ContainerType>
 static FORCENOINLINE void MarkReferencedClustersAsReachable(int32 ClusterIndex, ContainerType& ObjectsToSerialize)
 {
-	static_assert(IsWithClusters(Options));
-
 	UE::GC::GStats.IncNumClustersTraversed();
 
 	// If we run across some PendingKill objects we need to add all objects from this cluster
@@ -2808,12 +2804,12 @@ public:
 			// Objects that are part of a GC cluster should never have the unreachable flag set!
 			checkSlow(Metadata.ObjectItem->GetOwnerIndex() <= 0);
 
-			if (!IsWithClusters(Options) || !Metadata.Has(EInternalObjectFlags::ClusterRoot))
+			if (!Metadata.Has(EInternalObjectFlags::ClusterRoot))
 			{
 				// Add it to the list of objects to serialize.
 				Context.ObjectsToSerialize.Add<Options>(Reference.Object);
 			}
-			else if constexpr (IsWithClusters(Options))
+			else
 			{
 				// This is a cluster root reference so mark all referenced clusters as reachable
 				MarkReferencedClustersAsReachableThunk<Options>(Metadata.ObjectItem->GetClusterIndex(), Context.ObjectsToSerialize);
@@ -2821,7 +2817,7 @@ public:
 
 			return true;
 		}
-		else if constexpr (IsWithClusters(Options))
+		else
 		{
 			if ((Metadata.ObjectItem->GetOwnerIndex() > 0) & !Metadata.Has(EInternalObjectFlags::ReachableInCluster))
 			{
@@ -3541,28 +3537,21 @@ class FRealtimeGC : public FGarbageCollectionTracer
 	static FORCEINLINE int32 GetGCFunctionIndex(EGCOptions InOptions)
 	{
 		return (!!(InOptions & EGCOptions::Parallel)) |
-			(!!(InOptions & EGCOptions::WithClusters) << 1) |
-			(!!(InOptions & EGCOptions::WithPendingKill) << 2);
+			(!!(InOptions & EGCOptions::WithPendingKill) << 1);
 	}
 
 public:
 	/** Default constructor, initializing all members. */
 	FRealtimeGC()
 	{
-		MarkObjectsFunctions[GetGCFunctionIndex(EGCOptions::None)] = &FRealtimeGC::MarkObjectsAsUnreachable<false, false>;
-		MarkObjectsFunctions[GetGCFunctionIndex(EGCOptions::Parallel | EGCOptions::None)] = &FRealtimeGC::MarkObjectsAsUnreachable<true, false>;
-		MarkObjectsFunctions[GetGCFunctionIndex(EGCOptions::None | EGCOptions::WithClusters)] = &FRealtimeGC::MarkObjectsAsUnreachable<false, true>;
-		MarkObjectsFunctions[GetGCFunctionIndex(EGCOptions::Parallel | EGCOptions::WithClusters)] = &FRealtimeGC::MarkObjectsAsUnreachable<true, true>;
+		MarkObjectsFunctions[GetGCFunctionIndex(EGCOptions::None)] = &FRealtimeGC::MarkObjectsAsUnreachable<false>;
+		MarkObjectsFunctions[GetGCFunctionIndex(EGCOptions::Parallel | EGCOptions::None)] = &FRealtimeGC::MarkObjectsAsUnreachable<true>;
 
 		ReachabilityAnalysisFunctions[GetGCFunctionIndex(EGCOptions::None)] = &FRealtimeGC::PerformReachabilityAnalysisOnObjectsInternal<EGCOptions::None | EGCOptions::None>;
 		ReachabilityAnalysisFunctions[GetGCFunctionIndex(EGCOptions::Parallel | EGCOptions::None)] = &FRealtimeGC::PerformReachabilityAnalysisOnObjectsInternal<EGCOptions::Parallel | EGCOptions::None>;
-		ReachabilityAnalysisFunctions[GetGCFunctionIndex(EGCOptions::None | EGCOptions::WithClusters)] = &FRealtimeGC::PerformReachabilityAnalysisOnObjectsInternal<EGCOptions::None | EGCOptions::WithClusters>;
-		ReachabilityAnalysisFunctions[GetGCFunctionIndex(EGCOptions::Parallel | EGCOptions::WithClusters)] = &FRealtimeGC::PerformReachabilityAnalysisOnObjectsInternal<EGCOptions::Parallel | EGCOptions::WithClusters>;
 
-		ReachabilityAnalysisFunctions[GetGCFunctionIndex(EGCOptions::None | EGCOptions::WithPendingKill)] = &FRealtimeGC::PerformReachabilityAnalysisOnObjectsInternal<EGCOptions::None | EGCOptions::None | EGCOptions::WithPendingKill>;
-		ReachabilityAnalysisFunctions[GetGCFunctionIndex(EGCOptions::Parallel | EGCOptions::None | EGCOptions::WithPendingKill)] = &FRealtimeGC::PerformReachabilityAnalysisOnObjectsInternal<EGCOptions::Parallel | EGCOptions::None | EGCOptions::WithPendingKill>;
-		ReachabilityAnalysisFunctions[GetGCFunctionIndex(EGCOptions::None | EGCOptions::WithClusters | EGCOptions::WithPendingKill)] = &FRealtimeGC::PerformReachabilityAnalysisOnObjectsInternal<EGCOptions::None | EGCOptions::WithClusters | EGCOptions::WithPendingKill>;
-		ReachabilityAnalysisFunctions[GetGCFunctionIndex(EGCOptions::Parallel | EGCOptions::WithClusters | EGCOptions::WithPendingKill)] = &FRealtimeGC::PerformReachabilityAnalysisOnObjectsInternal<EGCOptions::Parallel | EGCOptions::WithClusters | EGCOptions::WithPendingKill>;
+		ReachabilityAnalysisFunctions[GetGCFunctionIndex(EGCOptions::None | EGCOptions::WithPendingKill)] = &FRealtimeGC::PerformReachabilityAnalysisOnObjectsInternal<EGCOptions::None | EGCOptions::WithPendingKill>;
+		ReachabilityAnalysisFunctions[GetGCFunctionIndex(EGCOptions::Parallel | EGCOptions::WithPendingKill)] = &FRealtimeGC::PerformReachabilityAnalysisOnObjectsInternal<EGCOptions::Parallel | EGCOptions::WithPendingKill>;
 
 		FGCObject::StaticInit();
 	}
@@ -3571,7 +3560,7 @@ public:
 	 * Marks all objects that don't have KeepFlags and EInternalObjectFlags::GarbageCollectionKeepFlags as unreachable
 	 * This function is a template to speed up the case where we don't need to assemble the token stream (saves about 6ms on PS4)
 	 */
-	template <bool bParallel, bool bWithClusters>
+	template <bool bParallel>
 	void MarkObjectsAsUnreachable(const EObjectFlags KeepFlags)
 	{
 		TRACE_CPUPROFILER_EVENT_SCOPE(MarkObjectsAsUnreachable);
@@ -3625,10 +3614,8 @@ public:
 					// Keep track of how many objects are around.
 					ObjectCountDuringMarkPhase++;
 					
-					if (bWithClusters)
-					{
-						ObjectItem->ClearFlags(EInternalObjectFlags::ReachableInCluster);
-					}
+					ObjectItem->ClearFlags(EInternalObjectFlags::ReachableInCluster);
+					
 					// Special case handling for objects that are part of the root set.
 					if (ObjectItem->IsRootSet())
 					{
@@ -3638,18 +3625,15 @@ public:
 #if DO_GUARD_SLOW
 						checkCode(if (ObjectItem->IsPendingKill()) { UE_LOG(LogGarbage, Fatal, TEXT("Object %s is part of root set though has been marked RF_PendingKill!"), *Object->GetFullName()); });
 #endif
-						if (bWithClusters)
+						if (ObjectItem->HasAnyFlags(EInternalObjectFlags::ClusterRoot) || ObjectItem->GetOwnerIndex() > 0)
 						{
-							if (ObjectItem->HasAnyFlags(EInternalObjectFlags::ClusterRoot) || ObjectItem->GetOwnerIndex() > 0)
-							{
-								KeepClusterRefsList.Push(ObjectItem);
-							}
+							KeepClusterRefsList.Push(ObjectItem);
 						}
 
 						LocalObjectsToSerialize.Add(Object);
 					}
 					// Cluster objects 
-					else if (bWithClusters && ObjectItem->GetOwnerIndex() > 0)
+					else if (ObjectItem->GetOwnerIndex() > 0)
 					{
 						// treat cluster objects with FastKeepFlags the same way as if they are in the root set
 						if (ObjectItem->HasAnyFlags(FastKeepFlags))
@@ -3659,7 +3643,7 @@ public:
 						}
 					}
 					// Regular objects or cluster root objects
-					else if (!bWithClusters || ObjectItem->GetOwnerIndex() <= 0)
+					else
 					{
 						bool bMarkAsUnreachable = true;
 						// Internal flags are super fast to check and is used by async loading and must have higher precedence than PendingKill
@@ -3673,7 +3657,7 @@ public:
 							bMarkAsUnreachable = false;
 						}
 						PRAGMA_DISABLE_DEPRECATION_WARNINGS
-						else if (ObjectItem->HasAnyFlags(EInternalObjectFlags::PendingKill | EInternalObjectFlags::Garbage) && bWithClusters && ObjectItem->HasAnyFlags(EInternalObjectFlags::ClusterRoot))
+						else if (ObjectItem->HasAnyFlags(EInternalObjectFlags::PendingKill | EInternalObjectFlags::Garbage) && ObjectItem->HasAnyFlags(EInternalObjectFlags::ClusterRoot))
 						PRAGMA_ENABLE_DEPRECATION_WARNINGS
 						{
 							ClustersToDissolveList.Push(ObjectItem);
@@ -3686,12 +3670,9 @@ public:
 							checkSlow(Object->IsValidLowLevel());
 							LocalObjectsToSerialize.Add(Object);
 
-							if (bWithClusters)
+							if (ObjectItem->HasAnyFlags(EInternalObjectFlags::ClusterRoot))
 							{
-								if (ObjectItem->HasAnyFlags(EInternalObjectFlags::ClusterRoot))
-								{
-									KeepClusterRefsList.Push(ObjectItem);
-								}
+								KeepClusterRefsList.Push(ObjectItem);
 							}
 						}
 						else
@@ -3721,7 +3702,6 @@ public:
 			ObjectsToSerializeArrays.Empty();
 		}
 
-		if constexpr (bWithClusters)
 		{
 			TArray<FUObjectItem*> ClustersToDissolve;
 			ClustersToDissolveList.PopAll(ClustersToDissolve);
@@ -3756,7 +3736,7 @@ public:
 						{
 							RootObjectItem->ClearFlags(EInternalObjectFlags::Unreachable);
 							// Make sure all referenced clusters are marked as reachable too
-							MarkReferencedClustersAsReachable<EGCOptions::WithClusters>(RootObjectItem->GetClusterIndex(), InitialObjects);
+							MarkReferencedClustersAsReachable<EGCOptions::None>(RootObjectItem->GetClusterIndex(), InitialObjects);
 						}
 					}
 				}
@@ -3765,7 +3745,7 @@ public:
 					checkSlow(ObjectItem->HasAnyFlags(EInternalObjectFlags::ClusterRoot));
 					// this thing is definitely not marked unreachable, so don't test it here
 					// Make sure all referenced clusters are marked as reachable too
-					MarkReferencedClustersAsReachable<EGCOptions::WithClusters>(ObjectItem->GetClusterIndex(), InitialObjects);
+					MarkReferencedClustersAsReachable<EGCOptions::None>(ObjectItem->GetClusterIndex(), InitialObjects);
 				}
 			}
 		}
@@ -4656,8 +4636,6 @@ void CollectGarbageImpl(EObjectFlags KeepFlags)
 				// Fall back to single threaded GC if processor count is 1 or parallel GC is disabled
 				// or detailed per class gc stats are enabled (not thread safe)
 				(ShouldForceSingleThreadedGC() ? EGCOptions::None : EGCOptions::Parallel) |
-				// Run with GC clustering code enabled only if clustering is enabled and there's actual allocated clusters
-				((!!GCreateGCClusters && GUObjectClusters.GetNumAllocatedClusters()) ? EGCOptions::WithClusters : EGCOptions::None) |
 				// Toggle between PendingKill enabled or disabled
 				(UObjectBaseUtility::IsPendingKillEnabled() ? EGCOptions::WithPendingKill : EGCOptions::None);
 
diff --git a/Engine/Source/Runtime/CoreUObject/Public/UObject/FastReferenceCollector.h b/Engine/Source/Runtime/CoreUObject/Public/UObject/FastReferenceCollector.h
--- a/Engine/Source/Runtime/CoreUObject/Public/UObject/FastReferenceCollector.h
+++ b/Engine/Source/Runtime/CoreUObject/Public/UObject/FastReferenceCollector.h
@@ -33,8 +33,7 @@ enum class EGCOptions : uint32
 	None = 0,
 	Parallel = 1 << 0,					// Use all task workers to collect references, must be started on main thread
 	AutogenerateSchemas = 1 << 1,		// Assemble schemas for new UClasses
-	WithClusters = 1 << 2,				// Use clusters, see FUObjectCluster
-	WithPendingKill = 1 << 3,			// Internal flag used by reachability analysis
+	WithPendingKill = 1 << 2,			// Internal flag used by reachability analysis
 };
 ENUM_CLASS_FLAGS(EGCOptions);
 
