ExternalPackages: Add more feedback for failed ensure(ExternalPackage) in GetExternalPackage.
#rnx

diff --git a/Engine/Source/Runtime/CoreUObject/Private/UObject/UObjectBase.cpp b/Engine/Source/Runtime/CoreUObject/Private/UObject/UObjectBase.cpp
--- a/Engine/Source/Runtime/CoreUObject/Private/UObject/UObjectBase.cpp
+++ b/Engine/Source/Runtime/CoreUObject/Private/UObject/UObjectBase.cpp
@@ -257,7 +257,8 @@ UPackage* UObjectBase::GetExternalPackage() const
 	{
 		ExternalPackage = GetObjectExternalPackageThreadSafe(this);
 		// if the flag is set there should be an override set.
-		ensure(ExternalPackage);
+		ensureMsgf(ExternalPackage, TEXT("Object %s ExternalPackage is invalid: RF_ExternalPackage is set, but ExternalPackage is nullptr."),
+			*static_cast<const UObjectBaseUtility*>(this)->GetPathName());
 	}
 	return ExternalPackage;
 }
