Optimize UClass::FindFunctionByName by extending the function cache (prevously named SuperFuncMap) to include functions in this class in addition to superclasses/interfaces. This means we do a single lock + map lookup instead of needing two if the function isn't found in this class.

This saves 100us on the gamethread, and also helps on async loading to some extent. This is a memory tradeoff since the cache is larger. It adds around ~170KB on average, increasing mem usage of this function from 1.82MB to 1.99MB (in a replay benchmark, as measured with a custom LLM tag).

#tests A/B perf tested in a cooked build on console, also tested with validation to ensure results were the same, tested editor, including a full cook.

diff --git a/Engine/Source/Runtime/CoreUObject/Private/UObject/Class.cpp b/Engine/Source/Runtime/CoreUObject/Private/UObject/Class.cpp
--- a/Engine/Source/Runtime/CoreUObject/Private/UObject/Class.cpp
+++ b/Engine/Source/Runtime/CoreUObject/Private/UObject/Class.cpp
@@ -6052,8 +6052,8 @@ void UClass::CreateLinkAndAddChildFunctionsToMap(const FClassFunctionLinkInfo* F
 
 void UClass::ClearFunctionMapsCaches()
 {
-	FUClassFuncScopeWriteLock ScopeLock(SuperFuncMapLock);
-	SuperFuncMap.Empty();
+	FUClassFuncScopeWriteLock ScopeLock(AllFunctionsCacheLock);
+	AllFunctionsCache.Empty();
 }
 
 UFunction* UClass::FindFunctionByName(FName InName, EIncludeSuperFlag::Type IncludeSuper) const
@@ -6064,44 +6064,71 @@ UFunction* UClass::FindFunctionByName(FName InName, EIncludeSuperFlag::Type Incl
 
 	UE_AUTORTFM_OPEN(
 	{
+		UClass* SuperClass = GetSuperClass();
+		if (IncludeSuper == EIncludeSuperFlag::ExcludeSuper || ( Interfaces.Num() == 0 && SuperClass == nullptr ) )
 		{
+			// Trivial case: just look up in this class's function map and don't involve the cache
 			FUClassFuncScopeReadLock ScopeLock(FuncMapLock);
 			Result = FuncMap.FindRef(InName);
 		}
-
-		if (Result == nullptr && IncludeSuper == EIncludeSuperFlag::IncludeSuper)
+		else
 		{
-			UClass* SuperClass = GetSuperClass();
-			if (SuperClass || Interfaces.Num() > 0)
+			// Check the cache
+			bool bFoundInCache = false;
 			{
-				bool bFoundInSuperFuncMap = false;
+				FUClassFuncScopeReadLock ScopeLock(AllFunctionsCacheLock);
+				if (UFunction** SuperResult = AllFunctionsCache.Find(InName))
 				{
-					FUClassFuncScopeReadLock ScopeLock(SuperFuncMapLock);
-					if (UFunction** SuperResult = SuperFuncMap.Find(InName))
-					{
-						Result = *SuperResult;
-						bFoundInSuperFuncMap = true;
-					}
+					Result = *SuperResult;
+					bFoundInCache = true;
+				}
+			}
+
+			if (!bFoundInCache)
+			{
+				// Try this class's FuncMap first
+				{
+					FUClassFuncScopeReadLock ScopeLock(FuncMapLock);
+					Result = FuncMap.FindRef(InName);
 				}
 
-				if (!bFoundInSuperFuncMap)
+				if (Result)
+				{
+					// Cache the result
+					FUClassFuncScopeWriteLock ScopeLock(AllFunctionsCacheLock);
+					AllFunctionsCache.Add(InName, Result);
+				}
+				else
 				{
-					for (const FImplementedInterface& Inter : Interfaces)
+					// Check superclass and interfaces
+					if (Interfaces.Num() > 0)
 					{
-						Result = Inter.Class ? Inter.Class->FindFunctionByName(InName) : nullptr;
-						if (Result)
+						for (const FImplementedInterface& Inter : Interfaces)
 						{
-							break;
+							Result = Inter.Class ? Inter.Class->FindFunctionByName(InName) : nullptr;
+							if (Result)
+							{
+								break;
+							}
 						}
 					}
 
-					if (SuperClass && Result == nullptr)
+					if (Result == nullptr && SuperClass != nullptr )
 					{
 						Result = SuperClass->FindFunctionByName(InName);
 					}
 
-					FUClassFuncScopeWriteLock ScopeLock(SuperFuncMapLock);
-					SuperFuncMap.Add(InName, Result);
+					{
+						// Do a final check to make sure the function still doesn't exist in this class before we add it to the cache, in case the function was added by another thread since we last checked
+						// This avoids us writing null (or a superclass func with the same name) to the cache if the function was just added
+						FUClassFuncScopeReadLock ScopeLockFuncMap(FuncMapLock);
+						if (FuncMap.FindRef(InName) == nullptr)
+						{
+							// Cache the result (even if it's nullptr)
+							FUClassFuncScopeWriteLock ScopeLock(AllFunctionsCacheLock);
+							AllFunctionsCache.Add(InName, Result);
+						}
+					}
 				}
 			}
 		}
diff --git a/Engine/Source/Runtime/CoreUObject/Public/UObject/Class.h b/Engine/Source/Runtime/CoreUObject/Public/UObject/Class.h
--- a/Engine/Source/Runtime/CoreUObject/Public/UObject/Class.h
+++ b/Engine/Source/Runtime/CoreUObject/Public/UObject/Class.h
@@ -2907,11 +2907,11 @@ private:
 	/** Scope lock to avoid the FuncMap being read and written to simultaneously on multiple threads. */
 	mutable FUClassFuncLock FuncMapLock;
 
-	/** A cache of all functions by name that exist in a parent (superclass or interface) context */
-	mutable TMap<FName, UFunction*> SuperFuncMap;
+	/** A cache of all functions by name that exist in this class or a parent (superclass or interface) context */
+	mutable TMap<FName, UFunction*> AllFunctionsCache;
 
 	/** Scope lock to avoid the SuperFuncMap being read and written to simultaneously on multiple threads. */
-	mutable FUClassFuncLock SuperFuncMapLock;
+	mutable FUClassFuncLock AllFunctionsCacheLock;
 
 public:
 	/**
@@ -2997,8 +2997,15 @@ public:
 	/** Add a function to the function map */
 	void AddFunctionToFunctionMap(UFunction* Function, FName FuncName)
 	{
-		FUClassFuncScopeWriteLock ScopeLock(FuncMapLock);
-		FuncMap.Add(FuncName, Function);
+		{
+			FUClassFuncScopeWriteLock ScopeLock(FuncMapLock);
+			FuncMap.Add(FuncName, Function);
+		}
+		{
+			// Remove from the function cache if it exists
+			FUClassFuncScopeWriteLock ScopeLock(AllFunctionsCacheLock);
+			AllFunctionsCache.Remove(FuncName);
+		}
 	}
 
 	COREUOBJECT_API void CreateLinkAndAddChildFunctionsToMap(const FClassFunctionLinkInfo* Functions, uint32 NumFunctions);
@@ -3006,8 +3013,15 @@ public:
 	/** Remove a function from the function map */
 	void RemoveFunctionFromFunctionMap(UFunction* Function)
 	{
-		FUClassFuncScopeWriteLock ScopeLock(FuncMapLock);
-		FuncMap.Remove(Function->GetFName());
+		{
+			FUClassFuncScopeWriteLock ScopeLock(FuncMapLock);
+			FuncMap.Remove(Function->GetFName());
+		}
+		{
+			// Remove from the function cache if it exists
+			FUClassFuncScopeWriteLock ScopeLock(AllFunctionsCacheLock);
+			AllFunctionsCache.Remove(Function->GetFName());
+		}
 	}
 
 	/** Clears the function name caches, in case things have changed */
