Fix crash during recursion by splitting nodes and waiting for dependencies instead of having recursive code
  - Also prevents constructed objects from being assigned to the wrong package during highly recursive phase

#rnx

diff --git a/Engine/Source/Runtime/CoreUObject/Private/Serialization/AsyncLoading2.cpp b/Engine/Source/Runtime/CoreUObject/Private/Serialization/AsyncLoading2.cpp
--- a/Engine/Source/Runtime/CoreUObject/Private/Serialization/AsyncLoading2.cpp
+++ b/Engine/Source/Runtime/CoreUObject/Private/Serialization/AsyncLoading2.cpp
@@ -1942,7 +1942,16 @@ enum class EAsyncPackageLoadingState2 : uint8
 	ProcessPackageSummary,
 	WaitingForDependencies,
 	DependenciesReady,
+#if ALT2_ENABLE_LINKERLOAD_SUPPORT
+// This is the path taken by LinkerLoad packages
+	CreateLinkerLoadExports,
+	WaitingForLinkerLoadDependencies,
+	ResolveLinkerLoadImports,
+	PreloadLinkerLoadExports,
+#endif
+// This is the path taken by Runtime/cooked packages
 	ProcessExportBundles,
+// Both LinkerLoad and Cooked packages should converge at this point
 	WaitingForExternalReads,
 	ExportsDone,
 	PostLoad,
@@ -2269,6 +2278,11 @@ enum EEventLoadNode2 : uint8
 {
 	Package_ProcessSummary,
 	Package_DependenciesReady,
+#if ALT2_ENABLE_LINKERLOAD_SUPPORT
+	Package_CreateLinkerLoadExports,
+	Package_ResolveLinkerLoadImports,
+	Package_PreloadLinkerLoadExports,
+#endif
 	Package_ExportsSerialized,
 	Package_NumPhases,
 
@@ -2649,7 +2663,8 @@ private:
 #if WITH_EDITORONLY_DATA
 		int32 MetaDataIndex = -1;
 #endif
-		bool bIsCurrentlyProcessingImports = false;
+		bool bIsCurrentlyResolvingImports = false;
+		bool bIsCurrentlyCreatingExports = false;
 		bool bContainsClasses = false;
 
 		FAsyncPackageLinkerLoadHeaderData LinkerLoadHeaderData;
@@ -2773,6 +2788,9 @@ private:
 		}
 	};
 	FAllDependenciesState		AllDependenciesSetupState;
+#if ALT2_ENABLE_LINKERLOAD_SUPPORT
+	FAllDependenciesState		AllDependenciesImportState;
+#endif
 	FAllDependenciesState		AllDependenciesFullyLoadedState;
 
 	/** True if our load has failed */
@@ -2805,6 +2823,11 @@ public:
 	/** [EDL] Begin Event driven loader specific stuff */
 
 	static EEventLoadNodeExecutionResult Event_ProcessExportBundle(FAsyncLoadingThreadState2& ThreadState, FAsyncPackage2* Package, int32 ExportBundleIndex);
+#if ALT2_ENABLE_LINKERLOAD_SUPPORT
+	static EEventLoadNodeExecutionResult Event_CreateLinkerLoadExports(FAsyncLoadingThreadState2& ThreadState, FAsyncPackage2* Package, int32);
+	static EEventLoadNodeExecutionResult Event_ResolveLinkerLoadImports(FAsyncLoadingThreadState2& ThreadState, FAsyncPackage2* Package, int32);
+	static EEventLoadNodeExecutionResult Event_PreloadLinkerLoadExports(FAsyncLoadingThreadState2& ThreadState, FAsyncPackage2* Package, int32);
+#endif
 	static EEventLoadNodeExecutionResult Event_ProcessPackageSummary(FAsyncLoadingThreadState2& ThreadState, FAsyncPackage2* Package, int32);
 	static EEventLoadNodeExecutionResult Event_DependenciesReady(FAsyncLoadingThreadState2& ThreadState, FAsyncPackage2* Package, int32);
 	static EEventLoadNodeExecutionResult Event_ExportsDone(FAsyncLoadingThreadState2& ThreadState, FAsyncPackage2* Package, int32);
@@ -2868,10 +2891,10 @@ private:
 
 #if ALT2_ENABLE_LINKERLOAD_SUPPORT
 	EEventLoadNodeExecutionResult ProcessLinkerLoadPackageSummary(FAsyncLoadingThreadState2& ThreadState);
-	bool ProcessLinkerLoadPackageImports(FAsyncLoadingThreadState2& ThreadState);
-	bool ProcessLinkerLoadPackageExports(FAsyncLoadingThreadState2& ThreadState);
-	bool ProcessLinkerLoadPackagePreloads(FAsyncLoadingThreadState2& ThreadState);
-	EEventLoadNodeExecutionResult ProcessLinkerLoadPackageImportsAndExports(FAsyncLoadingThreadState2& ThreadState);
+	bool CreateLinkerLoadExports(FAsyncLoadingThreadState2& ThreadState);
+	void ConditionalBeginResolveLinkerLoadImports(FAsyncLoadingThreadState2& ThreadState);
+	bool ResolveLinkerLoadImports(FAsyncLoadingThreadState2& ThreadState);
+	bool PreloadLinkerLoadExports(FAsyncLoadingThreadState2& ThreadState);
 	EEventLoadNodeExecutionResult PostLoadLinkerLoadPackageExports(FAsyncLoadingThreadState2& ThreadState);
 #endif
 
@@ -3078,6 +3101,7 @@ private:
 	TUniquePtr<FAsyncLoadingThreadState2> AsyncLoadingThreadState;
 
 	uint32 ConditionalBeginProcessExportsTick = 0;
+	uint32 ConditionalBeginResolveImportsTick = 0;
 	uint32 ConditionalFinishLoadingTick = 0;
 
 public:
@@ -5050,6 +5074,7 @@ void FAsyncPackage2::StartLoading(FAsyncLoadingThreadState2& ThreadState, FIoBat
 }
 
 #if ALT2_ENABLE_LINKERLOAD_SUPPORT
+
 EEventLoadNodeExecutionResult FAsyncPackage2::ProcessLinkerLoadPackageSummary(FAsyncLoadingThreadState2& ThreadState)
 {
 	TRACE_CPUPROFILER_EVENT_SCOPE(ProcessLinkerLoadPackageSummary);
@@ -5241,7 +5266,7 @@ EEventLoadNodeExecutionResult FAsyncPackage2::ProcessLinkerLoadPackageSummary(FA
 	return EEventLoadNodeExecutionResult::Complete;
 }
 
-bool FAsyncPackage2::ProcessLinkerLoadPackagePreloads(FAsyncLoadingThreadState2& ThreadState)
+bool FAsyncPackage2::PreloadLinkerLoadExports(FAsyncLoadingThreadState2& ThreadState)
 {
 	// Serialize exports
 	const int32 ExportCount = LinkerLoadState->Linker->ExportMap.Num();
@@ -5265,76 +5290,28 @@ bool FAsyncPackage2::ProcessLinkerLoadPackagePreloads(FAsyncLoadingThreadState2&
 	return true;
 }
 
-EEventLoadNodeExecutionResult FAsyncPackage2::ProcessLinkerLoadPackageImportsAndExports(FAsyncLoadingThreadState2& ThreadState)
+bool FAsyncPackage2::ResolveLinkerLoadImports(FAsyncLoadingThreadState2& ThreadState)
 {
-	ON_SCOPE_EXIT
-	{
-		FUObjectSerializeContext * LoadContext = GetSerializeContext();
-		TArray<UObject*>& ThreadObjLoaded = LoadContext->PRIVATE_GetObjectsLoadedInternalUseOnly();
-		ThreadObjLoaded.Reset();
-	};
+	check(AsyncPackageLoadingState >= EAsyncPackageLoadingState2::WaitingForLinkerLoadDependencies);
 
-	if (!ProcessLinkerLoadPackageImports(ThreadState))
-	{
-		return EEventLoadNodeExecutionResult::Timeout;
-	}
-	if (!ProcessLinkerLoadPackageExports(ThreadState))
-	{
-		return EEventLoadNodeExecutionResult::Timeout;
-	}
-	if (!ProcessLinkerLoadPackagePreloads(ThreadState))
-	{
-		return EEventLoadNodeExecutionResult::Timeout;
-	}
+	check (!LinkerLoadState->bIsCurrentlyResolvingImports);
+	TGuardValue<bool> GuardIsCurrentlyResolvingImports(LinkerLoadState->bIsCurrentlyResolvingImports, true);
 
-	check(AsyncPackageLoadingState == EAsyncPackageLoadingState2::ProcessExportBundles);
-	if (ExternalReadDependencies.Num() == 0)
-	{
-		GetPackageNode(Package_ExportsSerialized).ReleaseBarrier(&ThreadState);
-	}
-	else
-	{
-		AsyncPackageLoadingState = EAsyncPackageLoadingState2::WaitingForExternalReads;
-		AsyncLoadingThread.ExternalReadQueue.Enqueue(this);
-	}
-	return EEventLoadNodeExecutionResult::Complete;
-}
-
-bool FAsyncPackage2::ProcessLinkerLoadPackageImports(FAsyncLoadingThreadState2& ThreadState)
-{
-	check(AsyncPackageLoadingState >= EAsyncPackageLoadingState2::DependenciesReady);
-	if (LinkerLoadState->bIsCurrentlyProcessingImports)
-	{
-		return true;
-	}
-	TGuardValue<bool> GuardIsCurrentlyProcessingImports(LinkerLoadState->bIsCurrentlyProcessingImports, true);
+	// Validate that all imports are in the appropriate state and had their exports created
 	const int32 ImportedPackagesCount = Data.ImportedAsyncPackages.Num();
-	while (LinkerLoadState->ProcessingImportedPackageIndex < ImportedPackagesCount)
+	for (int32 ImportIndex = 0; ImportIndex < ImportedPackagesCount; ++ImportIndex)
 	{
-		const int32 ImportIndex = LinkerLoadState->ProcessingImportedPackageIndex++;
 		FAsyncPackage2* ImportedPackage = Data.ImportedAsyncPackages[ImportIndex];
 		if (ImportedPackage && ImportedPackage->LinkerLoadState.IsSet())
 		{
-			check(ImportedPackage->AsyncPackageLoadingState >= EAsyncPackageLoadingState2::DependenciesReady);
-			if (ImportedPackage->AsyncPackageLoadingState <= EAsyncPackageLoadingState2::ProcessExportBundles)
-			{
-				if (!ImportedPackage->ProcessLinkerLoadPackageImports(ThreadState))
-				{
-					return false;
-				}
-			}
-			if (ImportedPackage->AsyncPackageLoadingState <= EAsyncPackageLoadingState2::ProcessExportBundles)
+			if (ImportedPackage->AsyncPackageLoadingState < EAsyncPackageLoadingState2::WaitingForLinkerLoadDependencies)
 			{
-				if (!ImportedPackage->ProcessLinkerLoadPackageExports(ThreadState))
-				{
-					return false;
-				}
+				check(ThreadState.PackagesOnStack.Contains(ImportedPackage));
+				UE_LOG(LogStreaming, Warning, TEXT("Package %s might be missing an import from package %s because of a circular dependency between them."),
+					*Desc.UPackageName.ToString(),
+					*ImportedPackage->Desc.UPackageName.ToString());
 			}
 		}
-		if (AsyncPackageLoadingState > EAsyncPackageLoadingState2::ProcessExportBundles)
-		{
-			return true;
-		}
 	}
 
 	const int32 ImportCount = LinkerLoadState->Linker->ImportMap.Num();
@@ -5381,13 +5358,17 @@ bool FAsyncPackage2::ProcessLinkerLoadPackageImports(FAsyncLoadingThreadState2&
 			}
 		}
 	}
+
 	return true;
 }
 
-bool FAsyncPackage2::ProcessLinkerLoadPackageExports(FAsyncLoadingThreadState2& ThreadState)
+bool FAsyncPackage2::CreateLinkerLoadExports(FAsyncLoadingThreadState2& ThreadState)
 {
 	check(AsyncPackageLoadingState >= EAsyncPackageLoadingState2::DependenciesReady);
 
+	check (!LinkerLoadState->bIsCurrentlyCreatingExports);
+	TGuardValue<bool> GuardIsCurrentlyCreatingExports(LinkerLoadState->bIsCurrentlyCreatingExports, true);
+
 	// Create exports
 	const int32 ExportCount = LinkerLoadState->Linker->ExportMap.Num();
 	while (LinkerLoadState->CreateExportIndex < ExportCount)
@@ -5430,7 +5411,6 @@ bool FAsyncPackage2::ProcessLinkerLoadPackageExports(FAsyncLoadingThreadState2&
 		}
 	}
 
-
 	return true;
 }
 
@@ -5459,8 +5439,101 @@ EEventLoadNodeExecutionResult FAsyncPackage2::PostLoadLinkerLoadPackageExports(F
 	ConditionalFinishLoading(ThreadState);
 	return EEventLoadNodeExecutionResult::Complete;
 }
+
+EEventLoadNodeExecutionResult FAsyncPackage2::Event_CreateLinkerLoadExports(FAsyncLoadingThreadState2& ThreadState, FAsyncPackage2* Package, int32)
+{
+	TRACE_CPUPROFILER_EVENT_SCOPE(Event_CreateLinkerLoadExports);
+	UE_ASYNC_PACKAGE_DEBUG(Package->Desc);
+	check(Package->AsyncPackageLoadingState == EAsyncPackageLoadingState2::CreateLinkerLoadExports);
+
+	FAsyncPackageScope2 Scope(Package);
+#if WITH_EDITOR
+	UE::Core::Private::FPlayInEditorLoadingScope PlayInEditorIDScope(Package->Desc.PIEInstanceID);
+#endif
+
+	if (Package->LinkerLoadState.IsSet())
+	{
+		if (!Package->CreateLinkerLoadExports(ThreadState))
+		{
+			return EEventLoadNodeExecutionResult::Timeout;
+		}
+	}
+
+	Package->AsyncPackageLoadingState = EAsyncPackageLoadingState2::WaitingForLinkerLoadDependencies;
+	Package->ConditionalBeginResolveLinkerLoadImports(ThreadState);
+	return EEventLoadNodeExecutionResult::Complete;
+}
+
+void FAsyncPackage2::ConditionalBeginResolveLinkerLoadImports(FAsyncLoadingThreadState2& ThreadState)
+{
+	TRACE_CPUPROFILER_EVENT_SCOPE(ConditionalBeginResolveLinkerLoadImports);
+
+	WaitForAllDependenciesToReachState(ThreadState, &FAsyncPackage2::AllDependenciesImportState, EAsyncPackageLoadingState2::WaitingForLinkerLoadDependencies, AsyncLoadingThread.ConditionalBeginResolveImportsTick,
+		[&ThreadState](FAsyncPackage2* Package)
+		{
+			check (Package->LinkerLoadState.IsSet())
+			Package->GetPackageNode(Package_ResolveLinkerLoadImports).ReleaseBarrier(&ThreadState);
+		}
+	);
+}
+
+EEventLoadNodeExecutionResult FAsyncPackage2::Event_ResolveLinkerLoadImports(FAsyncLoadingThreadState2& ThreadState, FAsyncPackage2* Package, int32)
+{
+	TRACE_CPUPROFILER_EVENT_SCOPE(Event_ResolveLinkerLoadImports);
+	UE_ASYNC_PACKAGE_DEBUG(Package->Desc);
+	check(Package->AsyncPackageLoadingState == EAsyncPackageLoadingState2::ResolveLinkerLoadImports);
+
+	FAsyncPackageScope2 Scope(Package);
+#if WITH_EDITOR
+	UE::Core::Private::FPlayInEditorLoadingScope PlayInEditorIDScope(Package->Desc.PIEInstanceID);
+#endif
+
+	if (Package->LinkerLoadState.IsSet())
+	{
+		if (!Package->ResolveLinkerLoadImports(ThreadState))
+		{
+			return EEventLoadNodeExecutionResult::Timeout;
+		}
+	}
+
+	Package->AsyncPackageLoadingState = EAsyncPackageLoadingState2::PreloadLinkerLoadExports;
+	Package->GetPackageNode(Package_PreloadLinkerLoadExports).ReleaseBarrier(&ThreadState);
+	return EEventLoadNodeExecutionResult::Complete;
+}
+
+EEventLoadNodeExecutionResult FAsyncPackage2::Event_PreloadLinkerLoadExports(FAsyncLoadingThreadState2& ThreadState, FAsyncPackage2* Package, int32)
+{
+	TRACE_CPUPROFILER_EVENT_SCOPE(Event_PreloadLinkerLoadExports);
+	UE_ASYNC_PACKAGE_DEBUG(Package->Desc);
+	check(Package->AsyncPackageLoadingState == EAsyncPackageLoadingState2::PreloadLinkerLoadExports);
+
+	FAsyncPackageScope2 Scope(Package);
+#if WITH_EDITOR
+	UE::Core::Private::FPlayInEditorLoadingScope PlayInEditorIDScope(Package->Desc.PIEInstanceID);
 #endif
 
+	if (Package->LinkerLoadState.IsSet())
+	{
+		if (!Package->PreloadLinkerLoadExports(ThreadState))
+		{
+			return EEventLoadNodeExecutionResult::Timeout;
+		}
+	}
+
+	if (Package->ExternalReadDependencies.Num() == 0)
+	{
+		Package->GetPackageNode(Package_ExportsSerialized).ReleaseBarrier(&ThreadState);
+	}
+	else
+	{
+		Package->AsyncPackageLoadingState = EAsyncPackageLoadingState2::WaitingForExternalReads;
+		Package->AsyncLoadingThread.ExternalReadQueue.Enqueue(Package);
+	}
+	return EEventLoadNodeExecutionResult::Complete;
+}
+
+#endif // ALT2_ENABLE_LINKERLOAD_SUPPORT
+
 EEventLoadNodeExecutionResult FAsyncPackage2::Event_ProcessPackageSummary(FAsyncLoadingThreadState2& ThreadState, FAsyncPackage2* Package, int32)
 {
 	TRACE_CPUPROFILER_EVENT_SCOPE(Event_ProcessPackageSummary);
@@ -5582,7 +5655,7 @@ EEventLoadNodeExecutionResult FAsyncPackage2::Event_ProcessPackageSummary(FAsync
 
 EEventLoadNodeExecutionResult FAsyncPackage2::Event_DependenciesReady(FAsyncLoadingThreadState2& ThreadState, FAsyncPackage2* Package, int32)
 {
-	TRACE_CPUPROFILER_EVENT_SCOPE(Event_DependenciesReadt);
+	TRACE_CPUPROFILER_EVENT_SCOPE(Event_DependenciesReady);
 	UE_ASYNC_PACKAGE_DEBUG(Package->Desc);
 	check(Package->AsyncPackageLoadingState == EAsyncPackageLoadingState2::WaitingForDependencies);
 
@@ -5639,10 +5712,8 @@ EEventLoadNodeExecutionResult FAsyncPackage2::Event_ProcessExportBundle(FAsyncLo
 #endif
 
 #if ALT2_ENABLE_LINKERLOAD_SUPPORT
-	if (Package->LinkerLoadState.IsSet())
-	{
-		return Package->ProcessLinkerLoadPackageImportsAndExports(ThreadState);
-	}
+	// This code path should not be reached for LinkerLoad packages
+	check(!Package->LinkerLoadState.IsSet());
 #endif
 
 	check(InExportBundleIndex < Package->Data.TotalExportBundleCount);
@@ -6309,7 +6380,15 @@ EEventLoadNodeExecutionResult FAsyncPackage2::Event_ExportsDone(FAsyncLoadingThr
 {
 	TRACE_CPUPROFILER_EVENT_SCOPE(Event_ExportsDone);
 	UE_ASYNC_PACKAGE_DEBUG(Package->Desc);
-	check(Package->AsyncPackageLoadingState == EAsyncPackageLoadingState2::ProcessExportBundles || Package->AsyncPackageLoadingState == EAsyncPackageLoadingState2::WaitingForExternalReads);
+#if ALT2_ENABLE_LINKERLOAD_SUPPORT
+	check(Package->AsyncPackageLoadingState == EAsyncPackageLoadingState2::ProcessExportBundles || 
+		  Package->AsyncPackageLoadingState == EAsyncPackageLoadingState2::WaitingForExternalReads ||
+		  Package->AsyncPackageLoadingState == EAsyncPackageLoadingState2::PreloadLinkerLoadExports);
+#else
+	check(Package->AsyncPackageLoadingState == EAsyncPackageLoadingState2::ProcessExportBundles ||
+		  Package->AsyncPackageLoadingState == EAsyncPackageLoadingState2::WaitingForExternalReads);
+#endif
+
 	Package->AsyncPackageLoadingState = EAsyncPackageLoadingState2::ExportsDone;
 
 	if (!Package->bLoadHasFailed && Package->Desc.bCanBeImported)
@@ -6539,19 +6618,17 @@ void FAsyncPackage2::WaitForAllDependenciesToReachState(FAsyncLoadingThreadState
 void FAsyncPackage2::ConditionalBeginProcessPackageExports(FAsyncLoadingThreadState2& ThreadState)
 {
 	TRACE_CPUPROFILER_EVENT_SCOPE(ConditionalBeginProcessPackageExports);
-#if ALT2_ENABLE_LINKERLOAD_SUPPORT
-	// When releasing export process nodes for packages when supporting linker loader state
-	// partition the set of related packages so that they are released in this order:
-	// zen packages -> linker load packages containing classes -> other linker load packages
-	TArray<FAsyncPackage2*, TInlineAllocator<16>> PendingExportProcessToRelease;
+
 	WaitForAllDependenciesToReachState(ThreadState, &FAsyncPackage2::AllDependenciesSetupState, EAsyncPackageLoadingState2::DependenciesReady, AsyncLoadingThread.ConditionalBeginProcessExportsTick,
-		[&ThreadState, &PendingExportProcessToRelease](FAsyncPackage2* Package)
+		[&ThreadState](FAsyncPackage2* Package)
 		{
+#if ALT2_ENABLE_LINKERLOAD_SUPPORT
 			if (Package->LinkerLoadState.IsSet())
 			{
-				PendingExportProcessToRelease.Add(Package);
+				Package->GetPackageNode(EEventLoadNode2::Package_CreateLinkerLoadExports).ReleaseBarrier(&ThreadState);
 			}
 			else
+#endif
 			{
 				for (int32 ExportBundleIndex = 0; ExportBundleIndex < Package->Data.TotalExportBundleCount; ++ExportBundleIndex)
 				{
@@ -6559,31 +6636,6 @@ void FAsyncPackage2::ConditionalBeginProcessPackageExports(FAsyncLoadingThreadSt
 				}
 			}
 		});
-	Algo::Partition(PendingExportProcessToRelease, [](FAsyncPackage2* Package)
-		{
-			if (Package->LinkerLoadState->bContainsClasses)
-			{
-				return true;
-			}
-			return false;
-		});
-	for (FAsyncPackage2* Package : PendingExportProcessToRelease)
-	{
-		for (int32 ExportBundleIndex = 0; ExportBundleIndex < Package->Data.TotalExportBundleCount; ++ExportBundleIndex)
-		{
-			Package->GetExportBundleNode(EEventLoadNode2::ExportBundle_Process, ExportBundleIndex).ReleaseBarrier(&ThreadState);
-		}
-	}
-#else
-	WaitForAllDependenciesToReachState(ThreadState, &FAsyncPackage2::AllDependenciesSetupState, EAsyncPackageLoadingState2::DependenciesReady, AsyncLoadingThread.ConditionalBeginProcessExportsTick,
-		[&ThreadState](FAsyncPackage2* Package)
-		{
-			for (int32 ExportBundleIndex = 0; ExportBundleIndex < Package->Data.TotalExportBundleCount; ++ExportBundleIndex)
-			{
-				Package->GetExportBundleNode(EEventLoadNode2::ExportBundle_Process, ExportBundleIndex).ReleaseBarrier(&ThreadState);
-			}
-		});
-#endif
 }
 
 void FAsyncPackage2::ConditionalFinishLoading(FAsyncLoadingThreadState2& ThreadState)
@@ -7358,6 +7410,11 @@ FAsyncLoadingThread2::FAsyncLoadingThread2(FIoDispatcher& InIoDispatcher, IAsync
 	EventSpecs.AddDefaulted(EEventLoadNode2::Package_NumPhases + EEventLoadNode2::ExportBundle_NumPhases);
 	EventSpecs[EEventLoadNode2::Package_ProcessSummary] = { &FAsyncPackage2::Event_ProcessPackageSummary, &EventQueue, false };
 	EventSpecs[EEventLoadNode2::Package_DependenciesReady] = { &FAsyncPackage2::Event_DependenciesReady, &EventQueue, false };
+#if ALT2_ENABLE_LINKERLOAD_SUPPORT
+	EventSpecs[EEventLoadNode2::Package_CreateLinkerLoadExports] = { &FAsyncPackage2::Event_CreateLinkerLoadExports, &EventQueue, false };
+	EventSpecs[EEventLoadNode2::Package_ResolveLinkerLoadImports] = { &FAsyncPackage2::Event_ResolveLinkerLoadImports, &EventQueue, false };
+	EventSpecs[EEventLoadNode2::Package_PreloadLinkerLoadExports] = { &FAsyncPackage2::Event_PreloadLinkerLoadExports, &EventQueue, false };
+#endif
 	EventSpecs[EEventLoadNode2::Package_ExportsSerialized] = { &FAsyncPackage2::Event_ExportsDone, &EventQueue, true };
 
 	EventSpecs[EEventLoadNode2::Package_NumPhases + EEventLoadNode2::ExportBundle_Process] = { &FAsyncPackage2::Event_ProcessExportBundle, &EventQueue, false };
@@ -8386,7 +8443,7 @@ int32 FAsyncLoadingThread2::LoadPackage(const FPackagePath& InPackagePath, FName
 	TRACE_LOADTIME_BEGIN_REQUEST(RequestId);
 	AddPendingRequest(RequestId);
 
-	// Allocate delegate on Game Thread, it is not safe to copy delegates by value on other threads
+	// Allocate delegate before going async, it is not safe to copy delegates by value on other threads
 	TUniquePtr<FLoadPackageAsyncDelegate> CompletionDelegate = InCompletionDelegate.IsBound()
 		? MakeUnique<FLoadPackageAsyncDelegate>(MoveTemp(InCompletionDelegate))
 		: TUniquePtr<FLoadPackageAsyncDelegate>();
