Fix VC static analyzer "dereferencing NULL pointer" warning.

diff --git a/Engine/Source/Runtime/CoreUObject/Private/UObject/SavePackage2.cpp b/Engine/Source/Runtime/CoreUObject/Private/UObject/SavePackage2.cpp
--- a/Engine/Source/Runtime/CoreUObject/Private/UObject/SavePackage2.cpp
+++ b/Engine/Source/Runtime/CoreUObject/Private/UObject/SavePackage2.cpp
@@ -494,7 +494,7 @@ ESavePackageResult ValidateRealms(FSaveContext& SaveContext)
 					*Reference.FormatStringArg);
 				break;
 			case EIllegalRefReason::ExternalPackage:
-				if (Reference.From->GetOutermostObject() == Reference.To->GetOutermostObject() && SaveContext.IsCooking())
+				if (Reference.From && Reference.To && Reference.From->GetOutermostObject() == Reference.To->GetOutermostObject() && SaveContext.IsCooking())
 				{
 					ErrorMessage = FString::Printf(TEXT("Can't save %s: export (%s) has a reference to export (%s) which still has its external package set to (%s)."),
 						SaveContext.GetFilename(),
