CoreUObject deprecation cleanup pass to remove functions and classes deprecated before 5.0
This change does not remove things like PackageGuid that are still used in some cases

diff --git a/Engine/Source/Runtime/CoreUObject/Private/Serialization/ArchiveObjectGraph.cpp b/Engine/Source/Runtime/CoreUObject/Private/Serialization/ArchiveObjectGraph.cpp
deleted file mode 100644
diff --git a/Engine/Source/Runtime/CoreUObject/Private/Serialization/ArchiveTraceRoute.cpp b/Engine/Source/Runtime/CoreUObject/Private/Serialization/ArchiveTraceRoute.cpp
deleted file mode 100644
diff --git a/Engine/Source/Runtime/CoreUObject/Private/Serialization/TraceReferences.cpp b/Engine/Source/Runtime/CoreUObject/Private/Serialization/TraceReferences.cpp
deleted file mode 100644
diff --git a/Engine/Source/Runtime/CoreUObject/Private/UObject/Class.cpp b/Engine/Source/Runtime/CoreUObject/Private/UObject/Class.cpp
--- a/Engine/Source/Runtime/CoreUObject/Private/UObject/Class.cpp
+++ b/Engine/Source/Runtime/CoreUObject/Private/UObject/Class.cpp
@@ -2052,19 +2052,6 @@ void UStruct::SetSuperStruct(UStruct* NewSuperStruct)
 #endif
 }
 
-FString UStruct::PropertyNameToDisplayName(FName InName) const
-{
-	FFieldVariant FoundField = FindUFieldOrFProperty(this, InName);
-	if (FoundField.IsUObject())
-	{
-		return GetAuthoredNameForField(FoundField.Get<UField>());
-	}
-	else
-	{
-		return GetAuthoredNameForField(FoundField.Get<FField>());
-	}
-}
-
 FString UStruct::GetAuthoredNameForField(const UField* Field) const
 {
 	if (Field)
diff --git a/Engine/Source/Runtime/CoreUObject/Private/UObject/Linker.cpp b/Engine/Source/Runtime/CoreUObject/Private/UObject/Linker.cpp
--- a/Engine/Source/Runtime/CoreUObject/Private/UObject/Linker.cpp
+++ b/Engine/Source/Runtime/CoreUObject/Private/UObject/Linker.cpp
@@ -919,13 +919,6 @@ FLinkerLoad* LoadPackageLinker(UPackage* InOuter, const TCHAR* InLongPackageName
 	return LoadPackageLinker(InOuter, GetPackagePath(InOuter, InLongPackageName), LoadFlags, Sandbox, InReaderOverride, [](FLinkerLoad* InLinker) {});
 }
 
-
-void ResetLoadersForSave(UObject* InOuter, const TCHAR* Filename)
-{
-	UPackage* Package = dynamic_cast<UPackage*>(InOuter);
-	ResetLoadersForSave(Package, Filename);
-}
-
 void ResetLoadersForSave(UPackage* Package, const TCHAR* Filename)
 {
 	FLinkerLoad* Loader = FLinkerLoad::FindExistingLinkerForPackage(Package);
diff --git a/Engine/Source/Runtime/CoreUObject/Private/UObject/ScriptCore.cpp b/Engine/Source/Runtime/CoreUObject/Private/UObject/ScriptCore.cpp
--- a/Engine/Source/Runtime/CoreUObject/Private/UObject/ScriptCore.cpp
+++ b/Engine/Source/Runtime/CoreUObject/Private/UObject/ScriptCore.cpp
@@ -143,7 +143,6 @@ static struct F##inst##Registrar \
 PRAGMA_DISABLE_DEPRECATION_WARNINGS
 
 FBlueprintCoreDelegates::FOnScriptDebuggingEvent FBlueprintCoreDelegates::OnScriptException;
-FBlueprintCoreDelegates::FOnScriptExecutionEnd FBlueprintCoreDelegates::OnScriptExecutionEnd;
 FBlueprintCoreDelegates::FOnScriptInstrumentEvent FBlueprintCoreDelegates::OnScriptProfilingEvent;
 FBlueprintCoreDelegates::FOnToggleScriptProfiler FBlueprintCoreDelegates::OnToggleScriptProfiler;
 
diff --git a/Engine/Source/Runtime/CoreUObject/Private/UObject/UObjectGlobals.cpp b/Engine/Source/Runtime/CoreUObject/Private/UObject/UObjectGlobals.cpp
--- a/Engine/Source/Runtime/CoreUObject/Private/UObject/UObjectGlobals.cpp
+++ b/Engine/Source/Runtime/CoreUObject/Private/UObject/UObjectGlobals.cpp
@@ -963,11 +963,6 @@ void RemoveMountPointDefaultPackageFlags(const TArrayView<FString> InMountPoints
 }
 #endif //if WITH_EDITOR
 
-UPackage* CreatePackage(UObject* InOuter, const TCHAR* PackageName)
-{
-	return CreatePackage(PackageName);
-}
-
 UPackage* CreatePackage(const TCHAR* PackageName )
 {
 	FString InName;
@@ -4327,21 +4322,6 @@ FStaticConstructObjectParameters::FStaticConstructObjectParameters(const UClass*
 {
 }
 
-UObject* StaticConstructObject_Internal(const UClass* Class, UObject* InOuter, FName Name, EObjectFlags SetFlags, EInternalObjectFlags InternalSetFlags, UObject* Template, bool bCopyTransientsFromClassDefaults, FObjectInstancingGraph* InstanceGraph, bool bAssumeTemplateIsArchetype, UPackage* ExternalPackage)
-{
-	FStaticConstructObjectParameters Params(Class);
-	Params.Outer = InOuter;
-	Params.Name = Name;
-	Params.SetFlags = SetFlags;
-	Params.InternalSetFlags = InternalSetFlags;
-	Params.Template = Template;
-	Params.bCopyTransientsFromClassDefaults = bCopyTransientsFromClassDefaults;
-	Params.InstanceGraph = InstanceGraph;
-	Params.bAssumeTemplateIsArchetype = bAssumeTemplateIsArchetype;
-	Params.ExternalPackage = ExternalPackage;
-	return StaticConstructObject_Internal(Params);
-}
-
 UObject* StaticConstructObject_Internal(const FStaticConstructObjectParameters& Params)
 {
 	const UClass* InClass = Params.Class;
diff --git a/Engine/Source/Runtime/CoreUObject/Public/AssetRegistry/AssetBundleData.h b/Engine/Source/Runtime/CoreUObject/Public/AssetRegistry/AssetBundleData.h
--- a/Engine/Source/Runtime/CoreUObject/Public/AssetRegistry/AssetBundleData.h
+++ b/Engine/Source/Runtime/CoreUObject/Public/AssetRegistry/AssetBundleData.h
@@ -53,17 +53,9 @@ PRAGMA_ENABLE_DEPRECATION_WARNINGS
 
 	}
 
-	UE_DEPRECATED(4.27, "Bundles scopes are removed, please use FAssetBundleEntry(FName InBundleName) instead")
-	FAssetBundleEntry(const FPrimaryAssetId& InBundleScope, FName InBundleName)
-		: FAssetBundleEntry(InBundleName)
-	{
-		check(InBundleScope == FPrimaryAssetId());
-	}
-
 	/** Returns true if this represents a real entry */
 	bool IsValid() const { return !BundleName.IsNone(); }
 
-
 	/** Equality */
 	bool operator==(const FAssetBundleEntry& Other) const
 	{
@@ -106,13 +98,6 @@ struct FAssetBundleData
 
 	/** Returns pointer to an entry with given Scope/Name */
 	COREUOBJECT_API FAssetBundleEntry* FindEntry(FName SearchName);
-	
-	UE_DEPRECATED(4.27, "Bundles scopes are removed, please use FindEntry(FName) instead")
-	FAssetBundleEntry* FindEntry(const FPrimaryAssetId& SearchScope, FName SearchName)
-	{
-		check(SearchScope == FPrimaryAssetId());
-		return FindEntry(SearchName);
-	}
 
 	/** Adds or updates an entry with the given BundleName -> Path. Scope is empty and will be filled in later */
 	UE_DEPRECATED(5.1, "The contents of an asset bundle now use FTopLevelAssetPath instead of FSoftObjectPath.")
diff --git a/Engine/Source/Runtime/CoreUObject/Public/AssetRegistry/AssetData.h b/Engine/Source/Runtime/CoreUObject/Public/AssetRegistry/AssetData.h
--- a/Engine/Source/Runtime/CoreUObject/Public/AssetRegistry/AssetData.h
+++ b/Engine/Source/Runtime/CoreUObject/Public/AssetRegistry/AssetData.h
@@ -488,13 +488,6 @@ PRAGMA_ENABLE_DEPRECATION_WARNINGS
 		AppendObjectPath(Builder);
 		String.Append(FString(Builder));
 	}
-
-
-	UE_DEPRECATED(4.18, "ToStringReference was renamed to GetSoftObjectPath")
-	FSoftObjectPath ToStringReference() const
-	{
-		return GetSoftObjectPath();
-	}
 	
 	/** Gets primary asset id of this data */
 	COREUOBJECT_API FPrimaryAssetId GetPrimaryAssetId() const;
diff --git a/Engine/Source/Runtime/CoreUObject/Public/AssetRegistry/AssetDataTagMap.h b/Engine/Source/Runtime/CoreUObject/Public/AssetRegistry/AssetDataTagMap.h
--- a/Engine/Source/Runtime/CoreUObject/Public/AssetRegistry/AssetDataTagMap.h
+++ b/Engine/Source/Runtime/CoreUObject/Public/AssetRegistry/AssetDataTagMap.h
@@ -267,9 +267,6 @@ public:
 
 	COREUOBJECT_API bool		Equals(FStringView Str) const;
 
-	UE_DEPRECATED(4.27, "Use AsString(), AsName(), AsExportPath() or AsText() instead. ")
-	operator FString () const { return AsString(); }
-
 private:
 	/** Return whether this's value is a MarshalledFText, and copy it into out parameter if so */
 	bool						TryGetAsMarshalledText(FixedTagPrivate::FMarshalledText& Out) const;
@@ -378,19 +375,6 @@ public:
 		return FindTag(Tag).Equals(Value);
 	}
 
-	UE_DEPRECATED(4.27, "Use FindTag().As[String|Name|Text/ExportPath]() instead, this checks internally. ")
-	FString FindChecked(FName Key) const
-	{
-		return FindTag(Key).AsString();
-	}
-	
-	/** Find a value by key (default value if not found) */
-	UE_DEPRECATED(4.27, "Use FindTag() instead. ")
-	FString FindRef(FName Key) const
-	{
-		return FindTag(Key).AsString();
-	}
-
 	/** Determine whether a key is present in the map */
 	bool Contains(FName Key) const
 	{
@@ -408,12 +392,6 @@ public:
 		return Loose != nullptr ? Loose->Num() : 0;
 	}
 
-	UE_DEPRECATED(4.27, "Use CopyMap() instead if you really need to make a copy. ")
-	FAssetDataTagMap GetMap() const
-	{
-		return CopyMap();
-	}
-
 	/** Copy map contents to a loose FAssetDataTagMap */
 	FAssetDataTagMap CopyMap() const;
 
@@ -437,9 +415,6 @@ public:
 	COREUOBJECT_API friend bool operator==(const FAssetDataTagMapSharedView& A, const FAssetDataTagMap& B);
 	COREUOBJECT_API friend bool operator==(const FAssetDataTagMapSharedView& A, const FAssetDataTagMapSharedView& B);
 
-	UE_DEPRECATED(4.27, "Use FMemoryCounter instead. ")
-	uint32 GetAllocatedSize() const { return 0; }
-
 	///** Shrinks the contained map */
 	void Shrink();
 
diff --git a/Engine/Source/Runtime/CoreUObject/Public/Serialization/ArchiveObjectGraph.h b/Engine/Source/Runtime/CoreUObject/Public/Serialization/ArchiveObjectGraph.h
deleted file mode 100644
diff --git a/Engine/Source/Runtime/CoreUObject/Public/Serialization/ArchiveTraceRoute.h b/Engine/Source/Runtime/CoreUObject/Public/Serialization/ArchiveTraceRoute.h
deleted file mode 100644
diff --git a/Engine/Source/Runtime/CoreUObject/Public/Serialization/TraceReferences.h b/Engine/Source/Runtime/CoreUObject/Public/Serialization/TraceReferences.h
deleted file mode 100644
diff --git a/Engine/Source/Runtime/CoreUObject/Public/UObject/Class.h b/Engine/Source/Runtime/CoreUObject/Public/UObject/Class.h
--- a/Engine/Source/Runtime/CoreUObject/Public/UObject/Class.h
+++ b/Engine/Source/Runtime/CoreUObject/Public/UObject/Class.h
@@ -636,9 +636,6 @@ public:
 	 */
 	COREUOBJECT_API virtual void SetSuperStruct(UStruct* NewSuperStruct);
 
-	UE_DEPRECATED(4.23, "Replace with GetAuthoredNameForField or UField::GetAuthoredName")
-	COREUOBJECT_API virtual FString PropertyNameToDisplayName(FName InName) const;
-
 	/** Returns a human readable string for a given field, overridden for user defined structs */
 	COREUOBJECT_API virtual FString GetAuthoredNameForField(const UField* Field) const;
 
@@ -3444,9 +3441,6 @@ public:
 	 */
 	COREUOBJECT_API virtual bool IsFunctionImplementedInScript(FName InFunctionName) const;
 
-	UE_DEPRECATED(4.23, "IsFunctionImplementedInBlueprint is deprecated, call IsFunctionImplementedInScript instead")
-	bool IsFunctionImplementedInBlueprint(FName InFunctionName) const { return IsFunctionImplementedInScript(InFunctionName); }
-
 	/**
 	 * Checks if the property exists on this class or a parent class.
 	 * @param InProperty	The property to check if it is contained in this or a parent class.
diff --git a/Engine/Source/Runtime/CoreUObject/Public/UObject/CoreRedirects.h b/Engine/Source/Runtime/CoreUObject/Public/UObject/CoreRedirects.h
--- a/Engine/Source/Runtime/CoreUObject/Public/UObject/CoreRedirects.h
+++ b/Engine/Source/Runtime/CoreUObject/Public/UObject/CoreRedirects.h
@@ -52,11 +52,6 @@ enum class ECoreRedirectFlags : uint32
 	Option_MatchSubstring = 0x01000000, // Does a slow substring match
 	Option_MissingLoad =	0x02000000, // An automatically-created redirect that was created in response to a missing Thing during load. Redirect will be removed if and when the Thing is loaded.
 	Option_AllMask =		0xFF000000, // Bit mask of all possible Options
-
-	// Deprecated Names
-	Option_InstanceOnly UE_DEPRECATED(4.25, "Use Category_InstanceOnly instead") = Category_InstanceOnly,
-	Option_Removed UE_DEPRECATED(4.25, "Use Category_Removed instead") = Category_Removed,
-	Option_ExactMatchMask UE_DEPRECATED(4.25, "Use Category_AllMask instead") = Category_AllMask,
 };
 ENUM_CLASS_FLAGS(ECoreRedirectFlags);
 
diff --git a/Engine/Source/Runtime/CoreUObject/Public/UObject/Linker.h b/Engine/Source/Runtime/CoreUObject/Public/UObject/Linker.h
--- a/Engine/Source/Runtime/CoreUObject/Public/UObject/Linker.h
+++ b/Engine/Source/Runtime/CoreUObject/Public/UObject/Linker.h
@@ -372,40 +372,6 @@ public:
 	COREUOBJECT_API SIZE_T GetAllocatedSize() const;
 };
 
-
-struct UE_DEPRECATED(4.23, "Outdated since display index replaced FName as key") FLinkerNamePairKeyFuncs : DefaultKeyFuncs<FName, false>
-{
-	static FORCEINLINE bool Matches(FName A, FName B)
-	{
-		// The linker requires that FNames preserve case, but the numeric suffix can be ignored since
-		// that is stored separately for each FName instance saved
-		return A.GetDisplayIndex() == B.GetDisplayIndex();
-	}
-
-	static FORCEINLINE uint32 GetKeyHash(FName Key)
-	{
-		return GetTypeHash(Key.GetDisplayIndex());
-	}
-};
-
-
-template<typename ValueType>
-struct UE_DEPRECATED(4.23, "Outdated since display indexes replaced FNames as keys") TLinkerNameMapKeyFuncs : TDefaultMapKeyFuncs<FName, ValueType, false>
-{
-	static FORCEINLINE bool Matches(FName A, FName B)
-	{
-		// The linker requires that FNames preserve case, but the numeric suffix can be ignored since
-		// that is stored separately for each FName instance saved
-		return A.GetDisplayIndex() == B.GetDisplayIndex();
-	}
-
-	static FORCEINLINE uint32 GetKeyHash(FName Key)
-	{
-		return GetTypeHash(Key.GetDisplayIndex());
-	}
-};
-
-
 /*----------------------------------------------------------------------------
 	FLinker.
 ----------------------------------------------------------------------------*/
@@ -751,9 +717,6 @@ COREUOBJECT_API FLinkerLoad* GetPackageLinker(UPackage* InOuter, const TCHAR* In
 
 COREUOBJECT_API FString GetPrestreamPackageLinkerName(const TCHAR* InLongPackageName, bool bSkipIfExists = true);
 
-UE_DEPRECATED(4.25, "No longer used; use version that takes a UPackage* and call EnsureLoadingComplete separately.")
-COREUOBJECT_API void ResetLoadersForSave(UObject* InOuter, const TCHAR *Filename);
-
 /**
  *
  * Reset the loader for the given package if it is using the given filename, so we can write to the file
diff --git a/Engine/Source/Runtime/CoreUObject/Public/UObject/Object.h b/Engine/Source/Runtime/CoreUObject/Public/UObject/Object.h
--- a/Engine/Source/Runtime/CoreUObject/Public/UObject/Object.h
+++ b/Engine/Source/Runtime/CoreUObject/Public/UObject/Object.h
@@ -116,12 +116,6 @@ class UObject : public UObjectBaseUtility
 	/** DO NOT USE. This constructor is for internal usage only for hot-reload purposes. */
 	COREUOBJECT_API UObject(FVTableHelper& Helper);
 
-	UE_DEPRECATED(4.23, "CreateDefaultSubobject no longer takes bAbstract as a parameter.")
-	UObject* CreateDefaultSubobject(FName SubobjectFName, UClass* ReturnType, UClass* ClassToCreateByDefault, bool bIsRequired, bool bAbstract, bool bIsTransient)
-	{
-		return CreateDefaultSubobject(SubobjectFName, ReturnType, ClassToCreateByDefault, bIsRequired, bIsTransient);
-	}
-
 	/** Utility function for templates below */
 	COREUOBJECT_API UObject* CreateDefaultSubobject(FName SubobjectFName, UClass* ReturnType, UClass* ClassToCreateByDefault, bool bIsRequired, bool bIsTransient);
 
@@ -191,19 +185,6 @@ class UObject : public UObjectBaseUtility
 	{
 		return static_cast<TReturnType*>(CreateDefaultSubobject(SubobjectName, TReturnType::StaticClass(), TClassToConstructByDefault::StaticClass(), /*bIsRequired =*/ false, bTransient));
 	}
-	/**
-	 * Create a subobject that has the Abstract class flag, child classes are expected to override this by calling SetDefaultSubobjectClass with the same name and a non-abstract class.
-	 * @param	TReturnType					Class of return type, all overrides must be of this type
-	 * @param	SubobjectName				Name of the new component
-	 * @param	bTransient					True if the component is being assigned to a transient property. This does not make the component itself transient, but does stop it from inheriting parent defaults
-	 */
-	template<class TReturnType>
-	UE_DEPRECATED(4.23, "CreateAbstract did not work as intended and has been deprecated in favor of CreateDefaultObject")
-	TReturnType* CreateAbstractDefaultSubobject(FName SubobjectName, bool bTransient = false)
-	{
-		UClass* ReturnType = TReturnType::StaticClass();
-		return static_cast<TReturnType*>(CreateDefaultSubobject(SubobjectName, ReturnType, ReturnType, /*bIsRequired =*/ true, bTransient));
-	}
 
 	/**
 	 * Gets all default subobjects associated with this object instance.
@@ -1173,16 +1154,6 @@ public:
 	 */
 	virtual void ClearAllCachedCookedPlatformData() { }
 
-	/**
-	 * Called during cook to allow objects to generate additional cooked files alongside their cooked package.
-	 * @note These should typically match the name of the package, but with a different extension.
-	 *
-	 * @param	PackageFilename full path to the package that this object is being saved to on disk
-	 * @param	TargetPlatform	target platform to cook additional files for
-	 */
-	UE_DEPRECATED(4.23, "Use the new CookAdditionalFilesOverride that provides a function to write the files")
-	virtual void CookAdditionalFiles(const TCHAR* PackageFilename, const ITargetPlatform* TargetPlatform) { }
-
 	/**
 	 * Called during cook to allow objects to generate additional cooked files alongside their cooked package.
 	 * @note Implement CookAdditionalFilesOverride to define sub class behavior.
@@ -1211,9 +1182,6 @@ private:
 	virtual void CookAdditionalFilesOverride(const TCHAR* PackageFilename, const ITargetPlatform* TargetPlatform,
 		TFunctionRef<void(const TCHAR* Filename, void* Data, int64 Size)> WriteAdditionalFile)
 	{
-		PRAGMA_DISABLE_DEPRECATION_WARNINGS;
-		CookAdditionalFiles(PackageFilename, TargetPlatform);
-		PRAGMA_ENABLE_DEPRECATION_WARNINGS;
 	}
 
 public:
diff --git a/Engine/Source/Runtime/CoreUObject/Public/UObject/Script.h b/Engine/Source/Runtime/CoreUObject/Public/UObject/Script.h
--- a/Engine/Source/Runtime/CoreUObject/Public/UObject/Script.h
+++ b/Engine/Source/Runtime/CoreUObject/Public/UObject/Script.h
@@ -481,9 +481,6 @@ public:
 	// Callback for blueprint instrumentation enable/disable events
 	DECLARE_MULTICAST_DELEGATE_OneParam(FOnToggleScriptProfiler, bool );
 
-	// Deprecated
-	DECLARE_MULTICAST_DELEGATE(FOnScriptExecutionEnd);
-
 public:
 	// Called when a script exception occurs
 	static COREUOBJECT_API FOnScriptDebuggingEvent OnScriptException;
@@ -492,9 +489,6 @@ public:
 	// Called when a script profiler is enabled/disabled
 	static COREUOBJECT_API FOnToggleScriptProfiler OnToggleScriptProfiler;
 
-	UE_DEPRECATED(4.26, "OnScriptExecutionEnd is deprecated, bind to delegate inside FBlueprintContextTracker instead")
-	static COREUOBJECT_API FOnScriptExecutionEnd OnScriptExecutionEnd;
-
 public:
 	static COREUOBJECT_API void ThrowScriptException(const UObject* ActiveObject, struct FFrame& StackFrame, const FBlueprintExceptionInfo& Info);
 	static COREUOBJECT_API void InstrumentScriptEvent(const FScriptInstrumentationSignal& Info);
diff --git a/Engine/Source/Runtime/CoreUObject/Public/UObject/SparseDelegate.h b/Engine/Source/Runtime/CoreUObject/Public/UObject/SparseDelegate.h
--- a/Engine/Source/Runtime/CoreUObject/Public/UObject/SparseDelegate.h
+++ b/Engine/Source/Runtime/CoreUObject/Public/UObject/SparseDelegate.h
@@ -213,18 +213,6 @@ private:
 	}
 
 public:
-	/** Returns the multicast delegate if any delegates are bound to the sparse delegate */
-	UE_DEPRECATED(4.25, "This function has been deprecated - please use GetShared() instead")
-	MulticastDelegate* Get() const
-	{
-		MulticastDelegate* Result = nullptr;
-		if (bIsBound)
-		{
-			Result = static_cast<MulticastDelegate*>(FSparseDelegateStorage::GetMulticastDelegate(GetDelegateOwner(), GetDelegateName()));
-		}
-		return Result;
-	}
-
 	/** Returns the multicast delegate if any delegates are bound to the sparse delegate */
 	TSharedPtr<MulticastDelegate> GetShared() const
 	{
diff --git a/Engine/Source/Runtime/CoreUObject/Public/UObject/UObjectGlobals.h b/Engine/Source/Runtime/CoreUObject/Public/UObject/UObjectGlobals.h
--- a/Engine/Source/Runtime/CoreUObject/Public/UObject/UObjectGlobals.h
+++ b/Engine/Source/Runtime/CoreUObject/Public/UObject/UObjectGlobals.h
@@ -543,28 +543,6 @@ COREUOBJECT_API UClass* StaticLoadClass(UClass* BaseClass, UObject* InOuter, con
  */
 COREUOBJECT_API UObject* StaticConstructObject_Internal(const FStaticConstructObjectParameters& Params);
 
-/**
- * Create a new instance of an object.  The returned object will be fully initialized.  If InFlags contains RF_NeedsLoad (indicating that the object still needs to load its object data from disk), components
- * are not instanced (this will instead occur in PostLoad()).  The different between StaticConstructObject and StaticAllocateObject is that StaticConstructObject will also call the class constructor on the object
- * and instance any components.
- *
- * @param	Class		The class of the object to create
- * @param	InOuter		The object to create this object within (the Outer property for the new object will be set to the value specified here).
- * @param	Name		The name to give the new object. If no value (NAME_None) is specified, the object will be given a unique name in the form of ClassName_#.
- * @param	SetFlags	The ObjectFlags to assign to the new object. some flags can affect the behavior of constructing the object.
- * @param	InternalSetFlags	The InternalObjectFlags to assign to the new object. some flags can affect the behavior of constructing the object.
- * @param	Template	If specified, the property values from this object will be copied to the new object, and the new object's ObjectArchetype value will be set to this object.
- *						If nullptr, the class default object is used instead.
- * @param	bInCopyTransientsFromClassDefaults	If true, copy transient from the class defaults instead of the pass in archetype ptr (often these are the same)
- * @param	InstanceGraph	Contains the mappings of instanced objects and components to their templates
- * @param	bAssumeTemplateIsArchetype	If true, Template is guaranteed to be an archetype
- * @param	ExternalPackage	Assign an external Package to the created object if non-null
- *
- * @return	A pointer to a fully initialized object of the specified class.
- */
-UE_DEPRECATED(4.26, "Use version that takes parameter struct")
-COREUOBJECT_API UObject* StaticConstructObject_Internal(const UClass* Class, UObject* InOuter = (UObject*)GetTransientPackage(), FName Name = NAME_None, EObjectFlags SetFlags = RF_NoFlags, EInternalObjectFlags InternalSetFlags = EInternalObjectFlags::None, UObject* Template = nullptr, bool bCopyTransientsFromClassDefaults = false, struct FObjectInstancingGraph* InstanceGraph = nullptr, bool bAssumeTemplateIsArchetype = false, UPackage* ExternalPackage = nullptr);
-
 /**
  * Creates a copy of SourceObject using the Outer and Name specified, as well as copies of all objects contained by SourceObject.  
  * Any objects referenced by SourceOuter or RootObject and contained by SourceOuter are also copied, maintaining their name relative to SourceOuter.
@@ -948,15 +926,6 @@ COREUOBJECT_API void SetMountPointDefaultPackageFlags(const TMap<FString, EPacka
 COREUOBJECT_API void RemoveMountPointDefaultPackageFlags(const TArrayView<FString> InMountPoints);
 #endif
 
-/**
- * Find an existing package by name or create it if it doesn't exist
- * @param InOuter		The Outer object to search inside (unused)
- * @return The existing package or a newly created one
- *
- */
-UE_DEPRECATED(4.26, "Use CreatePackage overload that does not take the first Outer parameter. Specifying non-null outers for UPackages is no longer supported.")
-COREUOBJECT_API UPackage* CreatePackage( UObject* InOuter, const TCHAR* PackageName );
-
 /**
  * Find an existing package by name or create it if it doesn't exist
  * @return The existing package or a newly created one
diff --git a/Engine/Source/Runtime/CoreUObject/Public/UObject/UnrealType.h b/Engine/Source/Runtime/CoreUObject/Public/UObject/UnrealType.h
--- a/Engine/Source/Runtime/CoreUObject/Public/UObject/UnrealType.h
+++ b/Engine/Source/Runtime/CoreUObject/Public/UObject/UnrealType.h
@@ -4510,21 +4510,6 @@ public:
 		});
 	}
 
-	/**
-	 * Static version of Num() used when you don't need to bother to construct a FScriptMapHelper. Returns the number of elements in the map.
-	 *
-	 * @param  Target  Pointer to the raw memory associated with a FScriptMap
-	 *
-	 * @return The number of elements in the map.
-	 */
-	UE_DEPRECATED(4.25, "This shortcut is no longer valid - the Num() should be read from a proper map helper")
-	static FORCEINLINE int32 Num(const void* Target)
-	{
-		int32 Result = ((const FScriptMap*)Target)->Num();
-		checkSlow(Result >= 0); 
-		return Result;
-	}
-
 	/**
 	 * Returns a uint8 pointer to the pair in the map
 	 *
@@ -6369,20 +6354,6 @@ struct FPropertyChangedEvent
 	{
 	}
 
-	UE_DEPRECATED(4.25, "The FPropertyChangedEvent constructor taking a TArray* is deprecated. Use the version taking a TArrayView instead.")
-	FPropertyChangedEvent(FProperty* InProperty, EPropertyChangeType::Type InChangeType, const TArray<const UObject*>* InTopLevelObjects)
-		: Property(InProperty)
-		, MemberProperty(InProperty)
-		, ChangeType(InChangeType)
-		, ObjectIteratorIndex(INDEX_NONE)
-		, bFilterChangedInstances(false)
-	{
-		if (InTopLevelObjects)
-		{
-			TopLevelObjects = MakeArrayView(*InTopLevelObjects);
-		}
-	}
-
 	void SetActiveMemberProperty( FProperty* InActiveMemberProperty )
 	{
 		MemberProperty = InActiveMemberProperty;
@@ -6733,41 +6704,6 @@ struct TFieldRange
 	Field templates.
 -----------------------------------------------------------------------------*/
 
-//
-// Find a typed field in a struct.
-//
-template <class T>
-UE_DEPRECATED(4.25, "FindField will no longer return properties. Use FindFProperty instead or FindUField if you want to find functions or enums.")
- T* FindField( const UStruct* Owner, FName FieldName )
-{
-	// We know that a "none" field won't exist in this Struct
-	if( FieldName.IsNone() )
-	{
-		return nullptr;
-	}
-
-	// Search by comparing FNames (INTs), not strings
-	for( TFieldIterator<T>It( Owner ); It; ++It )
-	{
-		if( It->GetFName() == FieldName )
-		{
-			return *It;
-		}
-	}
-
-	// If we didn't find it, return no field
-	return nullptr;
-}
-
-template <class T>
-UE_DEPRECATED(4.25, "FindField will no longer return properties. Use FindFProperty instead or FindUField if you want to find UFunctions or UEnums.")
-T* FindField( const UStruct* Owner, const TCHAR* FieldName )
-{
-	// lookup the string name in the Name hash
-	FName Name(FieldName, FNAME_Find);
-	return FindField<T>(Owner, Name);
-}
-
 template <class T> 
 typename TEnableIf<TIsDerivedFrom<T, UField>::IsDerived, T*>::Type FindUField(const UStruct* Owner, FName FieldName, EFieldIterationFlags IterationFlags = EFieldIterationFlags::Default)
 {
diff --git a/Engine/Source/Runtime/CoreUObject/Public/UObject/WeakInterfacePtr.h b/Engine/Source/Runtime/CoreUObject/Public/UObject/WeakInterfacePtr.h
--- a/Engine/Source/Runtime/CoreUObject/Public/UObject/WeakInterfacePtr.h
+++ b/Engine/Source/Runtime/CoreUObject/Public/UObject/WeakInterfacePtr.h
@@ -55,16 +55,6 @@ struct TWeakInterfacePtr
 		}
 	}
 
-	UE_DEPRECATED(4.27, "Please use the constructor that takes a pointer")
-	TWeakInterfacePtr(T& Interface)
-	{
-		ObjectInstance = Cast<UObject>(&Interface);
-		if (ObjectInstance != nullptr)
-		{
-			InterfaceInstance = &Interface;
-		}
-	}
-
 	/**
 	 * Reset the weak pointer back to the null state.
 	 */
@@ -161,12 +151,6 @@ struct TWeakInterfacePtr
 		return InterfaceInstance != Other.InterfaceInstance;
 	}
 
-	UE_DEPRECATED(4.27, "Implicit equality with a UObject pointer has been deprecated - use GetObject() and test equality on its return value")
-	FORCEINLINE bool operator==(const UObject* Other) const
-	{
-		return Other == ObjectInstance.Get();
-	}
-
 	FORCEINLINE TScriptInterface<T> ToScriptInterface() const
 	{
 		UObject* Object = ObjectInstance.Get();
