Propagate set of request ids for FlushAsyncLoading into inner loops.

diff --git a/Engine/Source/Runtime/CoreUObject/Private/Serialization/AsyncLoading.cpp b/Engine/Source/Runtime/CoreUObject/Private/Serialization/AsyncLoading.cpp
--- a/Engine/Source/Runtime/CoreUObject/Private/Serialization/AsyncLoading.cpp
+++ b/Engine/Source/Runtime/CoreUObject/Private/Serialization/AsyncLoading.cpp
@@ -4657,9 +4657,13 @@ void FAsyncLoadingThread::CompleteFlushRequests(TArray<int32>& RequestIDs)
 	FScopeLock Lock(&FlushRequestCritical);
 	for (int32 RequestID : RequestIDs)
 	{
+		for (FFlushRequest& Flush : FlushRequests)
+		{
+			Flush.OnRequestComplete(RequestID);
+		}
 		FlushRequests.RemoveAll([RequestID](const FFlushRequest& Request)
 			{
-				return Request.GetId() == RequestID;
+				return Request.IsComplete();
 			});
 		DuplicateRequestMap.Remove(RequestID);
 	}
@@ -4670,9 +4674,9 @@ void FAsyncLoadingThread::MapDuplicateRequestID(int32 DuplicateRequestId, int32
 	DuplicateRequestMap.Add(DuplicateRequestId, MainRequestId);
 }
 
-FFlushRequest FAsyncLoadingThread::AddFlushRequest(int32 RequestID)
+FFlushRequest FAsyncLoadingThread::AddFlushRequest(TConstArrayView<int32> RequestIDs)
 {
-	FFlushRequest FlushRequest = FFlushRequest(RequestID);
+	FFlushRequest FlushRequest = FFlushRequest(RequestIDs);
 	{
 		FScopeLock Lock(&FlushRequestCritical);
 		FlushRequests.Push(FlushRequest);
@@ -4701,10 +4705,7 @@ void FAsyncLoadingThread::AdjustFlushRequest(FFlushRequest& FlushRequest)
 {
 	if (FlushRequest)
 	{
-		if (int32* MainRequest = DuplicateRequestMap.Find(FlushRequest.GetId()))
-		{
-			FlushRequest.SetId(*MainRequest);
-		}
+		FlushRequest.AdjustRequestIDs(DuplicateRequestMap);
 	}
 }
 
@@ -4715,7 +4716,8 @@ bool FAsyncLoadingThread::ShouldProcessPackage(FAsyncPackage* InAsyncPackage, co
 		return true;
 	}
 
-	return InAsyncPackage->IsDependencyOf(FlushRequest.GetId());
+	return Algo::AnyOf(FlushRequest.GetRequestIDs(), 
+		[InAsyncPackage](const std::atomic<int32>& RequestID) { return InAsyncPackage->IsDependencyOf(RequestID.load(std::memory_order_relaxed)); });
 }
 
 EAsyncPackageState::Type FAsyncLoadingThread::ProcessLoadedPackages(bool bUseTimeLimit, bool bUseFullTimeLimit, double TimeLimit, bool& bDidSomething, const FFlushRequest& FlushRequest)
@@ -4823,7 +4825,8 @@ EAsyncPackageState::Type FAsyncLoadingThread::ProcessLoadedPackages(bool bUseTim
 
 				CompletedPackages.Add(Package);
 
-				if (IsTimeLimitExceeded(TickStartTime, bUseTimeLimit, TimeLimit, TEXT("ProcessLoadedPackages Misc")) || (FlushRequest && !ContainsRequestInternal(FlushRequest.GetId())))
+				if (IsTimeLimitExceeded(TickStartTime, bUseTimeLimit, TimeLimit, TEXT("ProcessLoadedPackages Misc")) 
+				|| (FlushRequest && !ContainsAnyRequestInternal(FlushRequest.GetRequestIDs())))
 				{
 					// The only package we care about has finished loading, so we're good to exit
 					break;
@@ -6590,7 +6593,7 @@ EAsyncPackageState::Type FAsyncPackage::LoadImports()
 						check(IsInGameThread());
 						FAsyncLoadingThread::LeaveAsyncLoadingTick(AsyncLoadingThread.GetThreadIndex()); // Fix thread check that triggers if we recurse back into the uncooked package loader from here
 						int32 ImportRequestId = IoStorePackageLoader->LoadPackage(ImportedPackagePath, NAME_None, FLoadPackageAsyncDelegate(), PKG_None, INDEX_NONE, 0, nullptr, LOAD_None);
-						IoStorePackageLoader->FlushLoading(ImportRequestId);
+						IoStorePackageLoader->FlushLoading({ImportRequestId});
 						FAsyncLoadingThread::EnterAsyncLoadingTick(AsyncLoadingThread.GetThreadIndex());
 						bLoadedFromIoStore = true;
 					}
@@ -7482,7 +7485,7 @@ int32 FAsyncLoadingThread::LoadPackage(const FPackagePath& InPackagePath, FName
 	return RequestID;
 }
 
-void FAsyncLoadingThread::FlushLoading(int32 PackageID)
+void FAsyncLoadingThread::FlushLoading(TConstArrayView<int32> RequestIDs)
 {
  	if (IsAsyncLoadingPackages())
 	{
@@ -7491,7 +7494,7 @@ void FAsyncLoadingThread::FlushLoading(int32 PackageID)
 
 		SCOPE_CYCLE_COUNTER(STAT_FAsyncPackage_FlushAsyncLoadingGameThread);
 
-		if (PackageID != INDEX_NONE && !ContainsRequestInternal(PackageID))
+		if (RequestIDs.Num() != 0 && !ContainsAnyRequestInternal(RequestIDs))
 		{
 			return;
 		}
@@ -7503,16 +7506,16 @@ void FAsyncLoadingThread::FlushLoading(int32 PackageID)
 		// Flush async loaders without using a time limit. Needed for e.g. garbage collection.
 		{
 			FFlushRequest FlushRequest;
-			if (PackageID != INDEX_NONE)
+			if (RequestIDs.Num() != 0)
 			{
 				// This create and add a flush request to use on the async thread
-				FlushRequest = AddFlushRequest(PackageID);
+				FlushRequest = AddFlushRequest(RequestIDs);
 			}
 			SCOPE_CYCLE_COUNTER(STAT_FAsyncPackage_TickAsyncLoadingGameThread);
 			while (IsAsyncLoadingPackages())
 			{
 				EAsyncPackageState::Type Result = TickAsyncLoading(false, false, 0, FlushRequest);
-				if (PackageID != INDEX_NONE && !ContainsRequestInternal(PackageID))
+				if (RequestIDs.Num() != 0 && !ContainsAnyRequestInternal(RequestIDs))
 				{
 					break;
 				}
@@ -7535,7 +7538,7 @@ void FAsyncLoadingThread::FlushLoading(int32 PackageID)
 		GFlushAsyncLoadingTime += ElapsedTime;
 		GFlushAsyncLoadingCount++;
 
-		check(PackageID != INDEX_NONE || !IsAsyncLoadingPackages());
+		check(RequestIDs.Num() != 0 || !IsAsyncLoadingPackages());
 
 	}
 }
diff --git a/Engine/Source/Runtime/CoreUObject/Private/Serialization/AsyncLoading2.cpp b/Engine/Source/Runtime/CoreUObject/Private/Serialization/AsyncLoading2.cpp
--- a/Engine/Source/Runtime/CoreUObject/Private/Serialization/AsyncLoading2.cpp
+++ b/Engine/Source/Runtime/CoreUObject/Private/Serialization/AsyncLoading2.cpp
@@ -5,6 +5,7 @@
 =============================================================================*/
 
 #include "Serialization/AsyncLoading2.h"
+#include "Algo/AnyOf.h"
 #include "Algo/Partition.h"
 #include "IO/IoDispatcher.h"
 #include "Serialization/AsyncPackageLoader.h"
@@ -2129,9 +2130,10 @@ struct FAsyncLoadEventSpec
 class FAsyncLoadingSyncLoadContext
 {
 public:
-	FAsyncLoadingSyncLoadContext(int32 InRequestId)
-		: RequestId(InRequestId)
+	FAsyncLoadingSyncLoadContext(TConstArrayView<int32> InRequestIDs)
+		: RequestIDs(InRequestIDs)
 	{
+		RequestedPackages.AddZeroed(RequestIDs.Num());
 		ContextId = NextContextId++;
 		if (NextContextId == 0)
 		{
@@ -2159,10 +2161,10 @@ public:
 	}
 
 	uint64 ContextId;
-	int32 RequestId;
-	FAsyncPackage2* RequestedPackage = nullptr;
+	TArray<int32, TInlineAllocator<4>> RequestIDs;
+	TArray<FAsyncPackage2*, TInlineAllocator<4>> RequestedPackages;
 	FAsyncPackage2* RequestingPackageDebug = nullptr;
-	std::atomic<bool> bHasFoundRequestedPackage { false };
+	std::atomic<bool> bHasFoundRequestedPackages { false };
 
 private:
 	std::atomic<int32> RefCount = 1;
@@ -3225,7 +3227,7 @@ public:
 	* @param FlushTree Package dependency tree to be flushed
 	* @return The current state of async loading
 	*/
-	EAsyncPackageState::Type TickAsyncLoadingFromGameThread(FAsyncLoadingThreadState2& ThreadState, bool bUseTimeLimit, bool bUseFullTimeLimit, double TimeLimit, int32 FlushRequestID, bool& bDidSomething);
+	EAsyncPackageState::Type TickAsyncLoadingFromGameThread(FAsyncLoadingThreadState2& ThreadState, bool bUseTimeLimit, bool bUseFullTimeLimit, double TimeLimit, TConstArrayView<int32> FlushRequestIDs, bool& bDidSomething);
 
 	/**
 	* [ASYNC THREAD] Main thread loop
@@ -3279,7 +3281,7 @@ public:
 
 	virtual void ResumeLoading() override;
 
-	virtual void FlushLoading(int32 PackageId) override;
+	virtual void FlushLoading(TConstArrayView<int32> RequestIds) override;
 
 	virtual int32 GetNumQueuedPackages() override
 	{
@@ -3307,6 +3309,15 @@ public:
 		return PendingRequests.Contains(RequestID);
 	}
 
+	/**
+	 * [ASYNC/GAME THREAD] Checks if a request ID already is added to the loading queue
+	 */
+	bool ContainsAnyRequestID(TConstArrayView<int32> RequestIDs)
+	{
+		FScopeLock Lock(&PendingRequestsCritical);
+		return Algo::AnyOf(RequestIDs, [this](int32 RequestID) { return PendingRequests.Contains(RequestID); });
+	}
+
 	/**
 	 * [ASYNC/GAME THREAD] Adds a request ID to the list of pending requests
 	 */
@@ -3319,7 +3330,7 @@ public:
 	/**
 	 * [ASYNC/GAME THREAD] Removes a request ID from the list of pending requests
 	 */
-	void RemovePendingRequests(TArrayView<int32> RequestIDs)
+	void RemovePendingRequests(TConstArrayView<int32> RequestIDs)
 	{
 		FScopeLock Lock(&PendingRequestsCritical);
 		for (int32 ID : RequestIDs)
@@ -3463,7 +3474,7 @@ private:
 	* @param FlushTree Package dependency tree to be flushed
 	* @return The current state of async loading
 	*/
-	EAsyncPackageState::Type ProcessLoadedPackagesFromGameThread(FAsyncLoadingThreadState2& ThreadState, bool& bDidSomething, int32 FlushRequestID = INDEX_NONE);
+	EAsyncPackageState::Type ProcessLoadedPackagesFromGameThread(FAsyncLoadingThreadState2& ThreadState, bool& bDidSomething, TConstArrayView<int32> FlushRequestIDs = {});
 
 	void IncludePackageInSyncLoadContextRecursive(FAsyncLoadingThreadState2& ThreadState, uint64 ContextId, FAsyncPackage2* Package);
 	void UpdateSyncLoadContext(FAsyncLoadingThreadState2& ThreadState);
@@ -4727,7 +4738,7 @@ void FAsyncPackage2::ImportPackagesRecursiveInner(FAsyncLoadingThreadState2& Thr
 				check(IsInGameThread());
 				IoBatch.Issue(); // The batch might already contain requests for packages being imported from the uncooked one we're going to load so make sure that those are started before blocking
 				int32 ImportRequestId = AsyncLoadingThread.UncookedPackageLoader->LoadPackage(ImportedPackagePath, NAME_None, FLoadPackageAsyncDelegate(), PKG_None, INDEX_NONE, 0, nullptr, LOAD_None);
-				AsyncLoadingThread.UncookedPackageLoader->FlushLoading(ImportRequestId);
+				AsyncLoadingThread.UncookedPackageLoader->FlushLoading({ImportRequestId});
 				UncookedPackage = FindObjectFast<UPackage>(nullptr, ImportedPackagePath.GetPackageFName());
 				ImportedPackageRef.SetPackage(UncookedPackage);
 				if (UncookedPackage)
@@ -6935,7 +6946,8 @@ void FAsyncLoadingThread2::UpdateSyncLoadContext(FAsyncLoadingThreadState2& Thre
 	FAsyncLoadingSyncLoadContext* SyncLoadContext = ThreadState.SyncLoadContextStack.Top();
 	if (ThreadState.bIsAsyncLoadingThread)
 	{
-		while (!ContainsRequestID(SyncLoadContext->RequestId))
+		// Retire complete/invalid contexts for which we aren't loading any requests
+		while (!ContainsAnyRequestID(SyncLoadContext->RequestIDs))
 		{
 			SyncLoadContext->ReleaseRef();
 			ThreadState.SyncLoadContextStack.Pop();
@@ -6946,29 +6958,52 @@ void FAsyncLoadingThread2::UpdateSyncLoadContext(FAsyncLoadingThreadState2& Thre
 			SyncLoadContext = ThreadState.SyncLoadContextStack.Top();
 		}
 	}
-	else if (!ContainsRequestID(SyncLoadContext->RequestId))
+	else if (!ContainsAnyRequestID(SyncLoadContext->RequestIDs))
 	{
 		return;
 	}
-	if (ThreadState.bCanAccessAsyncLoadingThreadData && !SyncLoadContext->bHasFoundRequestedPackage.load(std::memory_order_relaxed))
+	if (ThreadState.bCanAccessAsyncLoadingThreadData && !SyncLoadContext->bHasFoundRequestedPackages.load(std::memory_order_relaxed))
 	{
 		// Ensure that we've created the package we're waiting for
 		CreateAsyncPackagesFromQueue(ThreadState);
-		if (FAsyncPackage2* RequestedPackage = RequestIdToPackageMap.FindRef(SyncLoadContext->RequestId))
+		int32 FoundPackages = 0;
+		for (int32 i=0; i < SyncLoadContext->RequestIDs.Num(); ++i)
+		{
+			int32 RequestID = SyncLoadContext->RequestIDs[i];
+			if (SyncLoadContext->RequestedPackages[i] != nullptr)
+			{
+				++FoundPackages;	
+			}
+			else if (FAsyncPackage2* RequestedPackage = RequestIdToPackageMap.FindRef(RequestID))
+			{
+				// Set RequestedPackage before setting bHasFoundRequestedPackage so that another thread looking at RequestedPackage
+				// after validating that bHasFoundRequestedPackage is true would see the proper value.
+				SyncLoadContext->RequestedPackages[i] = RequestedPackage;
+				IncludePackageInSyncLoadContextRecursive(ThreadState, SyncLoadContext->ContextId, RequestedPackage);
+				++FoundPackages;	
+			}
+		}
+		
+		// Only set when full list is available 
+		if (FoundPackages == SyncLoadContext->RequestIDs.Num())
 		{
-			// Set RequestedPackage before setting bHasFoundRequestedPackage so that another thread looking at RequestedPackage
-			// after validating that bHasFoundRequestedPackage is true would see the proper value.
-			SyncLoadContext->RequestedPackage = RequestedPackage;
-			SyncLoadContext->bHasFoundRequestedPackage.store(true, std::memory_order_release);
-			IncludePackageInSyncLoadContextRecursive(ThreadState, SyncLoadContext->ContextId, RequestedPackage);
+			SyncLoadContext->bHasFoundRequestedPackages.store(true, std::memory_order_release);
 		}
 	}
-	if (SyncLoadContext->bHasFoundRequestedPackage.load(std::memory_order_acquire) && ThreadState.PackagesOnStack.Contains(SyncLoadContext->RequestedPackage))
+	if (SyncLoadContext->bHasFoundRequestedPackages.load(std::memory_order_acquire))
 	{
-		// Flushing a package while it's already being processed on the stack, if we're done preloading we let it pass and remove the request id
-		bool bPreloadIsDone = SyncLoadContext->RequestedPackage->AsyncPackageLoadingState >= EAsyncPackageLoadingState2::DeferredPostLoad;
-		UE_CLOG(!bPreloadIsDone, LogStreaming, Fatal, TEXT("Flushing package %s while it's being preloaded in the same callstack is not permitted"), *SyncLoadContext->RequestedPackage->Desc.UPackageName.ToString());
-		RemovePendingRequests(TArrayView<int32>(&SyncLoadContext->RequestId, 1));
+		for (int32 i=0; i < SyncLoadContext->RequestIDs.Num(); ++i)
+		{
+			int32 RequestID = SyncLoadContext->RequestIDs[i];
+			FAsyncPackage2* RequestedPackage = SyncLoadContext->RequestedPackages[i];
+			if (RequestedPackage && ThreadState.PackagesOnStack.Contains(RequestedPackage))
+			{
+				// Flushing a package while it's already being processed on the stack, if we're done preloading we let it pass and remove the request id
+				bool bPreloadIsDone = RequestedPackage->AsyncPackageLoadingState >= EAsyncPackageLoadingState2::DeferredPostLoad;
+				UE_CLOG(!bPreloadIsDone, LogStreaming, Fatal, TEXT("Flushing package %s while it's being preloaded in the same callstack is not permitted"), *RequestedPackage->Desc.UPackageName.ToString());
+				RemovePendingRequests({RequestID});
+			}
+		}
 	}
 }
 
@@ -7057,7 +7092,7 @@ EAsyncPackageState::Type FAsyncLoadingThread2::ProcessAsyncLoadingFromGameThread
 	return EAsyncPackageState::Complete;
 }
 
-EAsyncPackageState::Type FAsyncLoadingThread2::ProcessLoadedPackagesFromGameThread(FAsyncLoadingThreadState2& ThreadState, bool& bDidSomething, int32 FlushRequestID)
+EAsyncPackageState::Type FAsyncLoadingThread2::ProcessLoadedPackagesFromGameThread(FAsyncLoadingThreadState2& ThreadState, bool& bDidSomething, TConstArrayView<int32> FlushRequestIDs)
 {
 	EAsyncPackageState::Type Result = EAsyncPackageState::Complete;
 
@@ -7276,7 +7311,8 @@ EAsyncPackageState::Type FAsyncLoadingThread2::ProcessLoadedPackagesFromGameThre
 		for (int32 CompletedPackageRequestIndex = CompletedPackageRequests.Num() - 1; CompletedPackageRequestIndex >= 0; --CompletedPackageRequestIndex)
 		{
 			FCompletedPackageRequest& CompletedPackageRequest = CompletedPackageRequests[CompletedPackageRequestIndex];
-			if (FlushRequestID == INDEX_NONE || CompletedPackageRequest.RequestIDs.Contains(FlushRequestID))
+			if (FlushRequestIDs.Num() == 0 
+				|| Algo::AnyOf(FlushRequestIDs, [&CompletedPackageRequest](int32 FlushRequestID) { return CompletedPackageRequest.RequestIDs.Contains(FlushRequestID); }))
 			{
 				RemovePendingRequests(CompletedPackageRequest.RequestIDs);
 				RequestsToProcess.Emplace(MoveTemp(CompletedPackageRequest));
@@ -7309,9 +7345,9 @@ EAsyncPackageState::Type FAsyncLoadingThread2::ProcessLoadedPackagesFromGameThre
 
 		bDidSomething = true;
 		
-		if (FlushRequestID != INDEX_NONE && !ContainsRequestID(FlushRequestID))
+		if (FlushRequestIDs.Num() != 0 && !ContainsAnyRequestID(FlushRequestIDs))
 		{
-			// The only package we care about has finished loading, so we're good to exit
+			// The only packages we care about have finished loading, so we're good to exit
 			break;
 		}
 	}
@@ -7323,7 +7359,7 @@ EAsyncPackageState::Type FAsyncLoadingThread2::ProcessLoadedPackagesFromGameThre
 	return Result;
 }
 
-EAsyncPackageState::Type FAsyncLoadingThread2::TickAsyncLoadingFromGameThread(FAsyncLoadingThreadState2& ThreadState, bool bUseTimeLimit, bool bUseFullTimeLimit, double TimeLimit, int32 FlushRequestID, bool& bDidSomething)
+EAsyncPackageState::Type FAsyncLoadingThread2::TickAsyncLoadingFromGameThread(FAsyncLoadingThreadState2& ThreadState, bool bUseTimeLimit, bool bUseFullTimeLimit, double TimeLimit, TConstArrayView<int32> FlushRequestIDs, bool& bDidSomething)
 {
 	SCOPE_CYCLE_COUNTER(STAT_FAsyncPackage_TickAsyncLoadingGameThread);
 	//TRACE_INT_VALUE(QueuedPackagesCounter, QueuedPackagesCounter);
@@ -7360,7 +7396,7 @@ EAsyncPackageState::Type FAsyncLoadingThread2::TickAsyncLoadingFromGameThread(FA
 		UpdateSyncLoadContext(ThreadState);
 
 		{
-			Result = ProcessLoadedPackagesFromGameThread(ThreadState, bDidSomething, FlushRequestID);
+			Result = ProcessLoadedPackagesFromGameThread(ThreadState, bDidSomething, FlushRequestIDs);
 			double TimeLimitUsedForProcessLoaded = FPlatformTime::Seconds() - TickStartTime;
 			UE_CLOG(!GIsEditor && bUseTimeLimit && TimeLimitUsedForProcessLoaded > .1f, LogStreaming, Warning, TEXT("Took %6.2fms to ProcessLoadedPackages"), float(TimeLimitUsedForProcessLoaded) * 1000.0f);
 		}
@@ -8507,11 +8543,11 @@ EAsyncPackageState::Type FAsyncLoadingThread2::ProcessLoadingFromGameThread(FAsy
 	CSV_CUSTOM_STAT(FileIO, ExistingQueuedPackagesQueueDepth, GetNumAsyncPackages(), ECsvCustomStatOp::Set);
 	
 	bool bDidSomething = false;
-	TickAsyncLoadingFromGameThread(ThreadState, bUseTimeLimit, bUseFullTimeLimit, TimeLimit, INDEX_NONE, bDidSomething);
+	TickAsyncLoadingFromGameThread(ThreadState, bUseTimeLimit, bUseFullTimeLimit, TimeLimit, {}, bDidSomething);
 	return IsAsyncLoadingPackages() ? EAsyncPackageState::TimeOut : EAsyncPackageState::Complete;
 }
 
-void FAsyncLoadingThread2::FlushLoading(int32 RequestId)
+void FAsyncLoadingThread2::FlushLoading(TConstArrayView<int32> RequestIDs)
 {
 	if (IsAsyncLoadingPackages())
 	{
@@ -8520,7 +8556,7 @@ void FAsyncLoadingThread2::FlushLoading(int32 RequestId)
 
 		SCOPE_CYCLE_COUNTER(STAT_FAsyncPackage_FlushAsyncLoadingGameThread);
 
-		if (RequestId != INDEX_NONE && !ContainsRequestID(RequestId))
+		if (RequestIDs.Num() != 0 && !ContainsAnyRequestID(RequestIDs))
 		{
 			return;
 		}
@@ -8541,7 +8577,7 @@ void FAsyncLoadingThread2::FlushLoading(int32 RequestId)
 		{
 			TRACE_CPUPROFILER_EVENT_SCOPE(HandleCurrentlyExecutingEventNode);
 
-			UE_CLOG(RequestId == INDEX_NONE, LogStreaming, Fatal, TEXT("Flushing async loading while creating, serializing or postloading an object is not permitted"));
+			UE_CLOG(RequestIDs.Num() == 0, LogStreaming, Fatal, TEXT("Flushing async loading while creating, serializing or postloading an object is not permitted"));
 			CurrentlyExecutingPackage = GameThreadState->CurrentlyExecutingEventNode->GetPackage();
 			GameThreadState->PackagesOnStack.Push(CurrentlyExecutingPackage);
 			// Update the state of any package that is waiting for the currently executing one
@@ -8571,9 +8607,9 @@ void FAsyncLoadingThread2::FlushLoading(int32 RequestId)
 		}
 
 		FAsyncLoadingSyncLoadContext* SyncLoadContext = nullptr;
-		if (RequestId != INDEX_NONE && GOnlyProcessRequiredPackagesWhenSyncLoading)
+		if (RequestIDs.Num() != 0 && GOnlyProcessRequiredPackagesWhenSyncLoading)
 		{
-			SyncLoadContext = new FAsyncLoadingSyncLoadContext(RequestId);
+			SyncLoadContext = new FAsyncLoadingSyncLoadContext(RequestIDs);
 			SyncLoadContext->RequestingPackageDebug = CurrentlyExecutingPackage;
 			GameThreadState->SyncLoadContextStack.Push(SyncLoadContext);
 			if (AsyncLoadingThreadState)
@@ -8589,8 +8625,8 @@ void FAsyncLoadingThread2::FlushLoading(int32 RequestId)
 			while (IsAsyncLoadingPackages())
 			{
 				bool bDidSomething = false;
-				EAsyncPackageState::Type Result = TickAsyncLoadingFromGameThread(*GameThreadState, false, false, 0.0, RequestId, bDidSomething);
-				if (RequestId != INDEX_NONE && !ContainsRequestID(RequestId))
+				EAsyncPackageState::Type Result = TickAsyncLoadingFromGameThread(*GameThreadState, false, false, 0.0, RequestIDs, bDidSomething);
+				if (RequestIDs.Num() != 0 && !ContainsAnyRequestID(RequestIDs))
 				{
 					break;
 				}
@@ -8642,7 +8678,8 @@ void FAsyncLoadingThread2::FlushLoading(int32 RequestId)
 		ConditionalProcessEditorCallbacks();
 #endif
 
-		check(RequestId != INDEX_NONE || !IsAsyncLoadingPackages());
+		// If we asked to flush everything, we should no longer have anything in the pipeline
+		check(RequestIDs.Num() != 0 || !IsAsyncLoadingPackages());
 	}
 }
 
diff --git a/Engine/Source/Runtime/CoreUObject/Private/Serialization/AsyncLoadingThread.h b/Engine/Source/Runtime/CoreUObject/Private/Serialization/AsyncLoadingThread.h
--- a/Engine/Source/Runtime/CoreUObject/Private/Serialization/AsyncLoadingThread.h
+++ b/Engine/Source/Runtime/CoreUObject/Private/Serialization/AsyncLoadingThread.h
@@ -6,6 +6,7 @@
 #pragma once
 
 #include "CoreMinimal.h"
+#include "Algo/AnyOf.h"
 #include "HAL/ThreadSafeCounter.h"
 #include "UObject/ObjectMacros.h"
 #include "Serialization/AsyncLoading.h"
@@ -110,6 +111,59 @@ struct FAsyncLoadEventQueue
 	}
 };
 
+// Shared data for FFlushRequest
+struct FFlushRequestData
+{
+	FFlushRequestData(TConstArrayView<int32> InRequestIDs)
+		: RequestIDs(MakeUnique<std::atomic<int32>[]>(InRequestIDs.Num()))
+		, NumRequests(InRequestIDs.Num())
+	{
+		for (int32 i=0; i < InRequestIDs.Num(); ++i)
+		{
+			RequestIDs.Get()[i] = InRequestIDs[i];
+		}
+	}
+	
+	TConstArrayView<std::atomic<int32>> GetRequestIDs() const
+	{
+		return MakeArrayView(RequestIDs.Get(), NumRequests);
+	}
+
+	TArrayView<std::atomic<int32>> GetRequestIDs() 
+	{
+		return MakeArrayView(RequestIDs.Get(), NumRequests);
+	}
+	
+	void OnRequestComplete(int32 InCompletedRequest)
+	{
+		for (std::atomic<int32>& RequestID : GetRequestIDs())
+		{
+			if (RequestID.load(std::memory_order_relaxed) == InCompletedRequest)
+			{
+				++NumCompletedRequests;
+			}
+		}
+	}
+	
+	bool IsComplete()
+	{
+		return NumRequests == NumCompletedRequests;
+	}
+
+private:
+	/** 
+	 * Shared request id state between main thread and async loading thread. 
+	 * @note because request id are assigned to package request after being generated
+	 * multiple request ids can identify a specific package request however only the original request id
+	 * will properly identify unloaded dependencies. The shared nature of the request allows the async loading thread to
+	 * fix up the main thread flush request id to identify the original request id instead of duplicate when trying to identify
+	 * which main thread packages should be processed when flushing.
+	 */
+	TUniquePtr<std::atomic<int32>[]> RequestIDs;	
+	int32 NumRequests;
+	int32 NumCompletedRequests = 0;
+};
+
 /** 
  * Request to flush a specific package using the request id 
  */
@@ -117,19 +171,24 @@ struct FFlushRequest
 {
 	FFlushRequest() = default;
 
-	FFlushRequest(int32 InRequestId)
-		: RequestId(MakeShared<std::atomic<int32>, ESPMode::ThreadSafe>(InRequestId))
+	FFlushRequest(TConstArrayView<int32> InRequestIDs)
+		: Data(MakeShared<FFlushRequestData>(InRequestIDs))
 	{
 	}
-
-	int32 GetId() const
+	
+	TConstArrayView<std::atomic<int32>> GetRequestIDs() const 
 	{
-		return RequestId->load(std::memory_order_relaxed);
+		return Data->GetRequestIDs();
 	}
 
 	bool IsValid() const
 	{
-		return RequestId.IsValid();
+		return Data.IsValid();
+	}
+	
+	bool IsComplete() const
+	{
+		return Data->IsComplete();
 	}
 
 	operator bool() const
@@ -137,28 +196,26 @@ struct FFlushRequest
 		return IsValid();
 	}
 
-	void Reset()
+private:
+	void OnRequestComplete(int32 RequestID)
 	{
-		RequestId.Reset();
+		Data->OnRequestComplete(RequestID);	
 	}
 
-private:
-	void SetId(int32 InId)
+	void AdjustRequestIDs(TMap<int32, int32>& DuplicateRequestMap)
 	{
-		RequestId->store(InId, std::memory_order_relaxed);
+		for (std::atomic<int32>& RequestID : Data->GetRequestIDs())
+		{	
+			if (int32* MainRequest = DuplicateRequestMap.Find(RequestID.load(std::memory_order_relaxed)))
+			{
+				RequestID.store(*MainRequest, std::memory_order_relaxed);
+			}
+		}
 	}
 
 	friend class FAsyncLoadingThread;
 
-	/** 
-	 * Shared request id state between main thread and async loading thread. 
-	 * @note because request id are assigned to package request after being generated
-	 * multiple request ids can identify a specific package request however only the original request id
-	 * will properly identify unloaded dependencies. The shared nature of the request allows the async loading thread to
-	 * fix up the main thread flush request id to identify the original request id instead of duplicate when trying to identify
-	 * which main thread packages should be processed when flushing.
-	 */
-	TSharedPtr<std::atomic<int32>, ESPMode::ThreadSafe> RequestId;
+	TSharedPtr<FFlushRequestData, ESPMode::ThreadSafe> Data;
 };
 
 /** Holds the maximum package summary size that can be set via ini files
@@ -333,7 +390,7 @@ public:
 
 	EAsyncPackageState::Type ProcessLoadingUntilComplete(TFunctionRef<bool()> CompletionPredicate, double TimeLimit) override;
 
-	void FlushLoading(int32 PackageId) override;
+	void FlushLoading(TConstArrayView<int32> RequestIDs) override;
 
 	void NotifyConstructedDuringAsyncLoading(UObject* Object, bool bSubObject) override;
 
@@ -650,13 +707,35 @@ private:
 		return PendingRequests.Contains(RequestID);
 	}
 
+	/**
+	 * [ASYNC/GAME THREAD] Checks if a request ID already is added to the loading queue
+	 */
+	bool ContainsAnyRequestInternal(TConstArrayView<int32> RequestIDs)
+	{
+#if THREADSAFE_UOBJECTS
+		FScopeLock Lock(&PendingRequestsCritical);
+#endif
+		return Algo::AnyOf(RequestIDs, [this](int32 RequestID){ return PendingRequests.Contains(RequestID); });
+	}
+
+	/**
+	 * [ASYNC/GAME THREAD] Checks if a request ID already is added to the loading queue
+	 */
+	bool ContainsAnyRequestInternal(TConstArrayView<std::atomic<int32>> RequestIDs)
+	{
+#if THREADSAFE_UOBJECTS
+		FScopeLock Lock(&PendingRequestsCritical);
+#endif
+		return Algo::AnyOf(RequestIDs, [this](const std::atomic<int32>& RequestID){ return PendingRequests.Contains(RequestID.load(std::memory_order_relaxed)); });
+	}
+
 	/** 
 	 * [GAME THREAD] Add a Flush request to the flush stack
-	 * @param RequestID the request id of the package to flush
+	 * @param RequestIDs the request id of the packages to flush
 	 * @returns the flushtree for the flush request to use on the game thread
 	 * @note the FlushRequest is a shared resource between the async loading thread and the main thread
 	 */
-	FFlushRequest AddFlushRequest(int32 RequestID);
+	FFlushRequest AddFlushRequest(TConstArrayView<int32> RequestIDs);
 
 	/** 
 	 * [ASYNC THREAD] Peek the top of the flush request stack
diff --git a/Engine/Source/Runtime/CoreUObject/Private/Serialization/AsyncPackageLoader.cpp b/Engine/Source/Runtime/CoreUObject/Private/Serialization/AsyncPackageLoader.cpp
--- a/Engine/Source/Runtime/CoreUObject/Private/Serialization/AsyncPackageLoader.cpp
+++ b/Engine/Source/Runtime/CoreUObject/Private/Serialization/AsyncPackageLoader.cpp
@@ -251,23 +251,19 @@ bool IsInAsyncLoadingThreadCoreUObjectInternal()
 	}
 }
 
-void FlushAsyncLoading(TConstArrayView<int32> RequestIds)
+void FlushAsyncLoading(int32 RequestId /* = INDEX_NONE */)
 {
-	if (RequestIds.Num() == 0)
+	if (RequestId == INDEX_NONE)
 	{
-		FlushAsyncLoading(INDEX_NONE);
+		FlushAsyncLoading(TConstArrayView<int32>());
 	}
 	else
 	{
-		for (int32 Id : RequestIds)
-		{
-			UE_CLOG(Id == INDEX_NONE, LogStreaming, Warning, TEXT("FlushAsyncLoading called with a list including request id -1, this will flush ALL async loading. If this is intentional, pass an empty list instead."));
-			FlushAsyncLoading(Id);
-		}
+		FlushAsyncLoading(MakeArrayView({RequestId}));
 	}
 }
 
-void FlushAsyncLoading(int32 RequestId /* = INDEX_NONE */)
+void FlushAsyncLoading(TConstArrayView<int32> RequestIds)
 {
 	TRACE_CPUPROFILER_EVENT_SCOPE(FlushAsyncLoading);
 
@@ -284,18 +280,19 @@ void FlushAsyncLoading(int32 RequestId /* = INDEX_NONE */)
 		{
 			// Log the flush, but only display once per frame to avoid log spam.
 			static uint64 LastFrameNumber = -1;
+			TStringBuilder<1024> Buffer;
 			if (LastFrameNumber != GFrameNumber)
 			{
-				UE_LOG(LogStreaming, Display, TEXT("FlushAsyncLoading(%d): %d QueuedPackages, %d AsyncPackages"), RequestId, GPackageLoader->GetNumQueuedPackages(), GPackageLoader->GetNumAsyncPackages());
+				UE_LOG(LogStreaming, Display, TEXT("FlushAsyncLoading(%s): %d QueuedPackages, %d AsyncPackages"), *Buffer.Join(RequestIds, TEXT(",")), GPackageLoader->GetNumQueuedPackages(), GPackageLoader->GetNumAsyncPackages());
 			}
 			else
 			{
-				UE_LOG(LogStreaming, Log, TEXT("FlushAsyncLoading(%d): %d QueuedPackages, %d AsyncPackages"), RequestId, GPackageLoader->GetNumQueuedPackages(), GPackageLoader->GetNumAsyncPackages());
+				UE_LOG(LogStreaming, Log, TEXT("FlushAsyncLoading(%s): %d QueuedPackages, %d AsyncPackages"), *Buffer.Join(RequestIds, TEXT(",")), GPackageLoader->GetNumQueuedPackages(), GPackageLoader->GetNumAsyncPackages());
 			}
 			LastFrameNumber = GFrameNumber;
 		}
 #endif
-		GPackageLoader->FlushLoading(RequestId);
+		GPackageLoader->FlushLoading(RequestIds);
 	}
 }
 
diff --git a/Engine/Source/Runtime/CoreUObject/Private/Serialization/EditorPackageLoader.cpp b/Engine/Source/Runtime/CoreUObject/Private/Serialization/EditorPackageLoader.cpp
--- a/Engine/Source/Runtime/CoreUObject/Private/Serialization/EditorPackageLoader.cpp
+++ b/Engine/Source/Runtime/CoreUObject/Private/Serialization/EditorPackageLoader.cpp
@@ -130,10 +130,10 @@ public:
 		UncookedPackageLoader->ResumeLoading();
 	}
 
-	virtual void FlushLoading(int32 PackageId) override
+	virtual void FlushLoading(TConstArrayView<int32> RequestIds) override
 	{
-		CookedPackageLoader->FlushLoading(PackageId);
-		UncookedPackageLoader->FlushLoading(PackageId);
+		CookedPackageLoader->FlushLoading(RequestIds);
+		UncookedPackageLoader->FlushLoading(RequestIds);
 	}
 
 	virtual int32 GetNumQueuedPackages() override
diff --git a/Engine/Source/Runtime/CoreUObject/Public/Serialization/AsyncPackageLoader.h b/Engine/Source/Runtime/CoreUObject/Public/Serialization/AsyncPackageLoader.h
--- a/Engine/Source/Runtime/CoreUObject/Public/Serialization/AsyncPackageLoader.h
+++ b/Engine/Source/Runtime/CoreUObject/Public/Serialization/AsyncPackageLoader.h
@@ -166,7 +166,7 @@ public:
 	 *
 	 * Note: Called from Game Thread.
 	 */
-	virtual void FlushLoading(int32 PackageId) = 0;
+	virtual void FlushLoading(TConstArrayView<int32> RequestIds) = 0;
 
 	/**
 	 *	Returns the number of queued packages.
diff --git a/Engine/Source/Runtime/CoreUObject/Public/UObject/UObjectGlobals.h b/Engine/Source/Runtime/CoreUObject/Public/UObject/UObjectGlobals.h
--- a/Engine/Source/Runtime/CoreUObject/Public/UObject/UObjectGlobals.h
+++ b/Engine/Source/Runtime/CoreUObject/Public/UObject/UObjectGlobals.h
@@ -840,7 +840,7 @@ COREUOBJECT_API void FlushAsyncLoading(int32 PackageID = INDEX_NONE);
 /**
  * Blocks till a set of pending async load requests are complete.
  *
- * @param RequestIds list of return values from LoadPackageAsync to wait for. 
+ * @param RequestIds list of return values from LoadPackageAsync to wait for. An empty list means all requests
  */
 COREUOBJECT_API void FlushAsyncLoading(TConstArrayView<int32> RequestIds);
 
