Make LoadPackageAsync callable from any thread for zenloader

diff --git a/Engine/Source/Runtime/Core/Private/Misc/CoreDelegates.cpp b/Engine/Source/Runtime/Core/Private/Misc/CoreDelegates.cpp
--- a/Engine/Source/Runtime/Core/Private/Misc/CoreDelegates.cpp
+++ b/Engine/Source/Runtime/Core/Private/Misc/CoreDelegates.cpp
@@ -245,7 +245,17 @@ TMulticastDelegate<void(bool)> FCoreDelegates::IsVanillaProductChanged;
 
 TMulticastDelegate<void()> FCoreDelegates::OnAsyncLoadingFlush;
 TMulticastDelegate<void()> FCoreDelegates::OnAsyncLoadingFlushUpdate;
+
+PRAGMA_DISABLE_DEPRECATION_WARNINGS
 TMulticastDelegate<void(const FString&)> FCoreDelegates::OnAsyncLoadPackage;
+PRAGMA_ENABLE_DEPRECATION_WARNINGS
+
+TTSMulticastDelegate<void(FStringView)>& FCoreDelegates::GetOnAsyncLoadPackage()
+{
+	static TTSMulticastDelegate<void(FStringView)> Singleton;
+	return Singleton;
+}
+
 TMulticastDelegate<void(const FString&)> FCoreDelegates::OnSyncLoadPackage;
 TMulticastDelegate<void()> FCoreDelegates::PostRenderingThreadCreated;
 TMulticastDelegate<void()> FCoreDelegates::PreRenderingThreadDestroyed;
diff --git a/Engine/Source/Runtime/Core/Public/Misc/CoreDelegates.h b/Engine/Source/Runtime/Core/Public/Misc/CoreDelegates.h
--- a/Engine/Source/Runtime/Core/Public/Misc/CoreDelegates.h
+++ b/Engine/Source/Runtime/Core/Public/Misc/CoreDelegates.h
@@ -80,8 +80,12 @@ public:
 	static CORE_API TMulticastDelegate<void()> OnAsyncLoadingFlushUpdate;
 
 	// Callback on the game thread when an async load is started. This goes off before the packages has finished loading
+	UE_DEPRECATED(5.3, "This delegate is not thread-safe, please use GetOnAsyncLoadPackage().")
 	static CORE_API TMulticastDelegate<void(const FString&)> OnAsyncLoadPackage;
 
+	// Thread-safe callback that is called on the same thread that LoadPackageAsync is issued from.
+	static CORE_API TTSMulticastDelegate<void(FStringView)>& GetOnAsyncLoadPackage();
+
 	static CORE_API TMulticastDelegate<void(const FString&)> OnSyncLoadPackage;
 
 	// get a hotfix delegate
diff --git a/Engine/Source/Runtime/CoreUObject/Private/Serialization/AsyncLoading.cpp b/Engine/Source/Runtime/CoreUObject/Private/Serialization/AsyncLoading.cpp
--- a/Engine/Source/Runtime/CoreUObject/Private/Serialization/AsyncLoading.cpp
+++ b/Engine/Source/Runtime/CoreUObject/Private/Serialization/AsyncLoading.cpp
@@ -7428,6 +7428,8 @@ bool FAsyncLoadingThread::ShouldAlwaysLoadPackageAsync(const FPackagePath& InPac
 
 int32 FAsyncLoadingThread::LoadPackage(const FPackagePath& InPackagePath, FName InCustomName, FLoadPackageAsyncDelegate InCompletionDelegate, EPackageFlags InPackageFlags, int32 InPIEInstanceID, int32 InPackagePriority, const FLinkerInstancingContext* InInstancingContext, uint32 InLoadFlags)
 {
+	checkf(IsInGameThread(), TEXT("LoadPackageAsync is only thread-safe when using the zenloader (i.e. AsyncLoading2)."));
+
 	static bool bOnce = false;
 	if (!bOnce && GEventDrivenLoaderEnabled)
 	{
@@ -7445,10 +7447,17 @@ int32 FAsyncLoadingThread::LoadPackage(const FPackagePath& InPackagePath, FName
 	FName PackageName = InCustomName.IsNone() ? InPackagePath.GetPackageFName() : InCustomName;
 	UE_SCOPED_COOK_STAT(PackageName, EPackageEventStatType::LoadPackage);
 
+PRAGMA_DISABLE_DEPRECATION_WARNINGS
 	if ( FCoreDelegates::OnAsyncLoadPackage.IsBound() )
 	{
 		FCoreDelegates::OnAsyncLoadPackage.Broadcast(PackageName.ToString());
 	}
+PRAGMA_ENABLE_DEPRECATION_WARNINGS
+
+	if (FCoreDelegates::GetOnAsyncLoadPackage().IsBound())
+	{
+		FCoreDelegates::GetOnAsyncLoadPackage().Broadcast(PackageName.ToString());
+	}
 
 	// Generate new request ID and add it immediately to the global request list (it needs to be there before we exit
 	// this function, otherwise it would be added when the packages are being processed on the async thread).
diff --git a/Engine/Source/Runtime/CoreUObject/Private/Serialization/AsyncLoading2.cpp b/Engine/Source/Runtime/CoreUObject/Private/Serialization/AsyncLoading2.cpp
--- a/Engine/Source/Runtime/CoreUObject/Private/Serialization/AsyncLoading2.cpp
+++ b/Engine/Source/Runtime/CoreUObject/Private/Serialization/AsyncLoading2.cpp
@@ -8467,11 +8467,21 @@ int32 FAsyncLoadingThread2::LoadPackage(const FPackagePath& InPackagePath, FName
 	{
 		InCustomName = NAME_None;
 	}
+
+PRAGMA_DISABLE_DEPRECATION_WARNINGS
 	if (FCoreDelegates::OnAsyncLoadPackage.IsBound())
 	{
+		checkf(IsInGameThread(), TEXT("FCoreDelegates::OnAsyncLoadPackage is not thread-safe and deprecated, update the callees to be thread-safe and register to FCoreDelegates::GetOnAsyncLoadPackage() instead before calling LoadPackageAsync from any other thread than the game-thread."));
 		const FName PackageName = InCustomName.IsNone() ? PackageNameToLoad : InCustomName;
 		FCoreDelegates::OnAsyncLoadPackage.Broadcast(PackageName.ToString());
 	}
+PRAGMA_ENABLE_DEPRECATION_WARNINGS
+
+	if (FCoreDelegates::GetOnAsyncLoadPackage().IsBound())
+	{
+		const FName PackageName = InCustomName.IsNone() ? PackageNameToLoad : InCustomName;
+		FCoreDelegates::GetOnAsyncLoadPackage().Broadcast(PackageName.ToString());
+	}
 
 	// Generate new request ID and add it immediately to the global request list (it needs to be there before we exit
 	// this function, otherwise it would be added when the packages are being processed on the async thread).
