Fix assert in zenloader steps now that we have multiple paths where states can come from

#rnx

diff --git a/Engine/Source/Runtime/CoreUObject/Private/Serialization/AsyncLoading2.cpp b/Engine/Source/Runtime/CoreUObject/Private/Serialization/AsyncLoading2.cpp
--- a/Engine/Source/Runtime/CoreUObject/Private/Serialization/AsyncLoading2.cpp
+++ b/Engine/Source/Runtime/CoreUObject/Private/Serialization/AsyncLoading2.cpp
@@ -5715,7 +5715,8 @@ EEventLoadNodeExecutionResult FAsyncPackage2::Event_ProcessExportBundle(FAsyncLo
 {
 	TRACE_CPUPROFILER_EVENT_SCOPE(Event_ProcessExportBundle);
 	UE_ASYNC_PACKAGE_DEBUG(Package->Desc);
-	check(Package->AsyncPackageLoadingState == EAsyncPackageLoadingState2::ProcessExportBundles);
+	check(Package->AsyncPackageLoadingState >= EAsyncPackageLoadingState2::DependenciesReady);
+	Package->AsyncPackageLoadingState = EAsyncPackageLoadingState2::ProcessExportBundles;
 
 	FAsyncPackageScope2 Scope(Package);
 #if WITH_EDITOR
