Add Defaults argument to FProperty::ConvertFromType.
This shouldn't result in any change in behavior yet, it's just plumbing the information for future changes to use.

diff --git a/Engine/Source/Runtime/CoreUObject/Private/UObject/Class.cpp b/Engine/Source/Runtime/CoreUObject/Private/UObject/Class.cpp
--- a/Engine/Source/Runtime/CoreUObject/Private/UObject/Class.cpp
+++ b/Engine/Source/Runtime/CoreUObject/Private/UObject/Class.cpp
@@ -1250,7 +1250,7 @@ void UStruct::LoadTaggedPropertiesFromText(FStructuredArchive::FSlot Slot, uint8
 
 					if (!(BreakRecursionIfFullyLoad && BreakRecursionIfFullyLoad->HasAllFlags(RF_LoadCompleted)))
 					{
-						switch (Property->ConvertFromType(Tag, ItemSlot.GetValue(), Data, DefaultsStruct))
+						switch (Property->ConvertFromType(Tag, ItemSlot.GetValue(), Data, DefaultsStruct, Defaults))
 						{
 						case EConvertFromTypeResult::Converted:
 							break;
@@ -1475,7 +1475,7 @@ void UStruct::SerializeVersionedTaggedProperties(FStructuredArchive::FSlot Slot,
 					{
 						FStructuredArchive::FSlot ValueSlot = PropertyRecord.EnterField(TEXT("Value"));
 
-						switch (Property->ConvertFromType(Tag, ValueSlot, Data, DefaultsStruct))
+						switch (Property->ConvertFromType(Tag, ValueSlot, Data, DefaultsStruct, Defaults))
 						{
 							case EConvertFromTypeResult::Converted:
 								bAdvanceProperty = true;
diff --git a/Engine/Source/Runtime/CoreUObject/Private/UObject/EnumProperty.cpp b/Engine/Source/Runtime/CoreUObject/Private/UObject/EnumProperty.cpp
--- a/Engine/Source/Runtime/CoreUObject/Private/UObject/EnumProperty.cpp
+++ b/Engine/Source/Runtime/CoreUObject/Private/UObject/EnumProperty.cpp
@@ -415,7 +415,7 @@ bool FEnumProperty::SameType(const FProperty* Other) const
 	return Super::SameType(Other) && static_cast<const FEnumProperty*>(Other)->Enum == Enum;
 }
 
-EConvertFromTypeResult FEnumProperty::ConvertFromType(const FPropertyTag& Tag, FStructuredArchive::FSlot Slot , uint8* Data, UStruct* DefaultsStruct)
+EConvertFromTypeResult FEnumProperty::ConvertFromType(const FPropertyTag& Tag, FStructuredArchive::FSlot Slot , uint8* Data, UStruct* DefaultsStruct, const uint8* Defaults)
 {
 	if ((Enum == nullptr) || (UnderlyingProp == nullptr))
 	{
diff --git a/Engine/Source/Runtime/CoreUObject/Private/UObject/FieldPathProperty.cpp b/Engine/Source/Runtime/CoreUObject/Private/UObject/FieldPathProperty.cpp
--- a/Engine/Source/Runtime/CoreUObject/Private/UObject/FieldPathProperty.cpp
+++ b/Engine/Source/Runtime/CoreUObject/Private/UObject/FieldPathProperty.cpp
@@ -28,7 +28,7 @@ FFieldPathProperty::FFieldPathProperty(UField* InField)
 }
 #endif // WITH_EDITORONLY_DATA
 
-EConvertFromTypeResult FFieldPathProperty::ConvertFromType(const FPropertyTag& Tag, FStructuredArchive::FSlot Slot, uint8* Data, UStruct* DefaultsStruct)
+EConvertFromTypeResult FFieldPathProperty::ConvertFromType(const FPropertyTag& Tag, FStructuredArchive::FSlot Slot, uint8* Data, UStruct* DefaultsStruct, const uint8* Defaults)
 {
 	// Convert UProperty object to TFieldPath
 	if (Tag.Type == NAME_ObjectProperty)
diff --git a/Engine/Source/Runtime/CoreUObject/Private/UObject/Property.cpp b/Engine/Source/Runtime/CoreUObject/Private/UObject/Property.cpp
--- a/Engine/Source/Runtime/CoreUObject/Private/UObject/Property.cpp
+++ b/Engine/Source/Runtime/CoreUObject/Private/UObject/Property.cpp
@@ -1330,7 +1330,7 @@ void FProperty::LinkInternal(FArchive& Ar)
 	check(0); // Link shouldn't call super...and we should never link an abstract property, like this base class
 }
 
-EConvertFromTypeResult FProperty::ConvertFromType(const FPropertyTag& Tag, FStructuredArchive::FSlot Slot, uint8* Data, UStruct* DefaultsStruct)
+EConvertFromTypeResult FProperty::ConvertFromType(const FPropertyTag& Tag, FStructuredArchive::FSlot Slot, uint8* Data, UStruct* DefaultsStruct, const uint8* Defaults)
 {
 	return EConvertFromTypeResult::UseSerializeItem;
 }
diff --git a/Engine/Source/Runtime/CoreUObject/Private/UObject/PropertyArray.cpp b/Engine/Source/Runtime/CoreUObject/Private/UObject/PropertyArray.cpp
--- a/Engine/Source/Runtime/CoreUObject/Private/UObject/PropertyArray.cpp
+++ b/Engine/Source/Runtime/CoreUObject/Private/UObject/PropertyArray.cpp
@@ -217,7 +217,7 @@ void FArrayProperty::SerializeItem(FStructuredArchive::FSlot Slot, void* Value,
 		{
 			const int64 StartOfProperty = Slot.GetUnderlyingArchive().Tell();
 			FStructProperty* StructProperty = CastFieldChecked<FStructProperty>(Inner);
-			switch(StructProperty->ConvertFromType(SerializeFromMismatchedTag.GetValue(), Slot, Item, nullptr))
+			switch(StructProperty->ConvertFromType(SerializeFromMismatchedTag.GetValue(), Slot, Item, nullptr, nullptr))
 			{
 				case EConvertFromTypeResult::Converted:
 					return;
@@ -803,7 +803,7 @@ bool FArrayProperty::SameType(const FProperty* Other) const
 	return Super::SameType(Other) && Inner && Inner->SameType(((FArrayProperty*)Other)->Inner);
 }
 
-EConvertFromTypeResult FArrayProperty::ConvertFromType(const FPropertyTag& Tag, FStructuredArchive::FSlot Slot, uint8* Data, UStruct* DefaultsStruct)
+EConvertFromTypeResult FArrayProperty::ConvertFromType(const FPropertyTag& Tag, FStructuredArchive::FSlot Slot, uint8* Data, UStruct* DefaultsStruct, const uint8* Defaults)
 {
 	// TODO: The ArrayProperty Tag really doesn't have adequate information for
 	// many types. This should probably all be moved in to ::SerializeItem
@@ -840,11 +840,11 @@ EConvertFromTypeResult FArrayProperty::ConvertFromType(const FPropertyTag& Tag,
 		{
 			FStructuredArchive::FStream ValueStream = Slot.EnterStream();
 
-			if (Inner->ConvertFromType(InnerPropertyTag, ValueStream.EnterElement(), ScriptArrayHelper.GetRawPtr(0), DefaultsStruct) == EConvertFromTypeResult::Converted)
+			if (Inner->ConvertFromType(InnerPropertyTag, ValueStream.EnterElement(), ScriptArrayHelper.GetRawPtr(0), DefaultsStruct, nullptr) == EConvertFromTypeResult::Converted)
 			{
 				for (int32 i = 1; i < ElementCount; ++i)
 				{
-					verify(Inner->ConvertFromType(InnerPropertyTag, ValueStream.EnterElement(), ScriptArrayHelper.GetRawPtr(i), DefaultsStruct) == EConvertFromTypeResult::Converted);
+					verify(Inner->ConvertFromType(InnerPropertyTag, ValueStream.EnterElement(), ScriptArrayHelper.GetRawPtr(i), DefaultsStruct, nullptr) == EConvertFromTypeResult::Converted);
 				}
 
 				return EConvertFromTypeResult::Converted;
diff --git a/Engine/Source/Runtime/CoreUObject/Private/UObject/PropertyBool.cpp b/Engine/Source/Runtime/CoreUObject/Private/UObject/PropertyBool.cpp
--- a/Engine/Source/Runtime/CoreUObject/Private/UObject/PropertyBool.cpp
+++ b/Engine/Source/Runtime/CoreUObject/Private/UObject/PropertyBool.cpp
@@ -274,7 +274,7 @@ void LoadFromType(FBoolProperty* Property, const FPropertyTag& Tag, FStructuredA
 	}
 }
 
-EConvertFromTypeResult FBoolProperty::ConvertFromType(const FPropertyTag& Tag, FStructuredArchive::FSlot Slot, uint8* Data, UStruct* DefaultsStruct)
+EConvertFromTypeResult FBoolProperty::ConvertFromType(const FPropertyTag& Tag, FStructuredArchive::FSlot Slot, uint8* Data, UStruct* DefaultsStruct, const uint8* Defaults)
 {
 	if (Tag.Type == NAME_IntProperty)
 	{
diff --git a/Engine/Source/Runtime/CoreUObject/Private/UObject/PropertyByte.cpp b/Engine/Source/Runtime/CoreUObject/Private/UObject/PropertyByte.cpp
--- a/Engine/Source/Runtime/CoreUObject/Private/UObject/PropertyByte.cpp
+++ b/Engine/Source/Runtime/CoreUObject/Private/UObject/PropertyByte.cpp
@@ -213,7 +213,7 @@ struct TConvertIntToEnumProperty
 	}
 };
 
-EConvertFromTypeResult FByteProperty::ConvertFromType(const FPropertyTag& Tag, FStructuredArchive::FSlot Slot, uint8* Data, UStruct* DefaultsStruct)
+EConvertFromTypeResult FByteProperty::ConvertFromType(const FPropertyTag& Tag, FStructuredArchive::FSlot Slot, uint8* Data, UStruct* DefaultsStruct, const uint8* Defaults)
 {
 	if (Tag.Type == NAME_ByteProperty  && ((Tag.EnumName == NAME_None) != (Enum == nullptr)))
 	{
diff --git a/Engine/Source/Runtime/CoreUObject/Private/UObject/PropertyMap.cpp b/Engine/Source/Runtime/CoreUObject/Private/UObject/PropertyMap.cpp
--- a/Engine/Source/Runtime/CoreUObject/Private/UObject/PropertyMap.cpp
+++ b/Engine/Source/Runtime/CoreUObject/Private/UObject/PropertyMap.cpp
@@ -1058,7 +1058,7 @@ bool FMapProperty::SameType(const FProperty* Other) const
 	return Super::SameType(Other) && KeyProp && ValueProp && KeyProp->SameType(MapProp->KeyProp) && ValueProp->SameType(MapProp->ValueProp);
 }
 
-EConvertFromTypeResult FMapProperty::ConvertFromType(const FPropertyTag& Tag, FStructuredArchive::FSlot Slot, uint8* Data, UStruct* DefaultsStruct)
+EConvertFromTypeResult FMapProperty::ConvertFromType(const FPropertyTag& Tag, FStructuredArchive::FSlot Slot, uint8* Data, UStruct* DefaultsStruct, const uint8* Defaults)
 {
 	FArchive& UnderlyingArchive = Slot.GetUnderlyingArchive();
 
@@ -1075,7 +1075,7 @@ EConvertFromTypeResult FMapProperty::ConvertFromType(const FPropertyTag& Tag, FS
 			CurrentType->SerializeItem(InnerSlot, DestAddress, nullptr);
 			return true;
 		}
-		else if( CurrentType->ConvertFromType(InTag, InnerSlot, InData, InDefaultsStruct) == EConvertFromTypeResult::Converted )
+		else if( CurrentType->ConvertFromType(InTag, InnerSlot, InData, InDefaultsStruct, nullptr) == EConvertFromTypeResult::Converted )
 		{
 			return true;
 		}
diff --git a/Engine/Source/Runtime/CoreUObject/Private/UObject/PropertyMulticastDelegate.cpp b/Engine/Source/Runtime/CoreUObject/Private/UObject/PropertyMulticastDelegate.cpp
--- a/Engine/Source/Runtime/CoreUObject/Private/UObject/PropertyMulticastDelegate.cpp
+++ b/Engine/Source/Runtime/CoreUObject/Private/UObject/PropertyMulticastDelegate.cpp
@@ -390,7 +390,7 @@ bool FMulticastDelegateProperty::SameType(const FProperty* Other) const
 	return Super::SameType(Other) && (SignatureFunction == ((FMulticastDelegateProperty*)Other)->SignatureFunction);
 }
 
-EConvertFromTypeResult FMulticastDelegateProperty::ConvertFromType(const FPropertyTag& Tag, FStructuredArchive::FSlot Slot, uint8* Data, UStruct* DefaultsStruct)
+EConvertFromTypeResult FMulticastDelegateProperty::ConvertFromType(const FPropertyTag& Tag, FStructuredArchive::FSlot Slot, uint8* Data, UStruct* DefaultsStruct, const uint8* Defaults)
 {
 	// Multicast delegate properties are serialization compatible
 	if (Tag.Type == NAME_MulticastDelegateProperty || Tag.Type == FMulticastInlineDelegateProperty::StaticClass()->GetFName() || Tag.Type == FMulticastSparseDelegateProperty::StaticClass()->GetFName())
diff --git a/Engine/Source/Runtime/CoreUObject/Private/UObject/PropertyName.cpp b/Engine/Source/Runtime/CoreUObject/Private/UObject/PropertyName.cpp
--- a/Engine/Source/Runtime/CoreUObject/Private/UObject/PropertyName.cpp
+++ b/Engine/Source/Runtime/CoreUObject/Private/UObject/PropertyName.cpp
@@ -73,7 +73,7 @@ const TCHAR* FNameProperty::ImportText_Internal( const TCHAR* Buffer, void* Cont
 	return Buffer;
 }
 
-EConvertFromTypeResult FNameProperty::ConvertFromType(const FPropertyTag& Tag, FStructuredArchive::FSlot Slot, uint8* Data, UStruct* DefaultsStruct)
+EConvertFromTypeResult FNameProperty::ConvertFromType(const FPropertyTag& Tag, FStructuredArchive::FSlot Slot, uint8* Data, UStruct* DefaultsStruct, const uint8* Defaults)
 {
 	if (Tag.Type == NAME_StrProperty)
 	{
diff --git a/Engine/Source/Runtime/CoreUObject/Private/UObject/PropertyObject.cpp b/Engine/Source/Runtime/CoreUObject/Private/UObject/PropertyObject.cpp
--- a/Engine/Source/Runtime/CoreUObject/Private/UObject/PropertyObject.cpp
+++ b/Engine/Source/Runtime/CoreUObject/Private/UObject/PropertyObject.cpp
@@ -37,7 +37,7 @@ FString FObjectProperty::GetCPPMacroType( FString& ExtendedTypeText ) const
 	return TEXT("OBJECT");
 }
 
-EConvertFromTypeResult FObjectProperty::ConvertFromType(const FPropertyTag& Tag, FStructuredArchive::FSlot Slot, uint8* Data, UStruct* DefaultsStruct)
+EConvertFromTypeResult FObjectProperty::ConvertFromType(const FPropertyTag& Tag, FStructuredArchive::FSlot Slot, uint8* Data, UStruct* DefaultsStruct, const uint8* Defaults)
 {
 	static FName NAME_AssetObjectProperty = "AssetObjectProperty"; // old name of soft object property
 
diff --git a/Engine/Source/Runtime/CoreUObject/Private/UObject/PropertySet.cpp b/Engine/Source/Runtime/CoreUObject/Private/UObject/PropertySet.cpp
--- a/Engine/Source/Runtime/CoreUObject/Private/UObject/PropertySet.cpp
+++ b/Engine/Source/Runtime/CoreUObject/Private/UObject/PropertySet.cpp
@@ -902,7 +902,7 @@ bool FSetProperty::SameType(const FProperty* Other) const
 	return Super::SameType(Other) && ElementProp && ElementProp->SameType(SetProp->ElementProp);
 }
 
-EConvertFromTypeResult FSetProperty::ConvertFromType(const FPropertyTag& Tag, FStructuredArchive::FSlot Slot, uint8* Data, UStruct* DefaultsStruct)
+EConvertFromTypeResult FSetProperty::ConvertFromType(const FPropertyTag& Tag, FStructuredArchive::FSlot Slot, uint8* Data, UStruct* DefaultsStruct, const uint8* Defaults)
 {
 	FArchive& UnderlyingArchive = Slot.GetUnderlyingArchive();
 
@@ -945,7 +945,7 @@ EConvertFromTypeResult FSetProperty::ConvertFromType(const FPropertyTag& Tag, FS
 				TempElementStorage = (uint8*)FMemory::Malloc(SetLayout.Size);
 				ElementProp->InitializeValue(TempElementStorage);
 
-				if (ElementProp->ConvertFromType(InnerPropertyTag, ElementsToRemoveArray.EnterElement(), TempElementStorage, DefaultsStruct) == EConvertFromTypeResult::Converted)
+				if (ElementProp->ConvertFromType(InnerPropertyTag, ElementsToRemoveArray.EnterElement(), TempElementStorage, DefaultsStruct, nullptr) == EConvertFromTypeResult::Converted)
 				{
 					int32 Found = ScriptSetHelper.FindElementIndex(TempElementStorage);
 					if (Found != INDEX_NONE)
@@ -955,7 +955,7 @@ EConvertFromTypeResult FSetProperty::ConvertFromType(const FPropertyTag& Tag, FS
 
 					for (int32 I = 1; I < NumElementsToRemove; ++I)
 					{
-						verify(ElementProp->ConvertFromType(InnerPropertyTag, ElementsToRemoveArray.EnterElement(), TempElementStorage, DefaultsStruct) == EConvertFromTypeResult::Converted);
+						verify(ElementProp->ConvertFromType(InnerPropertyTag, ElementsToRemoveArray.EnterElement(), TempElementStorage, DefaultsStruct, nullptr) == EConvertFromTypeResult::Converted);
 
 						Found = ScriptSetHelper.FindElementIndex(TempElementStorage);
 						if (Found != INDEX_NONE)
@@ -986,7 +986,7 @@ EConvertFromTypeResult FSetProperty::ConvertFromType(const FPropertyTag& Tag, FS
 
 					// and read the first entry, we have to check for conversion possibility again because 
 					// NumElementsToRemove may not have run (in fact, it likely did not):
-					if (ElementProp->ConvertFromType(InnerPropertyTag, ElementsArray.EnterElement(), TempElementStorage, DefaultsStruct) == EConvertFromTypeResult::Converted)
+					if (ElementProp->ConvertFromType(InnerPropertyTag, ElementsArray.EnterElement(), TempElementStorage, DefaultsStruct, nullptr) == EConvertFromTypeResult::Converted)
 					{
 						if (ScriptSetHelper.FindElementIndex(TempElementStorage) == INDEX_NONE)
 						{
@@ -1001,7 +1001,7 @@ EConvertFromTypeResult FSetProperty::ConvertFromType(const FPropertyTag& Tag, FS
 						for (int32 I = 1; I < Num; ++I)
 						{
 							// Read key into temporary storage
-							verify(ElementProp->ConvertFromType(InnerPropertyTag, ElementsArray.EnterElement(), TempElementStorage, DefaultsStruct) == EConvertFromTypeResult::Converted);
+							verify(ElementProp->ConvertFromType(InnerPropertyTag, ElementsArray.EnterElement(), TempElementStorage, DefaultsStruct, nullptr) == EConvertFromTypeResult::Converted);
 
 							// Add a new entry if the element doesn't currently exist in the set
 							if (ScriptSetHelper.FindElementIndex(TempElementStorage) == INDEX_NONE)
diff --git a/Engine/Source/Runtime/CoreUObject/Private/UObject/PropertySoftObjectPtr.cpp b/Engine/Source/Runtime/CoreUObject/Private/UObject/PropertySoftObjectPtr.cpp
--- a/Engine/Source/Runtime/CoreUObject/Private/UObject/PropertySoftObjectPtr.cpp
+++ b/Engine/Source/Runtime/CoreUObject/Private/UObject/PropertySoftObjectPtr.cpp
@@ -174,7 +174,7 @@ const TCHAR* FSoftObjectProperty::ImportText_Internal( const TCHAR* InBuffer, vo
 	}
 }
 
-EConvertFromTypeResult FSoftObjectProperty::ConvertFromType(const FPropertyTag& Tag, FStructuredArchive::FSlot Slot, uint8* Data, UStruct* DefaultsStruct)
+EConvertFromTypeResult FSoftObjectProperty::ConvertFromType(const FPropertyTag& Tag, FStructuredArchive::FSlot Slot, uint8* Data, UStruct* DefaultsStruct, const uint8* Defaults)
 {
 	static FName NAME_AssetObjectProperty = "AssetObjectProperty";
 	static FName NAME_SoftObjectPath = "SoftObjectPath";
diff --git a/Engine/Source/Runtime/CoreUObject/Private/UObject/PropertyStr.cpp b/Engine/Source/Runtime/CoreUObject/Private/UObject/PropertyStr.cpp
--- a/Engine/Source/Runtime/CoreUObject/Private/UObject/PropertyStr.cpp
+++ b/Engine/Source/Runtime/CoreUObject/Private/UObject/PropertyStr.cpp
@@ -17,7 +17,7 @@ FStrProperty::FStrProperty(FFieldVariant InOwner, const UECodeGen_Private::FStrP
 {
 }
 
-EConvertFromTypeResult FStrProperty::ConvertFromType(const FPropertyTag& Tag, FStructuredArchive::FSlot Slot, uint8* Data, UStruct* DefaultsStruct)
+EConvertFromTypeResult FStrProperty::ConvertFromType(const FPropertyTag& Tag, FStructuredArchive::FSlot Slot, uint8* Data, UStruct* DefaultsStruct, const uint8* Defaults)
 {
 	// Convert serialized text to string.
 	if (Tag.Type==NAME_TextProperty)
diff --git a/Engine/Source/Runtime/CoreUObject/Private/UObject/PropertyStruct.cpp b/Engine/Source/Runtime/CoreUObject/Private/UObject/PropertyStruct.cpp
--- a/Engine/Source/Runtime/CoreUObject/Private/UObject/PropertyStruct.cpp
+++ b/Engine/Source/Runtime/CoreUObject/Private/UObject/PropertyStruct.cpp
@@ -391,7 +391,7 @@ bool FStructProperty::SameType(const FProperty* Other) const
 	return Super::SameType(Other) && (Struct == ((FStructProperty*)Other)->Struct);
 }
 
-EConvertFromTypeResult FStructProperty::ConvertFromType(const FPropertyTag& Tag, FStructuredArchive::FSlot Slot, uint8* Data, UStruct* DefaultsStruct)
+EConvertFromTypeResult FStructProperty::ConvertFromType(const FPropertyTag& Tag, FStructuredArchive::FSlot Slot, uint8* Data, UStruct* DefaultsStruct, const uint8* Defaults)
 {
 	FArchive& UnderlyingArchive = Slot.GetUnderlyingArchive();
 
diff --git a/Engine/Source/Runtime/CoreUObject/Private/UObject/TextProperty.cpp b/Engine/Source/Runtime/CoreUObject/Private/UObject/TextProperty.cpp
--- a/Engine/Source/Runtime/CoreUObject/Private/UObject/TextProperty.cpp
+++ b/Engine/Source/Runtime/CoreUObject/Private/UObject/TextProperty.cpp
@@ -17,7 +17,7 @@ FTextProperty::FTextProperty(FFieldVariant InOwner, const UECodeGen_Private::FTe
 {
 }
 
-EConvertFromTypeResult FTextProperty::ConvertFromType(const FPropertyTag& Tag, FStructuredArchive::FSlot Slot, uint8* Data, UStruct* DefaultsStruct)
+EConvertFromTypeResult FTextProperty::ConvertFromType(const FPropertyTag& Tag, FStructuredArchive::FSlot Slot, uint8* Data, UStruct* DefaultsStruct, const uint8* Defaults)
 {
 	// Convert serialized string to text.
 	if (Tag.Type==NAME_StrProperty)
diff --git a/Engine/Source/Runtime/CoreUObject/Public/UObject/EnumProperty.h b/Engine/Source/Runtime/CoreUObject/Public/UObject/EnumProperty.h
--- a/Engine/Source/Runtime/CoreUObject/Public/UObject/EnumProperty.h
+++ b/Engine/Source/Runtime/CoreUObject/Public/UObject/EnumProperty.h
@@ -73,7 +73,7 @@ public:
 	COREUOBJECT_API virtual const TCHAR* ImportText_Internal(const TCHAR* Buffer, void* ContainerOrPropertyPtr, EPropertyPointerType PropertyPointerType, UObject* OwnerObject, int32 PortFlags, FOutputDevice* ErrorText) const override;
 	COREUOBJECT_API virtual int32 GetMinAlignment() const override;
 	COREUOBJECT_API virtual bool SameType(const FProperty* Other) const override;
-	COREUOBJECT_API virtual EConvertFromTypeResult ConvertFromType(const FPropertyTag& Tag, FStructuredArchive::FSlot Slot, uint8* Data, UStruct* DefaultsStruct) override;
+	COREUOBJECT_API virtual EConvertFromTypeResult ConvertFromType(const FPropertyTag& Tag, FStructuredArchive::FSlot Slot, uint8* Data, UStruct* DefaultsStruct, const uint8* Defaults) override;
 #if WITH_EDITORONLY_DATA
 	COREUOBJECT_API virtual void AppendSchemaHash(FBlake3& Builder, bool bSkipEditorOnly) const override;
 #endif
diff --git a/Engine/Source/Runtime/CoreUObject/Public/UObject/FieldPathProperty.h b/Engine/Source/Runtime/CoreUObject/Public/UObject/FieldPathProperty.h
--- a/Engine/Source/Runtime/CoreUObject/Public/UObject/FieldPathProperty.h
+++ b/Engine/Source/Runtime/CoreUObject/Public/UObject/FieldPathProperty.h
@@ -68,7 +68,7 @@ PRAGMA_ENABLE_DEPRECATION_WARNINGS
 	// End of UHT interface
 
 	// FProperty interface
-	COREUOBJECT_API virtual EConvertFromTypeResult ConvertFromType(const FPropertyTag& Tag, FStructuredArchive::FSlot Slot, uint8* Data, UStruct* DefaultsStruct) override;
+	COREUOBJECT_API virtual EConvertFromTypeResult ConvertFromType(const FPropertyTag& Tag, FStructuredArchive::FSlot Slot, uint8* Data, UStruct* DefaultsStruct, const uint8* Defaults) override;
 	COREUOBJECT_API virtual bool Identical( const void* A, const void* B, uint32 PortFlags ) const override;
 	COREUOBJECT_API virtual void SerializeItem(FStructuredArchive::FSlot Slot, void* Value, void const* Defaults) const override;
 	COREUOBJECT_API virtual void ExportText_Internal( FString& ValueStr, const void* PropertyValueOrContainer, EPropertyPointerType PropertyPointerType, const void* DefaultValue, UObject* Parent, int32 PortFlags, UObject* ExportRootScope ) const override;
diff --git a/Engine/Source/Runtime/CoreUObject/Public/UObject/TextProperty.h b/Engine/Source/Runtime/CoreUObject/Public/UObject/TextProperty.h
--- a/Engine/Source/Runtime/CoreUObject/Public/UObject/TextProperty.h
+++ b/Engine/Source/Runtime/CoreUObject/Public/UObject/TextProperty.h
@@ -57,7 +57,7 @@ PRAGMA_ENABLE_DEPRECATION_WARNINGS
 #endif // WITH_EDITORONLY_DATA
 
 	// FProperty interface
-	COREUOBJECT_API virtual EConvertFromTypeResult ConvertFromType(const FPropertyTag& Tag, FStructuredArchive::FSlot Slot, uint8* Data, UStruct* DefaultsStruct) override;
+	COREUOBJECT_API virtual EConvertFromTypeResult ConvertFromType(const FPropertyTag& Tag, FStructuredArchive::FSlot Slot, uint8* Data, UStruct* DefaultsStruct, const uint8* Defaults) override;
 	COREUOBJECT_API virtual bool Identical( const void* A, const void* B, uint32 PortFlags ) const override;
 	COREUOBJECT_API virtual void SerializeItem(FStructuredArchive::FSlot Slot, void* Value, void const* Defaults) const override;
 protected:
diff --git a/Engine/Source/Runtime/CoreUObject/Public/UObject/UnrealType.h b/Engine/Source/Runtime/CoreUObject/Public/UObject/UnrealType.h
--- a/Engine/Source/Runtime/CoreUObject/Public/UObject/UnrealType.h
+++ b/Engine/Source/Runtime/CoreUObject/Public/UObject/UnrealType.h
@@ -356,15 +356,16 @@ public:
 	* 
 	* @param	Tag			property tag of the loading data
 	* @param	Ar			the archive the data is being loaded from
-	* @param	Data		
+	* @param	Data		a pointer to the container to write the loaded data to
 	* @param	DefaultsStruct 
+	* @param	Defaults	if available, a pointer to the container containing the default value for this property, or null
 	*
 	* @return	A state which tells the tagged property system how the property dealt with the data.
 	*			Converted:        the function has handled the tag.
 	*			CannotConvert:    the tag is not something that the property can convert.
 	*			UseSerializeItem: no conversion was done on the property - this can mean that the tag is correct and normal serialization applies or that the tag is incompatible.
 	*/
-	COREUOBJECT_API virtual EConvertFromTypeResult ConvertFromType(const FPropertyTag& Tag, FStructuredArchive::FSlot Slot, uint8* Data, UStruct* DefaultsStruct);
+	COREUOBJECT_API virtual EConvertFromTypeResult ConvertFromType(const FPropertyTag& Tag, FStructuredArchive::FSlot Slot, uint8* Data, UStruct* DefaultsStruct, const uint8* Defaults);
 
 	/**
 	 * Determines whether the property values are identical.
@@ -1722,7 +1723,7 @@ private:
 	};
 
 public:
-	virtual EConvertFromTypeResult ConvertFromType(const FPropertyTag& Tag, FStructuredArchive::FSlot Slot, uint8* Data, UStruct* DefaultsStruct) override
+	virtual EConvertFromTypeResult ConvertFromType(const FPropertyTag& Tag, FStructuredArchive::FSlot Slot, uint8* Data, UStruct* DefaultsStruct, const uint8* Defaults) override
 	{
 		if (const EName * TagType = Tag.Type.ToEName())
 		{
@@ -1949,7 +1950,7 @@ protected:
 	virtual void ExportText_Internal( FString& ValueStr, const void* PropertyValueOrContainer, EPropertyPointerType PropertyPointerType, const void* DefaultValue, UObject* Parent, int32 PortFlags, UObject* ExportRootScope ) const override;
 	virtual const TCHAR* ImportText_Internal(const TCHAR* Buffer, void* ContainerOrPropertyPtr, EPropertyPointerType PropertyPointerType, UObject* Parent, int32 PortFlags, FOutputDevice* ErrorText) const override;
 public:
-	virtual EConvertFromTypeResult ConvertFromType(const FPropertyTag& Tag, FStructuredArchive::FSlot Slot, uint8* Data, UStruct* DefaultsStruct) override;
+	virtual EConvertFromTypeResult ConvertFromType(const FPropertyTag& Tag, FStructuredArchive::FSlot Slot, uint8* Data, UStruct* DefaultsStruct, const uint8* Defaults) override;
 #if WITH_EDITORONLY_DATA
 	virtual void AppendSchemaHash(FBlake3& Builder, bool bSkipEditorOnly) const override;
 #endif
@@ -2405,7 +2406,7 @@ public:
 	virtual void ClearValueInternal( void* Data ) const override;
 	virtual void InitializeValueInternal( void* Dest ) const override;
 	virtual int32 GetMinAlignment() const override;
-	virtual EConvertFromTypeResult ConvertFromType(const FPropertyTag& Tag, FStructuredArchive::FSlot Slot, uint8* Data, UStruct* DefaultsStruct) override;
+	virtual EConvertFromTypeResult ConvertFromType(const FPropertyTag& Tag, FStructuredArchive::FSlot Slot, uint8* Data, UStruct* DefaultsStruct, const uint8* Defaults) override;
 #if WITH_EDITORONLY_DATA
 	virtual void AppendSchemaHash(FBlake3& Builder, bool bSkipEditorOnly) const override;
 #endif
@@ -2873,7 +2874,7 @@ PRAGMA_ENABLE_DEPRECATION_WARNINGS
 	virtual void SerializeItem(FStructuredArchive::FSlot Slot, void* Value, void const* Defaults) const override;
 	virtual void EmitReferenceInfo(UE::GC::FSchemaBuilder& Schema, int32 BaseOffset, TArray<const FStructProperty*>& EncounteredStructProps, UE::GC::FPropertyStack& DebugPath) override;
 	virtual const TCHAR* ImportText_Internal(const TCHAR* Buffer, void* ContainerOrPropertyPtr, EPropertyPointerType PropertyPointerType, UObject* OwnerObject, int32 PortFlags, FOutputDevice* ErrorText) const override;
-	virtual EConvertFromTypeResult ConvertFromType(const FPropertyTag& Tag, FStructuredArchive::FSlot Slot, uint8* Data, UStruct* DefaultsStruct) override;
+	virtual EConvertFromTypeResult ConvertFromType(const FPropertyTag& Tag, FStructuredArchive::FSlot Slot, uint8* Data, UStruct* DefaultsStruct, const uint8* Defaults) override;
 
 private:
 	virtual uint32 GetValueTypeHashInternal(const void* Src) const override;
@@ -3159,7 +3160,7 @@ protected:
 	virtual void ExportText_Internal( FString& ValueStr, const void* PropertyValueOrContainer, EPropertyPointerType PropertyPointerType, const void* DefaultValue, UObject* Parent, int32 PortFlags, UObject* ExportRootScope ) const override;
 	virtual const TCHAR* ImportText_Internal(const TCHAR* Buffer, void* ContainerOrPropertyPtr, EPropertyPointerType PropertyPointerType, UObject* OwnerObject, int32 PortFlags, FOutputDevice* ErrorText) const override;
 public:
-	virtual EConvertFromTypeResult ConvertFromType(const FPropertyTag& Tag, FStructuredArchive::FSlot Slot, uint8* Data, UStruct* DefaultsStruct) override;
+	virtual EConvertFromTypeResult ConvertFromType(const FPropertyTag& Tag, FStructuredArchive::FSlot Slot, uint8* Data, UStruct* DefaultsStruct, const uint8* Defaults) override;
 	// End of FProperty interface
 
 	// FObjectProperty interface
@@ -3543,7 +3544,7 @@ protected:
 	virtual void ExportText_Internal( FString& ValueStr, const void* PropertyValueOrContainer, EPropertyPointerType PropertyPointerType, const void* DefaultValue, UObject* Parent, int32 PortFlags, UObject* ExportRootScope ) const override;
 	virtual const TCHAR* ImportText_Internal(const TCHAR* Buffer, void* ContainerOrPropertyPtr, EPropertyPointerType PropertyPointerType, UObject* OwnerObject, int32 PortFlags, FOutputDevice* ErrorText) const override;
 public:
-	virtual EConvertFromTypeResult ConvertFromType(const FPropertyTag& Tag, FStructuredArchive::FSlot Slot, uint8* Data, UStruct* DefaultsStruct) override;
+	virtual EConvertFromTypeResult ConvertFromType(const FPropertyTag& Tag, FStructuredArchive::FSlot Slot, uint8* Data, UStruct* DefaultsStruct, const uint8* Defaults) override;
 	virtual FString GetCPPTypeForwardDeclaration() const override;
 	uint32 GetValueTypeHashInternal(const void* Src) const override;
 	// End of FProperty interface
@@ -3599,7 +3600,7 @@ protected:
 	virtual void ExportText_Internal( FString& ValueStr, const void* PropertyValueOrContainer, EPropertyPointerType PropertyPointerType, const void* DefaultValue, UObject* Parent, int32 PortFlags, UObject* ExportRootScope ) const override;
 	virtual const TCHAR* ImportText_Internal(const TCHAR* Buffer, void* ContainerOrPropertyPtr, EPropertyPointerType PropertyPointerType, UObject* OwnerObject, int32 PortFlags, FOutputDevice* ErrorText) const override;
 public:
-	virtual EConvertFromTypeResult ConvertFromType(const FPropertyTag& Tag, FStructuredArchive::FSlot Slot, uint8* Data, UStruct* DefaultsStruct) override;
+	virtual EConvertFromTypeResult ConvertFromType(const FPropertyTag& Tag, FStructuredArchive::FSlot Slot, uint8* Data, UStruct* DefaultsStruct, const uint8* Defaults) override;
 	virtual FString GetCPPTypeForwardDeclaration() const override;
 	uint32 GetValueTypeHashInternal(const void* Src) const override;
 	// End of FProperty interface
@@ -3731,7 +3732,7 @@ public:
 	virtual bool ContainsObjectReference(TArray<const FStructProperty*>& EncounteredStructProps, EPropertyObjectReferenceType InReferenceType = EPropertyObjectReferenceType::Strong) const override;
 	virtual void EmitReferenceInfo(UE::GC::FSchemaBuilder& Schema, int32 BaseOffset, TArray<const FStructProperty*>& EncounteredStructProps, UE::GC::FPropertyStack& DebugPath) override;
 	virtual bool SameType(const FProperty* Other) const override;
-	virtual EConvertFromTypeResult ConvertFromType(const FPropertyTag& Tag, FStructuredArchive::FSlot Slot, uint8* Data, UStruct* DefaultsStruct) override;
+	virtual EConvertFromTypeResult ConvertFromType(const FPropertyTag& Tag, FStructuredArchive::FSlot Slot, uint8* Data, UStruct* DefaultsStruct, const uint8* Defaults) override;
 
 	virtual int32 GetMinAlignment() const override
 	{
@@ -3867,7 +3868,7 @@ public:
 	virtual bool ContainsObjectReference(TArray<const FStructProperty*>& EncounteredStructProps, EPropertyObjectReferenceType InReferenceType = EPropertyObjectReferenceType::Strong) const override;
 	virtual void EmitReferenceInfo(UE::GC::FSchemaBuilder& Schema, int32 BaseOffset, TArray<const FStructProperty*>& EncounteredStructProps, UE::GC::FPropertyStack& DebugPath) override;
 	virtual bool SameType(const FProperty* Other) const override;
-	virtual EConvertFromTypeResult ConvertFromType(const FPropertyTag& Tag, FStructuredArchive::FSlot Slot, uint8* Data, UStruct* DefaultsStruct) override;
+	virtual EConvertFromTypeResult ConvertFromType(const FPropertyTag& Tag, FStructuredArchive::FSlot Slot, uint8* Data, UStruct* DefaultsStruct, const uint8* Defaults) override;
 	virtual void* GetValueAddressAtIndex_Direct(const FProperty* Inner, void* InValueAddress, int32 Index) const override;
 	// End of FProperty interface
 
@@ -3991,7 +3992,7 @@ public:
 	virtual bool ContainsObjectReference(TArray<const FStructProperty*>& EncounteredStructProps, EPropertyObjectReferenceType InReferenceType = EPropertyObjectReferenceType::Strong) const override;
 	virtual void EmitReferenceInfo(UE::GC::FSchemaBuilder& Schema, int32 BaseOffset, TArray<const FStructProperty*>& EncounteredStructProps, UE::GC::FPropertyStack& DebugPath) override;
 	virtual bool SameType(const FProperty* Other) const override;
-	virtual EConvertFromTypeResult ConvertFromType(const FPropertyTag& Tag, FStructuredArchive::FSlot Slot, uint8* Data, UStruct* DefaultsStruct) override;
+	virtual EConvertFromTypeResult ConvertFromType(const FPropertyTag& Tag, FStructuredArchive::FSlot Slot, uint8* Data, UStruct* DefaultsStruct, const uint8* Defaults) override;
 	virtual void* GetValueAddressAtIndex_Direct(const FProperty* Inner, void* InValueAddress, int32 Index) const override;
 	// End of FProperty interface
 
@@ -5838,7 +5839,7 @@ public:
 	virtual bool ContainsObjectReference(TArray<const FStructProperty*>& EncounteredStructProps, EPropertyObjectReferenceType InReferenceType = EPropertyObjectReferenceType::Strong) const override;
 	virtual void EmitReferenceInfo(UE::GC::FSchemaBuilder& Schema, int32 BaseOffset, TArray<const FStructProperty*>& EncounteredStructProps, UE::GC::FPropertyStack& DebugPath) override;
 	virtual bool SameType(const FProperty* Other) const override;
-	virtual EConvertFromTypeResult ConvertFromType(const FPropertyTag& Tag, FStructuredArchive::FSlot Slot, uint8* Data, UStruct* DefaultsStruct) override;
+	virtual EConvertFromTypeResult ConvertFromType(const FPropertyTag& Tag, FStructuredArchive::FSlot Slot, uint8* Data, UStruct* DefaultsStruct, const uint8* Defaults) override;
 #if WITH_EDITORONLY_DATA
 	virtual void AppendSchemaHash(FBlake3& Builder, bool bSkipEditorOnly) const override;
 #endif
@@ -5989,7 +5990,7 @@ public:
 	virtual bool ContainsObjectReference(TArray<const FStructProperty*>& EncounteredStructProps, EPropertyObjectReferenceType InReferenceType = EPropertyObjectReferenceType::Strong) const override;
 	virtual void InstanceSubobjects( void* Data, void const* DefaultData, UObject* Owner, struct FObjectInstancingGraph* InstanceGraph ) override;
 	virtual bool SameType(const FProperty* Other) const override;
-	virtual EConvertFromTypeResult ConvertFromType(const FPropertyTag& Tag, FStructuredArchive::FSlot Slot, uint8* Data, UStruct* DefaultsStruct) override;
+	virtual EConvertFromTypeResult ConvertFromType(const FPropertyTag& Tag, FStructuredArchive::FSlot Slot, uint8* Data, UStruct* DefaultsStruct, const uint8* Defaults) override;
 	// End of FProperty interface
 
 	virtual const FMulticastScriptDelegate* GetMulticastDelegate(const void* PropertyValue) const PURE_VIRTUAL(FMulticastDelegateProperty::GetMulticastDelegate, return nullptr;);
