Linker Instancing Context:
- Add an optional package remapping functor that can be used to remap packages instead of building an instancing context.
- Remapping functor is invoked when the provided instancing context doesn't remap a package.
- Cache the result of the package remapping functor to avoid multiple calls with the same package to be remapped.

diff --git a/Engine/Source/Runtime/CoreUObject/Public/UObject/LinkerInstancingContext.h b/Engine/Source/Runtime/CoreUObject/Public/UObject/LinkerInstancingContext.h
--- a/Engine/Source/Runtime/CoreUObject/Public/UObject/LinkerInstancingContext.h
+++ b/Engine/Source/Runtime/CoreUObject/Public/UObject/LinkerInstancingContext.h
@@ -98,13 +98,25 @@ public:
 
 	bool IsInstanced() const
 	{
-		return InstancedPackageMap.IsInstanced() || PathMapping.Num() > 0;
+		return InstancedPackageMap.IsInstanced() || InstancedPackageMapFunc || PathMapping.Num() > 0;
 	}
 
 	/** Remap the package name from the import table to its instanced counterpart, otherwise return the name unmodified. */
 	FName RemapPackage(const FName& PackageName) const
 	{
-		return InstancedPackageMap.RemapPackage(PackageName);
+		FName RemappedPackageName = InstancedPackageMap.RemapPackage(PackageName);
+
+		if ((RemappedPackageName == PackageName) && InstancedPackageMapFunc)
+		{
+			RemappedPackageName = InstancedPackageMapFunc(PackageName);
+
+			if (RemappedPackageName != PackageName)
+			{
+				InstancedPackageMap.AddPackageMapping(PackageName, RemappedPackageName);
+			}
+		}
+
+		return RemappedPackageName;
 	}
 
 	/**
@@ -126,6 +138,12 @@ public:
 		InstancedPackageMap.AddPackageMapping(Original, Instanced);
 	}
 
+	/** Add a mapping function from a package name to a new package name. This function should be thread-safe, is it can be invoked from ALT. */
+	void AddPackageMappingFunc(TFunction<FName(FName)> InInstancedPackageMapFunc)
+	{
+		InstancedPackageMapFunc = MoveTemp(InInstancedPackageMapFunc);
+	}
+
 	/** Add a mapping from a top level asset path (/Path/To/Package.AssetName) to another. */
 	void AddPathMapping(FSoftObjectPath Original, FSoftObjectPath Instanced)
 	{
@@ -194,7 +212,9 @@ private:
 	friend class FLinkerInstancingContextTests;
 
 	/** Map of original package name to their instance counterpart. */
-	FLinkerInstancedPackageMap InstancedPackageMap;
+	mutable FLinkerInstancedPackageMap InstancedPackageMap;
+	/** Optional function to map original package name to their instance counterpart. The result of this function should be immutable, as it will be cached. */
+	TFunction<FName(FName)> InstancedPackageMapFunc;
 	/** Map of original top level asset path to their instance counterpart. */
 	TMap<FTopLevelAssetPath, FTopLevelAssetPath> PathMapping;
 	/** Tags can be used to determine some loading behavior. */
