Added extra checks when queuing up Slow ARO call to help track down CI issue

#rnx

diff --git a/Engine/Source/Runtime/CoreUObject/Private/UObject/GarbageCollection.cpp b/Engine/Source/Runtime/CoreUObject/Private/UObject/GarbageCollection.cpp
--- a/Engine/Source/Runtime/CoreUObject/Private/UObject/GarbageCollection.cpp
+++ b/Engine/Source/Runtime/CoreUObject/Private/UObject/GarbageCollection.cpp
@@ -2651,6 +2651,15 @@ void FSlowARO::CallSync(uint32 SlowAROIndex, UObject* Object, FReferenceCollecto
 
 bool FSlowARO::TryQueueCall(uint32 SlowAROIndex, UObject* Object, FWorkerContext& Context)
 {
+	check(Object);
+	check(Object->GetClass());
+	checkf(SlowAROIndex < uint32(GSlowARO.GetPostInit().NumAROs()), TEXT("SlowAROIndex out of bounds %d/%d"), SlowAROIndex, GSlowARO.GetPostInit().NumAROs());
+	checkf(SlowAROIndex == GSlowARO.GetPostInit().FindImplementation(Object->GetClass()->CppClassStaticFunctions.GetAddReferencedObjects()),
+		TEXT("Queueing a '%s' named '%s' for slow AddReferenceObjects call, but the class or ARO doesn't match. "
+			 "Slow ARO index from schema is %d but found index %d."),
+		*Object->GetClass()->GetFName().ToString(), *Object->GetFName().ToString(),
+		SlowAROIndex, GSlowARO.GetPostInit().FindImplementation(Object->GetClass()->CppClassStaticFunctions.GetAddReferencedObjects()));
+
 	return GSlowARO.GetPostInit().QueueCall(SlowAROIndex, Context.GetWorkerIndex(), Object);
 }
 
