Fix zenloader state progression when mixing cooked and uncooked assets

#rnx

diff --git a/Engine/Source/Runtime/CoreUObject/Private/Serialization/AsyncLoading2.cpp b/Engine/Source/Runtime/CoreUObject/Private/Serialization/AsyncLoading2.cpp
--- a/Engine/Source/Runtime/CoreUObject/Private/Serialization/AsyncLoading2.cpp
+++ b/Engine/Source/Runtime/CoreUObject/Private/Serialization/AsyncLoading2.cpp
@@ -5462,12 +5462,11 @@ EEventLoadNodeExecutionResult FAsyncPackage2::Event_CreateLinkerLoadExports(FAsy
 	UE::Core::Private::FPlayInEditorLoadingScope PlayInEditorIDScope(Package->Desc.PIEInstanceID);
 #endif
 
-	if (Package->LinkerLoadState.IsSet())
+	check(Package->LinkerLoadState.IsSet());
+
+	if (!Package->CreateLinkerLoadExports(ThreadState))
 	{
-		if (!Package->CreateLinkerLoadExports(ThreadState))
-		{
-			return EEventLoadNodeExecutionResult::Timeout;
-		}
+		return EEventLoadNodeExecutionResult::Timeout;
 	}
 
 	Package->AsyncPackageLoadingState = EAsyncPackageLoadingState2::WaitingForLinkerLoadDependencies;
@@ -5482,8 +5481,18 @@ void FAsyncPackage2::ConditionalBeginResolveLinkerLoadImports(FAsyncLoadingThrea
 	WaitForAllDependenciesToReachState(ThreadState, &FAsyncPackage2::AllDependenciesImportState, EAsyncPackageLoadingState2::WaitingForLinkerLoadDependencies, AsyncLoadingThread.ConditionalBeginResolveImportsTick,
 		[&ThreadState](FAsyncPackage2* Package)
 		{
-			check (Package->LinkerLoadState.IsSet())
-			Package->GetPackageNode(Package_ResolveLinkerLoadImports).ReleaseBarrier(&ThreadState);
+			if (Package->LinkerLoadState.IsSet())
+			{
+				check(Package->AsyncPackageLoadingState == EAsyncPackageLoadingState2::WaitingForLinkerLoadDependencies);
+				Package->AsyncPackageLoadingState = EAsyncPackageLoadingState2::ResolveLinkerLoadImports;
+				Package->GetPackageNode(Package_ResolveLinkerLoadImports).ReleaseBarrier(&ThreadState);
+			}
+			else
+			{
+				// Don't advance state of cooked package, nodes already have dependencies setup.
+				// Just validate that the dependency is doing its job and that the zen package we import have reached a proper state.
+				check(Package->AsyncPackageLoadingState >= EAsyncPackageLoadingState2::ExportsDone);
+			}
 		}
 	);
 }
@@ -5802,15 +5811,14 @@ EEventLoadNodeExecutionResult FAsyncPackage2::Event_ProcessExportBundle(FAsyncLo
 			Package->OptionalSegmentSerializationState->ReleaseIoRequest();
 		}
 #endif
+		check(Package->AsyncPackageLoadingState == EAsyncPackageLoadingState2::ProcessExportBundles);
 
 		if (Package->ExternalReadDependencies.Num() == 0)
 		{
-			check(Package->AsyncPackageLoadingState == EAsyncPackageLoadingState2::ProcessExportBundles);
 			Package->GetPackageNode(Package_ExportsSerialized).ReleaseBarrier(&ThreadState);
 		}
 		else
 		{
-			check(Package->AsyncPackageLoadingState == EAsyncPackageLoadingState2::ProcessExportBundles);
 			Package->AsyncPackageLoadingState = EAsyncPackageLoadingState2::WaitingForExternalReads;
 			Package->AsyncLoadingThread.ExternalReadQueue.Enqueue(Package);
 		}
@@ -6587,9 +6595,7 @@ FAsyncPackage2* FAsyncPackage2::UpdateDependenciesStateRecursive(FAsyncLoadingTh
 			}
 			else
 			{
-				check(InStronglyConnectedComponent->AsyncPackageLoadingState == Context.WaitForPackageState);
 				InStronglyConnectedComponentState.bAllDone = true;
-				InStronglyConnectedComponent->AsyncPackageLoadingState = static_cast<EAsyncPackageLoadingState2>(static_cast<uint32>(Context.WaitForPackageState) + 1);
 #if ALT2_ENABLE_PACKAGE_DEPENDENCY_DEBUGGING
 				InStronglyConnectedComponent->CheckThatAllDependenciesHaveReachedStateDebug(ThreadState, InStronglyConnectedComponent->AsyncPackageLoadingState, Context.WaitForPackageState);
 #endif
@@ -6634,14 +6640,18 @@ void FAsyncPackage2::ConditionalBeginProcessPackageExports(FAsyncLoadingThreadSt
 	WaitForAllDependenciesToReachState(ThreadState, &FAsyncPackage2::AllDependenciesSetupState, EAsyncPackageLoadingState2::DependenciesReady, AsyncLoadingThread.ConditionalBeginProcessExportsTick,
 		[&ThreadState](FAsyncPackage2* Package)
 		{
+			check(Package->AsyncPackageLoadingState == EAsyncPackageLoadingState2::DependenciesReady);
+
 #if ALT2_ENABLE_LINKERLOAD_SUPPORT
 			if (Package->LinkerLoadState.IsSet())
 			{
+				Package->AsyncPackageLoadingState = EAsyncPackageLoadingState2::CreateLinkerLoadExports;
 				Package->GetPackageNode(EEventLoadNode2::Package_CreateLinkerLoadExports).ReleaseBarrier(&ThreadState);
 			}
 			else
 #endif
 			{
+				Package->AsyncPackageLoadingState = EAsyncPackageLoadingState2::ProcessExportBundles;
 				for (int32 ExportBundleIndex = 0; ExportBundleIndex < Package->Data.TotalExportBundleCount; ++ExportBundleIndex)
 				{
 					Package->GetExportBundleNode(EEventLoadNode2::ExportBundle_Process, ExportBundleIndex).ReleaseBarrier(&ThreadState);
@@ -6656,6 +6666,8 @@ void FAsyncPackage2::ConditionalFinishLoading(FAsyncLoadingThreadState2& ThreadS
 	WaitForAllDependenciesToReachState(ThreadState, &FAsyncPackage2::AllDependenciesFullyLoadedState, EAsyncPackageLoadingState2::DeferredPostLoadDone, AsyncLoadingThread.ConditionalFinishLoadingTick,
 		[](FAsyncPackage2* Package)
 		{
+			check(Package->AsyncPackageLoadingState == EAsyncPackageLoadingState2::DeferredPostLoadDone);
+			Package->AsyncPackageLoadingState = EAsyncPackageLoadingState2::Finalize;
 			Package->AsyncLoadingThread.LoadedPackagesToProcess.Add(Package);
 		});
 }
