Move FOptionalProperty to CoreUObject.
Credit goes to Maxime.Mercier for a large part of the work on this.

diff --git a/Engine/Source/Runtime/Core/Public/UObject/UnrealNames.inl b/Engine/Source/Runtime/Core/Public/UObject/UnrealNames.inl
--- a/Engine/Source/Runtime/Core/Public/UObject/UnrealNames.inl
+++ b/Engine/Source/Runtime/Core/Public/UObject/UnrealNames.inl
@@ -50,6 +50,7 @@ REGISTER_NAME(33,CoreUObject)
 
 // More class properties
 REGISTER_NAME(34,EnumProperty)
+REGISTER_NAME(35,OptionalProperty)
 
 
 // Special types.
diff --git a/Engine/Source/Runtime/CoreUObject/Private/Internationalization/StabilizeLocalizationKeys.cpp b/Engine/Source/Runtime/CoreUObject/Private/Internationalization/StabilizeLocalizationKeys.cpp
--- a/Engine/Source/Runtime/CoreUObject/Private/Internationalization/StabilizeLocalizationKeys.cpp
+++ b/Engine/Source/Runtime/CoreUObject/Private/Internationalization/StabilizeLocalizationKeys.cpp
@@ -1,6 +1,7 @@
 // Copyright Epic Games, Inc. All Rights Reserved.
 
 #include "Internationalization/StabilizeLocalizationKeys.h"
+#include "UObject/PropertyOptional.h"
 #include "UObject/TextProperty.h"
 
 #if WITH_EDITOR
@@ -178,6 +179,24 @@ void StabilizeLocalizationKeys::StabilizeLocalizationKeysForProperty(FProperty*
 		}
 		return;
 	}
+
+	if (FOptionalProperty* OptionalProp = CastField<FOptionalProperty>(InProp))
+	{
+		FProperty* OptionalValueProp = OptionalProp->GetValueProperty();
+		if (ShouldStabilizeInnerProperty(OptionalValueProp))
+		{
+			if (void* ValueData = OptionalProp->GetValuePointerForReplaceIfSet(InPropData))
+			{
+				StabilizeLocalizationKeysForProperty(
+					OptionalValueProp,
+					ValueData,
+					InNamespace,
+					PropKeyRoot,
+					/*bAppendPropertyNameToKey*/false
+					);
+			}
+		}
+	}
 }
 
 void StabilizeLocalizationKeys::StabilizeLocalizationKeysForStruct(UStruct* InStruct, void* InStructData, const FString& InNamespace, const FString& InKeyRoot)
diff --git a/Engine/Source/Runtime/CoreUObject/Private/Serialization/ArchiveStackTrace.cpp b/Engine/Source/Runtime/CoreUObject/Private/Serialization/ArchiveStackTrace.cpp
--- a/Engine/Source/Runtime/CoreUObject/Private/Serialization/ArchiveStackTrace.cpp
+++ b/Engine/Source/Runtime/CoreUObject/Private/Serialization/ArchiveStackTrace.cpp
@@ -23,6 +23,7 @@
 #include "Misc/ScopeExit.h"
 #include "Compression/CompressionUtil.h"
 #include "Serialization/ZenPackageHeader.h"
+#include "UObject/PropertyOptional.h"
 
 DEFINE_LOG_CATEGORY_STATIC(LogArchiveDiff, Log, All);
 
@@ -661,6 +662,11 @@ namespace
 			}
 		}
 
+		if (FOptionalProperty* OptionalProp = CastField<FOptionalProperty>(Prop))
+		{
+			return ShouldDumpPropertyValueState(OptionalProp->GetValueProperty());
+		}
+
 		return false;
 	}
 }
diff --git a/Engine/Source/Runtime/CoreUObject/Private/Serialization/PropertyLocalizationDataGathering.cpp b/Engine/Source/Runtime/CoreUObject/Private/Serialization/PropertyLocalizationDataGathering.cpp
--- a/Engine/Source/Runtime/CoreUObject/Private/Serialization/PropertyLocalizationDataGathering.cpp
+++ b/Engine/Source/Runtime/CoreUObject/Private/Serialization/PropertyLocalizationDataGathering.cpp
@@ -9,6 +9,7 @@
 #include "UObject/UnrealType.h"
 #include "UObject/TextProperty.h"
 #include "UObject/PropertyPortFlags.h"
+#include "UObject/PropertyOptional.h"
 #include "HAL/UnrealMemory.h"
 #include "Internationalization/TextNamespaceUtil.h"
 #include "Internationalization/TextPackageNamespaceUtil.h"
@@ -146,6 +147,10 @@ const FPropertyLocalizationDataGatherer::FGatherableFieldsForType& FPropertyLoca
 				{
 					ProcessInnerProperty(SetProp->ElementProp, PropertyField);
 				}
+				if (const FOptionalProperty* OptionalProp = CastField<const FOptionalProperty>(PropertyField))
+				{
+					ProcessInnerProperty(OptionalProp->GetValueProperty(), PropertyField);
+				}
 			}
 		}
 	}
@@ -343,6 +348,7 @@ void FPropertyLocalizationDataGatherer::GatherLocalizationDataFromChildTextPrope
 	const FSetProperty* const SetProperty = CastField<const FSetProperty>(Property);
 	const FStructProperty* const StructProperty = CastField<const FStructProperty>(Property);
 	const FObjectPropertyBase* const ObjectProperty = CastField<const FObjectPropertyBase>(Property);
+	const FOptionalProperty* const OptionalProperty = CastField<const FOptionalProperty>(Property);
 
 	const EPropertyLocalizationGathererTextFlags FixedChildPropertyGatherTextFlags = GatherTextFlags | (Property->HasAnyPropertyFlags(CPF_EditorOnly) ? EPropertyLocalizationGathererTextFlags::ForceEditorOnly : EPropertyLocalizationGathererTextFlags::None);
 
@@ -538,6 +544,18 @@ void FPropertyLocalizationDataGatherer::GatherLocalizationDataFromChildTextPrope
 				GatherLocalizationDataFromObjectWithCallbacks(InnerObject, FixedChildPropertyGatherTextFlags);
 			}
 		}
+		else if (OptionalProperty)
+		{
+			if (const void* OptionalValueAddress = OptionalProperty->GetValuePointerForReadIfSet(ElementValueAddress))
+			{
+				GatherLocalizationDataFromChildTextProperties(
+					FString::Printf(TEXT("%s?"), *PathToElement),
+					OptionalProperty->GetValueProperty(),
+					OptionalValueAddress,
+					nullptr,
+					ElementChildPropertyGatherTextFlags);
+			}
+		}
 	}
 }
 
diff --git a/Engine/Source/Runtime/CoreUObject/Private/UObject/Class.cpp b/Engine/Source/Runtime/CoreUObject/Private/UObject/Class.cpp
--- a/Engine/Source/Runtime/CoreUObject/Private/UObject/Class.cpp
+++ b/Engine/Source/Runtime/CoreUObject/Private/UObject/Class.cpp
@@ -35,6 +35,7 @@
 #include "UObject/Interface.h"
 #include "UObject/LinkerPlaceholderClass.h"
 #include "UObject/LinkerPlaceholderFunction.h"
+#include "UObject/PropertyOptional.h"
 #include "UObject/StructOnScope.h"
 #include "UObject/StructScriptLoader.h"
 #include "UObject/PropertyHelper.h"
@@ -2420,10 +2421,11 @@ bool FindConstructorUninitialized(UStruct* BaseClass,uint8* Data,uint8* Defaults
 		int32 Size = P->GetSize();
 		bool bProblem = false;
 		check(Size);
-		FBoolProperty*   PB     = CastField<FBoolProperty>(P);
-		FStructProperty* PS     = CastField<FStructProperty>(P);
-		FStrProperty*    PStr   = CastField<FStrProperty>(P);
-		FArrayProperty*  PArray = CastField<FArrayProperty>(P);
+		FBoolProperty*     PB        = CastField<FBoolProperty>(P);
+		FStructProperty*   PS        = CastField<FStructProperty>(P);
+		FStrProperty*      PStr      = CastField<FStrProperty>(P);
+		FArrayProperty*    PArray    = CastField<FArrayProperty>(P);
+		FOptionalProperty* POptional = CastField<FOptionalProperty>(P);
 		if(PStr)
 		{
 			// string that actually have data would be false positives, since they would point to the same string, but actually be different pointers
@@ -2456,6 +2458,10 @@ bool FindConstructorUninitialized(UStruct* BaseClass,uint8* Data,uint8* Defaults
 		{
 			bProblem = !PArray->Identical_InContainer(Data, Defaults);
 		}
+		else if (POptional)
+		{
+			bProblem = !POptional->Identical_InContainer(Data, Defaults);
+		}
 		else
 		{
 			if (FMemory::Memcmp(P->ContainerPtrToValuePtr<uint8>(Data), P->ContainerPtrToValuePtr<uint8>(Defaults), Size) != 0)
diff --git a/Engine/Source/Runtime/CoreUObject/Private/UObject/CoreNet.cpp b/Engine/Source/Runtime/CoreUObject/Private/UObject/CoreNet.cpp
--- a/Engine/Source/Runtime/CoreUObject/Private/UObject/CoreNet.cpp
+++ b/Engine/Source/Runtime/CoreUObject/Private/UObject/CoreNet.cpp
@@ -5,6 +5,7 @@
 =============================================================================*/
 
 #include "UObject/CoreNet.h"
+#include "UObject/PropertyOptional.h"
 #include "UObject/UnrealType.h"
 #include "Misc/EngineNetworkCustomVersion.h"
 
@@ -109,6 +110,12 @@ uint32 FClassNetCacheMgr::GetPropertyChecksum( const FProperty* Property, uint32
 
 			const_cast< FClassNetCacheMgr* >( this )->DebugChecksumIndent--;
 		}
+
+		// Evolve checksum on optional inner
+		if (const FOptionalProperty* OptionalProperty = CastField<FOptionalProperty>(Property))
+		{
+			return GetPropertyChecksum(OptionalProperty->GetValueProperty(), Checksum, bIncludeChildren);
+		}
 	}
 
 	return Checksum;
diff --git a/Engine/Source/Runtime/CoreUObject/Private/UObject/GarbageCollection.cpp b/Engine/Source/Runtime/CoreUObject/Private/UObject/GarbageCollection.cpp
--- a/Engine/Source/Runtime/CoreUObject/Private/UObject/GarbageCollection.cpp
+++ b/Engine/Source/Runtime/CoreUObject/Private/UObject/GarbageCollection.cpp
@@ -38,6 +38,7 @@
 #include "UObject/FieldPathProperty.h"
 #include "UObject/GarbageCollectionHistory.h"
 #include "UObject/GarbageCollectionTesting.h"
+#include "UObject/PropertyOptional.h"
 
 #include <atomic>
 
@@ -5048,6 +5049,24 @@ bool FMulticastDelegateProperty::ContainsObjectReference(TArray<const FStructPro
 	return !!(InReferenceType & EPropertyObjectReferenceType::Weak);
 }
 
+void FOptionalProperty::EmitReferenceInfo(UE::GC::FSchemaBuilder& Schema, int32 BaseOffset, TArray<const FStructProperty*>& EncounteredStructProps, UE::GC::FPropertyStack& DebugPath)
+{
+	using namespace UE::GC;
+	if (IsValueNonNullablePointer())
+	{
+		ValueProperty->EmitReferenceInfo(Schema, BaseOffset + GetOffset_ForGC(), EncounteredStructProps, DebugPath);
+	}
+	else if (ValueProperty->ContainsObjectReference(EncounteredStructProps))
+	{
+		FSchemaBuilder InnerSchema(ValueProperty->GetSize());
+		{
+			FPropertyStackScope PropertyScope(DebugPath, ValueProperty);
+			ValueProperty->EmitReferenceInfo(InnerSchema, 0, EncounteredStructProps, DebugPath);
+		}
+		Schema.Add(DeclareMember(DebugPath, BaseOffset + GetOffset_ForGC(), EMemberType::Optional, InnerSchema.Build()));
+	}
+}
+
 void FProperty::EmitReferenceInfo(UE::GC::FSchemaBuilder& Schema, int32 BaseOffset, TArray<const FStructProperty*>& EncounteredStructProps, UE::GC::FPropertyStack& DebugPath)
 {
 }
@@ -5185,6 +5204,11 @@ void FFieldPathProperty::EmitReferenceInfo(UE::GC::FSchemaBuilder& Schema, int32
 	}
 }
 
+bool FOptionalProperty::ContainsObjectReference(TArray<const FStructProperty*>& EncounteredStructProps, EPropertyObjectReferenceType InReferenceType) const
+{
+	return ValueProperty->ContainsObjectReference(EncounteredStructProps, InReferenceType);
+}
+
 /** Both game thread and async loading can assemble schemas */
 static FCriticalSection GAssembleSchemaLock;
 
diff --git a/Engine/Source/Runtime/CoreUObject/Private/UObject/LinkerPlaceholderBase.cpp b/Engine/Source/Runtime/CoreUObject/Private/UObject/LinkerPlaceholderBase.cpp
--- a/Engine/Source/Runtime/CoreUObject/Private/UObject/LinkerPlaceholderBase.cpp
+++ b/Engine/Source/Runtime/CoreUObject/Private/UObject/LinkerPlaceholderBase.cpp
@@ -3,6 +3,7 @@
 #include "UObject/LinkerPlaceholderBase.h"
 
 #include "UObject/LinkerPlaceholderExportObject.h"
+#include "UObject/PropertyOptional.h"
 #include "UObject/UnrealType.h"
 #include "UObject/UnrealTypePrivate.h"
 #include "Blueprint/BlueprintSupport.h"
@@ -186,6 +187,22 @@ int32 FLinkerPlaceholderObjectImpl::ResolvePlaceholderValues(const TArray<FField
 			// With FProperties this should never happen
 			check(false);
 		}
+		else if (const FOptionalProperty* OptionalProperty = Property.Get<FOptionalProperty>())
+		{
+#if USE_DEFERRED_DEPENDENCY_CHECK_VERIFICATION_TESTS
+			const FProperty* NextProperty = PropertyChain[PropertyIndex - 1].Get<FProperty>();
+			check(NextProperty == OptionalProperty->GetValueProperty());
+#endif // USE_DEFERRED_DEPENDENCY_CHECK_VERIFICATION_TESTS
+
+			if (void* SetValueAddress = OptionalProperty->GetValuePointerForReplaceIfSet(ValueAddress))
+			{
+				ReplacementCount += ResolvePlaceholderValues(PropertyChain, PropertyIndex - 1, static_cast<uint8*>(SetValueAddress), OldValue, ReplacementValue);
+			}
+
+			// the above recursive call chewed through the rest of the
+			// PropertyChain, no need to keep on here
+			break;
+		}
 		else if (const FProperty* NextProperty = PropertyChain[PropertyIndex - 1].Get<FProperty>())
 		{
 			ValueAddress = NextProperty->ContainerPtrToValuePtr<uint8>(ValueAddress, /*ArrayIndex =*/0);
diff --git a/Engine/Source/Runtime/CoreUObject/Private/UObject/PropertyOptional.cpp b/Engine/Source/Runtime/CoreUObject/Private/UObject/PropertyOptional.cpp
new file mode 100644
--- /dev/null
+++ b/Engine/Source/Runtime/CoreUObject/Private/UObject/PropertyOptional.cpp
@@ -0,0 +1,478 @@
+// Copyright Epic Games, Inc. All Rights Reserved.
+
+#include "UObject\PropertyOptional.h"
+#include "Misc/Guid.h"
+#include "Serialization/CustomVersion.h"
+#include "String/LexFromString.h"
+#include "UObject/GarbageCollectionSchema.h"
+#include "UObject/PropertyHelper.h"
+
+DEFINE_LOG_CATEGORY_STATIC(LogOptionalProperty, Log, All);
+
+IMPLEMENT_FIELD(FOptionalProperty)
+
+// Custom serialization version for the legacy FOptionProperty
+// Note that this is deprecated, and new versions should not be added.
+// It's only used to support loading old versions of serialized data.
+struct FLegacyFOptionPropertyCustomVersion
+{
+	enum Type
+	{
+		BeforeCustomVersionWasAdded = 0,
+		ConditionalSerializeValue = 1,
+		AlwaysSavingIsSetForFObjectProperty = 2,
+		RemoveIsValueNonNullablePointerHack = 3
+	};
+
+	const FGuid GUID;
+	FCustomVersionRegistration Registration;
+		
+	FLegacyFOptionPropertyCustomVersion()
+	: GUID{0x68C409FC, 0x70954986, 0x8963ACD2, 0xC4865183}
+	, Registration(GUID, RemoveIsValueNonNullablePointerHack, TEXT("FPropertyOption"))
+	{}
+};
+static FLegacyFOptionPropertyCustomVersion LegacyFOptionPropertyCustomVersion;
+
+FOptionalProperty::FOptionalProperty(FFieldVariant InOwner, const FName& InName, EObjectFlags InObjectFlags)
+	: FProperty(InOwner, InName, InObjectFlags)
+{
+}
+
+FOptionalProperty::~FOptionalProperty()
+{
+	if (ValueProperty)
+	{
+		delete ValueProperty;
+		ValueProperty = nullptr;
+	}
+}
+
+void FOptionalProperty::SetValueProperty(FProperty* InValueProperty)
+{
+	check(!ValueProperty);
+	check(InValueProperty);
+	ValueProperty = InValueProperty;
+}
+
+void FOptionalProperty::AddReferencedObjects(FReferenceCollector& Collector)
+{
+	Super::AddReferencedObjects(Collector);
+	if (ValueProperty)
+	{
+		ValueProperty->AddReferencedObjects(Collector);
+	}
+}
+
+void FOptionalProperty::Serialize(FArchive& Ar)
+{
+	Super::Serialize(Ar);
+	SerializeSingleField(Ar, ValueProperty, this);
+	checkSlow(ValueProperty);
+}
+
+void FOptionalProperty::GetPreloadDependencies(TArray<UObject*>& OutDeps)
+{
+	Super::GetPreloadDependencies(OutDeps);
+	checkSlow(ValueProperty);
+	ValueProperty->GetPreloadDependencies(OutDeps);
+}
+
+void FOptionalProperty::PostDuplicate(const FField& InField)
+{
+	const FOptionalProperty& Source = static_cast<const FOptionalProperty&>(InField);
+	ValueProperty = CastFieldChecked<FProperty>(FField::Duplicate(Source.ValueProperty, this));
+	Super::PostDuplicate(InField);
+}
+
+FField* FOptionalProperty::GetInnerFieldByName(const FName& InName)
+{
+	checkSlow(ValueProperty);
+	if (ValueProperty->GetFName() == InName)
+	{
+		return ValueProperty;
+	}
+	return nullptr;
+}
+
+void FOptionalProperty::GetInnerFields(TArray<FField*>& OutFields)
+{
+	checkSlow(ValueProperty);
+	OutFields.Add(ValueProperty);
+	ValueProperty->GetInnerFields(OutFields);
+}
+
+FString FOptionalProperty::GetCPPType(FString* ExtendedTypeText, uint32 CPPExportFlags) const
+{
+	checkSlow(ValueProperty);
+	if (ExtendedTypeText != nullptr)
+	{
+		FString ValueExtendedTypeText;
+		FString ValueTypeText = ValueProperty->GetCPPType(&ValueExtendedTypeText, CPPExportFlags & ~CPPF_ArgumentOrReturnValue); // we won't consider the optional's inner value type to be "arguments or return values"
+		*ExtendedTypeText = FString::Printf(TEXT("<%s%s>"), *ValueTypeText, *ValueExtendedTypeText);
+	}
+	return TEXT("TOptional");
+}
+
+FString FOptionalProperty::GetCPPMacroType(FString& ExtendedTypeText) const
+{
+	checkSlow(ValueProperty);
+	ExtendedTypeText = ValueProperty->GetCPPType();
+	return TEXT("TOPTIONAL");
+}
+
+FString FOptionalProperty::GetCPPTypeForwardDeclaration() const
+{
+	// We assume that TOptional<> is globally known already and that we just need to make the value type known
+	checkSlow(ValueProperty);
+	return ValueProperty->GetCPPTypeForwardDeclaration();
+}
+
+void FOptionalProperty::LinkInternal(FArchive& Ar)
+{
+	checkSlow(ValueProperty);
+	ValueProperty->Link(Ar);
+	
+	// After ValueProperty's size has been computed, compute the size of this property.
+	ElementSize = CalcSize();
+
+	// Propagate CPF_ZeroConstructor, CPF_NoDestructor, CPF_IsPlainOldData, and CPF_HasGetValueTypeHash from the value property.
+	PropertyFlags |= (ValueProperty->PropertyFlags & (CPF_ZeroConstructor|CPF_NoDestructor|CPF_IsPlainOldData|CPF_HasGetValueTypeHash));
+	if (ValueProperty->ContainsInstancedObjectProperty())
+	{
+		PropertyFlags |= CPF_ContainsInstancedReference;
+	}
+}
+
+bool FOptionalProperty::Identical(const void* A, const void* B, uint32 PortFlags) const
+{
+	checkSlow(A && ValueProperty);
+
+	if (B == nullptr)
+	{
+		return !IsSet(A);
+	}
+
+	const bool bIsSetA = IsSet(A);
+	const bool bIsSetB = IsSet(B);
+	if (bIsSetA != bIsSetB)
+	{
+		return false;
+	}
+	else if (!bIsSetA)
+	{
+		return true;
+	}
+
+	return ValueProperty->Identical(GetValuePointerForRead(A), GetValuePointerForRead(B), PortFlags);
+}
+
+void FOptionalProperty::SerializeItem(FStructuredArchive::FSlot Slot, void* Data, void const* Defaults) const
+{
+	FArchive& UnderlyingArchive = Slot.GetUnderlyingArchive();
+	const bool bIsLoading = UnderlyingArchive.IsLoading();
+
+	FStructuredArchive::FRecord Record = Slot.EnterRecord();
+
+	// Use an optional field slot to encode whether the optional was set or not.
+	TOptional<FStructuredArchiveSlot> MaybeTaggedValueSlot = Record.TryEnterField(TEXT("TaggedValue"), IsSet(Data));
+	if (MaybeTaggedValueSlot.IsSet())
+	{
+		FStructuredArchive::FRecord TaggedValueRecord = MaybeTaggedValueSlot.GetValue().EnterRecord();
+		
+		const void* ValueDefaults = Defaults ? GetValuePointerForReadIfSet(Defaults) : nullptr;
+
+		if (bIsLoading)
+		{
+			// Serialize the value's tag.
+			FPropertyTag ValueTag;
+			TaggedValueRecord << SA_VALUE(TEXT("Tag"), ValueTag);
+
+			// Deserialize/convert the value.
+			void* ValueData = MarkSetAndGetInitializedValuePointerToReplace(Data);
+
+			const int64 ValueStartOffset = UnderlyingArchive.Tell();
+			FStructuredArchive::FSlot ValueSlot = TaggedValueRecord.EnterField(TEXT("Value"));
+			bool bSuccessfullyDeserialized = false;
+			switch (GetValueProperty()->ConvertFromType(ValueTag, ValueSlot, static_cast<uint8*>(ValueData), GetOwnerStruct(), static_cast<const uint8*>(ValueDefaults)))
+			{
+				case EConvertFromTypeResult::Converted:
+					bSuccessfullyDeserialized = true;
+					break;
+				case EConvertFromTypeResult::UseSerializeItem:
+					GetValueProperty()->SerializeItem(ValueSlot, ValueData, ValueDefaults);
+					bSuccessfullyDeserialized = !UnderlyingArchive.IsCriticalError();
+					break;
+				case EConvertFromTypeResult::CannotConvert:
+					break;
+				default:
+					check(false);
+			}
+
+			if (!bSuccessfullyDeserialized)
+			{
+				UnderlyingArchive.Seek(ValueStartOffset + ValueTag.Size);
+				
+				// If the deserialization or conversion of the value failed, behave as if an unset optional value was deserialized.
+				MarkUnset(Data);
+			}
+			else
+			{
+				const int64 NumLoadedBytes = UnderlyingArchive.Tell() - ValueStartOffset;
+				checkf(
+					ValueTag.Size == NumLoadedBytes,
+					TEXT("Value tagged with ID %s was saved as %i bytes, but loading it as a %s loaded %i bytes."),
+					*ValueTag.Type.ToString(),
+					ValueTag.Size,
+					*GetValueProperty()->GetClass()->GetFName().ToString(),
+					NumLoadedBytes);
+			}
+		}
+		else
+		{
+			void* ValueData = GetValuePointerForReadOrReplace(Data);
+
+			// Construct and serialize a tag for the value.
+			FPropertyTag ValueTag(UnderlyingArchive, GetValueProperty(), 0, static_cast<uint8*>(ValueData), static_cast<const uint8*>(ValueDefaults));
+			TaggedValueRecord << SA_VALUE(TEXT("Tag"), ValueTag);
+
+			// Serialize the value.
+			int64 ValueStartOffset = UnderlyingArchive.Tell();
+			ValueTag.SerializeTaggedProperty(TaggedValueRecord.EnterField(TEXT("Value")), GetValueProperty(), static_cast<uint8*>(ValueData), static_cast<const uint8*>(ValueDefaults));
+			
+			// If saving to a non-text archive, save the size of the serialized value so it can be skipped over on load if it's the wrong type.
+			const int64 ValueEndOffset = UnderlyingArchive.Tell();
+			ValueTag.Size = ValueEndOffset - ValueStartOffset;
+			if (ValueTag.Size > 0 && !UnderlyingArchive.IsTextFormat())
+			{
+				// The tag serialization set ValueTag.SizeOffset to the archive offset where a placeholder size field was serialized.
+				// Now that we know the size of the serialized value, seek back to the size field and serialize the true value over the placeholder.
+				UnderlyingArchive.Seek(ValueTag.SizeOffset);
+				UnderlyingArchive << ValueTag.Size;
+				UnderlyingArchive.Seek(ValueEndOffset);
+			}
+		}
+	}
+	else if (bIsLoading)
+	{
+		// If the field wasn't present, the serialized optional was unset.
+		MarkUnset(Data);
+	}
+}
+
+bool FOptionalProperty::NetSerializeItem(FArchive& Ar, UPackageMap* Map, void* Data, TArray<uint8>* MetaData /*= nullptr*/) const
+{
+	UE_LOG(LogOptionalProperty, Fatal, TEXT("Unsupported code path"));
+	return true;
+}
+
+void FOptionalProperty::ExportText_Internal(FString& ValueStr, const void* ContainerOrPropertyPtr, EPropertyPointerType PropertyPointerType, const void* DefaultValue, UObject* Parent, int32 PortFlags, UObject* ExportRootScope) const
+{
+	checkSlow(ValueProperty);
+	void* OptionalValuePointer = PointerToValuePtr(ContainerOrPropertyPtr, PropertyPointerType);
+	// Output (Value) or ()
+	ValueStr += TEXT("(");
+	if (const void* ValuePointer = GetValuePointerForReadIfSet(OptionalValuePointer))
+	{
+		ValueProperty->ExportTextItem_Direct(ValueStr, ValuePointer, DefaultValue ? GetValuePointerForReadIfSet(DefaultValue) : nullptr, Parent, PortFlags, ExportRootScope);
+	}
+	ValueStr += TEXT(")");
+}
+
+const TCHAR* FOptionalProperty::ImportText_Internal(const TCHAR* Buffer, void* ContainerOrPropertyPtr, EPropertyPointerType PropertyPointerType, UObject* Parent, int32 PortFlags, FOutputDevice* ErrorText) const
+{
+	checkSlow(ValueProperty);
+	void* Data = PointerToValuePtr(ContainerOrPropertyPtr, PropertyPointerType);
+
+	if (!Buffer || *Buffer++ != TCHAR('('))
+	{
+		return nullptr;
+	}
+	
+	SkipWhitespace(Buffer);
+
+	// Check if there is a value
+	if (*Buffer != TCHAR(')'))
+	{
+		// Yes, parse it
+		void* ValuePointer = MarkSetAndGetInitializedValuePointerToReplace(Data);
+		Buffer = ValueProperty->ImportText_Direct(Buffer, ValuePointer, Parent, PortFlags, ErrorText);
+		if (!Buffer)
+		{
+			MarkUnset(Data);
+			return nullptr;
+		}
+	}
+	else
+	{
+		MarkUnset(Data);
+	}
+	
+	SkipWhitespace(Buffer);
+
+	if (*Buffer++ != TCHAR(')'))
+	{
+		return nullptr;
+	}
+
+	return Buffer;
+}
+
+void FOptionalProperty::CopyValuesInternal(void* Dest, void const* Src, int32 Count) const
+{
+	int32 Size = GetSize();
+	for (int32 Index = 0; Index < Count; ++Index, (uint8*&)Dest += Size, (const uint8*&)Src += Size)
+	{
+		if (const void* SrcValue = GetValuePointerForReadIfSet(Src))
+		{
+			void* DestValue = MarkSetAndGetInitializedValuePointerToReplace(Dest);
+			ValueProperty->CopySingleValue(DestValue, SrcValue);
+		}
+		else
+		{
+			MarkUnset(Dest);
+		}
+	}
+}
+
+void FOptionalProperty::ClearValueInternal(void* Data) const
+{
+	MarkUnset(Data);
+}
+
+void FOptionalProperty::InitializeValueInternal(void* Data) const
+{
+	if (IsValueNonNullablePointer())
+	{
+		ValueProperty->InitializeValue(Data);
+	}
+	else
+	{
+		*GetIsSetPointer(Data) = false;
+	}
+}
+
+void FOptionalProperty::DestroyValueInternal(void* Data) const
+{
+	MarkUnset(Data);
+}
+
+void FOptionalProperty::InstanceSubobjects(void* Data, void const* DefaultData, UObject* InOwner, struct FObjectInstancingGraph* InstanceGraph)
+{
+	if (Data && IsSet(Data) && ValueProperty->ContainsInstancedObjectProperty())
+	{
+		ValueProperty->InstanceSubobjects(Data, DefaultData, InOwner, InstanceGraph);
+	}
+}
+
+int32 FOptionalProperty::GetMinAlignment() const
+{
+	return ValueProperty->GetMinAlignment();
+}
+
+EConvertFromTypeResult FOptionalProperty::ConvertFromType(const FPropertyTag& Tag, FStructuredArchive::FSlot Slot, uint8* ContainerData, UStruct* DefaultsStruct, const uint8* DefaultsContainer)
+{
+	static const FName NAME_OptionProperty("OptionProperty");
+	if (Tag.Type != NAME_OptionProperty)
+	{
+		return EConvertFromTypeResult::UseSerializeItem;
+	}
+	const int32 FPropertyOptionVersion = Slot.GetArchiveState().CustomVer(LegacyFOptionPropertyCustomVersion.GUID);
+
+	checkSlow(ValueProperty);
+
+	void* Data = ContainerPtrToValuePtr<void>(ContainerData, Tag.ArrayIndex);
+	const void* Defaults = DefaultsStruct ? ContainerPtrToValuePtrForDefaults<void>(DefaultsStruct, DefaultsContainer, Tag.ArrayIndex) : nullptr;
+	
+	bool bIsValueNonNullablePointer;
+	if (FPropertyOptionVersion < FLegacyFOptionPropertyCustomVersion::RemoveIsValueNonNullablePointerHack)
+	{
+		bIsValueNonNullablePointer =
+			  ValueProperty->IsA<FClassProperty>() ? false
+			: ValueProperty->IsA<FObjectProperty>() ? true
+			: (ValueProperty->GetPropertyFlags() & CPF_NonNullable) != 0;
+	}
+	else
+	{
+		bIsValueNonNullablePointer = IsValueNonNullablePointer();
+	}
+
+	if (bIsValueNonNullablePointer && FPropertyOptionVersion < FLegacyFOptionPropertyCustomVersion::AlwaysSavingIsSetForFObjectProperty)
+	{
+		// Code below is duplicated from FObjectProperty::SerializeItem with the exception of passing AllowNullValuesOnNonNullableProperty to PostSerializeObjectItem
+		// in order to maintain backwards compatibility with packages saved prior to non-nullable properties properly constructing new objects when failing to load
+		// valid (non-null) object values. We also don't need to make sure the duplicated code matches the most recent version of FObjectProperty::SerializeItem
+		// because we only run the duplicated code for packages created before the introduction of FLegacyFOptionPropertyCustomVersion::AlwaysSavingIsSetForFObjectProperty.
+
+		// We're also not duplicating 'if (UnderlyingArchive.IsObjectReferenceCollector())' code path from FObjectProperty::SerializeItem because reference collectors
+		// always have the latest custom version set. Sanity check below.
+		checkf(!Slot.GetArchiveState().IsObjectReferenceCollector(), TEXT("Reference collector archive (%s) serializing with an old FLegacyFOptionPropertyCustomVersion (%d)"),
+			*Slot.GetArchiveState().GetArchiveName(), FPropertyOptionVersion);
+
+		// Another sanity check since the duplicated code comes from FObjectProperty we are assuming ValueProperty is in fact an FObjectProperty or FClassProperty which shares SerializeItem code with FObjectProperty
+		// Note that we only support FObjectProperty and FClassProperty because these are the only two object referencing property types that have so far been used with FOptionalProperty. 
+		checkf(ValueProperty && (ValueProperty->GetClass() == FObjectProperty::StaticClass() || ValueProperty->GetClass() == FClassProperty::StaticClass()), 
+			TEXT("Unexpected value property type when deserializing optional non-nullable property %s: %s"),
+			*GetFullName(), ValueProperty ? *ValueProperty->GetClass()->GetName() : TEXT("None"));
+
+		FObjectProperty* ObjectValueProperty = CastFieldChecked<FObjectProperty>(ValueProperty);
+
+		// Begin duplicated code (no need to keep up to date with FObjectProperty::SerializeItem code, see comment above)
+		TObjectPtr<UObject> ObjectValuePtr = ObjectValueProperty->GetObjectPtrPropertyValue(Data);
+		check(ObjectValuePtr.IsResolved());
+		UObject* ObjectValue = UE::CoreUObject::Private::ReadObjectHandlePointerNoCheck(ObjectValuePtr.GetHandle());
+
+		Slot << ObjectValue;
+
+		TObjectPtr<UObject> CurrentValuePtr = ObjectValueProperty->GetObjectPtrPropertyValue(Data);
+		check(CurrentValuePtr.IsResolved());
+		UObject* CurrentValue = UE::CoreUObject::Private::ReadObjectHandlePointerNoCheck(CurrentValuePtr.GetHandle());
+
+		ObjectValueProperty->PostSerializeObjectItem(Slot.GetUnderlyingArchive(), Data, CurrentValue, ObjectValue, 
+			EObjectPropertyOptions::AllowNullValuesOnNonNullableProperty); // The original code was passing EObjectPropertySerializeOptions::None here
+		// End duplicated code
+	}
+	else if (FPropertyOptionVersion < FLegacyFOptionPropertyCustomVersion::ConditionalSerializeValue)
+	{
+		// Legacy: Serialize the inner FProperty whether it is set or not.
+		FStructuredArchive::FStream Stream = Slot.EnterStream();
+		bool bLoadedIsSet = false;
+		Stream.EnterElement() << bLoadedIsSet;
+		void* ValueData = MarkSetAndGetInitializedValuePointerToReplace(Data);
+		ValueProperty->SerializeItem(Stream.EnterElement(), ValueData, Defaults ? GetValuePointerForReadIfSet(Defaults) : nullptr);
+		if (!bLoadedIsSet)
+		{
+			MarkUnset(Data);
+		}
+	}
+	else
+	{
+		FStructuredArchive::FStream Stream = Slot.EnterStream();
+		bool bIsSetValue = IsSet(Data);
+		Stream.EnterElement() << bIsSetValue;
+
+		if (bIsSetValue)
+		{
+			void* ValuePointer = MarkSetAndGetInitializedValuePointerToReplace(Data);
+			ValueProperty->SerializeItem(Stream.EnterElement(), ValuePointer, Defaults ? GetValuePointerForReadIfSet(Defaults) : nullptr);
+		}
+		else
+		{
+			MarkUnset(Data);
+		}
+	}
+	return EConvertFromTypeResult::Converted;
+}
+
+uint32 FOptionalProperty::GetValueTypeHashInternal(const void* Src) const
+{
+	// If the option is set, return the value's hash, otherwise return 0.
+	if (const void* ValuePointer = GetValuePointerForReadIfSet(Src))
+	{
+		return ValueProperty->GetValueTypeHash(ValuePointer);
+	}
+	else
+	{
+		return 0;
+	}
+}
diff --git a/Engine/Source/Runtime/CoreUObject/Private/UObject/PropertyTag.cpp b/Engine/Source/Runtime/CoreUObject/Private/UObject/PropertyTag.cpp
--- a/Engine/Source/Runtime/CoreUObject/Private/UObject/PropertyTag.cpp
+++ b/Engine/Source/Runtime/CoreUObject/Private/UObject/PropertyTag.cpp
@@ -8,12 +8,13 @@
 #include "UObject/UnrealType.h"
 #include "UObject/EnumProperty.h"
 #include "UObject/BlueprintsObjectVersion.h"
+#include "UObject/PropertyOptional.h"
 
 /*-----------------------------------------------------------------------------
 FPropertyTag
 -----------------------------------------------------------------------------*/
 
-FPropertyTag::FPropertyTag( FArchive& InSaveAr, FProperty* Property, int32 InIndex, uint8* Value, uint8* Defaults )
+FPropertyTag::FPropertyTag( FArchive& InSaveAr, FProperty* Property, int32 InIndex, uint8* Value, const uint8* Defaults )
 	: Prop      (Property)
 	, Type      (Property->GetID())
 	, Name      (Property->GetFName())
@@ -72,6 +73,10 @@ FPropertyTag::FPropertyTag( FArchive& InSaveAr, FProperty* Property, int32 InInd
 		{
 			BoolVal = Bool->GetPropertyValue(Value);
 		}
+		else if (FOptionalProperty* OptionalProp = CastField<FOptionalProperty>(Property))
+		{
+			InnerType = OptionalProp->GetValueProperty()->GetID();
+		}
 	}
 }
 
@@ -176,7 +181,7 @@ void operator<<(FStructuredArchive::FSlot Slot, FPropertyTag& Tag)
 		{
 			Slot << SA_ATTRIBUTE(TEXT("EnumName"), Tag.EnumName);
 		}
-		// only need to serialize this for arrays
+		// need to serialize the InnerType for arrays
 		else if (TagType == NAME_ArrayProperty)
 		{
 			if (Version >= VAR_UE4_ARRAY_PROPERTY_INNER_TAGS)
@@ -184,6 +189,11 @@ void operator<<(FStructuredArchive::FSlot Slot, FPropertyTag& Tag)
 				Slot << SA_ATTRIBUTE(TEXT("InnerType"), Tag.InnerType);
 			}
 		}
+		// need to serialize the InnerType for optionals.
+		else if (TagType == NAME_OptionalProperty)
+		{
+			Slot << SA_ATTRIBUTE(TEXT("InnerType"), Tag.InnerType);
+		}
 		else if (Version >= VER_UE4_PROPERTY_TAG_SET_MAP_SUPPORT)
 		{
 			if (TagType == NAME_SetProperty)
@@ -218,12 +228,12 @@ void operator<<(FStructuredArchive::FSlot Slot, FPropertyTag& Tag)
 }
 
 // Property serializer.
-void FPropertyTag::SerializeTaggedProperty(FArchive& Ar, FProperty* Property, uint8* Value, uint8* Defaults) const
+void FPropertyTag::SerializeTaggedProperty(FArchive& Ar, FProperty* Property, uint8* Value, const uint8* Defaults) const
 {
 	SerializeTaggedProperty(FStructuredArchiveFromArchive(Ar).GetSlot(), Property, Value, Defaults);
 }
 
-void FPropertyTag::SerializeTaggedProperty(FStructuredArchive::FSlot Slot, FProperty* Property, uint8* Value, uint8* Defaults) const
+void FPropertyTag::SerializeTaggedProperty(FStructuredArchive::FSlot Slot, FProperty* Property, uint8* Value, const uint8* Defaults) const
 {
 	FArchive& UnderlyingArchive = Slot.GetUnderlyingArchive();
 	const int64 StartOfProperty = UnderlyingArchive.Tell();
diff --git a/Engine/Source/Runtime/CoreUObject/Private/UObject/UObjectGlobals.cpp b/Engine/Source/Runtime/CoreUObject/Private/UObject/UObjectGlobals.cpp
--- a/Engine/Source/Runtime/CoreUObject/Private/UObject/UObjectGlobals.cpp
+++ b/Engine/Source/Runtime/CoreUObject/Private/UObject/UObjectGlobals.cpp
@@ -74,6 +74,7 @@
 #include "Serialization/AsyncPackageLoader.h"
 #include "Containers/VersePath.h"
 #include "AutoRTFM/AutoRTFM.h"
+#include "UObject/PropertyOptional.h"
 
 #if UE_USE_VERSE_PATHS
 #include "Interfaces/IPluginManager.h"
@@ -4839,6 +4840,21 @@ void CollectSetReferences(FReferenceCollector& Collector, FSetProperty& Property
 	}
 }
 
+template<EPropertyCollectFlags CollectFlags>
+void CollectOptionalReference(FReferenceCollector& Collector, FOptionalProperty& Property, void* Instance, const UObject* Referencer)
+{
+	FProperty& InnerProperty = *Property.GetValueProperty();
+	EClassCastFlags InnerCastFlags = static_cast<EClassCastFlags>(InnerProperty.GetClass()->GetCastFlags());
+	if (MayContainStrongReference(InnerCastFlags))
+	{
+		if (void* ValueInstance = Property.GetValuePointerForReplaceIfSet(Instance))
+		{
+			CollectPropertyReferences<CollectFlags>(Collector, InnerProperty, ValueInstance, Referencer);			
+		}
+	}
+}
+
+
 // Process FObjectProperty or FObjectPtrProperty reference
 template<EPropertyCollectFlags CollectFlags>
 FORCEINLINE_DEBUGGABLE void CollectObjectReference(FReferenceCollector& Collector, FProperty& Property, void* Value, const UObject* Referencer)
@@ -4940,6 +4956,10 @@ void CollectPropertyReferences(FReferenceCollector& Collector, FProperty& Proper
 		{	
 			CollectInterfaceReference(Collector, static_cast<FInterfaceProperty&>(Property), *reinterpret_cast<FScriptInterface*>(Value), Referencer);
 		}
+		else if (EnumHasAnyFlags(CastFlags, CASTCLASS_FOptionalProperty))
+		{
+			CollectOptionalReference<CollectFlags>(Collector, static_cast<FOptionalProperty&>(Property), Value, Referencer);
+		}
 		else
 		{
 			// Fallback to virtual SerializeItem dispatch inside SerializeBin
diff --git a/Engine/Source/Runtime/CoreUObject/Private/UObject/UnrealType.cpp b/Engine/Source/Runtime/CoreUObject/Private/UObject/UnrealType.cpp
--- a/Engine/Source/Runtime/CoreUObject/Private/UObject/UnrealType.cpp
+++ b/Engine/Source/Runtime/CoreUObject/Private/UObject/UnrealType.cpp
@@ -1,6 +1,7 @@
 // Copyright Epic Games, Inc. All Rights Reserved.
 
 #include "UObject/UnrealType.h"
+#include "UObject/PropertyOptional.h"
 #include "Serialization/ArchiveUObjectFromStructuredArchive.h"
 
 DEFINE_LOG_CATEGORY(LogType);
@@ -35,10 +36,11 @@ FPropertyValueIterator::EPropertyValueFlags FPropertyValueIterator::GetPropertyV
 	if (RecursionFlags == EPropertyValueIteratorFlags::FullRecursion)
 	{
 		uint64 CastFlags = Property->GetClass()->GetCastFlags();
-		Flags = EPropertyValueFlags(  !!(CastFlags & CASTCLASS_FArrayProperty)	* uint32(EPropertyValueFlags::IsArray)
-									| !!(CastFlags & CASTCLASS_FMapProperty)	* uint32(EPropertyValueFlags::IsMap)
-									| !!(CastFlags & CASTCLASS_FSetProperty)	* uint32(EPropertyValueFlags::IsSet)
-									| !!(CastFlags & CASTCLASS_FStructProperty)	* uint32(EPropertyValueFlags::IsStruct));
+		Flags = EPropertyValueFlags(  !!(CastFlags & CASTCLASS_FArrayProperty)    * uint32(EPropertyValueFlags::IsArray)
+									| !!(CastFlags & CASTCLASS_FMapProperty)      * uint32(EPropertyValueFlags::IsMap)
+									| !!(CastFlags & CASTCLASS_FSetProperty)      * uint32(EPropertyValueFlags::IsSet)
+									| !!(CastFlags & CASTCLASS_FStructProperty)   * uint32(EPropertyValueFlags::IsStruct)
+									| !!(CastFlags & CASTCLASS_FOptionalProperty) * uint32(EPropertyValueFlags::IsOptional));
 	}
 	if (bMatchAll || Property->IsA(PropertyClass))
 	{
@@ -100,8 +102,21 @@ FORCEINLINE_DEBUGGABLE bool FPropertyValueIterator::NextValue(EPropertyValueIter
 		if (InRecursionFlags == EPropertyValueIteratorFlags::FullRecursion)
 		{
 			FPropertyValueStackEntry NewEntry(PropertyValue);
-
-			if (EnumHasAnyFlags(PropertyValueFlags, EPropertyValueFlags::IsArray))
+			
+			if (EnumHasAnyFlags(PropertyValueFlags, EPropertyValueFlags::IsOptional))
+			{
+				const FOptionalProperty* OptionalProperty = CastFieldChecked<FOptionalProperty>(Property);
+				const FProperty* InnerProperty = OptionalProperty->GetValueProperty();
+				EPropertyValueFlags InnerFlags = GetPropertyValueFlags(InnerProperty);
+				if (InnerFlags != EPropertyValueFlags::None)
+				{
+					if (const void* InnerValue = OptionalProperty->GetValuePointerForReadIfSet(PropertyValue))
+					{
+						NewEntry.ValueArray.Emplace(BasePairType(InnerProperty, InnerValue), InnerFlags);
+					}
+				}
+			}
+			else if (EnumHasAnyFlags(PropertyValueFlags, EPropertyValueFlags::IsArray))
 			{
 				const FArrayProperty* ArrayProperty = CastFieldChecked<FArrayProperty>(Property);
 				const FProperty* InnerProperty = ArrayProperty->Inner;
@@ -236,7 +251,12 @@ FString FPropertyValueIterator::GetPropertyPathDebugString() const
 
 		if (NextStackIndex < PropertyIteratorStack.Num())
 		{
-			if (CastField<FArrayProperty>(Property))
+			if (CastField<FOptionalProperty>(Property))
+			{
+				PropertyPath.Append(TEXT("?"));
+				StackIndex++;
+			}
+			else if (CastField<FArrayProperty>(Property))
 			{
 				const FPropertyValueStackEntry& NextEntry = PropertyIteratorStack[StackIndex+1];
 			
diff --git a/Engine/Source/Runtime/CoreUObject/Public/Templates/Casts.h b/Engine/Source/Runtime/CoreUObject/Public/Templates/Casts.h
--- a/Engine/Source/Runtime/CoreUObject/Public/Templates/Casts.h
+++ b/Engine/Source/Runtime/CoreUObject/Public/Templates/Casts.h
@@ -423,6 +423,7 @@ DECLARE_CAST_BY_FLAG(FSetProperty)						\
 DECLARE_CAST_BY_FLAG(USparseDelegateFunction)			\
 DECLARE_CAST_BY_FLAG(FMulticastInlineDelegateProperty)	\
 DECLARE_CAST_BY_FLAG(FMulticastSparseDelegateProperty)	\
+DECLARE_CAST_BY_FLAG(FOptionalProperty)					\
 FINISH_DECLARING_CAST_FLAGS		// This is here to hopefully remind people to include the "\" in all declarations above, especially when copy/pasting the final line.
 
 // Now actually declare the flags
diff --git a/Engine/Source/Runtime/CoreUObject/Public/UObject/ObjectMacros.h b/Engine/Source/Runtime/CoreUObject/Public/UObject/ObjectMacros.h
--- a/Engine/Source/Runtime/CoreUObject/Public/UObject/ObjectMacros.h
+++ b/Engine/Source/Runtime/CoreUObject/Public/UObject/ObjectMacros.h
@@ -374,6 +374,7 @@ enum EClassCastFlags : uint64
 	CASTCLASS_FObjectPtrProperty			= 0x0020000000000000,
 	CASTCLASS_FClassPtrProperty				= 0x0040000000000000,
 	CASTCLASS_FLargeWorldCoordinatesRealProperty = 0x0080000000000000,
+	CASTCLASS_FOptionalProperty				= 0x0100000000000000,
 };
 
 #define CASTCLASS_AllFlags ((EClassCastFlags)0xFFFFFFFFFFFFFFFF)
diff --git a/Engine/Source/Runtime/CoreUObject/Public/UObject/PropertyOptional.h b/Engine/Source/Runtime/CoreUObject/Public/UObject/PropertyOptional.h
new file mode 100644
--- /dev/null
+++ b/Engine/Source/Runtime/CoreUObject/Public/UObject/PropertyOptional.h
@@ -0,0 +1,228 @@
+// Copyright Epic Games, Inc. All Rights Reserved.
+
+#pragma once
+
+#include "UObject/UnrealType.h"
+#include "Misc/Optional.h"
+
+// 
+// Encapsulates the memory layout logic for an optional without implementing the full FProperty API.
+//
+struct COREUOBJECT_API FOptionalPropertyLayout
+{
+	explicit FOptionalPropertyLayout(FProperty* InValueProperty)
+		: ValueProperty(InValueProperty)
+	{
+		check(ValueProperty);
+	}
+
+	FProperty* GetValueProperty() const 
+	{
+		checkf(ValueProperty, TEXT("Expected ValueProperty to be initialized"));
+		return ValueProperty; 
+	}
+
+	FORCEINLINE bool IsSet(const void* Data) const
+	{
+		checkSlow(Data);
+		return IsValueNonNullablePointer()
+			? *reinterpret_cast<void*const*>(Data) != nullptr
+			: *GetIsSetPointer(Data);
+	}
+	FORCEINLINE void* MarkSetAndGetInitializedValuePointerToReplace(void* Data) const
+	{
+		checkSlow(Data);
+		if (IsValueNonNullablePointer())
+		{
+			ValueProperty->InitializeValue(Data);
+		}
+		else
+		{
+			bool* IsSetPointer = GetIsSetPointer(Data);
+			if (!*IsSetPointer)
+			{
+				ValueProperty->InitializeValue(Data);
+				*IsSetPointer = true;
+			}
+		}
+		return Data;
+	}
+	FORCEINLINE void MarkUnset(void* Data) const
+	{
+		checkSlow(Data);
+		if (IsValueNonNullablePointer())
+		{
+			ValueProperty->ClearValue(Data);
+		}
+		else
+		{
+			bool* IsSetPointer = GetIsSetPointer(Data);
+			if (*IsSetPointer)
+			{
+				ValueProperty->DestroyValue(Data);
+				*IsSetPointer = false;
+			}
+		}
+	}
+
+	// For reading the value of a set optional.
+	// Must be called on a non-null pointer to a set optional.
+	FORCEINLINE const void* GetValuePointerForRead(const void* Data) const 
+	{
+		checkSlow(Data && IsSet(Data));
+		return Data; 
+	}
+	
+	// For replacing the value of a set optional.
+	// Must be called on a non-null pointer to a set optional.
+	FORCEINLINE void* GetValuePointerForReplace(void* Data) const
+	{
+		checkSlow(Data && IsSet(Data));
+		return Data;
+	}
+
+	// For reading the value of a set optional.
+	// Must be called on a non-null pointer to an optional.
+	// If called on an unset optional, will return null.
+	FORCEINLINE const void* GetValuePointerForReadIfSet(const void* Data) const
+	{
+		checkSlow(Data);
+		return IsSet(Data) ? Data : nullptr;
+	}
+	
+	// For replacing the value of a set optional.
+	// Must be called on a non-null pointer to an optional.
+	// If called on an unset optional, will return null.
+	FORCEINLINE void* GetValuePointerForReplaceIfSet(void* Data) const
+	{
+		checkSlow(Data);
+		return IsSet(Data) ? Data : nullptr;
+	}
+	
+	// For calling from polymorphic code that doesn't know whether it needs the value pointer for
+	// read or replace, or whether it has a const pointer or not.
+	// Must be called on a non-null pointer to a set optional.
+	FORCEINLINE const void* GetValuePointerForReadOrReplace(const void* Data) const
+	{
+		checkSlow(Data && IsSet(Data));
+		return Data;
+	}
+	FORCEINLINE void* GetValuePointerForReadOrReplace(void* Data) const
+	{
+		checkSlow(Data && IsSet(Data));
+		return Data;
+	}
+	
+	// For calling from polymorphic code that doesn't know whether it needs the value pointer for
+	// read or replace, or whether it has a const pointer or not.
+	// Must be called on a non-null pointer to an optional.
+	// If called on an unset optional, will return null.
+	FORCEINLINE const void* GetValuePointerForReadOrReplaceIfSet(const void* Data) const
+	{
+		checkSlow(Data);
+		return IsSet(Data) ? Data : nullptr;
+	}
+	FORCEINLINE void* GetValuePointerForReadOrReplaceIfSet(void* Data) const
+	{
+		checkSlow(Data);
+		return IsSet(Data) ? Data : nullptr;
+	}
+	
+protected:
+
+	FOptionalPropertyLayout() : ValueProperty(nullptr) {}
+
+	// Variables
+	FProperty* ValueProperty; // The type of the value
+
+	FORCEINLINE int32 CalcIsSetOffset() const
+	{
+		check(!IsValueNonNullablePointer());
+		checkfSlow(
+			ValueProperty->GetSize() == Align(ValueProperty->GetSize(), ValueProperty->GetMinAlignment()),
+			TEXT("Expected optional value property to have aligned size, but got misaligned size %i for %s that has minimum alignment %i"),
+			ValueProperty->GetSize(),
+			*ValueProperty->GetFullName(),
+			ValueProperty->GetMinAlignment());
+		return ValueProperty->GetSize();
+	}
+	FORCEINLINE int32 CalcSize() const
+	{
+		if (IsValueNonNullablePointer())
+		{
+			return ValueProperty->GetSize();
+		}
+		else
+		{
+			return Align(CalcIsSetOffset() + 1, ValueProperty->GetMinAlignment());
+		}
+	}
+
+	FORCEINLINE bool IsValueNonNullablePointer() const
+	{
+		return (ValueProperty->GetPropertyFlags() & CPF_NonNullable) != 0;
+	}
+
+	FORCEINLINE bool* GetIsSetPointer(void* Data) const
+	{
+		return reinterpret_cast<bool*>(reinterpret_cast<uint8*>(Data) + CalcIsSetOffset());
+	}
+	FORCEINLINE const bool* GetIsSetPointer(const void* Data) const
+	{
+		return reinterpret_cast<const bool*>(reinterpret_cast<const uint8*>(Data) + CalcIsSetOffset());
+	}
+};
+
+//
+// A property corresponding to UE's optional type, TOptional<T>.
+// NOTE: this property is not yet handled by all UE subsystems that produce or consume properties.
+//
+class COREUOBJECT_API FOptionalProperty : public FProperty, public FOptionalPropertyLayout
+{
+	DECLARE_FIELD(FOptionalProperty, FProperty, CASTCLASS_FOptionalProperty)
+
+public:
+
+	FOptionalProperty(FFieldVariant InOwner, const FName& InName, EObjectFlags InObjectFlags);
+	virtual ~FOptionalProperty();
+
+	// Sets the optional property's value property.
+	void SetValueProperty(FProperty* InValueProperty);
+
+	// UObject interface
+	virtual void Serialize(FArchive& Ar) override;
+	virtual void GetPreloadDependencies(TArray<UObject*>& OutDeps) override;
+	// End of UObject interface
+
+	// Field interface
+	virtual void PostDuplicate(const FField& InField) override;
+	virtual FField* GetInnerFieldByName(const FName& InName) override;
+	virtual void GetInnerFields(TArray<FField*>& OutFields) override;
+	// End of Field interface
+
+	// UHT interface
+	virtual FString GetCPPType(FString* ExtendedTypeText = NULL, uint32 CPPExportFlags = 0) const override;
+	virtual FString GetCPPMacroType(FString& ExtendedTypeText) const override;
+	virtual FString GetCPPTypeForwardDeclaration() const override;
+	// End of UHT interface
+
+	// FProperty interface
+	virtual void LinkInternal(FArchive& Ar) override;
+	virtual bool Identical(const void* A, const void* B, uint32 PortFlags) const override;
+	virtual void SerializeItem(FStructuredArchive::FSlot Slot, void* Value, void const* Defaults) const override;
+	virtual bool NetSerializeItem(FArchive& Ar, UPackageMap* Map, void* Data, TArray<uint8>* MetaData = NULL) const override;
+	virtual void ExportText_Internal(FString& ValueStr, const void* ContainerOrPropertyPtr, EPropertyPointerType PropertyPointerType, const void* DefaultValue, UObject* Parent, int32 PortFlags, UObject* ExportRootScope) const override;
+	virtual const TCHAR* ImportText_Internal(const TCHAR* Buffer, void* ContainerOrPropertyPtr, EPropertyPointerType PropertyPointerType, UObject* Parent, int32 PortFlags, FOutputDevice* ErrorText) const override;
+	virtual void CopyValuesInternal(void* Dest, void const* Src, int32 Count) const override;
+	virtual void ClearValueInternal(void* Data) const override;
+	virtual void InitializeValueInternal(void* Data) const override;
+	virtual void DestroyValueInternal(void* Data) const override;
+	virtual void InstanceSubobjects(void* Data, void const* DefaultData, UObject* Owner, struct FObjectInstancingGraph* InstanceGraph) override;
+	virtual int32 GetMinAlignment() const override;
+	virtual bool ContainsObjectReference(TArray<const FStructProperty*>& EncounteredStructProps, EPropertyObjectReferenceType InReferenceType = EPropertyObjectReferenceType::Strong) const override;
+	virtual void EmitReferenceInfo(UE::GC::FSchemaBuilder& Schema, int32 BaseOffset, TArray<const FStructProperty*>& EncounteredStructProps, UE::GC::FPropertyStack& DebugPath) override;
+	virtual EConvertFromTypeResult ConvertFromType(const FPropertyTag& Tag, FStructuredArchive::FSlot Slot, uint8* Data, UStruct* DefaultsStruct, const uint8* Defaults) override;
+	virtual uint32 GetValueTypeHashInternal(const void* Src) const override;
+	virtual void AddReferencedObjects(FReferenceCollector& Collector) override;
+	// End of FProperty interface
+};
diff --git a/Engine/Source/Runtime/CoreUObject/Public/UObject/PropertyTag.h b/Engine/Source/Runtime/CoreUObject/Public/UObject/PropertyTag.h
--- a/Engine/Source/Runtime/CoreUObject/Public/UObject/PropertyTag.h
+++ b/Engine/Source/Runtime/CoreUObject/Public/UObject/PropertyTag.h
@@ -39,7 +39,7 @@ struct FPropertyTag
 
 	// Constructors.
 	FPropertyTag() {}
-	FPropertyTag( FArchive& InSaveAr, FProperty* Property, int32 InIndex, uint8* Value, uint8* Defaults );
+	FPropertyTag( FArchive& InSaveAr, FProperty* Property, int32 InIndex, uint8* Value, const uint8* Defaults );
 
 	// Set optional property guid
 	void SetPropertyGuid(const FGuid& InPropertyGuid);
@@ -49,8 +49,8 @@ struct FPropertyTag
 	friend void operator<<(FStructuredArchive::FSlot Slot, FPropertyTag& Tag);
 
 	// Property serializer.
-	void SerializeTaggedProperty( FArchive& Ar, FProperty* Property, uint8* Value, uint8* Defaults ) const;
-	void SerializeTaggedProperty(FStructuredArchive::FSlot Slot, FProperty* Property, uint8* Value, uint8* Defaults) const;
+	void SerializeTaggedProperty( FArchive& Ar, FProperty* Property, uint8* Value, const uint8* Defaults ) const;
+	void SerializeTaggedProperty(FStructuredArchive::FSlot Slot, FProperty* Property, uint8* Value, const uint8* Defaults) const;
 };
 
 #if UE_ENABLE_INCLUDE_ORDER_DEPRECATED_IN_5_2
diff --git a/Engine/Source/Runtime/CoreUObject/Public/UObject/UnrealType.h b/Engine/Source/Runtime/CoreUObject/Public/UObject/UnrealType.h
--- a/Engine/Source/Runtime/CoreUObject/Public/UObject/UnrealType.h
+++ b/Engine/Source/Runtime/CoreUObject/Public/UObject/UnrealType.h
@@ -6971,7 +6971,8 @@ private:
 		None = 0x0,
 		IsMatch = 0x01,
 
-		ContainerMask = 0xF0,
+		ContainerMask = 0xF8,
+		IsOptional = 0x08,
 		IsArray = 0x10,
 		IsMap = 0x20,
 		IsSet = 0x40,
diff --git a/Engine/Source/Runtime/Engine/Private/InstancedReferenceSubobjectHelper.cpp b/Engine/Source/Runtime/Engine/Private/InstancedReferenceSubobjectHelper.cpp
--- a/Engine/Source/Runtime/Engine/Private/InstancedReferenceSubobjectHelper.cpp
+++ b/Engine/Source/Runtime/Engine/Private/InstancedReferenceSubobjectHelper.cpp
@@ -1,6 +1,7 @@
 // Copyright Epic Games, Inc. All Rights Reserved.
 
 #include "InstancedReferenceSubobjectHelper.h"
+#include "UObject/PropertyOptional.h"
 #include "UObject/Package.h"
 #include "UObject/Package.h"
 	
@@ -95,6 +96,18 @@ UObject* FInstancedPropertyPath::Resolve(const UObject* Container) const
 				CurrentProp = nullptr;
 			}
 		}
+		else if(const FOptionalProperty* OptionalProperty = CastField<FOptionalProperty>(CurrentProp))
+		{
+			if(const uint8* InnerValuePtr = static_cast<const uint8*>(OptionalProperty->GetValuePointerForReadIfSet(ValuePtr)))
+			{
+				ValuePtr = InnerValuePtr;
+				CurrentProp = OptionalProperty->GetValueProperty();
+			}
+			else
+			{
+				CurrentProp = nullptr;
+			}
+		}
 		else
 		{
 			if (const FStructProperty* StructProperty = CastField<FStructProperty>(CurrentProp))
@@ -194,6 +207,21 @@ void FFindInstancedReferenceSubobjectHelper::ForEachInstancedSubObject(FInstance
 			}
 		}
 	}
+	else if (const FOptionalProperty* OptionalProperty = CastField<FOptionalProperty>(TargetProp))
+	{
+		// Exit early if the optional does not contain any instanced references.
+		if (!OptionalProperty->HasAnyPropertyFlags(CPF_ContainsInstancedReference))
+		{
+			return;
+		}
+		
+		if (T ValueAddress = static_cast<T>(OptionalProperty->GetValuePointerForReadOrReplaceIfSet(ContainerAddress)))
+		{
+			PropertyPath.Push(OptionalProperty->GetValueProperty());
+			ForEachInstancedSubObject(PropertyPath, ValueAddress, ObjRefFunc);
+			PropertyPath.Pop();
+		}
+	}
 	else if (const FStructProperty* StructProperty = CastField<const FStructProperty>(TargetProp))
 	{
 		// Exit early if the struct does not contain any instanced references or if the struct is invalid.
