Force warning as error in UObject class analysis tests

#rnx

diff --git a/Engine/Source/Runtime/CoreUObject/Private/UObject/Class.cpp b/Engine/Source/Runtime/CoreUObject/Private/UObject/Class.cpp
--- a/Engine/Source/Runtime/CoreUObject/Private/UObject/Class.cpp
+++ b/Engine/Source/Runtime/CoreUObject/Private/UObject/Class.cpp
@@ -3724,7 +3724,19 @@ FAutoConsoleCommandWithWorldAndArgs GCmdListBadScriptStructs(
 			FStructUtils::AttemptToFindUninitializedScriptStructMembers();
 		}));
 
-IMPLEMENT_SIMPLE_AUTOMATION_TEST(FAutomationTestAttemptToFindUninitializedScriptStructMembers, "UObject.Class AttemptToFindUninitializedScriptStructMembers", EAutomationTestFlags::EditorContext | EAutomationTestFlags::ClientContext | EAutomationTestFlags::ServerContext | EAutomationTestFlags::SmokeFilter)
+class FAutomationTestUObjectClassBase : public FAutomationTestBase
+{
+public:
+	FAutomationTestUObjectClassBase(const FString& InName, const bool bInComplexTask)
+		: FAutomationTestBase(InName, bInComplexTask)
+		{ }
+
+	virtual bool SuppressLogErrors() override { return false; }
+	virtual bool SuppressLogWarnings() override { return false; }
+	virtual bool ElevateLogWarningsToErrors() override { return true; }
+};
+
+IMPLEMENT_CUSTOM_SIMPLE_AUTOMATION_TEST(FAutomationTestAttemptToFindUninitializedScriptStructMembers, FAutomationTestUObjectClassBase, "UObject.Class AttemptToFindUninitializedScriptStructMembers", EAutomationTestFlags::EditorContext | EAutomationTestFlags::ClientContext | EAutomationTestFlags::ServerContext | EAutomationTestFlags::SmokeFilter)
 bool FAutomationTestAttemptToFindUninitializedScriptStructMembers::RunTest(const FString& Parameters)
 {
 	// This test fails when running tests under UHT because there is no TestUninitializedScriptStructMembersTest, so just skip it in that config.
@@ -3931,7 +3943,7 @@ FAutoConsoleCommandWithWorldAndArgs GCmdListShortTypeNamesInMetaData(
 			FStructUtils::AttemptToFindShortTypeNamesInMetaData();
 		}));
 
-IMPLEMENT_SIMPLE_AUTOMATION_TEST(FAutomationTestAttemptToFindShortTypeNamesInMetaData, "UObject.Class AttemptToFindShortTypeNamesInMetaData", EAutomationTestFlags::EditorContext | EAutomationTestFlags::ApplicationContextMask | EAutomationTestFlags::ServerContext | EAutomationTestFlags::SmokeFilter)
+IMPLEMENT_CUSTOM_SIMPLE_AUTOMATION_TEST(FAutomationTestAttemptToFindShortTypeNamesInMetaData, FAutomationTestUObjectClassBase, "UObject.Class AttemptToFindShortTypeNamesInMetaData", EAutomationTestFlags::EditorContext | EAutomationTestFlags::ApplicationContextMask | EAutomationTestFlags::ServerContext | EAutomationTestFlags::SmokeFilter)
 bool FAutomationTestAttemptToFindShortTypeNamesInMetaData::RunTest(const FString& Parameters)
 {
 	// This test is not necessary when running under UHT so just skip it in that config.
@@ -4084,7 +4096,7 @@ static bool FindUndeclaredObjectReferencesInStructSerializers(FAutomationTestBas
 }
 
 // Test for finding structs which misreport their object properties for native serialization
-IMPLEMENT_SIMPLE_AUTOMATION_TEST(FAutomationTestFindUndeclaredObjectReferencesInStructSerializers, "UObject.Class FindUndeclaredObjectReferencesInStructSerializers", EAutomationTestFlags::EditorContext | EAutomationTestFlags::ClientContext | EAutomationTestFlags::ServerContext | EAutomationTestFlags::SmokeFilter)
+IMPLEMENT_CUSTOM_SIMPLE_AUTOMATION_TEST(FAutomationTestFindUndeclaredObjectReferencesInStructSerializers, FAutomationTestUObjectClassBase, "UObject.Class FindUndeclaredObjectReferencesInStructSerializers", EAutomationTestFlags::EditorContext | EAutomationTestFlags::ClientContext | EAutomationTestFlags::ServerContext | EAutomationTestFlags::SmokeFilter)
 bool FAutomationTestFindUndeclaredObjectReferencesInStructSerializers::RunTest(const FString& Parameters)
 {
 	return FindUndeclaredObjectReferencesInStructSerializers(*this, false);
@@ -4092,7 +4104,7 @@ bool FAutomationTestFindUndeclaredObjectReferencesInStructSerializers::RunTest(c
 
 // Test for finding structs which unoptimally their object properties for native serialization - not a smoke test
 
-	IMPLEMENT_SIMPLE_AUTOMATION_TEST(FAutomationTestFindInexactObjectReferencesInStructSerializers, "UObject.Class FindInexactObjectReferencesInStructSerializers", EAutomationTestFlags::EditorContext | EAutomationTestFlags::ClientContext | EAutomationTestFlags::ServerContext | EAutomationTestFlags::EngineFilter | EAutomationTestFlags::RequiresUser)
+IMPLEMENT_CUSTOM_SIMPLE_AUTOMATION_TEST(FAutomationTestFindInexactObjectReferencesInStructSerializers, FAutomationTestUObjectClassBase, "UObject.Class FindInexactObjectReferencesInStructSerializers", EAutomationTestFlags::EditorContext | EAutomationTestFlags::ClientContext | EAutomationTestFlags::ServerContext | EAutomationTestFlags::EngineFilter | EAutomationTestFlags::RequiresUser)
 bool FAutomationTestFindInexactObjectReferencesInStructSerializers::RunTest(const FString& Parameters)
 {
 	return FindUndeclaredObjectReferencesInStructSerializers(*this, true);
