Fix assert when doing unversioned property serialization of FOptionalProperty by skipping the inner value tag serialization in that case.

diff --git a/Engine/Source/Runtime/CoreUObject/Private/UObject/PropertyOptional.cpp b/Engine/Source/Runtime/CoreUObject/Private/UObject/PropertyOptional.cpp
--- a/Engine/Source/Runtime/CoreUObject/Private/UObject/PropertyOptional.cpp
+++ b/Engine/Source/Runtime/CoreUObject/Private/UObject/PropertyOptional.cpp
@@ -182,7 +182,16 @@ void FOptionalProperty::SerializeItem(FStructuredArchive::FSlot Slot, void* Data
 		
 		const void* ValueDefaults = Defaults ? GetValuePointerForReadIfSet(Defaults) : nullptr;
 
-		if (bIsLoading)
+		if (Slot.GetArchiveState().UseUnversionedPropertySerialization())
+		{
+			// Simply serialize the inner value if using unversioned property serialization.
+			FStructuredArchive::FSlot ValueSlot = TaggedValueRecord.EnterField(TEXT("Value"));
+			void* ValueData = bIsLoading
+				? MarkSetAndGetInitializedValuePointerToReplace(Data)
+				: GetValuePointerForReadOrReplace(Data);
+			GetValueProperty()->SerializeItem(ValueSlot, ValueData, ValueDefaults);
+		}
+		else if (bIsLoading)
 		{
 			// Serialize the value's tag.
 			FPropertyTag ValueTag;
