Adjust Log verbosity settings for UObject Class tests so test automation framework caught errors.

#rnx

diff --git a/Engine/Source/Runtime/CoreUObject/Private/UObject/Class.cpp b/Engine/Source/Runtime/CoreUObject/Private/UObject/Class.cpp
--- a/Engine/Source/Runtime/CoreUObject/Private/UObject/Class.cpp
+++ b/Engine/Source/Runtime/CoreUObject/Private/UObject/Class.cpp
@@ -3732,8 +3732,8 @@ public:
 		{ }
 
 	virtual bool SuppressLogErrors() override { return false; }
-	virtual bool SuppressLogWarnings() override { return false; }
-	virtual bool ElevateLogWarningsToErrors() override { return true; }
+	virtual bool SuppressLogWarnings() override { return true; }
+	virtual bool ElevateLogWarningsToErrors() override { return false; }
 };
 
 IMPLEMENT_CUSTOM_SIMPLE_AUTOMATION_TEST(FAutomationTestAttemptToFindUninitializedScriptStructMembers, FAutomationTestUObjectClassBase, "UObject.Class AttemptToFindUninitializedScriptStructMembers", EAutomationTestFlags::EditorContext | EAutomationTestFlags::ClientContext | EAutomationTestFlags::ServerContext | EAutomationTestFlags::SmokeFilter)
@@ -3809,12 +3809,28 @@ static FString GetSuggestedPathNameForTypeShortName(const FString& ShortName)
 template <typename T>
 static void LogMetaDataShortTypeName(const T* Field, FName MetaDataKey, const FString& MetaDataValue)
 {
+	static struct FLogShortTypeNameInMetaDataCheckSettings
+	{
+		ELogVerbosity::Type LogVerbosity = ELogVerbosity::Warning;
+
+		FLogShortTypeNameInMetaDataCheckSettings()
+		{
+			{
+				FString VerbositySettingString;
+				if (GConfig->GetString(TEXT("CoreUObject.ShortTypeNameInMetaDataCheck"), TEXT("LogVerbosity"), VerbositySettingString, GEngineIni))
+				{
+					LogVerbosity = ParseLogVerbosityFromString(VerbositySettingString);
+				}
+			}
+		}
+	} Settings;
+
 	if constexpr (std::is_base_of_v<T, UField>)
 	{
 		if (const UFunction* Func = Cast<UFunction>(Field))
 		{
 			UStruct* FuncOwner = CastChecked<UStruct>(Func->GetOuter());
-			FMsg::Logf(__FILE__, __LINE__, LogClass.GetCategoryName(), ELogVerbosity::Warning, TEXT("Function %s%s::%s defines MetaData key \"%s\" which contains short type name \"%s\". %s%s"),
+			FMsg::Logf(__FILE__, __LINE__, LogClass.GetCategoryName(), Settings.LogVerbosity, TEXT("Function %s%s::%s defines MetaData key \"%s\" which contains short type name \"%s\". %s%s"),
 				FuncOwner->GetPrefixCPP(),
 				*FuncOwner->GetName(),
 				*Func->GetName(),
@@ -3825,7 +3841,7 @@ static void LogMetaDataShortTypeName(const T* Field, FName MetaDataKey, const FS
 		}
 		else
 		{
-			FMsg::Logf(__FILE__, __LINE__, LogClass.GetCategoryName(), ELogVerbosity::Warning, TEXT("%s %s%s defines MetaData key \"%s\" which contains short type name \"%s\". %s%s"),
+			FMsg::Logf(__FILE__, __LINE__, LogClass.GetCategoryName(), Settings.LogVerbosity, TEXT("%s %s%s defines MetaData key \"%s\" which contains short type name \"%s\". %s%s"),
 				*Field->GetClass()->GetName(),
 				Field->template IsA<UStruct>() ? CastChecked<UStruct>(Field)->GetPrefixCPP() : TEXT(""),
 				*Field->GetName(),
@@ -3838,7 +3854,7 @@ static void LogMetaDataShortTypeName(const T* Field, FName MetaDataKey, const FS
 	else if constexpr (std::is_same_v<T, FProperty>)
 	{
 		UStruct* OwnerStruct = Field->GetOwnerStruct();
-		FMsg::Logf(__FILE__, __LINE__, LogClass.GetCategoryName(), ELogVerbosity::Warning, TEXT("Property %s %s%s::%s defines MetaData key \"%s\" which contains short type name \"%s\". %s%s"),
+		FMsg::Logf(__FILE__, __LINE__, LogClass.GetCategoryName(), Settings.LogVerbosity, TEXT("Property %s %s%s::%s defines MetaData key \"%s\" which contains short type name \"%s\". %s%s"),
 			*Field->GetClass()->GetName(),
 			OwnerStruct->GetPrefixCPP(),
 			*OwnerStruct->GetName(),
