Use the PostConstructLinkNext pointer to initialize properties that aren't CPF_ZeroConstructor, rather than just starting at FirstPropertyToInit and initializing every remaining property (often redundantly)
Relevant CLs: 86560, 1444860, 1447974

diff --git a/Engine/Source/Editor/KismetCompiler/Private/KismetCompiler.cpp b/Engine/Source/Editor/KismetCompiler/Private/KismetCompiler.cpp
--- a/Engine/Source/Editor/KismetCompiler/Private/KismetCompiler.cpp
+++ b/Engine/Source/Editor/KismetCompiler/Private/KismetCompiler.cpp
@@ -2875,6 +2875,8 @@ void FKismetCompilerContext::SetCalculatedMetaDataAndFlags(UFunction* Function,
 	Function->ParmsSize = 0;
 	Function->NumParms = 0;
 	Function->ReturnValueOffset = MAX_uint16;
+	Function->FirstPropertyToInit = nullptr;
+	FProperty** ConstructLink = &Function->FirstPropertyToInit;
 
 	for (TFieldIterator<FProperty> PropIt(Function, EFieldIteratorFlags::ExcludeSuper); PropIt; ++PropIt)
 	{
@@ -2898,9 +2900,11 @@ void FKismetCompilerContext::SetCalculatedMetaDataAndFlags(UFunction* Function,
 		{
 			if (!Property->HasAnyPropertyFlags(CPF_ZeroConstructor))
 			{
-				Function->FirstPropertyToInit = Property;
+				(*ConstructLink) = Property;
+				Property->PostConstructLinkNext = nullptr;
+				ConstructLink = &Property->PostConstructLinkNext;
+
 				Function->FunctionFlags |= FUNC_HasDefaults;
-				break;
 			}
 		}
 	}
diff --git a/Engine/Source/Runtime/CoreUObject/Private/UObject/Class.cpp b/Engine/Source/Runtime/CoreUObject/Private/UObject/Class.cpp
--- a/Engine/Source/Runtime/CoreUObject/Private/UObject/Class.cpp
+++ b/Engine/Source/Runtime/CoreUObject/Private/UObject/Class.cpp
@@ -6623,6 +6623,7 @@ void UFunction::InitializeDerivedMembers()
 	NumParms = 0;
 	ParmsSize = 0;
 	ReturnValueOffset = MAX_uint16;
+	FProperty** ConstructLink = &FirstPropertyToInit;
 
 	for (FProperty* Property = CastField<FProperty>(ChildProperties); Property; Property = CastField<FProperty>(Property->Next))
 	{
@@ -6635,17 +6636,17 @@ void UFunction::InitializeDerivedMembers()
 				ReturnValueOffset = Property->GetOffset_ForUFunction();
 			}
 		}
-		else if ((FunctionFlags & FUNC_HasDefaults) != 0)
+		else if ((FunctionFlags & FUNC_HasDefaults) == 0)
 		{
-			if (!Property->HasAnyPropertyFlags(CPF_ZeroConstructor))
-			{
-				FirstPropertyToInit = Property;
-				break;
-			}
+			// we're done with parms and we've not been tagged as FUNC_HasDefaults, so we can abort
+			// this potentially costly loop:
+			break;
 		}
-		else
+		else if (!Property->HasAnyPropertyFlags(CPF_ZeroConstructor))
 		{
-			break;
+			*ConstructLink = Property;
+			Property->PostConstructLinkNext = nullptr;
+			ConstructLink = &Property->PostConstructLinkNext;
 		}
 	}
 }
diff --git a/Engine/Source/Runtime/CoreUObject/Private/UObject/ScriptCore.cpp b/Engine/Source/Runtime/CoreUObject/Private/UObject/ScriptCore.cpp
--- a/Engine/Source/Runtime/CoreUObject/Private/UObject/ScriptCore.cpp
+++ b/Engine/Source/Runtime/CoreUObject/Private/UObject/ScriptCore.cpp
@@ -1023,8 +1023,8 @@ void ProcessScriptFunction(UObject* Context, UFunction* Function, FFrame& Stack,
 
 	if (!bUsePersistentFrame)
 	{
-		// Initialize any local struct properties with defaults
-		for (FProperty* LocalProp = Function->FirstPropertyToInit; LocalProp != NULL; LocalProp = (FProperty*)(LocalProp->Next))
+		// Initialize any local properties that aren't CPF_ZeroConstruct:
+		for (FProperty* LocalProp = Function->FirstPropertyToInit; LocalProp != nullptr; LocalProp = (FProperty*)(LocalProp->PostConstructLinkNext))
 		{
 			LocalProp->InitializeValue_InContainer(NewStack.Locals);
 		}
@@ -2131,7 +2131,7 @@ void UObject::ProcessEvent( UFunction* Function, void* Parms )
 
 		if (!bUsePersistentFrame)
 		{
-			for (FProperty* LocalProp = Function->FirstPropertyToInit; LocalProp != NULL; LocalProp = (FProperty*)LocalProp->Next)
+			for (FProperty* LocalProp = Function->FirstPropertyToInit; LocalProp != nullptr; LocalProp = (FProperty*)LocalProp->PostConstructLinkNext)
 			{
 				LocalProp->InitializeValue_InContainer(NewStack.Locals);
 			}
