fix for warning in PluginManager test

diff --git a/Engine/Source/Runtime/CoreUObject/Tests/PluginManagerTest.cpp b/Engine/Source/Runtime/CoreUObject/Tests/PluginManagerTest.cpp
--- a/Engine/Source/Runtime/CoreUObject/Tests/PluginManagerTest.cpp
+++ b/Engine/Source/Runtime/CoreUObject/Tests/PluginManagerTest.cpp
@@ -14,7 +14,7 @@
 #include "UObject/Package.h"
 #include "HAL/IConsoleManager.h"
 #include "UObject/Class.h"
-
+#include "LowLevelTestsRunner/WarnFilterScope.h"
 #include "UObject/FastReferenceCollector.h"
 
 
@@ -175,6 +175,11 @@ namespace UE::CoreUObject::Private::Tests
 		};
 		RefObj->ObjectPtr = Obj; //add a reference
 
+		UE::Testing::FWarnFilterScope _([](const TCHAR* Message, ELogVerbosity::Type Verbosity, const FName& Category)
+			{
+				bool bFiltered = FCString::Strncmp(Message, TEXT("Marking leaking package"), FCString::Strlen(TEXT("Marking leaking package"))) == 0 && Verbosity == ELogVerbosity::Type::Warning && Category == TEXT("PluginHandlerLog");
+				return bFiltered;
+			});
 		FText Reason;
 		CHECK(PluginManager.UnmountExplicitlyLoadedPlugin(TEXT("MyTestPlugin"), &Reason));
 		CHECK(!FPackageName::MountPointExists(TEXT("/MyTestPlugin")));
