Experimental support for Incremental Reachability Analysis

Disabled by default. gc.AllowIncrementalReachability controls if it's enabled or not.
gc.IncrementalReachabilityTimeLimit controls time limit (defaults to 0.005 = 5ms)
gc.VerifyUnrachableObjects (defaults to 0) controls additional GC assumption verification that no unreachable objects are reachable after RA is complete
gc.ContinuousIncrementalGC (defaults to 0) is a new GC stress test that runs incremental GC continuously
TObjectPtr barrier (Obj->MarkAsReachable()) is also currently disabled (compiled out with a macro) until we have all the places (ARO/BPVM) properly running barriers
EngineTest tests can and will run Incremental Reachability Analysis tests even if all the above is disabled

diff --git a/Engine/Source/Developer/DesktopWidgets/DesktopWidgets.Build.cs b/Engine/Source/Developer/DesktopWidgets/DesktopWidgets.Build.cs
--- a/Engine/Source/Developer/DesktopWidgets/DesktopWidgets.Build.cs
+++ b/Engine/Source/Developer/DesktopWidgets/DesktopWidgets.Build.cs
@@ -9,6 +9,7 @@ public class DesktopWidgets : ModuleRules
 		PublicDependencyModuleNames.AddRange(
 			new string[] {
 				"Core",
+				"CoreUObject",
 				"Slate",
 				"SlateCore",
 			}
diff --git a/Engine/Source/Developer/DeviceManager/DeviceManager.Build.cs b/Engine/Source/Developer/DeviceManager/DeviceManager.Build.cs
--- a/Engine/Source/Developer/DeviceManager/DeviceManager.Build.cs
+++ b/Engine/Source/Developer/DeviceManager/DeviceManager.Build.cs
@@ -15,7 +15,8 @@ public class DeviceManager : ModuleRules
 		PrivateDependencyModuleNames.AddRange(
 			new string[] {
 				"Core",
-                "InputCore",
+				"CoreUObject",
+				"InputCore",
 				"Slate",
 				"SlateCore",
 				"TargetPlatform",
diff --git a/Engine/Source/Developer/SessionFrontend/SessionFrontend.Build.cs b/Engine/Source/Developer/SessionFrontend/SessionFrontend.Build.cs
--- a/Engine/Source/Developer/SessionFrontend/SessionFrontend.Build.cs
+++ b/Engine/Source/Developer/SessionFrontend/SessionFrontend.Build.cs
@@ -11,6 +11,7 @@ public class SessionFrontend : ModuleRules
 		PublicDependencyModuleNames.AddRange(
 			new string[] {
 				"Core",
+				"CoreUObject",
 				"Slate",
 			}
 		);
diff --git a/Engine/Source/Editor/PListEditor/PListEditor.Build.cs b/Engine/Source/Editor/PListEditor/PListEditor.Build.cs
--- a/Engine/Source/Editor/PListEditor/PListEditor.Build.cs
+++ b/Engine/Source/Editor/PListEditor/PListEditor.Build.cs
@@ -9,7 +9,8 @@ public class PListEditor : ModuleRules
 		PrivateDependencyModuleNames.AddRange(
 			new string[] { 
 				"Core",
-                "InputCore",
+				"CoreUObject",
+				"InputCore",
 				"Slate",
 				"SlateCore",
 				"DesktopPlatform",
diff --git a/Engine/Source/Editor/SequencerWidgets/SequencerWidgets.Build.cs b/Engine/Source/Editor/SequencerWidgets/SequencerWidgets.Build.cs
--- a/Engine/Source/Editor/SequencerWidgets/SequencerWidgets.Build.cs
+++ b/Engine/Source/Editor/SequencerWidgets/SequencerWidgets.Build.cs
@@ -21,6 +21,7 @@ public class SequencerWidgets : ModuleRules
         PrivateDependencyModuleNames.AddRange(
 			new string[] {
 				"Core",
+				"CoreUObject",
 				"InputCore",
 				"Slate",
 				"SlateCore",
diff --git a/Engine/Source/Programs/Shared/EpicGames.Core/UnrealEngineTypes.cs b/Engine/Source/Programs/Shared/EpicGames.Core/UnrealEngineTypes.cs
--- a/Engine/Source/Programs/Shared/EpicGames.Core/UnrealEngineTypes.cs
+++ b/Engine/Source/Programs/Shared/EpicGames.Core/UnrealEngineTypes.cs
@@ -647,6 +647,11 @@ namespace EpicGames.Core
 		None = 0,
 		//~ All the other bits are reserved, DO NOT ADD NEW FLAGS HERE!
 
+		/// <summary>
+		/// Flag set on all non-root objects at the beginning of Reachability Analysis
+		/// </summary>
+		MaybeUnreachable = 1 << 19,
+
 		/// <summary>
 		/// Object is ready to be imported by another package during loading
 		/// </summary>
@@ -705,7 +710,7 @@ namespace EpicGames.Core
 		GarbageCollectionKeepFlags = Native | Async | AsyncLoading | LoaderImport,
 
 		//~ Make sure this is up to date!
-		AllFlags = LoaderImport | Garbage | ReachableInCluster | ClusterRoot | Native | Async | AsyncLoading | Unreachable | PendingKill | RootSet | PendingConstruction
+		AllFlags = MaybeUnreachable | LoaderImport | Garbage | ReachableInCluster | ClusterRoot | Native | Async | AsyncLoading | Unreachable | PendingKill | RootSet | PendingConstruction
 	};
 
 	/// <summary>
diff --git a/Engine/Source/Runtime/CoreUObject/Private/UObject/ExpandingChunkedList.h b/Engine/Source/Runtime/CoreUObject/Private/UObject/ExpandingChunkedList.h
new file mode 100644
--- /dev/null
+++ b/Engine/Source/Runtime/CoreUObject/Private/UObject/ExpandingChunkedList.h
@@ -0,0 +1,115 @@
+// Copyright Epic Games, Inc. All Rights Reserved.
+
+/*=============================================================================
+	ExpandingChunkedList.h: Unreal realtime garbage collection internal helpers
+=============================================================================*/
+
+#pragma once
+
+#include "CoreMinimal.h"
+#include <atomic>
+
+template<typename T, int32 NumElementsPerChunk = 64>
+class TExpandingChunkedList
+{
+	struct FChunk
+	{
+		/** Chunk items */
+		T Items[NumElementsPerChunk];
+		/** Number of items in this chunk, this can be > NumElementsPerChunk if the chunk is full */
+		std::atomic<int32> NumItems = 0;
+		/** Pointer to the next chunk */
+		FChunk* Next = nullptr;
+	};
+
+	/** Head of the chunk list */
+	FChunk* Head = nullptr;
+
+public:
+
+	~TExpandingChunkedList()
+	{
+		Empty();
+	}
+
+	/**
+	* Thread safe: Pushes a new item into the list
+	**/
+	FORCEINLINE void Push(T Item)
+	{
+		FChunk* NewChunk = nullptr;
+		FChunk* OldHead = Head;
+		FChunk* DestChunk = OldHead;
+		int32 Index = DestChunk ? DestChunk->NumItems.fetch_add(1, std::memory_order_acq_rel) : NumElementsPerChunk;
+		while (Index >= NumElementsPerChunk)
+		{
+			// There's no chunks or no free chunks so allocate a new one
+			if (!NewChunk)
+			{
+				// Allocate a new chunk if we haven't already
+				NewChunk = new FChunk();
+				// reserve index 0 for the new item
+				Index = NewChunk->NumItems.fetch_add(1, std::memory_order_acq_rel);
+			}
+			// Try to swap the head chunk with our new chunk
+			FChunk* MaybeDifferentHead = (FChunk*)FPlatformAtomics::InterlockedCompareExchangePointer((void**)&Head, NewChunk, OldHead);
+			if (MaybeDifferentHead != OldHead)
+			{
+				// Someone got here first and added new chunk before we did
+				int32 MaybeFreeIndex = MaybeDifferentHead->NumItems.fetch_add(1, std::memory_order_acq_rel);
+				if (MaybeFreeIndex < NumElementsPerChunk)
+				{
+					// And it had a free index so dispose of the chunk we were going to add
+					delete NewChunk;
+					DestChunk = MaybeDifferentHead;
+					Index = MaybeFreeIndex;
+				}
+				// else keep trying to replace Head with NewChunk in the next iteration
+			}
+			else
+			{
+				// We successfully replaced Head with NewChunk
+				NewChunk->Next = OldHead;
+				DestChunk = NewChunk;
+			}
+		}
+
+		check(DestChunk != nullptr);
+		check(Index >= 0 && Index < NumElementsPerChunk);
+		DestChunk->Items[Index] = Item;
+	}
+
+	/**
+	* Not thread safe: checks if the list is empty
+	**/
+	FORCEINLINE bool IsEmpty() const
+	{
+		return Head == nullptr;
+	}
+
+	/**
+	* Not thread safe: Empties the list and frees its memory
+	**/
+	void Empty()
+	{
+		for (FChunk* Chunk = Head; Chunk;)
+		{
+			FChunk* NextChunk = Chunk->Next;
+			delete Chunk;
+			Chunk = NextChunk;
+		}
+		Head = nullptr;
+	}
+
+	/**
+	* Not thread safe: Moves all items to the provided array and empties the list
+	**/
+	FORCEINLINE void PopAllAndEmpty(TArray<T>& OutArray)
+	{
+		for (FChunk* Chunk = Head; Chunk; Chunk = Chunk->Next)
+		{
+			OutArray.Append(Chunk->Items, FMath::Min(Chunk->NumItems.load(std::memory_order_relaxed), NumElementsPerChunk));
+		}
+		Empty();
+	}
+};
diff --git a/Engine/Source/Runtime/CoreUObject/Private/UObject/GarbageCollection.cpp b/Engine/Source/Runtime/CoreUObject/Private/UObject/GarbageCollection.cpp
--- a/Engine/Source/Runtime/CoreUObject/Private/UObject/GarbageCollection.cpp
+++ b/Engine/Source/Runtime/CoreUObject/Private/UObject/GarbageCollection.cpp
@@ -39,6 +39,9 @@
 #include "UObject/GarbageCollectionHistory.h"
 #include "UObject/GarbageCollectionTesting.h"
 #include "UObject/PropertyOptional.h"
+#include "UObject/ExpandingChunkedList.h"
+#include "UObject/ReachabilityAnalysisState.h"
+#include "UObject/ReachabilityAnalysis.h"
 
 #include <atomic>
 
@@ -193,24 +196,16 @@ static FAutoConsoleVariableRef CMultithreadedDestructionEnabled(
 	ECVF_Default
 );
 
-static UObjectReachabilityStressData* GReachabilityStressData;
+static TArray<UObjectReachabilityStressData*> GReachabilityStressData;
 static void AllocateReachabilityStressData(FOutputDevice&)
 {
-	if (GReachabilityStressData)
-	{
-		return;
-	}
-	GReachabilityStressData = GenerateReachabilityStressData();
+	GenerateReachabilityStressData(GReachabilityStressData);
 }
 
 static void UnlinkReachabilityStressData(FOutputDevice&)
 {
-	if (!GReachabilityStressData)
-	{
-		return;
-	}
 	UnlinkReachabilityStressData(GReachabilityStressData);
-	GReachabilityStressData = nullptr;
+	GReachabilityStressData.Empty();
 }
 
 static FAutoConsoleCommandWithOutputDevice GGenerateReachabilityStressDataCmd(TEXT("gc.GenerateReachabilityStressData"),
@@ -235,6 +230,32 @@ static FAutoConsoleVariableRef CGarbageReferenceTrackingEnabled(
 );
 #endif // UE_BUILD_SHIPPING
 
+static int32 GAllowIncrementalReachability = 0;
+static FAutoConsoleVariableRef CVarAllowIncrementalGC(
+	TEXT("gc.AllowIncrementalReachability"),
+	GAllowIncrementalReachability,
+	TEXT("Set to control incremental Reachability Analysis (experimental)"),
+	ECVF_Default
+);
+
+static float GIncrementalReachabilityTimeLimit = 0.005f;
+static FAutoConsoleVariableRef CVarIncrementalReachabilityTimeLimit(
+	TEXT("gc.IncrementalReachabilityTimeLimit"),
+	GIncrementalReachabilityTimeLimit,
+	TEXT("Time in seconds (game time) we should allow for incremental GC (experimental)."),
+	ECVF_Default
+);
+
+#if VERIFY_DISREGARD_GC_ASSUMPTIONS
+static int32 GVerifyUnreachableObjects = 0;
+static FAutoConsoleVariableRef CVarVerifyUnreachableObjects(
+	TEXT("gc.VerifyUnreachableObjects"),
+	GVerifyUnreachableObjects,
+	TEXT("Enables or disables no unreachable objects are reachable verification"),
+	ECVF_Default
+);
+#endif // VERIFY_DISREGARD_GC_ASSUMPTIONS
+
 namespace UE::GC
 {
 
@@ -447,6 +468,11 @@ namespace UE::GC
 	}
 }
 
+namespace UE::GC::Private
+{
+	static TExpandingChunkedList<UObject*> GReachableObjects;
+	bool GIsIncrementalReachabilityPending = false;
+}
 
 /**
  * Helper class for destroying UObjects on a worker thread
@@ -823,6 +849,8 @@ void FContextPoolScope::ReturnToPool(FWorkerContext* Context)
 	--Pool.NumAllocated;
 	Context->FreeWorkerIndex();
 	Context->Stats = {};
+	Context->IncrementalStructs = {};
+	Context->bIsSuspended = false;
 	Pool.Reusable.Emplace(Context);
 }
 
@@ -877,9 +905,9 @@ static bool MarkClusterMutableObjectsAsReachable(FUObjectCluster& Cluster, Conta
 				if (!ReferencedMutableObjectItem->HasAnyFlags(EInternalObjectFlags::PendingKill | EInternalObjectFlags::Garbage))
 				PRAGMA_ENABLE_DEPRECATION_WARNINGS
 				{
-					if (ReferencedMutableObjectItem->IsUnreachable())
+					if (ReferencedMutableObjectItem->IsMaybeUnreachable())
 					{
-						if (ReferencedMutableObjectItem->ThisThreadAtomicallyClearedRFUnreachable())
+						if (ReferencedMutableObjectItem->ThisThreadAtomicallyClearedMaybeUnreachable())
 						{
 							// Needs doing because this is either a normal unclustered object (clustered objects are never unreachable) or a cluster root
 							ObjectsToSerialize.Add(static_cast<UObject*>(ReferencedMutableObjectItem->Object));
@@ -898,10 +926,10 @@ static bool MarkClusterMutableObjectsAsReachable(FUObjectCluster& Cluster, Conta
 						{
 							// Needs doing, we need to get its cluster root and process it too
 							FUObjectItem* ReferencedMutableObjectsClusterRootItem = GUObjectArray.IndexToObjectUnsafeForGC(ReferencedMutableObjectItem->GetOwnerIndex());
-							if (ReferencedMutableObjectsClusterRootItem->IsUnreachable())
+							if (ReferencedMutableObjectsClusterRootItem->IsMaybeUnreachable())
 							{
 								// The root is also maybe unreachable so process it and all the referenced clusters
-								if (ReferencedMutableObjectsClusterRootItem->ThisThreadAtomicallyClearedRFUnreachable())
+								if (ReferencedMutableObjectsClusterRootItem->ThisThreadAtomicallyClearedMaybeUnreachable())
 								{
 									MarkReferencedClustersAsReachable<Options>(ReferencedMutableObjectsClusterRootItem->GetClusterIndex(), ObjectsToSerialize);
 								}
@@ -939,9 +967,9 @@ static bool MarkClusterMutableObjectsAsReachable(FUObjectCluster& Cluster, Conta
 						
 					// If the root is also unreachable, process it and all its referenced clusters
 					FUObjectItem* ReferencedMutableObjectsClusterRootItem = GUObjectArray.IndexToObjectUnsafeForGC(ReferencedMutableObjectItem->GetOwnerIndex());
-					if (ReferencedMutableObjectsClusterRootItem->IsUnreachable())
+					if (ReferencedMutableObjectsClusterRootItem->IsMaybeUnreachable())
 					{
-						ReferencedMutableObjectsClusterRootItem->ClearFlags(EInternalObjectFlags::Unreachable);
+						ReferencedMutableObjectsClusterRootItem->ClearFlags(EInternalObjectFlags::MaybeUnreachable);
 						MarkReferencedClustersAsReachable<Options>(ReferencedMutableObjectsClusterRootItem->GetClusterIndex(), ObjectsToSerialize);
 					}
 				}
@@ -984,11 +1012,11 @@ static FORCENOINLINE void MarkReferencedClustersAsReachable(int32 ClusterIndex,
 				{
 					if constexpr (IsParallel(Options))
 					{
-						ReferencedClusterRootObjectItem->ThisThreadAtomicallyClearedFlag(EInternalObjectFlags::Unreachable);
+						ReferencedClusterRootObjectItem->ThisThreadAtomicallyClearedFlag(EInternalObjectFlags::MaybeUnreachable);
 					}
 					else
 					{
-						ReferencedClusterRootObjectItem->ClearFlags(EInternalObjectFlags::Unreachable);
+						ReferencedClusterRootObjectItem->ClearFlags(EInternalObjectFlags::MaybeUnreachable);
 					}
 				}
 			}
@@ -1078,9 +1106,7 @@ struct FReferenceMetadata
 
 //////////////////////////////////////////////////////////////////////////
 
-// Work-stealing algorithms are O(N^2), everyone steals from everyone.
-// Might want to improve that before going much wider.
-static constexpr int32 MaxWorkers = 16;
+static constexpr int32 MaxWorkers = FReachabilityAnalysisState::MaxWorkers;
 
 int32 GetNumSlowAROs();
 
@@ -1306,7 +1332,6 @@ template<class RefType> RefType				MakeReference(UObject*& Object);
 template<> FORCEINLINE FImmutableReference	MakeReference(UObject*& Object) { return {Object}; }
 template<> FORCEINLINE FMutableReference	MakeReference(UObject*& Object) { return {&Object}; }
 
-using FStructArray = Private::FStructArray;
 using FStridedReferenceArray = Private::FStridedReferenceArray;
 using FStridedReferenceView = Private::FStridedReferenceView;
 
@@ -1629,16 +1654,28 @@ class FStructBlockifier
 {
 	UE_NONCOPYABLE(FStructBlockifier);
 public:
-	explicit FStructBlockifier(int32 WorkerIdx)
-	: WorkerIndex(WorkerIdx)
+	explicit FStructBlockifier(FWorkerContext& Context)
+	: WorkerIndex(Context.GetWorkerIndex())
 	{
-		AllocateWipBlock(nullptr);
+		if (Context.IncrementalStructs.ContainsBatchData())
+		{
+			Wip = Context.IncrementalStructs.Wip;
+			WipIt = Context.IncrementalStructs.WipIt;
+			Context.IncrementalStructs = {};
+		}
+		else
+		{
+			AllocateWipBlock(nullptr);
+		}
 	}
 
 	~FStructBlockifier()
 	{
-		check(!CanPop());
-		FreeBlock(Wip);
+		if (Wip)
+		{
+			check(!CanPop());
+			FreeBlock(Wip);
+		}
 	}
 	
 	void Push(FStructArray AoS)
@@ -1699,6 +1736,14 @@ public:
 		GScratchPages.ReturnWorkerPage(WorkerIndex, Block);
 	}
 
+	FSuspendedStructBatch Suspend()
+	{
+		FSuspendedStructBatch SuspendData = { Wip, WipIt };
+		Wip = nullptr;
+		WipIt = nullptr;
+		return SuspendData;
+	}
+
 private:
 	FStructArray* WipIt; // Wip->Data cursor
 	FStructArrayBlock* Wip;
@@ -1763,7 +1808,7 @@ class FStructBatcher : public FBatcherBase
 	alignas (PLATFORM_CACHE_LINE_SIZE)	FStructBlockifier										ValidatedStructArrays;			// Drained by ProcessStructs, which feed back into TReferenceBatcher
 
 public:
-	explicit FStructBatcher(int32 WorkerIdx) : ValidatedStructArrays(WorkerIdx) {}
+	explicit FStructBatcher(FWorkerContext& Context) : ValidatedStructArrays(Context) {}
 
 	FORCEINLINE_DEBUGGABLE void PushSparseStructArray(FSchemaView Schema, FScriptSparseArray& Array)
 	{
@@ -1794,6 +1839,13 @@ public:
 	}
 
 	FStructBlockifier& GetUnboundedQueue() { return ValidatedStructArrays;	}
+
+	FSuspendedStructBatch Suspend()
+	{
+		FlushBoundedQueues();
+		return ValidatedStructArrays.Suspend();
+	}
+
 private:
 	FORCEINLINE_DEBUGGABLE void DrainSparseStructArrays(uint32 Num)
 	{
@@ -2701,9 +2753,9 @@ constexpr FORCEINLINE EKillable MayKill(EOrigin Origin, bool bAllowKill)
 }
 
 // Return whether flag was cleared. Only thread-safe for concurrent clear, not concurrent set+clear. Don't use during mark phase.
-FORCEINLINE static bool ClearUnreachableInterlocked(int32& Flags)
+FORCEINLINE static bool ClearMaybeUnreachableInterlocked(int32& Flags)
 {
-	static constexpr int32 FlagToClear = int32(EInternalObjectFlags::Unreachable);
+	static constexpr int32 FlagToClear = int32(EInternalObjectFlags::MaybeUnreachable);
 	if (FPlatformAtomics::AtomicRead_Relaxed(&Flags) & FlagToClear)
 	{
 		int32 Old = FPlatformAtomics::InterlockedAnd(&Flags, ~FlagToClear);
@@ -2713,8 +2765,7 @@ FORCEINLINE static bool ClearUnreachableInterlocked(int32& Flags)
 	return false;
 }
 
-
-
+FReachabilityAnalysisState GReachabilityState;
 
 template <EGCOptions InOptions>
 class TReachabilityProcessor 
@@ -2735,6 +2786,11 @@ public:
 		UE::GC::GStats.LogDetailedStatsSummary();
 	}
 
+	FORCEINLINE bool IsTimeLimitExceeded() const
+	{
+		return IsWithIncrementalReachabilityAnalysis() && GReachabilityState.IsTimeLimitExceeded();
+	}
+
 	static constexpr EGCOptions Options = InOptions;
 
 	static constexpr FORCEINLINE bool IsWithPendingKill() {	return !!(Options & EGCOptions::WithPendingKill);}
@@ -2745,6 +2801,8 @@ public:
 
 	constexpr static FORCEINLINE EKillable MayKill(EOrigin Origin, bool bAllowKill) { return UE::GC::MayKill<Options>(Origin, bAllowKill); }	
 
+	static constexpr FORCEINLINE bool IsWithIncrementalReachabilityAnalysis() { return !!(Options & EGCOptions::IncrementalReachability); }
+
 	static FORCEINLINE void ProcessReferenceDirectly(FWorkerContext& Context, FPermanentObjectPoolExtents PermanentPool, const UObject* ReferencingObject, UObject*& Object, FMemberId MemberId, EKillable Killable)
 	{
 		return Killable == EKillable::No	? ProcessReferenceDirectly<EKillable::No >(Context, PermanentPool, ReferencingObject, Object, MemberId)
@@ -2809,7 +2867,7 @@ public:
 
 	FORCEINLINE static bool HandleValidReference(FWorkerContext& Context, FImmutableReference Reference, FReferenceMetadata Metadata)
 	{
-		if (ClearUnreachableInterlocked(Metadata.ObjectItem->Flags))
+		if (ClearMaybeUnreachableInterlocked(Metadata.ObjectItem->Flags))
 		{
 			// Objects that are part of a GC cluster should never have the unreachable flag set!
 			checkSlow(Metadata.ObjectItem->GetOwnerIndex() <= 0);
@@ -2848,7 +2906,7 @@ public:
 				{
 					if constexpr (IsParallel(Options))
 					{
-						if (ClearUnreachableInterlocked(RootObjectItem->Flags))
+						if (ClearMaybeUnreachableInterlocked(RootObjectItem->Flags))
 						{
 							// Make sure all referenced clusters are marked as reachable too
 							MarkReferencedClustersAsReachableThunk<Options>(RootObjectItem->GetClusterIndex(), Context.ObjectsToSerialize);
@@ -2890,7 +2948,7 @@ struct TBatchDispatcher
 	, Collector(InCollector)
 	, KillableBatcher(InContext)
 	, ImmutableBatcher(InContext)
-	, StructBatcher(InContext.GetWorkerIndex())
+	, StructBatcher(InContext)
 	{}
 
 	// ProcessObjectArray API
@@ -3019,6 +3077,12 @@ struct TBatchDispatcher
 			}
 		}
 	}
+
+	void Suspend()
+	{
+		FlushQueuedReferences();
+		Context.IncrementalStructs = StructBatcher.Suspend();
+	}
 };
 
 //////////////////////////////////////////////////////////////////////////
@@ -3129,8 +3193,19 @@ public:
 		UE::GC::GStats.LogDetailedStatsSummary();
 	}
 
+	FORCEINLINE bool IsTimeLimitExceeded() const
+	{
+		if (IsWithIncrementalReachabilityAnalysis() && bTrackHistory && !bTrackGarbage)
+		{
+			return GReachabilityState.IsTimeLimitExceeded();
+		}
+		return false;
+	}
+
 	static constexpr EGCOptions Options = InOptions;
 
+	static constexpr FORCEINLINE bool IsWithIncrementalReachabilityAnalysis() { return !!(Options & EGCOptions::IncrementalReachability); }
+
 	TDebugReachabilityProcessor()
 	: bTrackGarbage(GGarbageReferenceTrackingEnabled != 0)
 	, bTrackHistory(FGCHistory::Get().IsActive())
@@ -3485,13 +3560,14 @@ void FReferenceFinder::HandleObjectReference( UObject*& InObject, const UObject*
 namespace UE::GC
 {
 
+
 class FRealtimeGC : public FGarbageCollectionTracer
 {
 	typedef void(FRealtimeGC::*MarkObjectsFn)(EObjectFlags);
 	typedef void(FRealtimeGC::*ReachabilityAnalysisFn)(FWorkerContext&);
 
 	/** Pointers to functions used for Marking objects as unreachable */
-	MarkObjectsFn MarkObjectsFunctions[4];
+	MarkObjectsFn MarkObjectsFunctions[2];
 	/** Pointers to functions used for Reachability Analysis */
 	ReachabilityAnalysisFn ReachabilityAnalysisFunctions[8];
 
@@ -3524,6 +3600,21 @@ class FRealtimeGC : public FGarbageCollectionTracer
 		return InitialReferences;
 	}
 
+	template<class CollectorType, class ProcessorType>
+	FORCEINLINE void CollectReferencesForGC(ProcessorType& Processor, UE::GC::FWorkerContext& Context)
+	{
+		using FastReferenceCollector = TFastReferenceCollector<ProcessorType, CollectorType>;
+
+		if constexpr (IsParallel(ProcessorType::Options))
+		{
+			ProcessAsync([](void* P, FWorkerContext& C) { FastReferenceCollector(*reinterpret_cast<ProcessorType*>(P)).ProcessObjectArray(C); }, &Processor, Context);
+		}
+		else
+		{
+			FastReferenceCollector(Processor).ProcessObjectArray(Context);
+		}
+	}
+
 	template <EGCOptions Options>
 	void PerformReachabilityAnalysisOnObjectsInternal(FWorkerContext& Context)
 	{
@@ -3534,20 +3625,21 @@ class FRealtimeGC : public FGarbageCollectionTracer
 		if (DebugProcessor.TracksHistory() | 
 			DebugProcessor.TracksGarbage() & Stats.bFoundGarbageRef)
 		{
-			CollectReferences<TDebugReachabilityCollector<Options>>(DebugProcessor, Context);
+			CollectReferencesForGC<TDebugReachabilityCollector<Options>>(DebugProcessor, Context);
 			return;
 		}
 #endif
 		
 		TReachabilityProcessor<Options> Processor;
-		CollectReferences<TReachabilityCollector<Options>>(Processor, Context);
+		CollectReferencesForGC<TReachabilityCollector<Options>>(Processor, Context);
 	}
 
 	/** Calculates GC function index based on current settings */
 	static FORCEINLINE int32 GetGCFunctionIndex(EGCOptions InOptions)
 	{
 		return (!!(InOptions & EGCOptions::Parallel)) |
-			(!!(InOptions & EGCOptions::WithPendingKill) << 1);
+			(!!(InOptions & EGCOptions::WithPendingKill) << 1) |
+			(!!(InOptions & EGCOptions::IncrementalReachability) << 2);
 	}
 
 public:
@@ -3555,7 +3647,7 @@ public:
 	FRealtimeGC()
 	{
 		MarkObjectsFunctions[GetGCFunctionIndex(EGCOptions::None)] = &FRealtimeGC::MarkObjectsAsUnreachable<false>;
-		MarkObjectsFunctions[GetGCFunctionIndex(EGCOptions::Parallel | EGCOptions::None)] = &FRealtimeGC::MarkObjectsAsUnreachable<true>;
+		MarkObjectsFunctions[GetGCFunctionIndex(EGCOptions::Parallel)] = &FRealtimeGC::MarkObjectsAsUnreachable<true>;
 
 		ReachabilityAnalysisFunctions[GetGCFunctionIndex(EGCOptions::None)] = &FRealtimeGC::PerformReachabilityAnalysisOnObjectsInternal<EGCOptions::None | EGCOptions::None>;
 		ReachabilityAnalysisFunctions[GetGCFunctionIndex(EGCOptions::Parallel | EGCOptions::None)] = &FRealtimeGC::PerformReachabilityAnalysisOnObjectsInternal<EGCOptions::Parallel | EGCOptions::None>;
@@ -3563,6 +3655,12 @@ public:
 		ReachabilityAnalysisFunctions[GetGCFunctionIndex(EGCOptions::None | EGCOptions::WithPendingKill)] = &FRealtimeGC::PerformReachabilityAnalysisOnObjectsInternal<EGCOptions::None | EGCOptions::WithPendingKill>;
 		ReachabilityAnalysisFunctions[GetGCFunctionIndex(EGCOptions::Parallel | EGCOptions::WithPendingKill)] = &FRealtimeGC::PerformReachabilityAnalysisOnObjectsInternal<EGCOptions::Parallel | EGCOptions::WithPendingKill>;
 
+		ReachabilityAnalysisFunctions[GetGCFunctionIndex(EGCOptions::None | EGCOptions::IncrementalReachability)] = &FRealtimeGC::PerformReachabilityAnalysisOnObjectsInternal<EGCOptions::None | EGCOptions::IncrementalReachability>;
+		ReachabilityAnalysisFunctions[GetGCFunctionIndex(EGCOptions::Parallel | EGCOptions::IncrementalReachability)] = &FRealtimeGC::PerformReachabilityAnalysisOnObjectsInternal<EGCOptions::Parallel | EGCOptions::IncrementalReachability>;
+
+		ReachabilityAnalysisFunctions[GetGCFunctionIndex(EGCOptions::WithPendingKill | EGCOptions::IncrementalReachability)] = &FRealtimeGC::PerformReachabilityAnalysisOnObjectsInternal<EGCOptions::None | EGCOptions::WithPendingKill | EGCOptions::IncrementalReachability>;
+		ReachabilityAnalysisFunctions[GetGCFunctionIndex(EGCOptions::Parallel | EGCOptions::WithPendingKill | EGCOptions::IncrementalReachability)] = &FRealtimeGC::PerformReachabilityAnalysisOnObjectsInternal<EGCOptions::Parallel | EGCOptions::WithPendingKill | EGCOptions::IncrementalReachability>;
+
 		FGCObject::StaticInit();
 	}
 
@@ -3621,6 +3719,9 @@ public:
 							GObjFinishDestroyHasBeenRoutedToAllObjects,
 							GGCObjectsPendingDestructionCount);
 
+					// Objects may still be marked as MaybeUnreachable if we forced a full purge in the middle of incremental GC
+					checkf(UE::GC::Private::GIsIncrementalReachabilityPending || !ObjectItem->IsMaybeUnreachable() || bIsRerun, TEXT("%s"), *Object->GetFullName());
+
 					// Keep track of how many objects are around.
 					ObjectCountDuringMarkPhase++;
 					
@@ -3687,7 +3788,7 @@ public:
 						}
 						else
 						{
-							ObjectItem->SetFlags(EInternalObjectFlags::Unreachable);
+							ObjectItem->SetFlags(EInternalObjectFlags::MaybeUnreachable);
 						}
 					}					
 				}
@@ -3742,9 +3843,9 @@ public:
 						FUObjectItem* RootObjectItem = GUObjectArray.IndexToObjectUnsafeForGC(OwnerIndex);
 						checkSlow(RootObjectItem->HasAnyFlags(EInternalObjectFlags::ClusterRoot));
 						// if it is reachable via keep flags we will do this below (or maybe already have)
-						if (RootObjectItem->IsUnreachable()) 
+						if (RootObjectItem->HasAnyFlags(EInternalObjectFlags::MaybeUnreachable)) 
 						{
-							RootObjectItem->ClearFlags(EInternalObjectFlags::Unreachable);
+							RootObjectItem->ClearFlags(EInternalObjectFlags::MaybeUnreachable);
 							// Make sure all referenced clusters are marked as reachable too
 							MarkReferencedClustersAsReachable<EGCOptions::None>(RootObjectItem->GetClusterIndex(), InitialObjects);
 						}
@@ -3761,20 +3862,37 @@ public:
 		}
 	}
 
-	/**
-	 * Performs reachability analysis.
-	 *
-	 * @param KeepFlags		Objects with these flags will be kept regardless of being referenced or not
-	 */
-	void PerformReachabilityAnalysis(EObjectFlags KeepFlags, const EGCOptions Options)
+private:
+
+	void ConditionallyAddBarrierReferencesToHistory(FWorkerContext& Context)
 	{
-		LLM_SCOPE(ELLMTag::GC);
+#if !UE_BUILD_SHIPPING && ENABLE_GC_HISTORY
+		if (FGCHistory::Get().IsActive())
+		{
+			FGCDirectReference BarrierReference;
+			BarrierReference.ReferencerName = TEXT("Barrier");
+			TArray<FGCDirectReference>*& DirectReferences = Context.History.FindOrAdd(FGCHistory::Get().GetBarrierObject());
+			if (DirectReferences == nullptr)
+			{
+				DirectReferences = new TArray<FGCDirectReference>();
+				DirectReferences->Reserve(InitialObjects.Num());
+			}
+			for (UObject* BarrierObject : InitialObjects)
+			{
+				BarrierReference.ReferencedObject = BarrierObject;
+				DirectReferences->Add(BarrierReference);
+			}
+		}
+#endif // !UE_BUILD_SHIPPING && ENABLE_GC_HISTORY
+	}
 
+	void StartReachabilityAnalysis(EObjectFlags KeepFlags, const EGCOptions Options)
+	{
 		BeginInitialReferenceCollection(Options);
 
 		// Reset object count.
 		GObjectCountDuringLastMarkPhase.Reset();
-		
+
 		InitialObjects.Reset();
 
 		// Make sure GC referencer object is checked for references to other objects even if it resides in permanent object pool
@@ -3786,25 +3904,73 @@ public:
 		{
 			const double StartTime = FPlatformTime::Seconds();
 			// Mark phase doesn't care about PendingKill being enabled or not so there's just fewer compiled in functions
-			const EGCOptions OptionsForMarkPhase = Options & ~EGCOptions::WithPendingKill;
+			const EGCOptions OptionsForMarkPhase = Options & ~(EGCOptions::WithPendingKill | EGCOptions::IncrementalReachability);
 			(this->*MarkObjectsFunctions[GetGCFunctionIndex(OptionsForMarkPhase)])(KeepFlags);
 			UE_LOG(LogGarbage, Verbose, TEXT("%f ms for MarkObjectsAsUnreachable Phase (%d Objects To Serialize)"), (FPlatformTime::Seconds() - StartTime) * 1000, InitialObjects.Num());
 		}
+	}
 
+	void PerformReachabilityAnalysisPass(const EGCOptions Options)
+	{
+		FContextPoolScope Pool;
+		FWorkerContext* Context = nullptr;
+
+		if (GReachabilityState.IsSuspended())
 		{
-			const double StartTime = FPlatformTime::Seconds();
-		
-			FContextPoolScope Pool;
-			FWorkerContext* Context = Pool.AllocateFromPool();
-			Context->InitialNativeReferences = GetInitialReferences(Options);
+			Context = GReachabilityState.GetContextArray()[0];
+		}
+		else
+		{
+			Context = Pool.AllocateFromPool();
+
+			if (!Private::GReachableObjects.IsEmpty())
+			{
+				InitialObjects.Reset();
+				Private::GReachableObjects.PopAllAndEmpty(InitialObjects);
+				UE_LOG(LogGarbage, Log, TEXT("GC restarting reachability due to Reachable Objects list not being empty (%d objects to process)"), InitialObjects.Num());
+				ConditionallyAddBarrierReferencesToHistory(*Context);
+			}
+			else
+			{
+				Context->InitialNativeReferences = GetInitialReferences(Options);
+			}
+
 			Context->SetInitialObjectsUnpadded(InitialObjects);
+		}
 
-			PerformReachabilityAnalysisOnObjects(Context, Options);
+		PerformReachabilityAnalysisOnObjects(Context, Options);
 
+		if (!GReachabilityState.CheckIfAnyContextIsSuspended())
+		{
 			Stats = Context->Stats;
-			Context->ResetInitialObjects();
-			Context->InitialNativeReferences = TConstArrayView<UObject**>();
 			Pool.ReturnToPool(Context);
+			GReachabilityState.UpdateStats(Stats);
+		}
+	}
+
+public:
+
+	/**
+	 * Performs reachability analysis.
+	 *
+	 * @param KeepFlags		Objects with these flags will be kept regardless of being referenced or not
+	 */
+	void PerformReachabilityAnalysis(EObjectFlags KeepFlags, const EGCOptions Options)
+	{
+		LLM_SCOPE(ELLMTag::GC);
+
+		if (!GReachabilityState.IsSuspended())
+		{
+			StartReachabilityAnalysis(KeepFlags, Options);
+		}
+
+		{
+			const double StartTime = FPlatformTime::Seconds();
+
+			do
+			{
+				PerformReachabilityAnalysisPass(Options);
+			} while (!Private::GReachableObjects.IsEmpty() && !GReachabilityState.IsSuspended());
 
 			UE_LOG(LogGarbage, Verbose, TEXT("%f ms for Reachability Analysis"), (FPlatformTime::Seconds() - StartTime) * 1000);
 		}
@@ -3812,7 +3978,10 @@ public:
 PRAGMA_DISABLE_DEPRECATION_WARNINGS
 		// Allowing external systems to add object roots. This can't be done through AddReferencedObjects
 		// because it may require tracing objects (via FGarbageCollectionTracer) multiple times
-		FCoreUObjectDelegates::TraceExternalRootsForReachabilityAnalysis.Broadcast(*this, KeepFlags, !(Options & EGCOptions::Parallel));
+		if (!GReachabilityState.IsSuspended())
+		{
+			FCoreUObjectDelegates::TraceExternalRootsForReachabilityAnalysis.Broadcast(*this, KeepFlags, !(Options & EGCOptions::Parallel));
+		}
 PRAGMA_ENABLE_DEPRECATION_WARNINGS
 	}
 
@@ -4036,7 +4205,8 @@ bool IncrementalDestroyGarbage(bool bUseTimeLimit, double TimeLimit)
 
 				//@todo UE - A prefetch was removed here. Re-add it. It wasn't right anyway, since it was ten items ahead and the consoles on have 8 prefetch slots
 
-				check(ObjectItem->IsUnreachable());
+				check(!ObjectItem->IsMaybeUnreachable());
+				if (ObjectItem->IsUnreachable())
 				{
 					UObject* Object = static_cast<UObject*>(ObjectItem->Object);
 					// Object should always have had BeginDestroy called on it and never already be destroyed
@@ -4097,7 +4267,7 @@ bool IncrementalDestroyGarbage(bool bUseTimeLimit, double TimeLimit)
 					UObject* Object = GGCObjectsPendingDestruction[ CurPendingObjIndex ];
 
 					// Object should never have been added to the list if it failed this criteria
-					check( Object != NULL && Object->IsUnreachable() );
+					check( Object != nullptr && Object->IsUnreachable() );
 
 					// Object should always have had BeginDestroy called on it and never already be destroyed
 					check( Object->HasAnyFlags( RF_BeginDestroyed ) && !Object->HasAnyFlags( RF_FinishDestroyed ) );
@@ -4354,8 +4524,10 @@ void GatherUnreachableObjects(bool bForceSingleThreaded)
 		for (int32 ObjectIndex = FirstObjectIndex; ObjectIndex <= LastObjectIndex; ++ObjectIndex)
 		{
 			FUObjectItem* ObjectItem = &GUObjectArray.GetObjectItemArrayUnsafe()[ObjectIndex];
-			if (ObjectItem->IsUnreachable())
+			if (ObjectItem->HasAnyFlags(EInternalObjectFlags::MaybeUnreachable | EInternalObjectFlags::Unreachable))
 			{
+				ObjectItem->ClearMaybeUnreachable();
+				ObjectItem->SetUnreachable();
 				ThisThreadUnreachableObjects.Add(ObjectItem);
 				if (ObjectItem->HasAnyFlags(EInternalObjectFlags::ClusterRoot))
 				{
@@ -4390,7 +4562,7 @@ void GatherUnreachableObjects(bool bForceSingleThreaded)
 
 				if (!ClusterObjectItem->HasAnyFlags(EInternalObjectFlags::ReachableInCluster))
 				{
-					ClusterObjectItem->SetFlags(EInternalObjectFlags::Unreachable);
+					ClusterObjectItem->SetUnreachable();
 					ClusterObjects++;
 					GUnreachableObjects.Add(ClusterObjectItem);
 				}
@@ -4534,18 +4706,23 @@ FORCEINLINE void CollectGarbageInternal(EObjectFlags KeepFlags, bool bPerformFul
 {
 	const double StartTime = FPlatformTime::Seconds();
 
-	if (bPerformFullPurge)
-	{
-		CollectGarbageFull(KeepFlags);
-	}
-	else
-	{
-		CollectGarbageIncremental(KeepFlags);
-	}
+	GReachabilityState.CollectGarbage(KeepFlags, bPerformFullPurge);
 
 	GTimingInfo.LastGCDuration = FPlatformTime::Seconds() - StartTime;
 }
 
+EGCOptions GetReferenceCollectorOptions(bool bPerformFullPurge)
+{
+	return
+		// Fall back to single threaded GC if processor count is 1 or parallel GC is disabled
+		// or detailed per class gc stats are enabled (not thread safe)
+		(ShouldForceSingleThreadedGC() ? EGCOptions::None : EGCOptions::Parallel) |
+		// Toggle between PendingKill enabled or disabled
+		(UObjectBaseUtility::IsPendingKillEnabled() ? EGCOptions::WithPendingKill : EGCOptions::None) |
+		// Toggle between Incremental Reachability enabled or disabled
+		((GAllowIncrementalReachability && !bPerformFullPurge) ? EGCOptions::IncrementalReachability : EGCOptions::None);
+}
+
 /** 
  * Deletes all unreferenced objects, keeping objects that have any of the passed in KeepFlags set
  *
@@ -4553,8 +4730,10 @@ FORCEINLINE void CollectGarbageInternal(EObjectFlags KeepFlags, bool bPerformFul
  * @param	bPerformFullPurge	if true, perform a full purge after the mark pass
  */
 template<bool bPerformFullPurge>
-void CollectGarbageImpl(EObjectFlags KeepFlags)
+void PreCollectGarbageImpl(EObjectFlags KeepFlags)
 {
+	using namespace UE::GC::Private;
+
 	FGCCSyncObject::Get().ResetGCIsWaiting();
 
 #if defined(WITH_CODE_GUARD_HANDLER) && WITH_CODE_GUARD_HANDLER
@@ -4572,7 +4751,7 @@ void CollectGarbageImpl(EObjectFlags KeepFlags)
 	GNumAttemptsSinceLastGC = 0;
 
 	// Flush streaming before GC if requested
-	if (GFlushStreamingOnGC && IsAsyncLoading())
+	if (!GIsIncrementalReachabilityPending && GFlushStreamingOnGC && IsAsyncLoading())
 	{
 		UE_LOG(LogGarbage, Log, TEXT("CollectGarbageInternal() is flushing async loading"));
 		ReleaseGCLock();
@@ -4582,6 +4761,7 @@ void CollectGarbageImpl(EObjectFlags KeepFlags)
 
 	// Route callbacks so we can ensure that we are e.g. not in the middle of loading something by flushing
 	// the async loading, etc...
+	if (!GIsIncrementalReachabilityPending)
 	{
 		TRACE_CPUPROFILER_EVENT_SCOPE(BroadcastPreGarbageCollect);
 		FCoreUObjectDelegates::GetPreGarbageCollectDelegate().Broadcast();
@@ -4589,13 +4769,7 @@ void CollectGarbageImpl(EObjectFlags KeepFlags)
 	GLastGCFrame = GFrameCounter;
 
 	{
-		LLM_SCOPE(ELLMTag::GC);
-
-		// Set 'I'm garbage collecting' flag - might be checked inside various functions.
-		// This has to be unlocked before we call post GC callbacks
-		TGuardValue<bool> GuardIsGarbageCollecting(GIsGarbageCollecting, true);
-
-		UE_LOG(LogGarbage, Log, TEXT("Collecting garbage%s"), IsAsyncLoading() ? TEXT(" while async loading") : TEXT(""));
+		GIsGarbageCollecting = true;
 
 		// Make sure previous incremental purge has finished or we do a full purge pass in case we haven't kicked one
 		// off yet since the last call to garbage collection.
@@ -4608,12 +4782,17 @@ void CollectGarbageImpl(EObjectFlags KeepFlags)
 			}
 		}
 
-		// Reachability analysis.
-		// When exiting this scope all objects to be destroyed have been marked unreachable and any weak references have been cleared.
-		{
-			// With the new FGCLockBehavior::Default behavior hash tables are only locked during this scope of reachability analysis.
-			FGCHashTableScopeLock GCHashTableLock;
 
+		UE_LOG(LogGarbage, Log, TEXT("%sCollecting garbage%s"),
+			(!bPerformFullPurge && GAllowIncrementalReachability) ? TEXT("Incrementally ") : TEXT(""),
+			IsAsyncLoading() ? TEXT(" while async loading") : TEXT(""));
+
+		// The hash tables are only locked during this scope of reachability analysis.
+		GIsGarbageCollectingAndLockingUObjectHashTables = true;
+		LockUObjectHashTables();
+
+		if (!GIsIncrementalReachabilityPending)
+		{
 			check(!GObjIncrementalPurgeIsInProgress);
 			check(!GObjPurgeIsRequired);
 
@@ -4641,73 +4820,82 @@ void CollectGarbageImpl(EObjectFlags KeepFlags)
 				UE_LOG(LogGarbage, Log, TEXT("%f ms for Verify GC Assumptions"), (FPlatformTime::Seconds() - StartTime) * 1000);
 			}
 #endif
+		}
+	}
+}
 
-			const EGCOptions Options = 
-				// Fall back to single threaded GC if processor count is 1 or parallel GC is disabled
-				// or detailed per class gc stats are enabled (not thread safe)
-				(ShouldForceSingleThreadedGC() ? EGCOptions::None : EGCOptions::Parallel) |
-				// Toggle between PendingKill enabled or disabled
-				(UObjectBaseUtility::IsPendingKillEnabled() ? EGCOptions::WithPendingKill : EGCOptions::None);
+template<bool bPerformFullPurge>
+void CollectGarbageImpl(EObjectFlags KeepFlags)
+{
+	{
+		// Reachability analysis.
+		{
+			const EGCOptions Options = GetReferenceCollectorOptions(bPerformFullPurge);
 
 			// Perform reachability analysis.
-			{
-				FRealtimeGC GC;
-				{
-					SCOPED_NAMED_EVENT(FRealtimeGC_PerformReachabilityAnalysis, FColor::Red);
-					DECLARE_SCOPE_CYCLE_COUNTER(TEXT("FRealtimeGC::PerformReachabilityAnalysis"), STAT_FArchiveRealtimeGC_PerformReachabilityAnalysis, STATGROUP_GC);
-					const double StartTime = FPlatformTime::Seconds();
-					GC.PerformReachabilityAnalysis(KeepFlags, Options);
-					const double Ms = (FPlatformTime::Seconds() - StartTime) * 1000;
-					UE_LOG(LogGarbage, Log, TEXT("%.2f ms for GC - %d refs/ms while processing %d references from %d objects with %d clusters"),
-							Ms, (int32)(GC.Stats.NumReferences / Ms), GC.Stats.NumReferences, GC.Stats.NumObjects, GUObjectClusters.GetNumAllocatedClusters());
-				}
+			FRealtimeGC GC;
+			GC.PerformReachabilityAnalysis(KeepFlags, Options);
+		}
+	}
+}
 
-				if (GC.Stats.bFoundGarbageRef && GGarbageReferenceTrackingEnabled > 0)
-				{
-					
-					CSV_SCOPED_TIMING_STAT_EXCLUSIVE(GarbageCollectionDebug);
-					SCOPED_NAMED_EVENT(FRealtimeGC_PerformReachabilityAnalysisRerun, FColor::Orange);
-					DECLARE_SCOPE_CYCLE_COUNTER(TEXT("FRealtimeGC::PerformReachabilityAnalysisRerun"), STAT_FArchiveRealtimeGC_PerformReachabilityAnalysisRerun, STATGROUP_GC);
-					const double StartTime = FPlatformTime::Seconds();
-					GC.PerformReachabilityAnalysis(KeepFlags, Options);
-					UE_LOG(LogGarbage, Log, TEXT("%.2f ms for GC rerun to track garbage references (gc.GarbageReferenceTrackingEnabled=%d)"), (FPlatformTime::Seconds() - StartTime) * 1000, GGarbageReferenceTrackingEnabled);
-				}
-			}
+template<bool bPerformFullPurge>
+void PostCollectGarbageImpl(EObjectFlags KeepFlags)
+{
+	using namespace UE::GC::Private;
 
+	if (!GIsIncrementalReachabilityPending)
+	{
+		FContextPoolScope ContextPool;
+		TConstArrayView<TUniquePtr<FWorkerContext>> AllContexts = ContextPool.PeekFree();
+		// This needs to happen before clusters get dissolved otherwisise cluster information will be missing from history
+		UpdateGCHistory(AllContexts);
 
-			FContextPoolScope ContextPool;
-			TConstArrayView<TUniquePtr<FWorkerContext>> AllContexts = ContextPool.PeekFree();
-			// This needs to happen before clusters get dissolved otherwisise cluster information will be missing from history
-			UpdateGCHistory(AllContexts);
+		// Reconstruct clusters if needed
+		if (GUObjectClusters.ClustersNeedDissolving())
+		{
+			const double StartTime = FPlatformTime::Seconds();
+			GUObjectClusters.DissolveClusters();
+			UE_LOG(LogGarbage, Log, TEXT("%f ms for dissolving GC clusters"), (FPlatformTime::Seconds() - StartTime) * 1000);
+		}
 
-			// Reconstruct clusters if needed
-			if (GUObjectClusters.ClustersNeedDissolving())
-			{
-				const double StartTime = FPlatformTime::Seconds();
-				GUObjectClusters.DissolveClusters();
-				UE_LOG(LogGarbage, Log, TEXT("%f ms for dissolving GC clusters"), (FPlatformTime::Seconds() - StartTime) * 1000);
-			}
+		DumpGarbageReferencers(AllContexts);
 
-			DumpGarbageReferencers(AllContexts);
+		GatherUnreachableObjects(!(GetReferenceCollectorOptions(bPerformFullPurge) & EGCOptions::Parallel));
 
-			GatherUnreachableObjects(!(Options & EGCOptions::Parallel));
+		// This needs to happen after GatherUnreachableObjects since it can mark more objects as unreachable
+		ClearWeakReferences(AllContexts);
 
+		if (bPerformFullPurge)
+		{
+			ContextPool.Cleanup();
+		}
+	}
 
-			// This needs to happen after GatherUnreachableObjects since it can mark more objects as unreachable
-			ClearWeakReferences(AllContexts);
+	GIsGarbageCollectingAndLockingUObjectHashTables = false;
+	UnlockUObjectHashTables();
 
-			if (bPerformFullPurge)
-			{
-				ContextPool.Cleanup();
-			}
+	GIsGarbageCollecting = false;
 
-			NotifyUnreachableObjects(GUnreachableObjects);
-		}
+	if (!GIsIncrementalReachabilityPending)
+	{
+		NotifyUnreachableObjects(GUnreachableObjects);
+	}
 
-		// The hash tables lock was released when exiting the reachability analysis scope above.
-		// BeginDestroy, FinishDestroy, destructors and callbacks are allowed to call functions like StaticAllocateObject and StaticFindObject.
-		// Now release the GC lock to allow async loading and other threads to perform UObject operations under the FGCScopeGuard.
-		ReleaseGCLock();
+	// The hash tables lock was released when reachability analysis was done.
+	// BeginDestroy, FinishDestroy, destructors and callbacks are allowed to call functions like StaticAllocateObject and StaticFindObject.
+	// Now release the GC lock to allow async loading and other threads to perform UObject operations under the FGCScopeGuard.
+	ReleaseGCLock();
+
+	if (!GIsIncrementalReachabilityPending)
+	{
+#if VERIFY_DISREGARD_GC_ASSUMPTIONS
+		if (GVerifyUnreachableObjects != 0 && GUnreachableObjects.Num() > 0)
+		{
+			DECLARE_SCOPE_CYCLE_COUNTER(TEXT("CollectGarbageInternal.VerifyUnreachableObjects"), STAT_CollectGarbageInternal_VerifyUnreachableObjects, STATGROUP_GC);
+			VerifyUnreachableObjects();
+		}
+#endif // VERIFY_DISREGARD_GC_ASSUMPTIONS
 
 		// Fire post-reachability analysis hooks
 		{
@@ -4742,20 +4930,207 @@ void CollectGarbageImpl(EObjectFlags KeepFlags)
 		{
 			FMemory::Trim();
 		}
+
+		// Route callbacks to verify GC assumptions
+		{
+			TRACE_CPUPROFILER_EVENT_SCOPE(BroadcastPostGarbageCollect);
+			FCoreUObjectDelegates::GetPostGarbageCollect().Broadcast();
+		}
 	}
 
-	// Route callbacks to verify GC assumptions
+	GTimingInfo.LastGCTime = FPlatformTime::Seconds();
+	STAT_ADD_CUSTOMMESSAGE_NAME(STAT_NamedMarker, TEXT("GarbageCollection - End"));
+}
+
+void FReachabilityAnalysisState::CollectGarbage(EObjectFlags KeepFlags, bool bFullPurge)
+{
+	using namespace UE::GC::Private;
+
+	if (GIsIncrementalReachabilityPending)
 	{
-		TRACE_CPUPROFILER_EVENT_SCOPE(BroadcastPostGarbageCollect);
-		FCoreUObjectDelegates::GetPostGarbageCollect().Broadcast();
+		// Something triggered a new GC run but we're in the middle of incremental reachability analysis.
+		// Finish the current GC pass (including purging all unreachable objects) and then kick off another GC run as requested
+		bPerformFullPurge = true;
+		PerformReachabilityAnalysisAndConditionallyPurgeGarbage(/*bReachabilityUsingTimeLimit =*/ false);
+
+		checkf(!GIsIncrementalReachabilityPending, TEXT("Flushing incremental reachability analysis did not complete properly"));
+
+		// Need to acquire GC lock again as it was released in PerformReachabilityAnalysisAndConditionallyPurgeGarbage() -> UE::GC::PostCollectGarbageImpl()
+		AcquireGCLock();
 	}
 
-	GTimingInfo.LastGCTime = FPlatformTime::Seconds();
-	STAT_ADD_CUSTOMMESSAGE_NAME( STAT_NamedMarker, TEXT( "GarbageCollection - End" ) );
+	ObjectKeepFlags = KeepFlags;
+	bPerformFullPurge = bFullPurge;
+
+	const bool bReachabilityUsingTimeLimit = !bFullPurge && GAllowIncrementalReachability;
+	PerformReachabilityAnalysisAndConditionallyPurgeGarbage(bReachabilityUsingTimeLimit);
+}
+
+void FReachabilityAnalysisState::PerformReachabilityAnalysisAndConditionallyPurgeGarbage(bool bReachabilityUsingTimeLimit)
+{
+	using namespace UE::GC::Private;
+
+	LLM_SCOPE(ELLMTag::GC);
+
+	if (bPerformFullPurge)
+	{
+		UE::GC::PreCollectGarbageImpl<true>(ObjectKeepFlags);
+	}
+	else
+	{
+		UE::GC::PreCollectGarbageImpl<false>(ObjectKeepFlags);
+	}	
+	
+	const bool bForceNonIncrementalReachability = 
+		!GIsIncrementalReachabilityPending &&
+		(bPerformFullPurge || !GAllowIncrementalReachability);
+
+	{
+		// When incremental reachability is enabled we start the timer before acquiring GC lock
+		// and here we keep track of reference processing time only
+		const double ReferenceProcessingStartTime = FPlatformTime::Seconds();
+		// When performing the first iteration of reachability analysis start the timer when we actually start processing
+		// iteration as we don't have control over various callbacks being fired in PreCollectGarbageImpl and can't be responsible for any hitches in them
+		if (IterationStartTime == 0.0)
+		{
+			IterationStartTime = ReferenceProcessingStartTime;
+			IterationTimeLimit = bReachabilityUsingTimeLimit ? GIncrementalReachabilityTimeLimit : 0.0;
+		}
+
+		SCOPED_NAMED_EVENT(FRealtimeGC_PerformReachabilityAnalysis, FColor::Red);
+		DECLARE_SCOPE_CYCLE_COUNTER(TEXT("FRealtimeGC::PerformReachabilityAnalysis"), STAT_FArchiveRealtimeGC_PerformReachabilityAnalysis, STATGROUP_GC);
+
+		if (bForceNonIncrementalReachability)
+		{
+			IncrementalMarkPhaseTotalTime = 0.0;
+			ReferenceProcessingTotalTime = 0.0;
+			PerformReachabilityAnalysis();
+		}
+		else
+		{
+			if (!GIsIncrementalReachabilityPending)
+			{
+				UE_LOG(LogGarbage, Log, TEXT("Starting incremental GC"));
+				ReferenceProcessingTotalTime = 0.0;
+				IncrementalMarkPhaseTotalTime = 0.0;
+			}
+			else
+			{
+				UE_LOG(LogGarbage, Log, TEXT("Resuming incremental GC"));
+			}
+
+			PerformReachabilityAnalysis();
+		}
+
+		GIsIncrementalReachabilityPending = GReachabilityState.IsSuspended();
+
+		const double CurrentTime = FPlatformTime::Seconds();
+		const double ReferenceProcessingElapsedTime = CurrentTime - ReferenceProcessingStartTime;
+		const double ElapsedTime = CurrentTime - IterationStartTime;
+		ReferenceProcessingTotalTime += ReferenceProcessingElapsedTime;
+		IncrementalMarkPhaseTotalTime += ElapsedTime;
+
+		if (GIsIncrementalReachabilityPending)
+		{
+			UE_LOG(LogGarbage, Log, TEXT("GC Reachability Analysis iteration time: %f.2 ms (%f.2 ms on reference traversal)"), ElapsedTime * 1000, ReferenceProcessingElapsedTime * 1000);
+			const double SuspendLatency = CurrentTime - (IterationStartTime + IterationTimeLimit);
+			UE_LOG(LogGarbage, Log, TEXT("GC suspend latency: %f ms"), SuspendLatency * 1000.0);
+		}
+		else
+		{
+			const double ReferenceProcessingTotalTimeMs = ReferenceProcessingTotalTime * 1000;
+			const double IncrementalMarkPhaseTotalTimeMs = IncrementalMarkPhaseTotalTime * 1000;			
+			UE_LOG(LogGarbage, Log, TEXT("GC Reachability Analysis total time: %f.2 ms (%f.2 ms on reference traversal)"), IncrementalMarkPhaseTotalTimeMs, ReferenceProcessingTotalTimeMs);
+			UE_LOG(LogGarbage, Log, TEXT("%.2f ms for %sGC - %d refs/ms while processing %d references from %d objects with %d clusters"),
+				IncrementalMarkPhaseTotalTimeMs,
+				bForceNonIncrementalReachability ? TEXT("") : TEXT("Incremental "),
+				(int32)(Stats.NumReferences / ReferenceProcessingTotalTimeMs),
+				Stats.NumReferences, 
+				Stats.NumObjects, 
+				GUObjectClusters.GetNumAllocatedClusters());			
+		}
+
+		// Reset timer and the time limit. These values will be set to their target values in the next iteration but we don't want 
+		// them to be set before the debug reachability run below
+		IterationTimeLimit = 0.0;
+		IterationStartTime = 0.0;
+	}
+
+	if (!GIsIncrementalReachabilityPending && Stats.bFoundGarbageRef && GGarbageReferenceTrackingEnabled > 0)
+	{
+		CSV_SCOPED_TIMING_STAT_EXCLUSIVE(GarbageCollectionDebug);
+		SCOPED_NAMED_EVENT(FRealtimeGC_PerformReachabilityAnalysisRerun, FColor::Orange);
+		DECLARE_SCOPE_CYCLE_COUNTER(TEXT("FRealtimeGC::PerformReachabilityAnalysisRerun"), STAT_FArchiveRealtimeGC_PerformReachabilityAnalysisRerun, STATGROUP_GC);		
+		const double StartTime = FPlatformTime::Seconds();
+		{
+			TGuardValue<bool> GuardReachabilityUsingTimeLimit(bReachabilityUsingTimeLimit, false);
+			FRealtimeGC GC;
+			GC.Stats = Stats; // This is to pass Stats.bFoundGarbageRef to CG
+			GC.PerformReachabilityAnalysis(ObjectKeepFlags, GetReferenceCollectorOptions(bPerformFullPurge));
+		}
+		UE_LOG(LogGarbage, Log, TEXT("%.2f ms for GC rerun to track garbage references (gc.GarbageReferenceTrackingEnabled=%d)"), (FPlatformTime::Seconds() - StartTime) * 1000, GGarbageReferenceTrackingEnabled);
+	}
+	// Maybe purge garbage (if we're done with incremental reachability and there's still time left)
+
+	if (bPerformFullPurge)
+	{
+		UE::GC::PostCollectGarbageImpl<true>(ObjectKeepFlags);
+	}
+	else
+	{
+		UE::GC::PostCollectGarbageImpl<false>(ObjectKeepFlags);
+	}
+}
+
+void FReachabilityAnalysisState::PerformReachabilityAnalysis()
+{
+	if (bPerformFullPurge)
+	{
+		UE::GC::CollectGarbageFull(ObjectKeepFlags);
+	}
+	else
+	{
+		UE::GC::CollectGarbageIncremental(ObjectKeepFlags);
+	}
 }
 
 } // namespace UE::GC
 
+void SetIncrementalReachabilityAnalysisEnabled(bool bEnabled)
+{
+	GAllowIncrementalReachability = bEnabled;
+}
+
+bool GetIncrementalReachabilityAnalysisEnabled()
+{
+	return !!GAllowIncrementalReachability;
+}
+
+void SetReachabilityAnalysisTimeLimit(float TimeLimitSeconds)
+{
+	GIncrementalReachabilityTimeLimit = TimeLimitSeconds;
+}
+
+float GetReachabilityAnalysisTimeLimit()
+{
+	return GIncrementalReachabilityTimeLimit;
+}
+
+bool IsIncrementalReachabilityAnalysisPending()
+{
+	return UE::GC::Private::GIsIncrementalReachabilityPending;
+}
+
+void PerformIncrementalReachabilityAnalysis()
+{
+	checkf(UE::GC::Private::GIsIncrementalReachabilityPending, TEXT("Incremental reachability must be pending to perform its next iteration"));
+	// When performing Reachability Analysis iterations start the internal timer before acquiring GC lock 
+	// so that we don't spend more time than GIncrementalReachabilityTimeLimit on fully completing an iteration
+	UE::GC::GReachabilityState.StartTimer(GIncrementalReachabilityTimeLimit);
+	AcquireGCLock();	
+	UE::GC::GReachabilityState.PerformReachabilityAnalysisAndConditionallyPurgeGarbage(/*bReachabilityUsingTimeLimit = */ true);
+}
+
 FString FGarbageReferenceInfo::GetReferencingObjectInfo() const
 {
 	if (bReferencerUObject)
@@ -4887,32 +5262,34 @@ bool TryCollectGarbage(EObjectFlags KeepFlags, bool bPerformFullPurge)
 		return false;
 	}
 
-	// No other thread may be performing UObject operations while we're running
-	bool bCanRunGC = FGCCSyncObject::Get().TryGCLock();
-	if (!bCanRunGC)
+	// No other thread may be performing UObject operations while we're running so try to acquire GC lock
+	if (UE::GC::Private::GIsIncrementalReachabilityPending)
+	{
+		// Since we're already in the middle of a previous GC acquire GC lock even if it means we have to block main thread
+		AcquireGCLock();
+	}	
+	else if (!FGCCSyncObject::Get().TryGCLock())
 	{
 		if (GNumRetriesBeforeForcingGC > 0 && GNumAttemptsSinceLastGC > GNumRetriesBeforeForcingGC)
 		{
-			// Force GC and block main thread			
+			// Force acquire GC lock and block main thread		
 			UE_LOG(LogGarbage, Warning, TEXT("TryCollectGarbage: forcing GC after %d skipped attempts."), GNumAttemptsSinceLastGC);
 			GNumAttemptsSinceLastGC = 0;
 			AcquireGCLock();
-			bCanRunGC = true;
+		}
+		else
+		{
+			++GNumAttemptsSinceLastGC;
+			return false;
 		}
 	}
-	if (bCanRunGC)
-	{ 
-		// Perform actual garbage collection
-		UE::GC::CollectGarbageInternal(KeepFlags, bPerformFullPurge);
 
-		// GC lock was released after reachability analysis inside CollectGarbageInternal
-	}
-	else
-	{
-		GNumAttemptsSinceLastGC++;
-	}
+	// Perform actual garbage collection
+	UE::GC::CollectGarbageInternal(KeepFlags, bPerformFullPurge);
+
+	// GC lock was released after reachability analysis inside CollectGarbageInternal
 
-	return bCanRunGC;
+	return true;
 }
 
 void UObject::CallAddReferencedObjects(FReferenceCollector& Collector)
@@ -4930,6 +5307,23 @@ bool UObject::IsDestructionThreadSafe() const
 	return false;
 }
 
+void UObject::MarkAsReachable() const
+{	
+	using namespace UE::GC::Private;
+
+	FUObjectItem* ObjectItem = GUObjectArray.ObjectToObjectItem(const_cast<UObject*>(this));
+	if (ObjectItem->IsMaybeUnreachable())
+	{
+		checkf(GIsIncrementalReachabilityPending, TEXT("%s is marked as MaybeUnreachable but Incremental Reachability Analysis is not in progress"), *GetFullName());
+		if (ObjectItem->ThisThreadAtomicallyClearedMaybeUnreachable())
+		{
+			// This object became reachable so add it to a list of new objects to process in the next iteration of incremental GC because
+			// we need to mark objects it's referencing as reachable too
+			GReachableObjects.Push(const_cast<UObject*>(this));
+		}
+	}
+}
+
 /*-----------------------------------------------------------------------------
 	Implementation of realtime garbage collection helper functions in 
 	FProperty, UClass, ...
@@ -5437,6 +5831,11 @@ public:
 		return false;
 	}
 
+	void Suspend()
+	{
+		++NumStopped;
+	}
+
 	FORCENOINLINE void SpinUntilAllStopped()
 	{
 		while (NumStopped.load() < Contexts.Num())
@@ -5446,7 +5845,13 @@ public:
 	}
 };
 
-ELoot StealWork(FWorkerContext& Context, FReferenceCollector& Collector, FWorkBlock*& OutBlock)
+void SuspendWork(FWorkerContext& Context)
+{
+	Context.bIsSuspended = true;
+	Context.Coordinator->Suspend();
+}
+
+ELoot StealWork(FWorkerContext& Context, FReferenceCollector& Collector, FWorkBlock*& OutBlock, EGCOptions Options)
 {
 	FWorkBlockifier& RemainingObjects = Context.ObjectsToSerialize;
 	FWorkCoordinator& Tailspin = *Context.Coordinator;
@@ -5459,31 +5864,43 @@ ELoot StealWork(FWorkerContext& Context, FReferenceCollector& Collector, FWorkBl
 	{
 		return ELoot::ARO;
 	}
-	else if (FWorkerContext* StolenContext = Tailspin.StealContext())
+	else if (!(Options & EGCOptions::IncrementalReachability))
 	{
-		Context.InitialNativeReferences = StolenContext->InitialNativeReferences;
-		Context.SetInitialObjectsPrepadded(StolenContext->GetInitialObjects());
-		return ELoot::Context;
+		if (FWorkerContext* StolenContext = Tailspin.StealContext())
+		{
+			Context.InitialNativeReferences = StolenContext->InitialNativeReferences;
+			Context.SetInitialObjectsPrepadded(StolenContext->GetInitialObjects());
+			return ELoot::Context;
+		}
 	}
-	else if (Tailspin.ReportOutOfWork())
+	if (Tailspin.ReportOutOfWork())
 	{
 		while (Tailspin.KeepSpinning())
 		{
 			FPlatformProcess::Yield();
+
 			OutBlock = RemainingObjects.StealFullBlock();
 			if (OutBlock || FSlowARO::ProcessAllCalls(Context, Collector))
 			{
 				Tailspin.ReportBackToWork();
 				return OutBlock ? ELoot::Block : ELoot::ARO;
 			}
+
+			if (GReachabilityState.IsTimeLimitExceeded())
+			{
+				SuspendWork(Context);
+				break;
+			}
 		}
 	}
 
 	return ELoot::Nothing;
 }
 
-void ProcessAsync(void (*ProcessSync)(void*, FWorkerContext&), void* Processor, FWorkerContext& InContext)
+TArrayView<FWorkerContext*> InitializeAsyncProcessingContexts(FWorkerContext& InContext)
 {
+	FContextPoolScope ContextPool;
+
 	checkf(InContext.ObjectsToSerialize.IsUnused(), TEXT("Use InitialObjects instead, ObjectsToSerialize.Add() may only be called during reference processing"));
 	check(InContext.Stats.NumObjects == 0 && InContext.Stats.NumReferences == 0 && InContext.Stats.bFoundGarbageRef == false);
 
@@ -5492,24 +5909,24 @@ void ProcessAsync(void (*ProcessSync)(void*, FWorkerContext&), void* Processor,
 
 	int32 NumTaskgraphWorkers = FTaskGraphInterface::Get().GetNumWorkerThreads();
 	const int32 NumWorkers = FMath::Clamp(NumTaskgraphWorkers, 1, MaxWorkers);
-	const int32 ObjPerWorker =  (InitialObjects.Num() + NumWorkers - 1) / NumWorkers;
-	const int32 RefPerWorker =  (InitialReferences.Num() + NumWorkers - 1) / NumWorkers;
+	const int32 ObjPerWorker = (InitialObjects.Num() + NumWorkers - 1) / NumWorkers;
+	const int32 RefPerWorker = (InitialReferences.Num() + NumWorkers - 1) / NumWorkers;
 
-	// Allocate contexts
-	FContextPoolScope ContextPool;
+	GReachabilityState.InitReachabilityAnalysis(NumWorkers);
+
+	// Allocate contexts	
 	checkf(ContextPool.NumAllocated() == 1, TEXT("Other contexts forbidden during parallel reference collection. Work-stealing from all live contexts. "));
-	
-	FWorkerContext* ContextArray[MaxWorkers];
-	TArrayView<FWorkerContext*> Contexts = MakeArrayView(ContextArray, NumWorkers);
+
+	TArrayView<FWorkerContext*> Contexts = MakeArrayView(GReachabilityState.GetContextArray(), GReachabilityState.GetNumWorkers());
 	Contexts[0] = &InContext;
 	for (FWorkerContext*& Context : Contexts.RightChop(1))
 	{
 		Context = ContextPool.AllocateFromPool();
 	}
 
-	// Setup work-stealing queues and distribute initial workload across worker contexts
-	TSharedRef<FWorkCoordinator> Coordinator = MakeShared<FWorkCoordinator>(Contexts, NumTaskgraphWorkers);
 	GSlowARO.GetPostInit().SetupWorkerQueues(NumWorkers);
+
+	// Setup work-stealing queues and distribute initial workload across worker contexts
 	for (FWorkerContext* Context : Contexts)
 	{
 		int32 Idx = Context->GetWorkerIndex();
@@ -5518,15 +5935,62 @@ void ProcessAsync(void (*ProcessSync)(void*, FWorkerContext&), void* Processor,
 		// Initial objects is already padded at the end and its safe to prefetch in the middle too
 		Context->SetInitialObjectsPrepadded(InitialObjects.Mid(Idx * ObjPerWorker, ObjPerWorker));
 		Context->InitialNativeReferences = InitialReferences.Mid(Idx * RefPerWorker, RefPerWorker);
-		Context->Coordinator = &Coordinator.Get();
+		checkf(!Context->IncrementalStructs.ContainsBatchData(), TEXT("Reachability analysis is done but worker context holds suspended dispatcher state"));
+	}
+
+	return Contexts;
+}
+
+void ReleaseAsyncProcessingContexts(FWorkerContext& InContext, TArrayView<FWorkerContext*> Contexts)
+{
+	FContextPoolScope ContextPool;
+
+	// Tear down contexts and work-stealing queues
+	for (FWorkerContext* Context : Contexts)
+	{
+		Context->Coordinator = nullptr;
+		GWorkstealingManager.Queues[Context->GetWorkerIndex()].CheckEmpty();
+		Context->ObjectsToSerialize.ResetAsyncQueue();
+	}
+
+	GSlowARO.GetPostInit().ResetWorkerQueues();
+
+	for (FWorkerContext* Context : Contexts.RightChop(1))
+	{
+		InContext.Stats.AddStats(Context->Stats);
+		ContextPool.ReturnToPool(Context);
+	}
+
+	GReachabilityState.FinishReachabilityAnalysis();
+}
+
+void ProcessAsync(void (*ProcessSync)(void*, FWorkerContext&), void* Processor, FWorkerContext& InContext)
+{
+	using namespace UE::GC;
+
+	TArrayView<FWorkerContext*> Contexts;
+	if (!GReachabilityState.IsSuspended())
+	{
+		Contexts = InitializeAsyncProcessingContexts(InContext);
+	}
+	else
+	{
+		Contexts = MakeArrayView(GReachabilityState.GetContextArray(), GReachabilityState.GetNumWorkers());
 	}
 	
+	TSharedRef<FWorkCoordinator> WorkCoordinator = MakeShared<FWorkCoordinator>(Contexts, FTaskGraphInterface::Get().GetNumWorkerThreads());
+	for (FWorkerContext* Context : Contexts)
+	{
+		Context->bIsSuspended = false;
+		Context->Coordinator = &WorkCoordinator.Get();
+	}
+
 	// Kick workers
-	for (int32 Idx = 1; Idx < NumWorkers; ++Idx)
+	for (int32 Idx = 1; Idx < GReachabilityState.GetNumWorkers(); ++Idx)
 	{
 		Tasks::Launch(TEXT("CollectReferences"), [=]() 
 			{
-				if (FWorkerContext* Context = Coordinator->TryStartWorking(Idx))
+				if (FWorkerContext* Context = WorkCoordinator->TryStartWorking(Idx))
 				{
 					ProcessSync(Processor, *Context);
 				}
@@ -5534,31 +5998,25 @@ void ProcessAsync(void (*ProcessSync)(void*, FWorkerContext&), void* Processor,
 	}
 
 	// Start working ourselves
-	if (FWorkerContext* Context = Coordinator->TryStartWorking(0))
+	if (FWorkerContext* Context = WorkCoordinator->TryStartWorking(0))
 	{
 		ProcessSync(Processor, *Context);
 	}
 
 	// Wait until all work is complete. Current thread can steal and complete everything 
 	// alone if task workers are busy with long-running tasks.
-	Coordinator->SpinUntilAllStopped();
+	WorkCoordinator->SpinUntilAllStopped();
 
-	// Tear down contexts and work-stealing queues
+	// Reset initial object sets so that we don't process again them in the next iteration or when we're done processing references
 	for (FWorkerContext* Context : Contexts)
 	{
-		Context->Coordinator = nullptr;
-		Context->InitialNativeReferences = TConstArrayView<UObject**>();
 		Context->ResetInitialObjects();
-		GWorkstealingManager.Queues[Context->GetWorkerIndex()].CheckEmpty();
-		Context->ObjectsToSerialize.ResetAsyncQueue();
+		Context->InitialNativeReferences = TConstArrayView<UObject**>();
 	}
-	
-	GSlowARO.GetPostInit().ResetWorkerQueues();
 
-	for (FWorkerContext* Context : Contexts.RightChop(1))
+	if (!GReachabilityState.CheckIfAnyContextIsSuspended())
 	{
-		InContext.Stats.AddStats(Context->Stats);
-		ContextPool.ReturnToPool(Context);
+		ReleaseAsyncProcessingContexts(InContext, Contexts);		
 	}
 }
 
diff --git a/Engine/Source/Runtime/CoreUObject/Private/UObject/GarbageCollectionHistory.cpp b/Engine/Source/Runtime/CoreUObject/Private/UObject/GarbageCollectionHistory.cpp
--- a/Engine/Source/Runtime/CoreUObject/Private/UObject/GarbageCollectionHistory.cpp
+++ b/Engine/Source/Runtime/CoreUObject/Private/UObject/GarbageCollectionHistory.cpp
@@ -53,6 +53,19 @@ static FAutoConsoleCommand CmdSetHistorySize(
 		})
 );
 
+/** Dummy class to represent GC barrier inside of GC history object graph */
+class UGCBarrier : public UObject
+{
+	DECLARE_CLASS_INTRINSIC(UGCBarrier,
+		UObject,
+		CLASS_Transient,
+		TEXT("/Script/CoreUObject"));
+};
+
+IMPLEMENT_CORE_INTRINSIC_CLASS(UGCBarrier, UObject,
+	{
+	});
+
 FGCHistory::~FGCHistory()
 {
 	Cleanup();
@@ -84,6 +97,7 @@ void FGCHistory::Cleanup()
 	{
 		Cleanup(Snapshot);
 	}
+	GCBarrier = nullptr;
 	Snapshots.Empty();
 	MostRecentSnapshotIndex = -1;	
 }
@@ -99,6 +113,18 @@ void FGCHistory::SetHistorySize(int32 HistorySize)
 			Snapshots.AddDefaulted(HistorySize);
 		}		
 	}
+	if (HistorySize > 0)
+	{
+		if (!GCBarrier)
+		{
+			GCBarrier = NewObject<UGCBarrier>(GetTransientPackage(), TEXT("GC_Barrier"), RF_Transient);
+			GCBarrier->AddToRoot();
+		}
+	}
+	else if (GCBarrier)
+	{
+		GCBarrier->RemoveFromRoot();
+	}
 }
 
 void FGCHistory::MergeArrayStructHistory(TMap<const UObject*, TArray<FGCDirectReference>*>& History, FGCSnapshot& Snapshot)
diff --git a/Engine/Source/Runtime/CoreUObject/Private/UObject/GarbageCollectionTesting.cpp b/Engine/Source/Runtime/CoreUObject/Private/UObject/GarbageCollectionTesting.cpp
--- a/Engine/Source/Runtime/CoreUObject/Private/UObject/GarbageCollectionTesting.cpp
+++ b/Engine/Source/Runtime/CoreUObject/Private/UObject/GarbageCollectionTesting.cpp
@@ -17,17 +17,45 @@
 
 static void GenerateReachabilityStressData(int Levels, UObjectReachabilityStressData* Data);
 
-UObjectReachabilityStressData* GenerateReachabilityStressData()
+static int32 CanAllocateMoreUObjects()
 {
-	UObjectReachabilityStressData* Data = NewObject<UObjectReachabilityStressData>();
-	GenerateReachabilityStressData(20, Data);
-	return Data;
+	const int32 MinAvailableObjectCount = 64 * 1024;
+	return GUObjectArray.GetObjectArrayEstimatedAvailable() >= MinAvailableObjectCount;
 }
 
-static void GenerateReachabilityStressData(int Levels, UObjectReachabilityStressData* Data)
+static UObjectReachabilityStressData* ConditionallyAllocateNewStressDataObject()
+{
+	if (CanAllocateMoreUObjects())
+	{
+		return NewObject<UObjectReachabilityStressData>();
+	}
+	return nullptr;
+}
+
+void GenerateReachabilityStressData(TArray<UObjectReachabilityStressData*>& Data)
 {
-	Data->AddToRoot();
+	// Roughly NumRootObjects * 2^SubLevels of objects
+	const int32 NumRootObjects = 50;
+	const int32 SubLevels = 13;
 
+	for (int32 Index = 0; Index < NumRootObjects; ++Index)
+	{
+		UObjectReachabilityStressData* RootData = ConditionallyAllocateNewStressDataObject();
+		if (RootData)
+		{
+			Data.Add(RootData);
+			RootData->AddToRoot();
+			GenerateReachabilityStressData(SubLevels, RootData);
+		}
+		else
+		{
+			break;
+		}
+	}
+}
+
+static void GenerateReachabilityStressData(int Levels, UObjectReachabilityStressData* Data)
+{
 	if (Levels == 0)
 	{
 		return;
@@ -36,18 +64,24 @@ static void GenerateReachabilityStressData(int Levels, UObjectReachabilityStress
 	const int N = 2;
 	for (int I = 0; I < N; ++I)
 	{
-		UObjectReachabilityStressData* Child = NewObject<UObjectReachabilityStressData>();
-		GenerateReachabilityStressData(Levels - 1, Child);
-		Data->Children.Add(Child);
+		UObjectReachabilityStressData* Child = ConditionallyAllocateNewStressDataObject();
+		if (Child)
+		{
+			GenerateReachabilityStressData(Levels - 1, Child);
+			Data->Children.Add(Child);
+		}
+		else
+		{
+			break;
+		}
 	}
 }
 
-void UnlinkReachabilityStressData(UObjectReachabilityStressData* Data)
+void UnlinkReachabilityStressData(TArray<UObjectReachabilityStressData*>& Data)
 {
-	Data->RemoveFromRoot();
-	for (const auto Child : Data->Children)
+	for (UObjectReachabilityStressData* RootData : Data)
 	{
-		UnlinkReachabilityStressData(Child);
+		RootData->RemoveFromRoot();
 	}
 }
 
diff --git a/Engine/Source/Runtime/CoreUObject/Private/UObject/GarbageCollectionTesting.h b/Engine/Source/Runtime/CoreUObject/Private/UObject/GarbageCollectionTesting.h
--- a/Engine/Source/Runtime/CoreUObject/Private/UObject/GarbageCollectionTesting.h
+++ b/Engine/Source/Runtime/CoreUObject/Private/UObject/GarbageCollectionTesting.h
@@ -18,5 +18,5 @@ public:
 	TArray<UObjectReachabilityStressData*> Children;
 };
 
-UObjectReachabilityStressData* GenerateReachabilityStressData();
-void UnlinkReachabilityStressData(UObjectReachabilityStressData* Data);
+void GenerateReachabilityStressData(TArray<UObjectReachabilityStressData*>& Data);
+void UnlinkReachabilityStressData(TArray<UObjectReachabilityStressData*>& Data);
diff --git a/Engine/Source/Runtime/CoreUObject/Private/UObject/GarbageCollectionVerification.cpp b/Engine/Source/Runtime/CoreUObject/Private/UObject/GarbageCollectionVerification.cpp
--- a/Engine/Source/Runtime/CoreUObject/Private/UObject/GarbageCollectionVerification.cpp
+++ b/Engine/Source/Runtime/CoreUObject/Private/UObject/GarbageCollectionVerification.cpp
@@ -378,6 +378,99 @@ void VerifyObjectFlagMirroring()
 	UE_CLOG(NumErrors > 0, LogGarbage, Fatal, TEXT("Encountered %d object(s) breaking Object and Internal flag mirroring assumptions. Please check log for details."), (uint32)NumErrors);
 }
 
+/**
+* Finds only direct references of objects passed to the TFastReferenceCollector and verifies if they ae reachable or not
+*/
+class FUnreachableReferenceProcessor : public FSimpleReferenceProcessorBase
+{
+	int32 NumErrors = 0;
+
+public:
+	FUnreachableReferenceProcessor() = default;
+
+	int32 GetErrorCount() const
+	{
+		return NumErrors;
+	}
+
+	FORCEINLINE_DEBUGGABLE void HandleTokenStreamObjectReference(FGCArrayStruct& ObjectsToSerializeStruct, UObject* ReferencingObject, UObject*& Object, FMemberId MemberId, EOrigin Origin, bool bAllowReferenceElimination)
+	{
+		if (Object)
+		{
+			FUObjectItem* ObjectItem = GUObjectArray.ObjectToObjectItem(Object);
+			if (ObjectItem->HasAnyFlags(EInternalObjectFlags::MaybeUnreachable | EInternalObjectFlags::Unreachable))
+			{
+				if (ReferencingObject)
+				{
+					FUObjectItem* ReferencingObjectItem = GUObjectArray.ObjectToObjectItem(ReferencingObject);
+					if (ReferencingObjectItem->GetOwnerIndex() > 0 && ReferencingObjectItem->GetOwnerIndex() == GUObjectArray.ObjectToIndex(Object))
+					{
+						// ReferencingObject is in Object's cluster so it's fine it it's still reachable
+						return;
+					}
+				}
+
+				FString DebugInfo;
+				if (UClass* Class = (ReferencingObject ? ReferencingObject->GetClass() : nullptr))
+				{
+					UE::GC::FMemberInfo Member = GetMemberDebugInfo(Class->ReferenceSchema.Get(), MemberId);
+					DebugInfo = FString::Printf(TEXT("property: %s, offset: %d"),
+						*Member.Name.GetPlainNameString(), Member.Offset);
+				}
+				else
+				{
+					// This means this objects is most likely being referenced by AddReferencedObjects
+					DebugInfo = TEXT("Native Reference");
+				}
+
+				UE_LOG(LogGarbage, Warning, TEXT("Unreachable object %s is being referenced by reachable object %s through %s"),
+					*Object->GetFullName(),
+					*GetFullNameSafe(ReferencingObject),
+					*DebugInfo);
+
+				NumErrors++;
+			}
+		}
+	}
+};
+
+void VerifyUnreachableObjects()
+{
+	const double StartTime = FPlatformTime::Seconds();
+	const int32 MaxNumberOfReachableObjects = GUObjectArray.GetObjectArrayNum() - GUObjectArray.GetFirstGCIndex();
+	const int32 NumThreads = GetNumCollectReferenceWorkers();
+	const int32 NumberOfObjectsPerThread = (MaxNumberOfReachableObjects / NumThreads) + 1;
+	std::atomic<uint32> NumErrors(0);
+
+	ParallelFor(TEXT("GC.VerifyUnreachableObjects"), NumThreads, 1, [&NumErrors, NumberOfObjectsPerThread, NumThreads, MaxNumberOfReachableObjects](int32 ThreadIndex)
+	{
+		int32 FirstObjectIndex = ThreadIndex * NumberOfObjectsPerThread;
+		int32 NumObjects = (ThreadIndex < (NumThreads - 1)) ? NumberOfObjectsPerThread : (MaxNumberOfReachableObjects - (NumThreads - 1)*NumberOfObjectsPerThread);
+		TArray<UObject*> ObjectsToSerialize;
+		ObjectsToSerialize.Reserve(NumberOfObjectsPerThread);		
+
+		for (int32 ObjectIndex = 0; ObjectIndex < NumObjects && (FirstObjectIndex + ObjectIndex + GUObjectArray.GetFirstGCIndex()) < GUObjectArray.GetObjectArrayNum(); ++ObjectIndex)
+		{
+			FUObjectItem& ObjectItem = GUObjectArray.GetObjectItemArrayUnsafe()[FirstObjectIndex + ObjectIndex + GUObjectArray.GetFirstGCIndex()];
+			if (ObjectItem.Object && !ObjectItem.HasAnyFlags(EInternalObjectFlags::MaybeUnreachable | EInternalObjectFlags::Unreachable))
+			{
+				ObjectsToSerialize.Add(static_cast<UObject*>(ObjectItem.Object));
+			}
+		}
+
+		FUnreachableReferenceProcessor Processor;
+		UE::GC::FWorkerContext Context;
+		Context.SetInitialObjectsUnpadded(ObjectsToSerialize);
+		CollectReferences(Processor, Context);
+		NumErrors.fetch_add(Processor.GetErrorCount(), std::memory_order_acq_rel);
+	});
+
+
+	UE_CLOG(NumErrors > 0, LogGarbage, Fatal, TEXT("Encountered %d unreachable object(s) that are still reachable."), NumErrors.load());
+
+	UE_LOG(LogGarbage, Log, TEXT("%f ms for VerifyUnreachableObjects"), (FPlatformTime::Seconds() - StartTime) * 1000);
+}
+
 #endif // VERIFY_DISREGARD_GC_ASSUMPTIONS
 #if PROFILE_GCConditionalBeginDestroy
 
diff --git a/Engine/Source/Runtime/CoreUObject/Private/UObject/GarbageCollectionVerification.h b/Engine/Source/Runtime/CoreUObject/Private/UObject/GarbageCollectionVerification.h
--- a/Engine/Source/Runtime/CoreUObject/Private/UObject/GarbageCollectionVerification.h
+++ b/Engine/Source/Runtime/CoreUObject/Private/UObject/GarbageCollectionVerification.h
@@ -25,6 +25,9 @@ void VerifyClustersAssumptions();
 /** Verifies Object and Internal flag mirroring assumptions */
 void VerifyObjectFlagMirroring();
 
+/** Verifies Reachability assumptions */
+void VerifyUnreachableObjects();
+
 #endif // VERIFY_DISREGARD_GC_ASSUMPTIONS
 
 
diff --git a/Engine/Source/Runtime/CoreUObject/Private/UObject/ReachabilityAnalysisState.cpp b/Engine/Source/Runtime/CoreUObject/Private/UObject/ReachabilityAnalysisState.cpp
new file mode 100644
--- /dev/null
+++ b/Engine/Source/Runtime/CoreUObject/Private/UObject/ReachabilityAnalysisState.cpp
@@ -0,0 +1,44 @@
+// Copyright Epic Games, Inc. All Rights Reserved.
+
+/*=============================================================================
+	ReachabilityAnalysisState.cpp: Incremental reachability analysis support.
+=============================================================================*/
+
+#include "UObject/ReachabilityAnalysisState.h"
+#include "HAL/Runnable.h"
+#include "HAL/RunnableThread.h"
+#include "Misc/App.h"
+#include "Misc/Fork.h"
+
+namespace UE::GC
+{
+
+void FReachabilityAnalysisState::InitReachabilityAnalysis(int32 InNumWorkers)
+{
+	NumWorkers = InNumWorkers;
+	Stats = {};
+	bIsSuspended = false;
+}
+
+void FReachabilityAnalysisState::UpdateStats(const FProcessorStats& InStats)
+{
+	Stats.AddStats(InStats);
+}
+
+void FReachabilityAnalysisState::FinishReachabilityAnalysis()
+{
+	NumWorkers = 0;
+}
+
+bool FReachabilityAnalysisState::CheckIfAnyContextIsSuspended()
+{
+	bIsSuspended = false;
+	for (int32 WorkerIndex = 0; WorkerIndex < NumWorkers && !bIsSuspended; ++WorkerIndex)
+	{
+		bIsSuspended = Contexts[WorkerIndex]->bIsSuspended;
+	}
+	return bIsSuspended;
+}
+
+} // namespce UE::GC
+
diff --git a/Engine/Source/Runtime/CoreUObject/Private/UObject/ReachabilityAnalysisState.h b/Engine/Source/Runtime/CoreUObject/Private/UObject/ReachabilityAnalysisState.h
new file mode 100644
--- /dev/null
+++ b/Engine/Source/Runtime/CoreUObject/Private/UObject/ReachabilityAnalysisState.h
@@ -0,0 +1,114 @@
+// Copyright Epic Games, Inc. All Rights Reserved.
+
+/*=============================================================================
+	ReachabilityAnalysisState.h: Support for incremental reachability analysis
+=============================================================================*/
+
+#pragma once
+
+#include "CoreMinimal.h"
+#include "UObject/UObjectBase.h"
+#include "UObject/Object.h"
+#include "UObject/FastReferenceCollector.h"
+
+namespace UE::GC
+{
+
+/**
+* Reachability analysis state holds information about suspended worker contexts.
+*/
+class FReachabilityAnalysisState
+{
+public:
+
+	// Work-stealing algorithms are O(N^2), everyone steals from everyone.
+	// Might want to improve that before going much wider.
+	static constexpr int32 MaxWorkers = 16;
+
+private:
+	/** Number of workers threads */
+	int32 NumWorkers = 0;
+	/** Worker thread contexts */
+	FWorkerContext* Contexts[MaxWorkers] = { nullptr };
+
+	/** Additional object flags for rooting UObjects reachability analysis was kicked off with */
+	EObjectFlags ObjectKeepFlags = RF_NoFlags;
+	/** True if reachability analysis started as part of full GC */
+	bool bPerformFullPurge = false;
+	/** True if reachability analysis was paused due to time limit */
+	bool bIsSuspended = false;
+	/** Stats of the last reachability analysis run (all iterations) */
+	UE::GC::FProcessorStats Stats;
+
+	/** Total time of all reachability analysis iterations */
+	double IncrementalMarkPhaseTotalTime = 0.0;
+
+	/** Total time of the actual reference traversal */
+	double ReferenceProcessingTotalTime = 0.0;
+
+	alignas (PLATFORM_CACHE_LINE_SIZE) double IterationStartTime = 0.0;
+	alignas (PLATFORM_CACHE_LINE_SIZE) double IterationTimeLimit = 0.0;
+
+public:
+
+	FReachabilityAnalysisState() = default;
+
+	void StartTimer(double InTimeLimit)
+	{
+		IterationTimeLimit = InTimeLimit;
+		IterationStartTime = FPlatformTime::Seconds();
+	}
+
+	const UE::GC::FProcessorStats& GetStats() const
+	{
+		return Stats;
+	}
+
+	bool IsSuspended() const
+	{
+		return bIsSuspended;
+	}
+
+	/** Kicks off new Garbage Collection cycle */
+	void CollectGarbage(EObjectFlags KeepFlags, bool bFullPurge);
+
+	/** Performs Reachability Analysis (also incrementally) and destroys objects (also incrementally) */
+	void PerformReachabilityAnalysisAndConditionallyPurgeGarbage(bool bReachabilityUsingTimeLimit);
+
+	/** Checks if Time Limit has been exceeded. This function needs to be super fast as it's called very frequently. */
+	FORCEINLINE bool IsTimeLimitExceeded() const
+	{
+		if (IterationTimeLimit > 0.0)
+		{
+			return (FPlatformTime::Seconds() - IterationStartTime) >= IterationTimeLimit;
+		}
+		return false;
+	}
+
+	FORCEINLINE int32 GetNumWorkers() const
+	{
+		return NumWorkers;
+	}
+
+	FORCEINLINE FWorkerContext** GetContextArray()
+	{
+		return Contexts;
+	}
+
+	/** Initializes internal variables and sets the expected number of worker threads */
+	void InitReachabilityAnalysis(int32 InNumWorkers);
+
+	/** Main Garbage Collection function executed on Reachability Analysis Thread when StartEvent has been triggered */
+	void PerformReachabilityAnalysis();
+
+	/** Updates reachability analysis state after RA iteration */
+	void UpdateStats(const UE::GC::FProcessorStats& InStats);
+
+	/** Checks if any of the active worker contexts is suspended and updates the IsSuspended state */
+	bool CheckIfAnyContextIsSuspended();
+
+	/** Resets workers after reachability analysis is fully complete */
+	void FinishReachabilityAnalysis();
+};
+
+} // namespace UE::GC
\ No newline at end of file
diff --git a/Engine/Source/Runtime/CoreUObject/Private/UObject/ReferenceChainSearch.cpp b/Engine/Source/Runtime/CoreUObject/Private/UObject/ReferenceChainSearch.cpp
--- a/Engine/Source/Runtime/CoreUObject/Private/UObject/ReferenceChainSearch.cpp
+++ b/Engine/Source/Runtime/CoreUObject/Private/UObject/ReferenceChainSearch.cpp
@@ -255,6 +255,10 @@ namespace UE::ReferenceChainSearch
 					}
 				}
 			}
+			FORCEINLINE bool IsTimeLimitExceeded() const
+			{
+				return false;
+			}
 		};
 
 		class FCollector final: public FReferenceCollector
@@ -1270,6 +1274,12 @@ FString FReferenceChainSearch::GetObjectFlags(FGCObjectInfo* InObject)
 	{
 		Flags += TEXT("(ClusterRoot) ");
 	}
+
+	if (InObject->HasAnyInternalFlags(EInternalObjectFlags::MaybeUnreachable))
+	{
+		Flags += TEXT("(MaybeUnreachable) ");
+	}
+
 	if (InObject->GetOwnerIndex() > 0)
 	{
 		Flags += TEXT("(Clustered) ");
diff --git a/Engine/Source/Runtime/CoreUObject/Private/UObject/SavePackage/SavePackageUtilities.cpp b/Engine/Source/Runtime/CoreUObject/Private/UObject/SavePackage/SavePackageUtilities.cpp
--- a/Engine/Source/Runtime/CoreUObject/Private/UObject/SavePackage/SavePackageUtilities.cpp
+++ b/Engine/Source/Runtime/CoreUObject/Private/UObject/SavePackage/SavePackageUtilities.cpp
@@ -972,6 +972,21 @@ FEDLCookChecker::FEDLNodeHash::FEDLNodeHash()
 {
 }
 
+FEDLCookChecker::FEDLNodeHash::FEDLNodeHash(const FEDLNodeHash& Other)
+{
+	bIsNode = Other.bIsNode;
+	if (bIsNode)
+	{
+		Nodes = Other.Nodes;
+	}
+	else
+	{
+		Object = Other.Object;
+	}
+	NodeID = Other.NodeID;
+	ObjectEvent = Other.ObjectEvent;
+}
+
 FEDLCookChecker::FEDLNodeHash::FEDLNodeHash(const TArray<FEDLNodeData>* InNodes, FEDLNodeID InNodeID, EObjectEvent InObjectEvent)
 	: Nodes(InNodes)
 	, NodeID(InNodeID)
@@ -1013,6 +1028,23 @@ bool FEDLCookChecker::FEDLNodeHash::operator==(const FEDLNodeHash& Other) const
 	return LocalName.IsNone() == OtherName.IsNone();
 }
 
+FEDLCookChecker::FEDLNodeHash& FEDLCookChecker::FEDLNodeHash::operator=(const FEDLNodeHash& Other)
+{
+	bIsNode = Other.bIsNode;
+	if (bIsNode)
+	{
+		Nodes = Other.Nodes;
+	}
+	else
+	{
+		Object = Other.Object;
+	}
+	NodeID = Other.NodeID;
+	ObjectEvent = Other.ObjectEvent;
+
+	return *this;
+}
+
 uint32 GetTypeHash(const FEDLCookChecker::FEDLNodeHash& A)
 {
 	uint32 Hash = 0;
diff --git a/Engine/Source/Runtime/CoreUObject/Private/UObject/SavePackage/SavePackageUtilities.h b/Engine/Source/Runtime/CoreUObject/Private/UObject/SavePackage/SavePackageUtilities.h
--- a/Engine/Source/Runtime/CoreUObject/Private/UObject/SavePackage/SavePackageUtilities.h
+++ b/Engine/Source/Runtime/CoreUObject/Private/UObject/SavePackage/SavePackageUtilities.h
@@ -190,7 +190,9 @@ public: // FEDLNodeHash is public only so that GetTypeHash can be defined
 		FEDLNodeHash(); // creates an uninitialized node; only use this to provide as an out parameter
 		FEDLNodeHash(const TArray<FEDLNodeData>* InNodes, FEDLNodeID InNodeID, EObjectEvent InObjectEvent);
 		FEDLNodeHash(TObjectPtr<UObject> InObject, EObjectEvent InObjectEvent);
+		FEDLNodeHash(const FEDLNodeHash& Other);
 		bool operator==(const FEDLNodeHash& Other) const;
+		FEDLNodeHash& operator=(const FEDLNodeHash& Other);
 		friend uint32 GetTypeHash(const FEDLNodeHash& A);
 
 		FName GetName() const;
diff --git a/Engine/Source/Runtime/CoreUObject/Private/UObject/UObjectBase.cpp b/Engine/Source/Runtime/CoreUObject/Private/UObject/UObjectBase.cpp
--- a/Engine/Source/Runtime/CoreUObject/Private/UObject/UObjectBase.cpp
+++ b/Engine/Source/Runtime/CoreUObject/Private/UObject/UObjectBase.cpp
@@ -243,6 +243,10 @@ void UObjectBase::LowLevelRename(FName NewName,UObject *NewOuter)
 		OuterPrivate = NewOuter;
 	}
 	HashObject(this);
+	if (UE::GC::Private::GIsIncrementalReachabilityPending)
+	{
+		static_cast<UObject*>(this)->MarkAsReachable();
+	}
 }
 
 UPackage* UObjectBase::GetExternalPackage() const
diff --git a/Engine/Source/Runtime/CoreUObject/Private/UObject/UObjectClusters.cpp b/Engine/Source/Runtime/CoreUObject/Private/UObject/UObjectClusters.cpp
--- a/Engine/Source/Runtime/CoreUObject/Private/UObject/UObjectClusters.cpp
+++ b/Engine/Source/Runtime/CoreUObject/Private/UObject/UObjectClusters.cpp
@@ -145,7 +145,7 @@ void FUObjectClusterContainer::DissolveCluster(FUObjectCluster& Cluster)
 	TArray<int32> ReferencedByClusters = MoveTemp(Cluster.ReferencedByClusters);
 
 	// Unreachable clusters will be removed by GC during BeginDestroy phase (unhashing)
-	if (!RootObjectItem->IsUnreachable())
+	if (!RootObjectItem->IsMaybeUnreachable())
 	{
 #if UE_GCCLUSTER_VERBOSE_LOGGING
 		UObject* ClusterRootObject = static_cast<UObject*>(RootObjectItem->Object);
@@ -190,7 +190,7 @@ void FUObjectClusterContainer::DissolveClusterAndMarkObjectsAsUnreachable(FUObje
 	{
 		FUObjectItem* ClusterObjectItem = GUObjectArray.IndexToObjectUnsafeForGC(ClusterObjectIndex);
 		ClusterObjectItem->SetOwnerIndex(0);
-			ClusterObjectItem->SetFlags(EInternalObjectFlags::Unreachable);
+			ClusterObjectItem->SetMaybeUnreachable();
 		}
 
 #if !UE_GCCLUSTER_VERBOSE_LOGGING
@@ -206,7 +206,7 @@ void FUObjectClusterContainer::DissolveClusterAndMarkObjectsAsUnreachable(FUObje
 		FUObjectItem* ReferencedByClusterRootItem = GUObjectArray.IndexToObjectUnsafeForGC(ReferencedByClusterRootIndex);
 		if (ReferencedByClusterRootItem->HasAnyFlags(EInternalObjectFlags::ClusterRoot))
 		{
-				ReferencedByClusterRootItem->SetFlags(EInternalObjectFlags::Unreachable);
+				ReferencedByClusterRootItem->SetMaybeUnreachable();
 			DissolveClusterAndMarkObjectsAsUnreachable(ReferencedByClusterRootItem);
 		}
 	}
diff --git a/Engine/Source/Runtime/CoreUObject/Private/UObject/UObjectGlobals.cpp b/Engine/Source/Runtime/CoreUObject/Private/UObject/UObjectGlobals.cpp
--- a/Engine/Source/Runtime/CoreUObject/Private/UObject/UObjectGlobals.cpp
+++ b/Engine/Source/Runtime/CoreUObject/Private/UObject/UObjectGlobals.cpp
@@ -1310,6 +1310,11 @@ UObject* StaticLoadObjectInternal(UClass* ObjectClass, UObject* InOuter, const T
 				UObjectRedirector* Redirector = FindObjectFast<UObjectRedirector>(InOuter, *StrName);
 				if (Redirector && Redirector->DestinationObject && Redirector->DestinationObject->IsA(ObjectClass ? ObjectClass : UObject::StaticClass()))
 				{
+					if (UE::GC::Private::GIsIncrementalReachabilityPending)
+					{
+						Redirector->MarkAsReachable();
+						Redirector->DestinationObject->MarkAsReachable();
+					}
 					return Redirector->DestinationObject;
 				}
 			}
@@ -1332,6 +1337,10 @@ UObject* StaticLoadObjectInternal(UClass* ObjectClass, UObject* InOuter, const T
 	}
 #endif
 
+	if (Result && UE::GC::Private::GIsIncrementalReachabilityPending)
+	{
+		Result->MarkAsReachable();
+	}
 	return Result;
 }
 
diff --git a/Engine/Source/Runtime/CoreUObject/Private/UObject/UObjectHash.cpp b/Engine/Source/Runtime/CoreUObject/Private/UObject/UObjectHash.cpp
--- a/Engine/Source/Runtime/CoreUObject/Private/UObject/UObjectHash.cpp
+++ b/Engine/Source/Runtime/CoreUObject/Private/UObject/UObjectHash.cpp
@@ -811,6 +811,10 @@ UObject* StaticFindObjectFastInternalThreadSafe(FUObjectHashTables& ThreadHash,
 		}
 	}
 	// Not found.
+	if (Result && UE::GC::Private::GIsIncrementalReachabilityPending)
+	{
+		Result->MarkAsReachable();
+	}
 	return Result;
 }
 
@@ -1135,6 +1139,10 @@ void GetObjectsWithOuter(const class UObjectBase* Outer, TArray<UObject *>& Resu
 			UObject *Object = static_cast<UObject *>(*It);
 			if (!Object->HasAnyFlags(ExclusionFlags) && !Object->HasAnyInternalFlags(ExclusionInternalFlags))
 			{
+				if (UE::GC::Private::GIsIncrementalReachabilityPending)
+				{
+					Object->MarkAsReachable();
+				}
 				Results.Add(Object);
 			}
 		}
@@ -1154,6 +1162,10 @@ void GetObjectsWithOuter(const class UObjectBase* Outer, TArray<UObject *>& Resu
 						UObject *Object = static_cast<UObject *>(*It);
 						if (!Object->HasAnyFlags(ExclusionFlags) && !Object->HasAnyInternalFlags(ExclusionInternalFlags))
 						{
+							if (UE::GC::Private::GIsIncrementalReachabilityPending)
+							{
+								Object->MarkAsReachable();
+							}
 							Results.Add(Object);
 						}
 					}
@@ -1202,6 +1214,10 @@ void ForEachObjectWithOuterBreakable(const class UObjectBase* Outer, TFunctionRe
 			UObject *Object = static_cast<UObject*>(*It);
 			if (!Object->HasAnyFlags(ExclusionFlags) && !Object->HasAnyInternalFlags(ExclusionInternalFlags))
 			{
+				if (UE::GC::Private::GIsIncrementalReachabilityPending)
+				{
+					Object->MarkAsReachable();
+				}
 				if (!Operation(Object))
 				{
 					AllInners.Empty();
@@ -1255,6 +1271,10 @@ UObjectBase* FindObjectWithOuter(const class UObjectBase* Outer, const class UCl
 				Result = Object;
 				break;
 			}
+			if (Result && UE::GC::Private::GIsIncrementalReachabilityPending)
+			{
+				Result->MarkAsReachable();
+			}
 		}
 	}
 	return Result;
@@ -1312,6 +1332,10 @@ void ForEachObjectWithPackage(const class UPackage* Package, TFunctionRef<bool(U
 				!Object->HasAnyInternalFlags(ExclusionInternalFlags) &&
 				bIsInPackage)
 			{
+				if (UE::GC::Private::GIsIncrementalReachabilityPending)
+				{
+					Object->MarkAsReachable();
+				}
 				if (!Operation(Object))
 				{
 					AllInners.Empty(AllInners.Max());
@@ -1400,6 +1424,10 @@ FORCEINLINE void ForEachObjectOfClasses_Implementation(FUObjectHashTables& Threa
 				UObject* Object = static_cast<UObject*>(*ObjectIt);
 				if (!Object->HasAnyFlags(ExcludeFlags) && !Object->HasAnyInternalFlags(ExclusionInternalFlags))
 				{
+					if (UE::GC::Private::GIsIncrementalReachabilityPending)
+					{
+						Object->MarkAsReachable();
+					}
 					Operation(Object);
 				}
 			}
diff --git a/Engine/Source/Runtime/CoreUObject/Public/UObject/FastReferenceCollector.h b/Engine/Source/Runtime/CoreUObject/Public/UObject/FastReferenceCollector.h
--- a/Engine/Source/Runtime/CoreUObject/Public/UObject/FastReferenceCollector.h
+++ b/Engine/Source/Runtime/CoreUObject/Public/UObject/FastReferenceCollector.h
@@ -34,6 +34,7 @@ enum class EGCOptions : uint32
 	Parallel = 1 << 0,					// Use all task workers to collect references, must be started on main thread
 	AutogenerateSchemas = 1 << 1,		// Assemble schemas for new UClasses
 	WithPendingKill = 1 << 2,			// Internal flag used by reachability analysis
+	IncrementalReachability = 1 << 3	// Run Reachability Analysis incrementally
 };
 ENUM_CLASS_FLAGS(EGCOptions);
 
@@ -60,6 +61,8 @@ public:
 namespace UE::GC
 {
 
+struct FStructArrayBlock;
+
 static constexpr uint32 ObjectLookahead = 16;
 
 // Prefetches ClassPrivate, OuterPrivate, class schema and schema data while iterating over an object array
@@ -240,6 +243,26 @@ struct FProcessorStats
 	}
 };
 
+struct FStructArray
+{
+	FSchemaView Schema{ NoInit };
+	uint8* Data;
+	int32 Num;
+	uint32 Stride;
+};
+
+struct FSuspendedStructBatch
+{
+	FStructArrayBlock* Wip = nullptr;
+	FStructArray* WipIt = nullptr;
+
+	FORCEINLINE bool ContainsBatchData() const
+	{
+		return !!Wip;
+	}
+};
+
+
 /** Thread-local context containing initial objects and references to collect */
 struct alignas(PLATFORM_CACHE_LINE_SIZE) FWorkerContext
 {
@@ -269,6 +292,9 @@ public:
 	TMap<const UObject*, TArray<FGCDirectReference>*> History;
 #endif
 
+	FSuspendedStructBatch IncrementalStructs;
+	bool bIsSuspended = false;
+
 	FORCEINLINE UObject* GetReferencingObject()	{ return ReferencingObject;	}
 
 	TConstArrayView<UObject*> GetInitialObjects() { return InitialObjects; }
@@ -321,14 +347,6 @@ struct FMemberWordUnpacked
 	FMemberUnpacked Members[4];
 };
 
-struct FStructArray
-{
-	FSchemaView Schema{NoInit};
-	uint8* Data;
-	int32 Num;
-	uint32 Stride;
-};
-
 struct FStridedReferenceArray
 {
 	FScriptArray* Array;
@@ -624,6 +642,10 @@ struct TDirectDispatcher
 	{
 		HandleKillableReferences(ToView(Array), MemberId, Origin);
 	}
+
+	void Suspend()
+	{
+	}
 };
 
 // Default implementation is to create new direct dispatcher
@@ -643,7 +665,9 @@ struct TGetDispatcherType
 //////////////////////////////////////////////////////////////////////////
 
 enum class ELoot { Nothing, Block, ARO, Context };
-COREUOBJECT_API ELoot StealWork(FWorkerContext& Context, FReferenceCollector& Collector, FWorkBlock*& OutBlock);
+COREUOBJECT_API ELoot StealWork(FWorkerContext& Context, FReferenceCollector& Collector, FWorkBlock*& OutBlock, EGCOptions Options);
+
+COREUOBJECT_API void SuspendWork(FWorkerContext& Context);
 
 /** Allocates contexts and coordinator, kicks worker tasks that also call ProcessSync. Processor is type-erased to void* to avoid templated code. */
 COREUOBJECT_API void ProcessAsync(void (*ProcessSync)(void*, FWorkerContext&), void* Processor, FWorkerContext& InitialContext);
@@ -694,6 +718,13 @@ StoleContext:
 				Context.ObjectsToSerialize.FreeOwningBlock(CurrentObjects.GetData());
 			}
 
+			if (Processor.IsTimeLimitExceeded())
+			{
+				Dispatcher.Suspend();
+				SuspendWork(Context);				
+				return;
+			}
+
 			int32 BlockSize = FWorkBlock::ObjectCapacity;
 			FWorkBlockifier& RemainingObjects = Context.ObjectsToSerialize;
 			FWorkBlock* Block = RemainingObjects.PopFullBlock<Options>();
@@ -719,7 +750,7 @@ StoleARO:
 				else if (Block = RemainingObjects.PopPartialBlock(/* out if successful */ BlockSize); Block);
 				else if (bIsParallel) // if constexpr yields MSVC unreferenced label warning
 				{
-					switch (StealWork(/* in-out */ Context, Collector, /* out */ Block))
+					switch (StealWork(/* in-out */ Context, Collector, /* out */ Block, Options))
 					{
 						case ELoot::Nothing:	break;				// Done, stop working
 						case ELoot::Block:		break;				// Stole full block, process it
@@ -837,6 +868,11 @@ public:
 	void UpdateDetailedStats(UObject* CurrentObject) {}
 	void LogDetailedStatsSummary() {}
 
+	FORCEINLINE bool IsTimeLimitExceeded() const
+	{
+		return false;
+	}
+
 	// Implement this in your derived class, don't make this virtual as it will affect performance!
 	//FORCEINLINE void HandleTokenStreamObjectReference(FWorkerContext& Context, UObject* ReferencingObject, UObject*& Object, FMemberId MemberId, EOrigin Origin, bool bAllowReferenceElimination)
 };
@@ -848,7 +884,7 @@ FORCEINLINE static void CollectReferences(ProcessorType& Processor, UE::GC::FWor
 	using namespace UE::GC;
 	using FastReferenceCollector = TFastReferenceCollector<ProcessorType, CollectorType>;
 	
-	if constexpr (IsParallel(ProcessorType::Options))
+	if (IsParallel(ProcessorType::Options) && !UE::GC::Private::GIsIncrementalReachabilityPending)
 	{
 		ProcessAsync([](void* P, FWorkerContext& C) { FastReferenceCollector(*reinterpret_cast<ProcessorType*>(P)).ProcessObjectArray(C); }, &Processor, Context);
 	}
diff --git a/Engine/Source/Runtime/CoreUObject/Public/UObject/GarbageCollectionHistory.h b/Engine/Source/Runtime/CoreUObject/Public/UObject/GarbageCollectionHistory.h
--- a/Engine/Source/Runtime/CoreUObject/Public/UObject/GarbageCollectionHistory.h
+++ b/Engine/Source/Runtime/CoreUObject/Public/UObject/GarbageCollectionHistory.h
@@ -58,6 +58,8 @@ class FGCHistory
 	TArray<FGCSnapshot> Snapshots;
 	/** Index of the last recorded snapshot in the Snapshots array */
 	int32 MostRecentSnapshotIndex = -1;
+	/** Object representing GC barrier */
+	UObject* GCBarrier = nullptr;
 
 	COREUOBJECT_API ~FGCHistory();
 
@@ -108,6 +110,12 @@ public:
 
 	/** Returns the number of bytes allocated by GC history */
 	COREUOBJECT_API int64 GetAllocatedSize() const;
+
+	/** Returns an object representing GC Barrier */
+	FORCEINLINE UObject* GetBarrierObject() const
+	{
+		return GCBarrier;
+	}
 };
 
-#endif // ENABLE_GC_HISTORY
+#endif // ENABLE_GC_HISTORY
\ No newline at end of file
diff --git a/Engine/Source/Runtime/CoreUObject/Public/UObject/Object.h b/Engine/Source/Runtime/CoreUObject/Public/UObject/Object.h
--- a/Engine/Source/Runtime/CoreUObject/Public/UObject/Object.h
+++ b/Engine/Source/Runtime/CoreUObject/Public/UObject/Object.h
@@ -772,6 +772,11 @@ public:
 	 */
 	COREUOBJECT_API void CallAddReferencedObjects(FReferenceCollector& Collector);
 
+	/**
+	* Marks the object as Reachable if it's currently marked as MaybeUnreachable by incremental GC.
+	*/
+	COREUOBJECT_API void MarkAsReachable() const;
+
 	/**
 	 * Save information for StaticAllocateObject in the case of overwriting an existing object.
 	 * StaticAllocateObject will call delete on the result after calling Restore()
@@ -1809,6 +1814,22 @@ const T* GetValid(const T* Test)
 	return IsValid(Test) ? Test : nullptr;
 }
 
+/**
+* Helper function to mark object as reachable when performing Incremental GC
+*
+* @param	Object			Object to mark as reachable
+* @return	The same Object pointer that was passed to the function
+*/
+template <typename T>
+FORCEINLINE T* MarkAsReachable(T* Object)
+{
+	if (Object)
+	{
+		Object->MarkAsReachable();
+	}
+	return Object;
+}
+
 #if WITH_EDITOR
 /** Callback for editor object selection. This must be in core instead of editor for UObject::IsSelectedInEditor to work */
 extern COREUOBJECT_API TFunction<bool(const UObject*)> GIsObjectSelectedInEditor;
diff --git a/Engine/Source/Runtime/CoreUObject/Public/UObject/ObjectMacros.h b/Engine/Source/Runtime/CoreUObject/Public/UObject/ObjectMacros.h
--- a/Engine/Source/Runtime/CoreUObject/Public/UObject/ObjectMacros.h
+++ b/Engine/Source/Runtime/CoreUObject/Public/UObject/ObjectMacros.h
@@ -599,6 +599,7 @@ enum class EInternalObjectFlags : int32
 {
 	None = 0,
 
+	MaybeUnreachable = 1 << 19, ///< Flag set on all non-root objects at the beginning of Reachability Analysis
 	LoaderImport = 1 << 20, ///< Object is ready to be imported by another package during loading
 	Garbage = 1 << 21, ///< Garbage from logical point of view and should not be referenced. This flag is mirrored in EObjectFlags as RF_Garbage for performance
 	ReachableInCluster = 1 << 23, ///< External reference to object in cluster exists
@@ -616,7 +617,7 @@ enum class EInternalObjectFlags : int32
 	MirroredFlags = Garbage | PendingKill, /// Flags mirrored in EObjectFlags
 
 	//~ Make sure this is up to date!
-	AllFlags = LoaderImport | Garbage | ReachableInCluster | ClusterRoot | Native | Async | AsyncLoading | Unreachable | PendingKill | RootSet | PendingConstruction
+	AllFlags = MaybeUnreachable | LoaderImport | Garbage | ReachableInCluster | ClusterRoot | Native | Async | AsyncLoading | Unreachable | PendingKill | RootSet | PendingConstruction
 	PRAGMA_ENABLE_DEPRECATION_WARNINGS
 };
 ENUM_CLASS_FLAGS(EInternalObjectFlags);
diff --git a/Engine/Source/Runtime/CoreUObject/Public/UObject/ObjectPtr.h b/Engine/Source/Runtime/CoreUObject/Public/UObject/ObjectPtr.h
--- a/Engine/Source/Runtime/CoreUObject/Public/UObject/ObjectPtr.h
+++ b/Engine/Source/Runtime/CoreUObject/Public/UObject/ObjectPtr.h
@@ -17,6 +17,10 @@
 	#define UE_OBJPTR_DEPRECATED(Version, Message) 
 #endif
 
+#ifndef UE_OBJECT_PTR_GC_BARRIER
+	#define UE_OBJECT_PTR_GC_BARRIER 0
+#endif
+
 /** 
  * Wrapper macro for use in places where code needs to allow for a pointer type that could be a TObjectPtr<T> or a raw object pointer during a transitional period.
  * The coding standard disallows general use of the auto keyword, but in wrapping it in this macro, we have a record
@@ -98,13 +102,45 @@ public:
 		return UE::CoreUObject::Private::ResolveObjectHandleClass(Handle);
 	}
 
+
+#if UE_OBJECT_PTR_GC_BARRIER
+	FObjectPtr(FObjectPtr&& InOther)
+		: Handle(MoveTemp(InOther.Handle))
+	{
+		ConditionallyMarkAsReachable(*this);
+	}
+
+	FObjectPtr(const FObjectPtr& InOther)
+		: Handle(InOther.Handle)
+	{
+		ConditionallyMarkAsReachable(*this);
+	}
+
+	FObjectPtr& operator=(FObjectPtr&& InOther)
+	{
+		ConditionallyMarkAsReachable(InOther);
+		Handle = MoveTemp(InOther.Handle);
+		return *this;
+	}
+
+	FObjectPtr& operator=(const FObjectPtr& InOther)
+	{
+		ConditionallyMarkAsReachable(InOther);
+		Handle = InOther.Handle;
+		return *this;
+	}
+#else
 	FObjectPtr(FObjectPtr&&) = default;
 	FObjectPtr(const FObjectPtr&) = default;
 	FObjectPtr& operator=(FObjectPtr&&) = default;
 	FObjectPtr& operator=(const FObjectPtr&) = default;
+#endif // UE_OBJECT_PTR_GC_BARRIER
 
 	FObjectPtr& operator=(UObject* Other)
 	{
+#if UE_OBJECT_PTR_GC_BARRIER
+		ConditionallyMarkAsReachable(Other);
+#endif // UE_OBJECT_PTR_GC_BARRIER
 		Handle = UE::CoreUObject::Private::MakeObjectHandle(Other);
 		return *this;
 	}
@@ -235,6 +271,26 @@ private:
 		// is an odd/uneven number, that means the object reference is unresolved and you will not be able to dereference it successfully.
 		UObject* DebugPtr;
 	};
+
+#if UE_OBJECT_PTR_GC_BARRIER
+	FORCEINLINE void ConditionallyMarkAsReachable(const FObjectPtr& InPtr) const
+	{
+		if (UE::GC::Private::GIsIncrementalReachabilityPending && InPtr.IsResolved())
+		{
+			if (UObject* Obj = UE::CoreUObject::Private::ReadObjectHandlePointerNoCheck(InPtr.GetHandleRef()))
+			{
+				Obj->MarkAsReachable();
+			}
+		}
+	}
+	FORCEINLINE void ConditionallyMarkAsReachable(const UObject* InObj) const
+	{
+		if (UE::GC::Private::GIsIncrementalReachabilityPending && InObj)
+		{
+			InObj->MarkAsReachable();
+		}
+	}
+#endif // UE_OBJECT_PTR_GC_BARRIER
 };
 
 template <typename T>
@@ -365,8 +421,19 @@ public:
 	{
 	}
 
+#if UE_OBJECT_PTR_GC_BARRIER
+	TObjectPtr(TObjectPtr<T>&& Other)
+		: ObjectPtr(MoveTemp(Other.ObjectPtr))
+	{
+	}
+	TObjectPtr(const TObjectPtr<T>& Other)
+		: ObjectPtr(Other.ObjectPtr)
+	{
+	}
+#else
 	TObjectPtr(TObjectPtr<T>&& Other) = default;
 	TObjectPtr(const TObjectPtr<T>& Other) = default;
+#endif // UE_OBJECT_PTR_GC_BARRIER
 
 	explicit FORCEINLINE TObjectPtr(ENoInit)
 		: ObjectPtr(NoInit)
@@ -415,8 +482,21 @@ public:
 	{
 	}
 
+#if UE_OBJECT_PTR_GC_BARRIER
+	TObjectPtr<T>& operator=(TObjectPtr<T>&& Other)
+	{
+		ObjectPtr = MoveTemp(Other.ObjectPtr);
+		return *this;
+	}
+	TObjectPtr<T>& operator=(const TObjectPtr<T>& Other)
+	{
+		ObjectPtr = Other.ObjectPtr;
+		return *this;
+	}
+#else
 	TObjectPtr<T>& operator=(TObjectPtr<T>&&) = default;
 	TObjectPtr<T>& operator=(const TObjectPtr<T>&) = default;
+#endif // UE_OBJECT_PTR_GC_BARRIER
 
 	FORCEINLINE TObjectPtr<T>& operator=(TYPE_OF_NULLPTR)
 	{
@@ -923,7 +1003,7 @@ struct TIsZeroConstructType<TObjectPtr<T>>
 template <typename T>
 struct TIsBitwiseConstructible<TObjectPtr<T>, T*>
 {
-	enum { Value = true };
+	enum { Value = !UE_OBJECT_PTR_GC_BARRIER };
 };
 
 template <typename T, class PREDICATE_CLASS>
@@ -1218,7 +1298,12 @@ namespace UE::Core::Private // private facilities; not for direct use
 	{
 		static void Close(ViewType& View)
 		{
-			// TODO: Run GC barrier here
+#if UE_OBJECT_PTR_GC_BARRIER
+			if (UE::GC::Private::GIsIncrementalReachabilityPending && View)
+			{
+				View->MarkAsReachable();
+			}
+#endif // UE_OBJECT_PTR_GC_BARRIER
 		}
 	};
 	
@@ -1227,8 +1312,19 @@ namespace UE::Core::Private // private facilities; not for direct use
 	{
 		static void Close(ViewType& View)
 		{
-			// TODO: reenable this once GC Barriers merge
-			//RunGCBarriers(reinterpret_cast<const UObject* const*>(View.GetData()), View.Num());
+#if UE_OBJECT_PTR_GC_BARRIER
+			if (UE::GC::Private::GIsIncrementalReachabilityPending)
+			{
+				const UObject* const* Data = reinterpret_cast<const UObject* const*>(View.GetData());
+				for (int32 Index = 0; Index < View.Num(); ++Index)
+				{
+					if (Data[Index])
+					{
+						Data[Index]->MarkAsReachable();
+					}
+				}
+			}
+#endif // UE_OBJECT_PTR_GC_BARRIER
 		}
 	};
 
diff --git a/Engine/Source/Runtime/CoreUObject/Public/UObject/ReachabilityAnalysis.h b/Engine/Source/Runtime/CoreUObject/Public/UObject/ReachabilityAnalysis.h
new file mode 100644
--- /dev/null
+++ b/Engine/Source/Runtime/CoreUObject/Public/UObject/ReachabilityAnalysis.h
@@ -0,0 +1,37 @@
+// Copyright Epic Games, Inc. All Rights Reserved.
+
+/*=============================================================================
+	ReachabilityAnalysis.h: Utility functions for configuring and running
+	Incremental reachability analysis
+=============================================================================*/
+
+#pragma once
+
+#include "CoreMinimal.h"
+
+
+/**
+ * Enables or disables incremental reachability analysis.
+ *
+ * @param	bEnabled True if incremental reachability analysis is to be enabled
+ */
+COREUOBJECT_API void SetIncrementalReachabilityAnalysisEnabled(bool bEnabled);
+
+/**
+ * Returns true if incremental reachability analysis is enabled
+ */
+COREUOBJECT_API bool GetIncrementalReachabilityAnalysisEnabled();
+
+/**
+ * Sets time limit for incremental rachability analysis (if enabled).
+ *
+ * @param	TimeLimitSeconds Time limit (in seconds) for incremental raachability analysis
+ */
+COREUOBJECT_API void SetReachabilityAnalysisTimeLimit(float TimeLimitSeconds);
+
+/**
+ * Returns the time limit for incremental rachability analysis.
+ *
+ * @return	time limit (in seconds) for incremental rachability analysis.
+ */
+COREUOBJECT_API float GetReachabilityAnalysisTimeLimit();
diff --git a/Engine/Source/Runtime/CoreUObject/Public/UObject/UObjectArray.h b/Engine/Source/Runtime/CoreUObject/Public/UObject/UObjectArray.h
--- a/Engine/Source/Runtime/CoreUObject/Public/UObject/UObjectArray.h
+++ b/Engine/Source/Runtime/CoreUObject/Public/UObject/UObjectArray.h
@@ -167,19 +167,34 @@ struct FUObjectItem
 	{
 		ThisThreadAtomicallySetFlag(EInternalObjectFlags::Unreachable);
 	}
+	FORCEINLINE void SetMaybeUnreachable()
+	{
+		ThisThreadAtomicallySetFlag(EInternalObjectFlags::MaybeUnreachable);
+	}
 	FORCEINLINE void ClearUnreachable()
 	{
 		ThisThreadAtomicallyClearedFlag(EInternalObjectFlags::Unreachable);
 	}
+	FORCEINLINE void ClearMaybeUnreachable()
+	{
+		ThisThreadAtomicallyClearedFlag(EInternalObjectFlags::MaybeUnreachable);
+	}
 	FORCEINLINE bool IsUnreachable() const
 	{
 		return !!(GetFlagsInternal() & int32(EInternalObjectFlags::Unreachable));
 	}
+	FORCEINLINE bool IsMaybeUnreachable() const
+	{
+		return !!(Flags & int32(EInternalObjectFlags::MaybeUnreachable));
+	}
 	FORCEINLINE bool ThisThreadAtomicallyClearedRFUnreachable()
 	{
 		return ThisThreadAtomicallyClearedFlag(EInternalObjectFlags::Unreachable);
 	}
-
+	FORCEINLINE bool ThisThreadAtomicallyClearedMaybeUnreachable()
+	{
+		return ThisThreadAtomicallyClearedFlag(EInternalObjectFlags::MaybeUnreachable);
+	}
 	FORCEINLINE void SetPendingKill()
 	{
 		PRAGMA_DISABLE_DEPRECATION_WARNINGS
diff --git a/Engine/Source/Runtime/CoreUObject/Public/UObject/UObjectBaseUtility.h b/Engine/Source/Runtime/CoreUObject/Public/UObject/UObjectBaseUtility.h
--- a/Engine/Source/Runtime/CoreUObject/Public/UObject/UObjectBaseUtility.h
+++ b/Engine/Source/Runtime/CoreUObject/Public/UObject/UObjectBaseUtility.h
@@ -326,7 +326,7 @@ public:
 	/** Checks if the object is unreachable. */
 	FORCEINLINE bool IsUnreachable() const
 	{
-		return GUObjectArray.IndexToObject(InternalIndex)->IsUnreachable();
+		return GUObjectArray.IndexToObject(InternalIndex)->HasAnyFlags(EInternalObjectFlags::Unreachable);
 	}
 
 	/** Checks if the object is pending kill or unreachable. INTERNAL USE ONLY! If you want to check if your object is valid use IsValid(Object)/IsValidObjectChecked(Object)/GetValid(Object) instead. */
diff --git a/Engine/Source/Runtime/CoreUObject/Public/UObject/UObjectGlobals.h b/Engine/Source/Runtime/CoreUObject/Public/UObject/UObjectGlobals.h
--- a/Engine/Source/Runtime/CoreUObject/Public/UObject/UObjectGlobals.h
+++ b/Engine/Source/Runtime/CoreUObject/Public/UObject/UObjectGlobals.h
@@ -770,6 +770,23 @@ COREUOBJECT_API bool IsIncrementalPurgePending();
  */
 COREUOBJECT_API void GatherUnreachableObjects(bool bForceSingleThreaded);
 
+/**
+ * Returns whether an incremental reachability analysis is still pending/ in progress.
+ *
+ * @return	true if incremental reachability analysis needs to be kicked off or is currently in progress, false othwerise.
+ */
+COREUOBJECT_API bool IsIncrementalReachabilityAnalysisPending();
+
+/**
+ * Incrementally purge garbage by deleting all unreferenced objects after routing Destroy.
+ *
+ * Calling code needs to be EXTREMELY careful when and how to call this function as
+ * RF_Unreachable cannot change on any objects unless any pending purge has completed!
+ *
+ * @param	bUseTimeLimit	whether the time limit parameter should be used
+ */
+COREUOBJECT_API void PerformIncrementalReachabilityAnalysis();
+
 /**
  * Incrementally purge garbage by deleting all unreferenced objects after routing Destroy.
  *
@@ -3767,6 +3784,13 @@ enum class EDataValidationResult : uint8
  */
 COREUOBJECT_API EDataValidationResult CombineDataValidationResults(EDataValidationResult Result1, EDataValidationResult Result2);
 
+
+namespace UE::GC::Private
+{
+	/** true if incremental reachability analysis is in progress (global for faster access in low level structs and functions otherwise use IsIncrementalReachabilityAnalisysPending()) */
+	extern COREUOBJECT_API bool GIsIncrementalReachabilityPending;
+}
+
 #if UE_ENABLE_INCLUDE_ORDER_DEPRECATED_IN_5_2
 #include "CoreMinimal.h"
 #endif
diff --git a/Engine/Source/Runtime/CoreUObject/Tests/ObjectPtrTest.cpp b/Engine/Source/Runtime/CoreUObject/Tests/ObjectPtrTest.cpp
--- a/Engine/Source/Runtime/CoreUObject/Tests/ObjectPtrTest.cpp
+++ b/Engine/Source/Runtime/CoreUObject/Tests/ObjectPtrTest.cpp
@@ -34,11 +34,13 @@ static_assert(sizeof(FObjectPtr) == sizeof(void*), "FObjectPtr type must always
 static_assert(sizeof(TObjectPtr<UObject>) == sizeof(void*), "TObjectPtr<UObject> type must always compile to something equivalent to a pointer size.");
 
 // Ensure that a TObjectPtr is trivially copyable, (copy/move) constructible, (copy/move) assignable, and destructible
+#if !UE_OBJECT_PTR_GC_BARRIER
 static_assert(std::is_trivially_copyable<FMutableObjectPtr>::value, "TObjectPtr must be trivially copyable");
 static_assert(std::is_trivially_copy_constructible<FMutableObjectPtr>::value, "TObjectPtr must be trivially copy constructible");
 static_assert(std::is_trivially_move_constructible<FMutableObjectPtr>::value, "TObjectPtr must be trivially move constructible");
 static_assert(std::is_trivially_copy_assignable<FMutableObjectPtr>::value, "TObjectPtr must be trivially copy assignable");
 static_assert(std::is_trivially_move_assignable<FMutableObjectPtr>::value, "TObjectPtr must be trivially move assignable");
+#endif // !UE_OBJECT_PTR_GC_BARRIER
 static_assert(std::is_trivially_destructible<FMutableObjectPtr>::value, "TObjectPtr must be trivially destructible");
 
 // Ensure that raw pointers can be used to construct wrapped object pointers and that const-ness isn't stripped when constructing or converting with raw pointers
diff --git a/Engine/Source/Runtime/Engine/Private/UnrealEngine.cpp b/Engine/Source/Runtime/Engine/Private/UnrealEngine.cpp
--- a/Engine/Source/Runtime/Engine/Private/UnrealEngine.cpp
+++ b/Engine/Source/Runtime/Engine/Private/UnrealEngine.cpp
@@ -1552,6 +1552,12 @@ static TAutoConsoleVariable<int32> CVarCollectGarbageEveryFrame(
 	0,
 	TEXT("Used to debug garbage collection...Collects garbage every frame if the value is > 0."));
 
+
+static TAutoConsoleVariable<int32> CVarContinuousIncrementalGC(
+	TEXT("gc.ContinuousIncrementalGC"),
+	0,
+	TEXT("Used to debug garbage collection...Kicks off Incremental Garbage Collection as soon as the previous one finishes."));
+
 static float GTimeBetweenPurgingPendingKillObjects = 60.0f;
 static FAutoConsoleVariableRef CVarTimeBetweenPurgingPendingKillObjects(
 	TEXT("gc.TimeBetweenPurgingPendingKillObjects"),
@@ -1798,6 +1804,11 @@ void UEngine::ConditionalCollectGarbage()
 					{
 						bShouldDelayGarbageCollect = false;
 					}
+					else if (IsIncrementalReachabilityAnalysisPending())
+					{
+						SCOPE_CYCLE_COUNTER(STAT_GCMarkTime);
+						PerformIncrementalReachabilityAnalysis();
+					}
 					// Perform incremental purge update if it's pending or in progress.
 					else if (!IsIncrementalPurgePending()
 						// Purge reference to pending kill objects every now and so often.
@@ -1827,13 +1838,21 @@ void UEngine::ConditionalCollectGarbage()
 			}
 		}
 
-		if (CVarCollectGarbageEveryFrame.GetValueOnGameThread() > 0)
+		if (CVarCollectGarbageEveryFrame.GetValueOnGameThread())
 		{
 			ForceGarbageCollection(true);
 		}
+		else if (CVarContinuousIncrementalGC.GetValueOnGameThread() > 0 && !IsIncrementalReachabilityAnalysisPending() && !IsIncrementalUnhashPending() && !IsIncrementalPurgePending())
+		{
+			ForceGarbageCollection(false);
+		}
 
 		LastGCFrame = GFrameCounter;
 	}
+	else if (IsIncrementalReachabilityAnalysisPending())
+	{
+		PerformIncrementalReachabilityAnalysis();
+	}
 }
 
 void UEngine::PerformGarbageCollectionAndCleanupActors()
