Fixing CIS error (global var name conflict)

diff --git a/Engine/Source/Runtime/CoreUObject/Private/UObject/GarbageCollection.cpp b/Engine/Source/Runtime/CoreUObject/Private/UObject/GarbageCollection.cpp
--- a/Engine/Source/Runtime/CoreUObject/Private/UObject/GarbageCollection.cpp
+++ b/Engine/Source/Runtime/CoreUObject/Private/UObject/GarbageCollection.cpp
@@ -247,10 +247,10 @@ static FAutoConsoleVariableRef CVarIncrementalReachabilityTimeLimit(
 );
 
 #if VERIFY_DISREGARD_GC_ASSUMPTIONS
-static int32 GVerifyUnreachableObjects = 0;
-static FAutoConsoleVariableRef CVarVerifyUnreachableObjects(
-	TEXT("gc.VerifyUnreachableObjects"),
-	GVerifyUnreachableObjects,
+static int32 GVerifyNoUnreachableObjects = 0;
+static FAutoConsoleVariableRef CVarVerifyNoUnreachableObjects(
+	TEXT("gc.VerifyNoUnreachableObjects"),
+	GVerifyNoUnreachableObjects,
 	TEXT("Enables or disables no unreachable objects are reachable verification"),
 	ECVF_Default
 );
@@ -4890,10 +4890,10 @@ void PostCollectGarbageImpl(EObjectFlags KeepFlags)
 	if (!GIsIncrementalReachabilityPending)
 	{
 #if VERIFY_DISREGARD_GC_ASSUMPTIONS
-		if (GVerifyUnreachableObjects != 0 && GUnreachableObjects.Num() > 0)
+		if (GVerifyNoUnreachableObjects != 0 && GUnreachableObjects.Num() > 0)
 		{
-			DECLARE_SCOPE_CYCLE_COUNTER(TEXT("CollectGarbageInternal.VerifyUnreachableObjects"), STAT_CollectGarbageInternal_VerifyUnreachableObjects, STATGROUP_GC);
-			VerifyUnreachableObjects();
+			DECLARE_SCOPE_CYCLE_COUNTER(TEXT("CollectGarbageInternal.VerifyNoUnreachableObjects"), STAT_CollectGarbageInternal_VerifyNoUnreachableObjects, STATGROUP_GC);
+			VerifyNoUnreachableObjects();
 		}
 #endif // VERIFY_DISREGARD_GC_ASSUMPTIONS
 
diff --git a/Engine/Source/Runtime/CoreUObject/Private/UObject/GarbageCollectionVerification.cpp b/Engine/Source/Runtime/CoreUObject/Private/UObject/GarbageCollectionVerification.cpp
--- a/Engine/Source/Runtime/CoreUObject/Private/UObject/GarbageCollectionVerification.cpp
+++ b/Engine/Source/Runtime/CoreUObject/Private/UObject/GarbageCollectionVerification.cpp
@@ -434,7 +434,7 @@ public:
 	}
 };
 
-void VerifyUnreachableObjects()
+void VerifyNoUnreachableObjects()
 {
 	const double StartTime = FPlatformTime::Seconds();
 	const int32 MaxNumberOfReachableObjects = GUObjectArray.GetObjectArrayNum() - GUObjectArray.GetFirstGCIndex();
@@ -442,7 +442,7 @@ void VerifyUnreachableObjects()
 	const int32 NumberOfObjectsPerThread = (MaxNumberOfReachableObjects / NumThreads) + 1;
 	std::atomic<uint32> NumErrors(0);
 
-	ParallelFor(TEXT("GC.VerifyUnreachableObjects"), NumThreads, 1, [&NumErrors, NumberOfObjectsPerThread, NumThreads, MaxNumberOfReachableObjects](int32 ThreadIndex)
+	ParallelFor(TEXT("GC.VerifyNoUnreachableObjects"), NumThreads, 1, [&NumErrors, NumberOfObjectsPerThread, NumThreads, MaxNumberOfReachableObjects](int32 ThreadIndex)
 	{
 		int32 FirstObjectIndex = ThreadIndex * NumberOfObjectsPerThread;
 		int32 NumObjects = (ThreadIndex < (NumThreads - 1)) ? NumberOfObjectsPerThread : (MaxNumberOfReachableObjects - (NumThreads - 1)*NumberOfObjectsPerThread);
@@ -468,7 +468,7 @@ void VerifyUnreachableObjects()
 
 	UE_CLOG(NumErrors > 0, LogGarbage, Fatal, TEXT("Encountered %d unreachable object(s) that are still reachable."), NumErrors.load());
 
-	UE_LOG(LogGarbage, Log, TEXT("%f ms for VerifyUnreachableObjects"), (FPlatformTime::Seconds() - StartTime) * 1000);
+	UE_LOG(LogGarbage, Log, TEXT("%f ms for VerifyNoUnreachableObjects"), (FPlatformTime::Seconds() - StartTime) * 1000);
 }
 
 #endif // VERIFY_DISREGARD_GC_ASSUMPTIONS
diff --git a/Engine/Source/Runtime/CoreUObject/Private/UObject/GarbageCollectionVerification.h b/Engine/Source/Runtime/CoreUObject/Private/UObject/GarbageCollectionVerification.h
--- a/Engine/Source/Runtime/CoreUObject/Private/UObject/GarbageCollectionVerification.h
+++ b/Engine/Source/Runtime/CoreUObject/Private/UObject/GarbageCollectionVerification.h
@@ -26,7 +26,7 @@ void VerifyClustersAssumptions();
 void VerifyObjectFlagMirroring();
 
 /** Verifies Reachability assumptions */
-void VerifyUnreachableObjects();
+void VerifyNoUnreachableObjects();
 
 #endif // VERIFY_DISREGARD_GC_ASSUMPTIONS
 
