Adding CVar to control the Archetype cache usage

diff --git a/Engine/Source/Runtime/CoreUObject/Private/UObject/UObjectArchetype.cpp b/Engine/Source/Runtime/CoreUObject/Private/UObject/UObjectArchetype.cpp
--- a/Engine/Source/Runtime/CoreUObject/Private/UObject/UObjectArchetype.cpp
+++ b/Engine/Source/Runtime/CoreUObject/Private/UObject/UObjectArchetype.cpp
@@ -12,6 +12,7 @@
 #include "UObject/Package.h"
 #include "UObject/UObjectAnnotation.h"
 #include "Stats/StatsMisc.h"
+#include "HAL/IConsoleManager.h"
 
 #define UE_CACHE_ARCHETYPE (1 && !WITH_EDITORONLY_DATA)
 #define UE_VERIFY_CACHED_ARCHETYPE 0
@@ -51,8 +52,21 @@ struct FArchetypeInfo
 	int32 SerialNumber;
 };
 
-static FUObjectAnnotationChunked<FArchetypeInfo, true, 8192> ArchetypeAnnotation;
-
+namespace
+{
+FUObjectAnnotationChunked<FArchetypeInfo, true, 8192> ArchetypeAnnotation;
+
+//CVar to specify if we should use the Achetype cache.
+// default is true.
+// 
+bool bEnableArchetypeCache = true;
+FAutoConsoleVariableRef CVarEnableArchetypeCache(
+	TEXT("EnableArchetypeCache"),
+	bEnableArchetypeCache,
+	TEXT("If set to false, this will disable the use of the ArchetypeCache."),
+	ECVF_Default
+);
+}
 #endif // UE_CACHE_ARCHETYPE
 
 UObject* GetArchetypeFromRequiredInfoImpl(const UClass* Class, const UObject* Outer, FName Name, EObjectFlags ObjectFlags, bool bUseUpToDateClass)
@@ -173,6 +187,11 @@ UObject* UObject::GetArchetype() const
 	//SCOPE_SECONDS_ACCUMULATOR(STAT_FArchiveRealtimeGC_GetArchetype);
 
 #if UE_CACHE_ARCHETYPE
+	if (!bEnableArchetypeCache)
+	{
+		return GetArchetypeFromRequiredInfo(GetClass(), GetOuter(), GetFName(), GetFlags());
+	}
+
 	UObject* Archetype = nullptr;
 	FArchetypeInfo Annoatation = ArchetypeAnnotation.GetAnnotation(this);
 	int32 ArchetypeIndex = Annoatation.ArchetypeIndex;
