Disable PIE soft object path fixup tests outside of editor.

diff --git a/Engine/Source/Runtime/CoreUObject/Private/UObject/SoftObjectPath.cpp b/Engine/Source/Runtime/CoreUObject/Private/UObject/SoftObjectPath.cpp
--- a/Engine/Source/Runtime/CoreUObject/Private/UObject/SoftObjectPath.cpp
+++ b/Engine/Source/Runtime/CoreUObject/Private/UObject/SoftObjectPath.cpp
@@ -1008,6 +1008,8 @@ bool FSoftObjectPathTrySetPathTests::RunTest(const FString& Parameters)
 	return true;
 }
 
+#if WITH_EDITOR
+
 IMPLEMENT_SIMPLE_AUTOMATION_TEST(FSoftObjectPathFixupForPIETests, "System.CoreUObject.SoftObjectPath.FixupForPIE", EAutomationTestFlags::ApplicationContextMask | EAutomationTestFlags::EngineFilter);
 bool FSoftObjectPathFixupForPIETests::RunTest(const FString& Parameters)
 {
@@ -1021,4 +1023,6 @@ bool FSoftObjectPathFixupForPIETests::RunTest(const FString& Parameters)
 	return true;
 }
 
+#endif 
+
 #endif
\ No newline at end of file
diff --git a/Engine/Source/Runtime/Engine/Private/PieFixupSerializer.cpp b/Engine/Source/Runtime/Engine/Private/PieFixupSerializer.cpp
--- a/Engine/Source/Runtime/Engine/Private/PieFixupSerializer.cpp
+++ b/Engine/Source/Runtime/Engine/Private/PieFixupSerializer.cpp
@@ -83,7 +83,7 @@ FArchive& FPIEFixupSerializer::operator<<(FSoftObjectPtr& Value)
 }
 
 
-#if WITH_DEV_AUTOMATION_TESTS
+#if WITH_DEV_AUTOMATION_TESTS && WITH_EDITOR
 
 #include "Tests/PieFixupTestObjects.h"
 
@@ -123,4 +123,4 @@ bool FPIEFixupSerializerTest::RunTest(const FString& Parameters)
 	return true;
 }
 
-#endif
\ No newline at end of file
+#endif // WITH_DEV_AUTOMATION_TESTS && WITH_EDITOR
\ No newline at end of file
