Adding missing GC barriers to runtime

diff --git a/Engine/Source/Runtime/CoreUObject/Private/UObject/PropertyClassPtr.cpp b/Engine/Source/Runtime/CoreUObject/Private/UObject/PropertyClassPtr.cpp
--- a/Engine/Source/Runtime/CoreUObject/Private/UObject/PropertyClassPtr.cpp
+++ b/Engine/Source/Runtime/CoreUObject/Private/UObject/PropertyClassPtr.cpp
@@ -89,7 +89,7 @@ uint32 FClassPtrProperty::GetValueTypeHashInternal(const void* Src) const
 
 void FClassPtrProperty::CopySingleValueToScriptVM( void* Dest, void const* Src ) const
 {
-	*(UObject**)Dest = ((const FObjectPtr*)Src)->Get();
+	*(FObjectPtr*)Dest = ((const FObjectPtr*)Src)->Get();
 }
 
 void FClassPtrProperty::CopySingleValueFromScriptVM( void* Dest, void const* Src ) const
diff --git a/Engine/Source/Runtime/CoreUObject/Private/UObject/PropertyLazyObjectPtr.cpp b/Engine/Source/Runtime/CoreUObject/Private/UObject/PropertyLazyObjectPtr.cpp
--- a/Engine/Source/Runtime/CoreUObject/Private/UObject/PropertyLazyObjectPtr.cpp
+++ b/Engine/Source/Runtime/CoreUObject/Private/UObject/PropertyLazyObjectPtr.cpp
@@ -158,7 +158,7 @@ uint32 FLazyObjectProperty::GetValueTypeHashInternal(const void* Src) const
 
 void FLazyObjectProperty::CopySingleValueToScriptVM( void* Dest, void const* Src ) const
 {
-	*(UObject**)Dest = ((const FLazyObjectPtr*)Src)->Get();
+	*(FObjectPtr*)Dest = ((const FLazyObjectPtr*)Src)->Get();
 }
 
 void FLazyObjectProperty::CopySingleValueFromScriptVM( void* Dest, void const* Src ) const
diff --git a/Engine/Source/Runtime/CoreUObject/Private/UObject/PropertyObject.cpp b/Engine/Source/Runtime/CoreUObject/Private/UObject/PropertyObject.cpp
--- a/Engine/Source/Runtime/CoreUObject/Private/UObject/PropertyObject.cpp
+++ b/Engine/Source/Runtime/CoreUObject/Private/UObject/PropertyObject.cpp
@@ -266,7 +266,7 @@ void FObjectProperty::SetObjectPropertyValue_InContainer(void* ContainerAddress,
 
 void FObjectProperty::CopySingleValueToScriptVM(void* Dest, const void* Src) const
 {
-	*(UObject**)Dest = ((const FObjectPtr*)Src)->Get();
+	*(FObjectPtr*)Dest = ((const FObjectPtr*)Src)->Get();
 }
 
 void FObjectProperty::CopySingleValueFromScriptVM(void* Dest, const void* Src) const
diff --git a/Engine/Source/Runtime/CoreUObject/Private/UObject/PropertyObjectPtr.cpp b/Engine/Source/Runtime/CoreUObject/Private/UObject/PropertyObjectPtr.cpp
--- a/Engine/Source/Runtime/CoreUObject/Private/UObject/PropertyObjectPtr.cpp
+++ b/Engine/Source/Runtime/CoreUObject/Private/UObject/PropertyObjectPtr.cpp
@@ -218,7 +218,7 @@ uint32 FObjectPtrProperty::GetValueTypeHashInternal(const void* Src) const
 
 void FObjectPtrProperty::CopySingleValueToScriptVM(void* Dest, const void* Src) const
 {
-	*(UObject**)Dest = ((const FObjectPtr*)Src)->Get();
+	*(FObjectPtr*)Dest = ((const FObjectPtr*)Src)->Get();
 }
 
 void FObjectPtrProperty::CopySingleValueFromScriptVM(void* Dest, const void* Src) const
diff --git a/Engine/Source/Runtime/CoreUObject/Private/UObject/PropertyWeakObjectPtr.cpp b/Engine/Source/Runtime/CoreUObject/Private/UObject/PropertyWeakObjectPtr.cpp
--- a/Engine/Source/Runtime/CoreUObject/Private/UObject/PropertyWeakObjectPtr.cpp
+++ b/Engine/Source/Runtime/CoreUObject/Private/UObject/PropertyWeakObjectPtr.cpp
@@ -120,7 +120,11 @@ uint32 FWeakObjectProperty::GetValueTypeHashInternal(const void* Src) const
 
 void FWeakObjectProperty::CopySingleValueToScriptVM( void* Dest, void const* Src ) const
 {
-	*(UObject**)Dest = ((const FWeakObjectPtr*)Src)->Get();
+	#if UE_GC_RUN_WEAKPTR_BARRIERS
+	*(FObjectPtr*)Dest = ((const FWeakObjectPtr*)Src)->Get();
+	#else
+	*(UObject**)Dest = ((const FWeakObjectPtr*)Src)->Get();	
+	#endif
 }
 
 void FWeakObjectProperty::CopySingleValueFromScriptVM( void* Dest, void const* Src ) const
diff --git a/Engine/Source/Runtime/CoreUObject/Public/UObject/UObjectGlobals.h b/Engine/Source/Runtime/CoreUObject/Public/UObject/UObjectGlobals.h
--- a/Engine/Source/Runtime/CoreUObject/Public/UObject/UObjectGlobals.h
+++ b/Engine/Source/Runtime/CoreUObject/Public/UObject/UObjectGlobals.h
@@ -94,6 +94,9 @@ DECLARE_DWORD_COUNTER_STAT_EXTERN(TEXT("FindObjectFast"),STAT_FindObjectFast,STA
 #define	INVALID_OBJECT	(UObject*)-1
 #define PERF_TRACK_DETAILED_ASYNC_STATS (0)
 
+#ifndef UE_GC_RUN_WEAKPTR_BARRIERS
+#define UE_GC_RUN_WEAKPTR_BARRIERS 0
+#endif
 
 // Private system wide variables.
 
