Make sure to clear the RF_HasExternalPackage flag before removing from the external package hash table to avoid a race condition.

diff --git a/Engine/Source/Runtime/CoreUObject/Private/UObject/UObjectBase.cpp b/Engine/Source/Runtime/CoreUObject/Private/UObject/UObjectBase.cpp
--- a/Engine/Source/Runtime/CoreUObject/Private/UObject/UObjectBase.cpp
+++ b/Engine/Source/Runtime/CoreUObject/Private/UObject/UObjectBase.cpp
@@ -286,14 +286,15 @@ void UObjectBase::SetExternalPackage(UPackage* InPackage)
 		check(GetClass()->IsChildOf(UPackage::StaticClass()) && (InPackage == this || InPackage == nullptr));
 		return;
 	}
-	HashObjectExternalPackage(this, InPackage);
 	if (InPackage)
 	{
+		HashObjectExternalPackage(this, InPackage);
 		AtomicallySetFlags(RF_HasExternalPackage);
 	}
 	else
 	{
 		AtomicallyClearFlags(RF_HasExternalPackage);
+		UnhashObjectExternalPackage(this);
 	}
 }
 
