Improve comments on LazyObjectPtr to make the limitations clear and warn about future deprecation in favor of SoftObjectPtr

diff --git a/Engine/Source/Runtime/CoreUObject/Public/UObject/LazyObjectPtr.h b/Engine/Source/Runtime/CoreUObject/Public/UObject/LazyObjectPtr.h
--- a/Engine/Source/Runtime/CoreUObject/Public/UObject/LazyObjectPtr.h
+++ b/Engine/Source/Runtime/CoreUObject/Public/UObject/LazyObjectPtr.h
@@ -22,7 +22,8 @@ template <typename T> struct TIsPODType;
 template <typename T> struct TIsWeakPointerType;
 
 /**
- * Wrapper structure for a GUID that uniquely identifies a UObject
+ * Wrapper structure for a GUID that uniquely identifies registered UObjects.
+ * The actual GUID is stored in an object annotation that is updated when a new reference is made.
  */
 struct FUniqueObjectGuid
 {
@@ -123,11 +124,13 @@ template<> struct TIsPODType<FUniqueObjectGuid> { enum { Value = true }; };
 
 /**
  * FLazyObjectPtr is a type of weak pointer to a UObject that uses a GUID created at save time.
+ * Objects will only have consistent GUIDs if they are referenced by a lazy pointer and then saved.
  * It will change back and forth between being valid or pending as the referenced object loads or unloads.
  * It has no impact on if the object is garbage collected or not.
  * It can't be directly used across a network.
  *
- * This is useful for cross level references or places where you need to point to an object whose name changes often.
+ * NOTE: Because this only stores a GUID, it does not know how to load the destination object and does not work with Play In Editor.
+ * This will be deprecated in a future engine version and new features should use FSoftObjectPtr instead.
  */
 struct FLazyObjectPtr : public TPersistentObjectPtr<FUniqueObjectGuid>
 {
@@ -155,7 +158,7 @@ public:
 		TPersistentObjectPtr<FUniqueObjectGuid>::operator=(InObjectID);
 	}
 
-	/** Fixes up this FLazyObjectPtr to target the right UID as set in PIEGuidMap */
+	/** Fixes up this FLazyObjectPtr to target the right UID as set in PIEGuidMap, this only works for directly serialized pointers */
 	FORCEINLINE void FixupForPIE(int32 PIEInstance)
 	{
 		*this = GetUniqueID().FixupForPIE(PIEInstance);
@@ -172,7 +175,8 @@ template <> struct TIsPODType<FLazyObjectPtr> { enum { Value = TIsPODType<TPersi
 template <> struct TIsWeakPointerType<FLazyObjectPtr> { enum { Value = TIsWeakPointerType<TPersistentObjectPtr<FUniqueObjectGuid> >::Value }; };
 
 /**
- * TLazyObjectPtr is templatized version of the generic FLazyObjectPtr
+ * TLazyObjectPtr is the templatized version of the generic FLazyObjectPtr.
+ * NOTE: This will be deprecated in a future engine version and new features should use TSoftObjectPtr instead.
  */
 template<class T=UObject>
 struct TLazyObjectPtr : private FLazyObjectPtr
