Garbage Collector barrier support for MutableView of TSet and TMap

diff --git a/Engine/Source/Runtime/CoreUObject/Public/UObject/ObjectPtr.h b/Engine/Source/Runtime/CoreUObject/Public/UObject/ObjectPtr.h
--- a/Engine/Source/Runtime/CoreUObject/Public/UObject/ObjectPtr.h
+++ b/Engine/Source/Runtime/CoreUObject/Public/UObject/ObjectPtr.h
@@ -1328,12 +1328,56 @@ namespace UE::Core::Private // private facilities; not for direct use
 		}
 	};
 
+	template <typename V, typename ViewType>
+	struct TMutableViewTraits<TSet<TObjectPtr<V>>, ViewType>
+	{
+		static void Close(ViewType& View)
+		{
+#if UE_OBJECT_PTR_GC_BARRIER
+			if (UE::GC::Private::GIsIncrementalReachabilityPending)
+			{
+				for (const typename ViewType::ElementType& Element : View)
+				{
+					if (Element)
+					{
+						Element->MarkAsReachable();
+					}
+				}
+			}
+#endif // UE_OBJECT_PTR_GC_BARRIER
+		}
+	};
+
 	template <typename K, typename V, typename ViewType>
 	struct TMutableViewTraits<TMap<K, V>, ViewType>
 	{
 		static void Close(ViewType& View)
 		{
-			// TODO: Run GC Barriers
+#if UE_OBJECT_PTR_GC_BARRIER
+			static constexpr bool bKeyReference = TIsTObjectPtr<K>::Value;
+			static constexpr bool bValueReference = TIsTObjectPtr<V>::Value;
+			static_assert(bKeyReference || bValueReference);
+			if (UE::GC::Private::GIsIncrementalReachabilityPending)
+			{
+				for (const typename ViewType::ElementType& Pair : View)
+				{
+					if constexpr (bKeyReference)
+					{
+						if (Pair.Key)
+						{
+							Pair.Key->MarkAsReachable();
+						}
+					}
+					if constexpr (bValueReference)
+					{
+						if (Pair.Value)
+						{
+							Pair.Value->MarkAsReachable();
+						}
+					}
+				}
+			}
+#endif // UE_OBJECT_PTR_GC_BARRIER
 		}
 	};
 	
