Add optional support to Unversioned Property Serialization

Previously small optionals worked, by accident rather than design, but large optionals didn't

diff --git a/Engine/Source/Runtime/CoreUObject/Private/Serialization/UnversionedPropertySerialization.cpp b/Engine/Source/Runtime/CoreUObject/Private/Serialization/UnversionedPropertySerialization.cpp
--- a/Engine/Source/Runtime/CoreUObject/Private/Serialization/UnversionedPropertySerialization.cpp
+++ b/Engine/Source/Runtime/CoreUObject/Private/Serialization/UnversionedPropertySerialization.cpp
@@ -5,6 +5,7 @@
 #include "Hash/Blake3.h"
 #include "Interfaces/ITargetPlatform.h"
 #include "Misc/ScopeRWLock.h"
+#include "UObject/PropertyOptional.h"
 #include "UObject/UnrealType.h"
 
 #if WITH_EDITORONLY_DATA
@@ -49,12 +50,17 @@ public:
 #if CACHE_UNVERSIONED_PROPERTY_SCHEMA
 		, Offset(Property->GetOffset_ForInternal() + Property->ElementSize * InArrayIndex)
 		, bSerializeAsInteger(CanSerializeAsInteger(Property))
+		, bIsOptional(IsOptional(Property->GetClass()->GetCastFlags()))
 		, IntType(GetIntType(Property->GetMinAlignment()))
-		, FastZeroIntNum(CanSerializeAsZero(InProperty, IntType) ? static_cast<uint8>(GetIntNum(InProperty, IntType)) : uint8(0))
 #else
 		, ArrayIndex(InArrayIndex)
 #endif
-	{}
+	{
+#if CACHE_UNVERSIONED_PROPERTY_SCHEMA
+		uint32 IntNum = GetIntNum(InProperty, IntType);
+		FastZeroIntNum = (CanSerializeAsZero(InProperty, IntType) & (IntNum < 256)) ? static_cast<uint8>(IntNum) : uint8(0);
+#endif
+	}
 
 	FProperty* GetProperty() const
 	{
@@ -110,16 +116,27 @@ public:
 
 	FORCEINLINE void LoadZero(uint8* Data) const
 	{
-#if !CACHE_UNVERSIONED_PROPERTY_SCHEMA
+		void* ValueData = GetValue(Data);
+
+#if CACHE_UNVERSIONED_PROPERTY_SCHEMA
+		// Cached FastZeroIntNum is only uint8 and not sufficient for large unset optionals
+		if (FastZeroIntNum == 0)
+		{
+			checkf(bIsOptional && Property->ElementSize >= 256, TEXT("Only large unset optionals should hit this loading path"));
+			FMemory::Memzero(ValueData, Property->ElementSize);
+			return;
+		}
+#else
 		EIntegerType IntType = GetIntType(Property->GetMinAlignment());
 		uint32 FastZeroIntNum = GetIntNum(Property, IntType);
 #endif
+
 		switch (IntType)
 		{
-			case EIntegerType::Uint8 : MemZeroRange<uint8 >(GetValue(Data), FastZeroIntNum); break;
-			case EIntegerType::Uint16: MemZeroRange<uint16>(GetValue(Data), FastZeroIntNum); break;
-			case EIntegerType::Uint32: MemZeroRange<uint32>(GetValue(Data), FastZeroIntNum); break;
-			case EIntegerType::Uint64: MemZeroRange<uint64>(GetValue(Data), FastZeroIntNum); break;
+			case EIntegerType::Uint8 : MemZeroRange<uint8 >(ValueData, FastZeroIntNum); break;
+			case EIntegerType::Uint16: MemZeroRange<uint16>(ValueData, FastZeroIntNum); break;
+			case EIntegerType::Uint32: MemZeroRange<uint32>(ValueData, FastZeroIntNum); break;
+			case EIntegerType::Uint64: MemZeroRange<uint64>(ValueData, FastZeroIntNum); break;
 			default: UE_ASSUME(0);
 		}
 	}
@@ -128,11 +145,15 @@ public:
 	{
 #if !CACHE_UNVERSIONED_PROPERTY_SCHEMA
 		EIntegerType IntType = GetIntType(Property->GetMinAlignment());
-		uint32 FastZeroIntNum = CanSerializeAsZero(Property, IntType) ? GetIntNum(Property, IntType) : uint8(0);
+		bool bIsOptional = IsOptional(Property->GetClass()->GetCastFlags());
+		uint32 FastZeroIntNum = CanSerializeAsZero(Property, IntType) ? GetIntNum(Property, IntType) : uint32(0);
 #endif
 
-		// Can simplified and faster using a switch() statement like LoadZero()
-		if (FastZeroIntNum == 1)
+		if (bIsOptional)
+		{
+			return !static_cast<const FOptionalProperty*>(Property)->IsSet(GetValue(Data));
+		}
+		else if (FastZeroIntNum == 1) // Can likely be simplified and faster using a switch() statement like LoadZero()
 		{
 			return IsIntZero(GetValue(Data), IntType);
 		}
@@ -150,6 +171,11 @@ public:
 private:
 	enum class EIntegerType : uint8 { Uint8, Uint16, Uint32, Uint64 };
 
+	static bool IsOptional(uint64 CastFlags)
+	{
+		return !!(CastFlags & CASTCLASS_FOptionalProperty);
+	}
+
 	static uint32 GetIntNum(const FProperty* Property, EIntegerType IntType)
 	{
 		return Property->ElementSize / GetSizeOf(IntType);
@@ -163,7 +189,7 @@ private:
 
 		if ((CastFlags & (CASTCLASS_FStructProperty | CASTCLASS_FBoolProperty)) == 0)
 		{
-			checkf(GetIntNum(Property, IntType) < MaxZeroComparisons, TEXT("Unexpectedly large property type encountered %s"), *Property->GetName());
+			checkf(GetIntNum(Property, IntType) < MaxZeroComparisons || IsOptional(CastFlags),  TEXT("Unexpectedly large property type encountered %s"), *Property->GetName());
 
 			return true;
 		}
@@ -272,6 +298,7 @@ private:
 #if CACHE_UNVERSIONED_PROPERTY_SCHEMA
 	uint32 Offset;
 	bool bSerializeAsInteger;
+	bool bIsOptional;
 	EIntegerType IntType;
 	uint8 FastZeroIntNum;
 #else
diff --git a/Engine/Source/Runtime/CoreUObject/Private/UObject/PropertyOptional.cpp b/Engine/Source/Runtime/CoreUObject/Private/UObject/PropertyOptional.cpp
--- a/Engine/Source/Runtime/CoreUObject/Private/UObject/PropertyOptional.cpp
+++ b/Engine/Source/Runtime/CoreUObject/Private/UObject/PropertyOptional.cpp
@@ -174,6 +174,19 @@ void FOptionalProperty::SerializeItem(FStructuredArchive::FSlot Slot, void* Data
 
 	FStructuredArchive::FRecord Record = Slot.EnterRecord();
 
+	if (UnderlyingArchive.UseUnversionedPropertySerialization())
+	{
+		// UPS only calls SerializeItem on set optionals, just need to serialize the inner property
+		checkf(bIsLoading || IsSet(Data), TEXT("UPS should never call SerializeItem to save an empty Optional"));
+
+		FStructuredArchive::FSlot ValueSlot = Record.EnterField(TEXT("Value"));
+		const void* ValueDefaults = Defaults ? GetValuePointerForReadIfSet(Defaults) : nullptr;
+		void* ValueData = bIsLoading ? MarkSetAndGetInitializedValuePointerToReplace(Data) : GetValuePointerForReadOrReplace(Data);
+		GetValueProperty()->SerializeItem(ValueSlot, ValueData, ValueDefaults);
+
+		return;
+	}
+
 	// Use an optional field slot to encode whether the optional was set or not.
 	TOptional<FStructuredArchiveSlot> MaybeTaggedValueSlot = Record.TryEnterField(TEXT("TaggedValue"), IsSet(Data));
 	if (MaybeTaggedValueSlot.IsSet())
@@ -182,16 +195,7 @@ void FOptionalProperty::SerializeItem(FStructuredArchive::FSlot Slot, void* Data
 		
 		const void* ValueDefaults = Defaults ? GetValuePointerForReadIfSet(Defaults) : nullptr;
 
-		if (Slot.GetArchiveState().UseUnversionedPropertySerialization())
-		{
-			// Simply serialize the inner value if using unversioned property serialization.
-			FStructuredArchive::FSlot ValueSlot = TaggedValueRecord.EnterField(TEXT("Value"));
-			void* ValueData = bIsLoading
-				? MarkSetAndGetInitializedValuePointerToReplace(Data)
-				: GetValuePointerForReadOrReplace(Data);
-			GetValueProperty()->SerializeItem(ValueSlot, ValueData, ValueDefaults);
-		}
-		else if (bIsLoading)
+		if (bIsLoading)
 		{
 			// Serialize the value's tag.
 			FPropertyTag ValueTag;
