Fixed typo.

diff --git a/Engine/Source/Runtime/CoreUObject/Private/UObject/Field.cpp b/Engine/Source/Runtime/CoreUObject/Private/UObject/Field.cpp
--- a/Engine/Source/Runtime/CoreUObject/Private/UObject/Field.cpp
+++ b/Engine/Source/Runtime/CoreUObject/Private/UObject/Field.cpp
@@ -343,10 +343,10 @@ FField::FField(UField* InField)
 		Owner = OriginalOuter;
 	}
 
-	TMap<FName, FString>* FeldMetaDataMap = UMetaData::GetMapForObject(InField);
-	if (FeldMetaDataMap && FeldMetaDataMap->Num())
+	TMap<FName, FString>* FieldMetaDataMap = UMetaData::GetMapForObject(InField);
+	if (FieldMetaDataMap && FieldMetaDataMap->Num())
 	{
-		MetaDataMap = new TMap<FName, FString>(*FeldMetaDataMap);
+		MetaDataMap = new TMap<FName, FString>(*FieldMetaDataMap);
 	}
 }
 #endif // WITH_EDITORONLY_DATA
