Minor optimization -  Early out when trying to load an empty bulkdata payload rather than allocating then freeing a zero byte buffer.

#rnx

- Note that this should not change the outcome as ::TryLoadDataIntoMemory will return false if the length of the payload is zero s the previous behaviour was to allocate, fail to load, then delete. The new code path and the old code path both end up with the bulkdata containing no allocation.

diff --git a/Engine/Source/Runtime/CoreUObject/Private/Serialization/BulkData.cpp b/Engine/Source/Runtime/CoreUObject/Private/Serialization/BulkData.cpp
--- a/Engine/Source/Runtime/CoreUObject/Private/Serialization/BulkData.cpp
+++ b/Engine/Source/Runtime/CoreUObject/Private/Serialization/BulkData.cpp
@@ -1635,6 +1635,12 @@ void FBulkData::MakeSureBulkDataIsLoaded()
 	}
 
 	const int64 BulkDataSize = GetBulkDataSize();
+
+	if (BulkDataSize == 0)
+	{
+		return;
+	}
+
 	void* Dest = ReallocateData(BulkDataSize);
 
 	if (TryLoadDataIntoMemory(FIoBuffer(FIoBuffer::Wrap, Dest, BulkDataSize)) == false)
