Pack FUObjectItem in Shipping builds to avoid 20% mem waste

diff --git a/Engine/Source/Runtime/CoreUObject/Public/UObject/UObjectArray.h b/Engine/Source/Runtime/CoreUObject/Public/UObject/UObjectArray.h
--- a/Engine/Source/Runtime/CoreUObject/Public/UObject/UObjectArray.h
+++ b/Engine/Source/Runtime/CoreUObject/Public/UObject/UObjectArray.h
@@ -22,7 +22,12 @@
 /**
 * Single item in the UObject array.
 */
-struct FUObjectItem
+struct 
+#if !STATS
+	// Packing avoids 20% mem waste and improves perf
+	GCC_PACK(4)
+#endif
+	FUObjectItem
 {
 	// Pointer to the allocated object
 	class UObjectBase* Object;
