[Backout] - CL26285807
Original CL Desc
-----------------------------------------------------------------
Pack FUObjectItem in Shipping builds to avoid 20% mem waste

diff --git a/Engine/Source/Runtime/CoreUObject/Public/UObject/UObjectArray.h b/Engine/Source/Runtime/CoreUObject/Public/UObject/UObjectArray.h
--- a/Engine/Source/Runtime/CoreUObject/Public/UObject/UObjectArray.h
+++ b/Engine/Source/Runtime/CoreUObject/Public/UObject/UObjectArray.h
@@ -22,12 +22,7 @@
 /**
 * Single item in the UObject array.
 */
-struct 
-#if !STATS
-	// Packing avoids 20% mem waste and improves perf
-	GCC_PACK(4)
-#endif
-	FUObjectItem
+struct FUObjectItem
 {
 	// Pointer to the allocated object
 	class UObjectBase* Object;
