Fix FOptionalProperty text-based serialization.

diff --git a/Engine/Source/Runtime/CoreUObject/Private/UObject/PropertyOptional.cpp b/Engine/Source/Runtime/CoreUObject/Private/UObject/PropertyOptional.cpp
--- a/Engine/Source/Runtime/CoreUObject/Private/UObject/PropertyOptional.cpp
+++ b/Engine/Source/Runtime/CoreUObject/Private/UObject/PropertyOptional.cpp
@@ -188,24 +188,27 @@ void FOptionalProperty::SerializeItem(FStructuredArchive::FSlot Slot, void* Data
 	}
 
 	// Use an optional field slot to encode whether the optional was set or not.
-	TOptional<FStructuredArchiveSlot> MaybeTaggedValueSlot = Record.TryEnterField(TEXT("TaggedValue"), IsSet(Data));
-	if (MaybeTaggedValueSlot.IsSet())
+	TOptional<FStructuredArchiveSlot> MaybeValueSlot = Record.TryEnterField(TEXT("Value"), IsSet(Data));
+	if (MaybeValueSlot.IsSet())
 	{
-		FStructuredArchive::FRecord TaggedValueRecord = MaybeTaggedValueSlot.GetValue().EnterRecord();
+		FStructuredArchive::FSlot ValueSlot = MaybeValueSlot.GetValue();
 		
 		const void* ValueDefaults = Defaults ? GetValuePointerForReadIfSet(Defaults) : nullptr;
 
 		if (bIsLoading)
 		{
-			// Serialize the value's tag.
 			FPropertyTag ValueTag;
-			TaggedValueRecord << SA_VALUE(TEXT("Tag"), ValueTag);
+
+			// Deserializing a FPropertyTag from text won't deserialize the ArrayIndex, leaving it uninitialized as INDEX_NONE.
+			ValueTag.ArrayIndex = 0;
+
+			// Serialize the value's tag.
+			ValueSlot << ValueTag;
 
 			// Deserialize/convert the value.
 			void* ValueData = MarkSetAndGetInitializedValuePointerToReplace(Data);
 
 			const int64 ValueStartOffset = UnderlyingArchive.Tell();
-			FStructuredArchive::FSlot ValueSlot = TaggedValueRecord.EnterField(TEXT("Value"));
 			bool bSuccessfullyDeserialized = false;
 			switch (GetValueProperty()->ConvertFromType(ValueTag, ValueSlot, static_cast<uint8*>(ValueData), GetOwnerStruct(), static_cast<const uint8*>(ValueDefaults)))
 			{
@@ -247,11 +250,11 @@ void FOptionalProperty::SerializeItem(FStructuredArchive::FSlot Slot, void* Data
 
 			// Construct and serialize a tag for the value.
 			FPropertyTag ValueTag(UnderlyingArchive, GetValueProperty(), 0, static_cast<uint8*>(ValueData), static_cast<const uint8*>(ValueDefaults));
-			TaggedValueRecord << SA_VALUE(TEXT("Tag"), ValueTag);
+			ValueSlot << ValueTag;
 
 			// Serialize the value.
 			int64 ValueStartOffset = UnderlyingArchive.Tell();
-			ValueTag.SerializeTaggedProperty(TaggedValueRecord.EnterField(TEXT("Value")), GetValueProperty(), static_cast<uint8*>(ValueData), static_cast<const uint8*>(ValueDefaults));
+			ValueTag.SerializeTaggedProperty(ValueSlot, GetValueProperty(), static_cast<uint8*>(ValueData), static_cast<const uint8*>(ValueDefaults));
 			
 			// If saving to a non-text archive, save the size of the serialized value so it can be skipped over on load if it's the wrong type.
 			const int64 ValueEndOffset = UnderlyingArchive.Tell();
