Fix a bug where loading a bulkdata payload from a cooked package mounted in the IoStore could end up loading more data off disk than required.

### Problem
- The IoStore does not support bulkdata payloads being compressed at the package level, instead compression will be applied to the .ucas file when it is created, this means that we don't really use the SizeOnDisk value for cooked bulkdata as it should always be the same as the in memory size.
- Some code paths in the editor however might support compression in the workspace domain and so use the SizeOnDisk value, such as calling ::GetCopy.
- In this cause the size being passed to the IoDispatcher would be INDEX_NONE, which the system converts to "please load the entire .ubulk section of the package"
- The buffer returned to the user would eventually be trimmed to the correct in memory size so this bug would be hard to notice, the only downside is wasting time loading more data off disk than required.

### Fix
- AsyncLoading2 now sets the SizeOnDisk to be a valid value (the size of the payload in memory) if we are not using runtime only bulkdata (USE_RUNTIME_BULKDATA) so when GetSizeOnDisk is called, it returns the correct value and we only load the payload we want.
- An assert has been added to make sure that the payload is not compressed (in which case the size would be wrong and we could end up overrunning valid memory) even though the flag is removed when cooking and it should not be possible to get to that point.

diff --git a/Engine/Source/Runtime/CoreUObject/Private/Serialization/AsyncLoading2.cpp b/Engine/Source/Runtime/CoreUObject/Private/Serialization/AsyncLoading2.cpp
--- a/Engine/Source/Runtime/CoreUObject/Private/Serialization/AsyncLoading2.cpp
+++ b/Engine/Source/Runtime/CoreUObject/Private/Serialization/AsyncLoading2.cpp
@@ -1895,6 +1895,18 @@ private:
 			Meta.SetFlags(static_cast<EBulkDataFlags>(Entry.Flags));
 			Meta.SetOffset(Entry.SerialOffset);
 			Meta.SetSize(Entry.SerialSize);
+
+#if !USE_RUNTIME_BULKDATA
+			// If the payload was compressed at package level then we will not be able to decompress it properly as that requires
+			// us to know the compressed size (SizeOnDisk) which we do not keep track of when the package is stored by the IoDispatcher.
+			// The BULKDATA_SerializeCompressed flag is removed during cooking/staging so the flag should never be set at this point,
+			// the assert is just a paranoid safety check.
+			checkf(Meta.HasAnyFlags(BULKDATA_SerializeCompressed) == false, TEXT("Package level compression is not supported by the IoDispatcher: '%s'"), *PackageDesc->UPackageName.ToString());
+
+			// Since we know that the payload is not compressed there is no difference between the in memory size and the size of disk
+			Meta.SetSizeOnDisk(Entry.SerialSize);
+#endif //!USE_RUNTIME_BULKDATA
+
 			DuplicateSerialOffset = Entry.DuplicateSerialOffset;
 		}
 
