Fix double redirector not properly resolving themselves in some loading order edge case

diff --git a/Engine/Source/Runtime/CoreUObject/Private/UObject/LinkerLoad.cpp b/Engine/Source/Runtime/CoreUObject/Private/UObject/LinkerLoad.cpp
--- a/Engine/Source/Runtime/CoreUObject/Private/UObject/LinkerLoad.cpp
+++ b/Engine/Source/Runtime/CoreUObject/Private/UObject/LinkerLoad.cpp
@@ -5528,6 +5528,7 @@ UObject* FLinkerLoad::CreateImport( int32 Index )
 					UObjectRedirector* Redirector = dynamic_cast<UObjectRedirector*>(Import.XObject);
 					if( Redirector )
 					{
+						Preload(Redirector);
 						Import.XObject = Redirector->DestinationObject;
 					}
 				}
