Fix for type-mismatched leading to a fatal error when compiling, cooking, and loading level data. Now resets non-nullable properties to an instandce of the CDO.

diff --git a/Engine/Source/Runtime/CoreUObject/Private/UObject/PropertyBaseObject.cpp b/Engine/Source/Runtime/CoreUObject/Private/UObject/PropertyBaseObject.cpp
--- a/Engine/Source/Runtime/CoreUObject/Private/UObject/PropertyBaseObject.cpp
+++ b/Engine/Source/Runtime/CoreUObject/Private/UObject/PropertyBaseObject.cpp
@@ -717,6 +717,11 @@ void FObjectPropertyBase::CheckValidObject(void* ValueAddress, TObjectPtr<UObjec
 		bool bIsReplacingClassRefs = PropertyClass && PropertyClass->HasAnyClassFlags(CLASS_NewerVersionExists) != ObjectClass->HasAnyClassFlags(CLASS_NewerVersionExists);
 		if (!bIsReplacingClassRefs && !IsDeferringValueLoad())
 		{
+			if (UPackage* ObjectPackage = Object->GetPackage())
+			{
+				ObjectPackage->SetDirtyFlag(true);
+			}
+
 			if (!HasAnyPropertyFlags(CPF_NonNullable))
 			{
 				UE_LOG(LogProperty, Warning,
@@ -730,6 +735,7 @@ void FObjectPropertyBase::CheckValidObject(void* ValueAddress, TObjectPtr<UObjec
 			}
 			else
 			{
+				Object->Rename(nullptr, GetTransientPackage(), REN_ForceNoResetLoaders | REN_DoNotDirty | REN_DontCreateRedirectors);
 				UObject* DefaultValue = ConstructDefaultObjectValueIfNecessary(OldValue);
 
 				UE_LOG(LogProperty, Warning,
