ObjectPropertyTest fixes:
1) Fix for test package names not starting with proper prefix
2) Fix incompatible value that is passed into CheckValidObject to NOT be the outer of the fallback objecte. The incompatible value now gets renamed by CheckValidObject, and renaming the outer of the fallback object would break the test.
#rnx

diff --git a/Engine/Source/Runtime/CoreUObject/Tests/ObjectPropertyTest.cpp b/Engine/Source/Runtime/CoreUObject/Tests/ObjectPropertyTest.cpp
--- a/Engine/Source/Runtime/CoreUObject/Tests/ObjectPropertyTest.cpp
+++ b/Engine/Source/Runtime/CoreUObject/Tests/ObjectPropertyTest.cpp
@@ -32,7 +32,7 @@ TEST_CASE("UE::CoreUObject::FObjectProperty::CheckValidAddress")
 	FObjectProperty* Property = CastField<FObjectProperty>(Class->FindPropertyByName(TEXT("ObjectPtr")));
 	REQUIRE(Property != nullptr);
 
-	UPackage* TestPackage = NewObject<UPackage>(nullptr, TEXT("TestPackageName"), RF_Transient);
+	UPackage* TestPackage = NewObject<UPackage>(nullptr, TEXT("/Temp/TestPackageName"), RF_Transient);
 	TestPackage->AddToRoot();
 	ON_SCOPE_EXIT
 	{
@@ -94,7 +94,8 @@ TEST_CASE("UE::CoreUObject::FObjectProperty::CheckValidAddressNonNullable")
 	FObjectProperty* Property = CastField<FObjectProperty>(Class->FindPropertyByName(TEXT("ObjectPtrNonNullable")));
 	REQUIRE(Property != nullptr);
 
-	UPackage* TestPackage = NewObject<UPackage>(nullptr, TEXT("TestPackageName"), RF_Transient);
+	UPackage* TestPackage = NewObject<UPackage>(nullptr, TEXT("/Temp/TestPackageName"), RF_Transient);
+	UPackage* OtherTestPackage = NewObject<UPackage>(nullptr, TEXT("/Temp/CheckValidAddressNonNullableOther"), RF_Transient);
 	TestPackage->AddToRoot();
 	ON_SCOPE_EXIT
 	{
@@ -122,7 +123,7 @@ TEST_CASE("UE::CoreUObject::FObjectProperty::CheckValidAddressNonNullable")
 
 	bAllowRead = true; //has resolve the old value to construct a new default value for the property
 	//assign a bad value to the pointer
-	Obj->ObjectPtrNonNullable = reinterpret_cast<UObjectPtrTestClass*>(Obj);
+	Obj->ObjectPtrNonNullable = reinterpret_cast<UObjectPtrTestClass*>(OtherTestPackage);
 	CHECK(Obj->ObjectPtrNonNullable != nullptr);
 
 	UE::Testing::FWarnFilterScope _([](const TCHAR* Message, ELogVerbosity::Type Verbosity, const FName& Category)
@@ -165,7 +166,7 @@ static void TestSerializeItem(FName ObjectName)
 	FObjectProperty* Property = CastField<FObjectProperty>(Class->FindPropertyByName(TEXT("ObjectPtr")));
 	REQUIRE(Property != nullptr);
 
-	UPackage* TestPackage = NewObject<UPackage>(nullptr, TEXT("TestPackageName"), RF_Transient);
+	UPackage* TestPackage = NewObject<UPackage>(nullptr, TEXT("/Temp/TestPackageName"), RF_Transient);
 	TestPackage->AddToRoot();
 	ON_SCOPE_EXIT
 	{
@@ -291,7 +292,7 @@ TEST_CASE("UE::CoreUObject::FObjectProperty::ParseObjectPropertyValue")
 	AssetData.AssetClassPath.TrySetPath(UObjectPtrTestClass::StaticClass());
 
 	MockAssetRegistry.AssetData.Add(Path, AssetData);
-	UPackage* TestPackage = NewObject<UPackage>(nullptr, TEXT("TestPackageName"), RF_Transient);
+	UPackage* TestPackage = NewObject<UPackage>(nullptr, TEXT("/Temp/TestPackageName"), RF_Transient);
 	TestPackage->AddToRoot();
 	ON_SCOPE_EXIT
 	{
diff --git a/Engine/Source/Runtime/CoreUObject/Tests/PropertyClassTest.cpp b/Engine/Source/Runtime/CoreUObject/Tests/PropertyClassTest.cpp
--- a/Engine/Source/Runtime/CoreUObject/Tests/PropertyClassTest.cpp
+++ b/Engine/Source/Runtime/CoreUObject/Tests/PropertyClassTest.cpp
@@ -30,7 +30,7 @@ TEST_CASE("UE::CoreUObject::FClassProperty::Identical")
 	REQUIRE(Property != nullptr);
 
 
-	UPackage* TestPackage = NewObject<UPackage>(nullptr, TEXT("TestPackageName"), RF_Transient);
+	UPackage* TestPackage = NewObject<UPackage>(nullptr, TEXT("/Temp/TestPackageName"), RF_Transient);
 	TestPackage->AddToRoot();
 	ON_SCOPE_EXIT
 	{
