PackageReload: Optimized ReloadPackages by doing a pre-pass using FReferencerFinder::GetAllReferencers to get all potential referencers and only inspecting those when fixing up references. The potential referencers set is expanded as new packages get (re)loaded through the FCoreUObjectDelegates::OnObjectsReplaced and FCoreUObjectDelegates::OnObjectConstructed delegates, whilst garbage collected packages end up being skipped through the use of TWeakObjectPtrs.

#rnx

diff --git a/Engine/Source/Runtime/CoreUObject/Private/UObject/PackageReload.cpp b/Engine/Source/Runtime/CoreUObject/Private/UObject/PackageReload.cpp
--- a/Engine/Source/Runtime/CoreUObject/Private/UObject/PackageReload.cpp
+++ b/Engine/Source/Runtime/CoreUObject/Private/UObject/PackageReload.cpp
@@ -11,6 +11,136 @@
 #include "Templates/Casts.h"
 #include "UObject/UObjectIterator.h"
 
+#if WITH_EDITOR
+#include "UObject/ReferencerFinder.h"
+#endif
+
+#if WITH_EDITOR
+TAutoConsoleVariable<bool> CVarPackageReloadEnableFastPath(
+	TEXT("PackageReload.EnableFastPath"),
+	true,
+	TEXT("When 'true', an optimized codepath is used to speed up reloading packages (experimental)."),
+	ECVF_Default);
+#endif
+
+class FPackageReferencersHelper
+{
+#if WITH_EDITOR
+	/**
+	 * Set of all objects that could potentially refer to any of the packages that are about to be reloaded.
+	 */
+	TSet<TWeakObjectPtr<UObject>> PotentialReferencerObjects;
+
+	/**
+	 * Callback handles.
+	 */
+	FDelegateHandle OnObjectsReplacedHandle;
+	FDelegateHandle OnObjectConstructedHandle;
+
+	/**
+	 * Callback to be called when one or more UObjects are replaced with others.
+	 */
+	void OnObjectsReplaced(const TMap<UObject*, UObject*>& OldToNewInstanceMap)
+	{
+		for (const TPair<UObject*, UObject*>& Pair : OldToNewInstanceMap)
+		{
+			if (Pair.Key && Pair.Value)
+			{
+				TWeakObjectPtr<UObject> OldObject = MakeWeakObjectPtr<UObject>(Pair.Key);
+				TWeakObjectPtr<UObject> NewObject = MakeWeakObjectPtr<UObject>(Pair.Value);
+				if (PotentialReferencerObjects.Contains(OldObject))
+				{
+					PotentialReferencerObjects.Add(NewObject);
+					PotentialReferencerObjects.Remove(OldObject);
+				}
+			}
+		}
+	}
+
+	/**
+	 * Callback to be called when a new UObject is constructed.
+	 */
+	void OnObjectConstructed(UObject* InObject)
+	{
+		TWeakObjectPtr<UObject> Object = MakeWeakObjectPtr<UObject>(InObject);
+		PotentialReferencerObjects.Add(Object);
+	}
+#endif
+
+public:
+	FPackageReferencersHelper(const TArrayView<FReloadPackageData>& InPackagesToReload)
+	{
+#if WITH_EDITOR
+		if (CVarPackageReloadEnableFastPath.GetValueOnAnyThread())
+		{
+			OnObjectsReplacedHandle = FCoreUObjectDelegates::OnObjectsReplaced.AddRaw(this, &FPackageReferencersHelper::OnObjectsReplaced);
+			OnObjectConstructedHandle = FCoreUObjectDelegates::OnObjectConstructed.AddRaw(this, &FPackageReferencersHelper::OnObjectConstructed);
+
+			// Build list of UObjects that represent the packages about to be reloaded.
+			TArray<UObject*> ObjectsToReload;
+			ObjectsToReload.Reserve(InPackagesToReload.Num());
+
+			for (const FReloadPackageData& PackageToReloadData : InPackagesToReload)
+			{
+				ObjectsToReload.Add(PackageToReloadData.PackageToReload);
+			}
+
+			// Retrieve all other UObjects that refer to any of them.
+			TArray<UObject*> ReferencerObjects = FReferencerFinder::GetAllReferencers(ObjectsToReload, nullptr, EReferencerFinderFlags::None);
+
+			// Convert them into weak pointers to deal with objects that get GC'd during the reload operation.
+			PotentialReferencerObjects.Reserve(ReferencerObjects.Num());
+			for (UObject* ReferencerObject : ReferencerObjects)
+			{
+				PotentialReferencerObjects.Emplace(MakeWeakObjectPtr(ReferencerObject));
+			}
+		}
+#endif
+	}
+
+	~FPackageReferencersHelper()
+	{
+#if WITH_EDITOR
+		if (OnObjectConstructedHandle.IsValid())
+		{
+			FCoreUObjectDelegates::OnObjectConstructed.Remove(OnObjectConstructedHandle);
+			OnObjectConstructedHandle.Reset();
+		}
+		if (OnObjectsReplacedHandle.IsValid())
+		{
+			FCoreUObjectDelegates::OnObjectsReplaced.Remove(OnObjectsReplacedHandle);
+			OnObjectsReplacedHandle.Reset();
+		}
+#endif
+	}
+
+public:
+	void ForEachObject(TFunctionRef<void(UObject*)> Operation)
+	{
+#if WITH_EDITOR
+		if (CVarPackageReloadEnableFastPath.GetValueOnAnyThread())
+		{
+			for (TWeakObjectPtr<UObject> Object : PotentialReferencerObjects)
+			{
+				UObject* PotentialReferencer = Object.Get();
+				if (PotentialReferencer != nullptr)
+				{
+					Operation(PotentialReferencer);
+				}
+			}
+		}
+		else
+#endif
+		{
+			for (FThreadSafeObjectIterator ObjIter(UObject::StaticClass(), false, RF_NoFlags, EInternalObjectFlags::Garbage); ObjIter; ++ObjIter)
+			{
+				UObject* PotentialReferencer = *ObjIter;
+				Operation(PotentialReferencer);
+			}
+		}
+	}
+};
+
 namespace PackageReloadInternal
 {
 
@@ -472,6 +602,8 @@ UPackage* ReloadPackage(UPackage* InPackageToReload, const uint32 InLoadFlags)
 
 void ReloadPackages(const TArrayView<FReloadPackageData>& InPackagesToReload, TArray<UPackage*>& OutReloadedPackages, int32 InNumPackagesPerBatch)
 {
+	FPackageReferencersHelper PackageReferencersHelper(InPackagesToReload);
+
 	// Interdependencies between packages (in particular Blueprints) make it unsafe to run this logic in batches. 
 	// There are a number of edge cases that would have to be addressed if the batching logic were to be re-enabled, but 
 	// most likely the blueprint reparenting step would have to take in a TMap<UObject*, UObject*> that it could update
@@ -592,29 +724,29 @@ void ReloadPackages(const TArrayView<FReloadPackageData>& InPackagesToReload, TA
 			//The FThreadSafeObjectIterator will lock the global UObject array, to avoid potential deadlock we simply build the list
 			//of the potential referencers and do the reference fix serialize outside of the FThreadSafeObjectIterator.
 			TArray<UObject*> PotentialReferencers;
-			for (FThreadSafeObjectIterator ObjIter(UObject::StaticClass(), false, RF_NoFlags, EInternalObjectFlags::Garbage); ObjIter; ++ObjIter)
-			{
-				UObject* PotentialReferencer = *ObjIter;
 
-				// Mutating the old versions of classes can result in us replacing the SuperStruct pointer, which results
-				// in class layout change and subsequently crashes because instances will not match this new class layout:
-				UClass* AsClass = Cast<UClass>(PotentialReferencer);
-				if (!AsClass)
+			PackageReferencersHelper.ForEachObject(
+				[&PotentialReferencers] (UObject* PotentialReferencer)
 				{
-					AsClass = PotentialReferencer->GetTypedOuter<UClass>();
-				}
+					// Mutating the old versions of classes can result in us replacing the SuperStruct pointer, which results
+					// in class layout change and subsequently crashes because instances will not match this new class layout:
+					UClass* AsClass = Cast<UClass>(PotentialReferencer);
+					if (!AsClass)
+					{
+						AsClass = PotentialReferencer->GetTypedOuter<UClass>();
+					}
 
-				if (AsClass)
-				{
-					if (AsClass->HasAnyClassFlags(CLASS_NewerVersionExists) ||
-						AsClass->HasAnyFlags(RF_NewerVersionExists))
+					if (AsClass)
 					{
-						continue;
+						if (AsClass->HasAnyClassFlags(CLASS_NewerVersionExists) ||
+							AsClass->HasAnyFlags(RF_NewerVersionExists))
+						{
+							return;
+						}
 					}
+					PotentialReferencers.Add(PotentialReferencer);
 				}
-				PotentialReferencers.Add(PotentialReferencer);
-			}
-
+			);
 
 			for (UObject* PotentialReferencer : PotentialReferencers)
 			{
