Added an additional FAssetRegistryTag constructor that leverages move semantics to efficiently transfer the value string into the member variable, potentially eliminating an extra copy operation

diff --git a/Engine/Source/Runtime/CoreUObject/Public/UObject/Object.h b/Engine/Source/Runtime/CoreUObject/Public/UObject/Object.h
--- a/Engine/Source/Runtime/CoreUObject/Public/UObject/Object.h
+++ b/Engine/Source/Runtime/CoreUObject/Public/UObject/Object.h
@@ -847,6 +847,8 @@ public:
 
 		FAssetRegistryTag(FName InName, const FString& InValue, ETagType InType, uint32 InDisplayFlags = TD_None)
 			: Name(InName), Value(InValue), Type(InType), DisplayFlags(InDisplayFlags) {}
+		FAssetRegistryTag(FName InName, FString&& InValue, ETagType InType, uint32 InDisplayFlags = TD_None)
+			: Name(InName), Value(MoveTemp(InValue)), Type(InType), DisplayFlags(InDisplayFlags) {}
 
 #if WITH_EDITOR
 		/** Callback  */
