Publish soft path literals to reference collectors that run over bytecode

diff --git a/Engine/Source/Runtime/CoreUObject/Public/UObject/ScriptSerialization.h b/Engine/Source/Runtime/CoreUObject/Public/UObject/ScriptSerialization.h
--- a/Engine/Source/Runtime/CoreUObject/Public/UObject/ScriptSerialization.h
+++ b/Engine/Source/Runtime/CoreUObject/Public/UObject/ScriptSerialization.h
@@ -408,7 +408,37 @@
 		}
 		case EX_SoftObjectConst:
 		{
-			SerializeExpr(iCode, Ar);
+			// if collecting references inform the archive of the reference:
+			if (Ar.IsSaving() && Ar.IsObjectReferenceCollector())
+			{
+				XFER(uint8);
+				Expr = (EExprToken)Script[iCode - 1];
+				check(Expr == EX_StringConst || Expr == EX_UnicodeStringConst);
+				FString LongPath;
+				if (Expr == EX_StringConst)
+				{
+					LongPath = (ANSICHAR*)&Script[iCode];
+					XFERSTRING();
+				}
+				else
+				{
+					LongPath = FString((UCS2CHAR*)&Script[iCode]);
+
+					// Inline combine any surrogate pairs in the data when loading into a UTF-32 string
+					StringConv::InlineCombineSurrogates(LongPath);
+					XFERUNICODESTRING();
+				}
+				FSoftObjectPath Path(LongPath);
+				Ar << Path;
+				// we can't patch the path, but we could log an attempt to do so
+				// or change the implementation to support patching (allocating 
+				// these strings in a special region or distinct object)
+			}
+			else
+			{
+				// else just write the string literal instructions:
+				SerializeExpr(iCode, Ar);
+			}
 			break;
 		}
 		case EX_FieldPathConst:
