Share LinkerInstancingContext between requests instead of copying
   - Saves a lot of memory when tons of actors remapping are used for world partition. Multiple GB in PIE is not unheard of.
   - Add LLM counter specifically for LinkerInstancingContext memory usage
   - Make sure everything is thread-safe and ready for async loading thread usage

diff --git a/Engine/Source/Runtime/CoreUObject/Private/UObject/LinkerInstancingContext.cpp b/Engine/Source/Runtime/CoreUObject/Private/UObject/LinkerInstancingContext.cpp
--- a/Engine/Source/Runtime/CoreUObject/Private/UObject/LinkerInstancingContext.cpp
+++ b/Engine/Source/Runtime/CoreUObject/Private/UObject/LinkerInstancingContext.cpp
@@ -9,6 +9,260 @@
 #include "Templates/Tuple.h"
 #include "Misc/AutomationTest.h"
 
+LLM_DEFINE_TAG(Loading_LinkerInstancingContext);
+
+class FLinkerInstancingContext::FSharedLinkerInstancingContextData
+{
+	/** The shared content needs to be protected as it could be modified by the caller as well as the loading thread */
+	mutable FRWLock Lock;
+	/** Map of original package name to their instance counterpart. */
+	FLinkerInstancedPackageMap InstancedPackageMap;
+	/** Optional function to map original package name to their instance counterpart. The result of this function should be immutable, as it will be cached. */
+	TFunction<FName(FName)> InstancedPackageMapFunc;
+	/** Map of original top level asset path to their instance counterpart. */
+	TMap<FTopLevelAssetPath, FTopLevelAssetPath> PathMapping;
+	/** Tags can be used to determine some loading behavior. */
+	TSet<FName> Tags;
+	/** Remap soft object paths */
+	std::atomic<bool> bSoftObjectPathRemappingEnabled { true };
+
+public:
+	FSharedLinkerInstancingContextData() = default;
+	explicit FSharedLinkerInstancingContextData(TSet<FName> InTags)
+		: Tags(InTags)
+	{
+	}
+
+	explicit FSharedLinkerInstancingContextData(bool bInSoftObjectPathRemappingEnabled)
+		: bSoftObjectPathRemappingEnabled(bInSoftObjectPathRemappingEnabled)
+	{
+	}
+		
+	void AddPackageMapping(FName Original, FName Instanced)
+	{
+		FWriteScopeLock ScopeLock(Lock);
+		LLM_SCOPE_BYTAG(Loading_LinkerInstancingContext);
+		InstancedPackageMap.AddPackageMapping(Original, Instanced);
+	}
+
+	FName FindPackageMapping(FName Original) const
+	{
+		FReadScopeLock ScopeLock(Lock);
+		LLM_SCOPE_BYTAG(Loading_LinkerInstancingContext);
+		return InstancedPackageMap.InstancedPackageMapping.FindRef(Original);
+	}
+
+	void FixupSoftObjectPath(FSoftObjectPath& InOutSoftObjectPath) const
+	{
+		FReadScopeLock ScopeLock(Lock);
+		LLM_SCOPE_BYTAG(Loading_LinkerInstancingContext);
+		InstancedPackageMap.FixupSoftObjectPath(InOutSoftObjectPath);
+	}
+
+	FSoftObjectPath RemapPath(const FSoftObjectPath& Path) const
+	{
+		FReadScopeLock ScopeLock(Lock);
+		if (const FTopLevelAssetPath* Remapped = PathMapping.Find(Path.GetAssetPath()))
+		{
+			return FSoftObjectPath(*Remapped, Path.GetSubPathString());
+		}
+		return Path;
+	}
+
+	bool IsInstanced() const
+	{
+		FReadScopeLock ScopeLock(Lock);
+		return InstancedPackageMap.IsInstanced() || InstancedPackageMapFunc || PathMapping.Num() > 0;
+	}
+
+	void EnableAutomationTest()
+	{
+		FWriteScopeLock ScopeLock(Lock);
+		InstancedPackageMap.EnableAutomationTest();
+	}
+
+	void BuildPackageMapping(FName Original, FName Instanced, bool bInSoftObjectPathRemappingEnabled)
+	{
+		FWriteScopeLock ScopeLock(Lock);
+		LLM_SCOPE_BYTAG(Loading_LinkerInstancingContext);
+		InstancedPackageMap.BuildPackageMapping(Original, Instanced, bInSoftObjectPathRemappingEnabled);
+	}
+
+	void AddTag(FName NewTag)
+	{
+		FWriteScopeLock ScopeLock(Lock);
+		LLM_SCOPE_BYTAG(Loading_LinkerInstancingContext);
+		Tags.Add(NewTag);
+	}
+
+	void AppendTags(const TSet<FName>& NewTags)
+	{
+		FWriteScopeLock ScopeLock(Lock);
+		LLM_SCOPE_BYTAG(Loading_LinkerInstancingContext);
+		Tags.Append(NewTags);
+	}
+
+	bool HasTag(FName Tag) const
+	{
+		FReadScopeLock ScopeLock(Lock);
+		return Tags.Contains(Tag);
+	}
+
+	void SetSoftObjectPathRemappingEnabled(bool bInSoftObjectPathRemappingEnabled)
+	{
+		bSoftObjectPathRemappingEnabled = bInSoftObjectPathRemappingEnabled;
+	}
+
+	bool GetSoftObjectPathRemappingEnabled() const
+	{
+		return bSoftObjectPathRemappingEnabled;
+	}
+
+	void AddPathMapping(FSoftObjectPath Original, FSoftObjectPath Instanced)
+	{
+		FWriteScopeLock ScopeLock(Lock);
+		LLM_SCOPE_BYTAG(Loading_LinkerInstancingContext);
+		PathMapping.Emplace(Original.GetAssetPath(), Instanced.GetAssetPath());
+	}
+
+	void SetPackageMappingFunc(TFunction<FName(FName)> InInstancedPackageMapFunc)
+	{
+		FWriteScopeLock ScopeLock(Lock);
+		InstancedPackageMapFunc = MoveTemp(InInstancedPackageMapFunc);
+	}
+
+	FName RemapPackage(const FName& PackageName)
+	{
+		bool bAddPackageMapping = false;
+
+		FName RemappedPackageName;
+		{
+			FReadScopeLock ScopeLock(Lock);
+			RemappedPackageName = InstancedPackageMap.RemapPackage(PackageName);
+
+			if ((RemappedPackageName == PackageName) && InstancedPackageMapFunc)
+			{
+				RemappedPackageName = InstancedPackageMapFunc(PackageName);
+
+				if (RemappedPackageName != PackageName)
+				{
+					bAddPackageMapping = true; 
+				}
+			}
+		}
+
+		if (bAddPackageMapping)
+		{
+			FWriteScopeLock ScopeLock(Lock);
+			LLM_SCOPE_BYTAG(Loading_LinkerInstancingContext);
+			InstancedPackageMap.AddPackageMapping(PackageName, RemappedPackageName);
+		}
+
+		return RemappedPackageName;
+	}
+};
+
+FLinkerInstancingContext::FLinkerInstancingContext()
+{
+	LLM_SCOPE_BYTAG(Loading_LinkerInstancingContext);
+	SharedData = MakeShared<FSharedLinkerInstancingContextData>();
+}
+
+FLinkerInstancingContext::FLinkerInstancingContext(TSet<FName> InTags)
+{
+	LLM_SCOPE_BYTAG(Loading_LinkerInstancingContext);
+	SharedData = MakeShared<FSharedLinkerInstancingContextData>(MoveTemp(InTags));
+}
+
+FLinkerInstancingContext::FLinkerInstancingContext(bool bInSoftObjectPathRemappingEnabled)
+{
+	LLM_SCOPE_BYTAG(Loading_LinkerInstancingContext);
+	SharedData = MakeShared<FSharedLinkerInstancingContextData>(bInSoftObjectPathRemappingEnabled);
+}
+
+void FLinkerInstancingContext::EnableAutomationTest() 
+{ 
+	SharedData->EnableAutomationTest();
+}
+
+void FLinkerInstancingContext::BuildPackageMapping(FName Original, FName Instanced)
+{
+	SharedData->BuildPackageMapping(Original, Instanced, GetSoftObjectPathRemappingEnabled());
+}
+
+FName FLinkerInstancingContext::FindPackageMapping(FName Original) const
+{
+	return SharedData->FindPackageMapping(Original);
+}
+
+bool FLinkerInstancingContext::IsInstanced() const
+{
+	return SharedData->IsInstanced();
+}
+
+/** Remap the package name from the import table to its instanced counterpart, otherwise return the name unmodified. */
+FName FLinkerInstancingContext::RemapPackage(const FName& PackageName) const
+{
+	return SharedData->RemapPackage(PackageName);
+}
+
+/**
+ * Remap the top level asset part of the path name to its instanced counterpart, otherwise return the name unmodified.
+ * i.e. remaps /Path/To/Package.AssetName:Inner to /NewPath/To/NewPackage.NewAssetName:Inner
+ */
+FSoftObjectPath FLinkerInstancingContext::RemapPath(const FSoftObjectPath& Path) const
+{
+	return SharedData->RemapPath(Path);
+}
+
+/** Add a mapping from a package name to a new package name. There should be no separators (. or :) in these strings. */
+void FLinkerInstancingContext::AddPackageMapping(FName Original, FName Instanced)
+{
+	SharedData->AddPackageMapping(Original, Instanced);
+}
+
+/** Add a mapping function from a package name to a new package name. This function should be thread-safe, as it can be invoked from ALT. */
+void FLinkerInstancingContext::AddPackageMappingFunc(TFunction<FName(FName)> InInstancedPackageMapFunc)
+{
+	SharedData->SetPackageMappingFunc(InInstancedPackageMapFunc);
+}
+
+/** Add a mapping from a top level asset path (/Path/To/Package.AssetName) to another. */
+void FLinkerInstancingContext::AddPathMapping(FSoftObjectPath Original, FSoftObjectPath Instanced)
+{
+	ensureAlwaysMsgf(Original.GetSubPathString().IsEmpty(),
+		TEXT("Linker instance remap paths should be top-level assets only: %s->"), *Original.ToString());
+	ensureAlwaysMsgf(Instanced.GetSubPathString().IsEmpty(),
+		TEXT("Linker instance remap paths should be top-level assets only: ->%s"), *Instanced.ToString());
+
+	SharedData->AddPathMapping(Original, Instanced);
+}
+
+void FLinkerInstancingContext::AddTag(FName NewTag)
+{
+	SharedData->AddTag(NewTag);
+}
+
+void FLinkerInstancingContext::AppendTags(const TSet<FName>& NewTags)
+{
+	SharedData->AppendTags(NewTags);
+}
+
+bool FLinkerInstancingContext::HasTag(FName Tag) const
+{
+	return SharedData->HasTag(Tag);
+}
+
+void FLinkerInstancingContext::SetSoftObjectPathRemappingEnabled(bool bInSoftObjectPathRemappingEnabled)
+{
+	SharedData->SetSoftObjectPathRemappingEnabled(bInSoftObjectPathRemappingEnabled);
+}
+
+bool FLinkerInstancingContext::GetSoftObjectPathRemappingEnabled() const
+{
+	return SharedData->GetSoftObjectPathRemappingEnabled();
+}
+
 void FLinkerInstancingContext::FixupSoftObjectPath(FSoftObjectPath& InOutSoftObjectPath) const
 {
 	if (IsInstanced() && GetSoftObjectPathRemappingEnabled())
@@ -20,7 +274,7 @@ void FLinkerInstancingContext::FixupSoftObjectPath(FSoftObjectPath& InOutSoftObj
 		}
 		else
 		{
-			InstancedPackageMap.FixupSoftObjectPath(InOutSoftObjectPath);
+			SharedData->FixupSoftObjectPath(InOutSoftObjectPath);
 		}
 	}
 }
diff --git a/Engine/Source/Runtime/CoreUObject/Private/UObject/LinkerLoad.cpp b/Engine/Source/Runtime/CoreUObject/Private/UObject/LinkerLoad.cpp
--- a/Engine/Source/Runtime/CoreUObject/Private/UObject/LinkerLoad.cpp
+++ b/Engine/Source/Runtime/CoreUObject/Private/UObject/LinkerLoad.cpp
@@ -2095,11 +2095,12 @@ FLinkerLoad::ELinkerStatus FLinkerLoad::PopulateInstancingContext()
 			// add remapping for all the packages that should be instantiated along with this one
 			for (const FName& InstancingName : InstancingPackageName)
 			{
-				FName& InstancedName = InstancingContext.FindOrAddPackageMapping(InstancingName);
+				FName InstancedName = InstancingContext.FindPackageMapping(InstancingName);
 				// if there's isn't already a remapping for that package, create one
 				if (InstancedName.IsNone())
 				{
 					InstancedName = *FLinkerInstancingContext::GetInstancedPackageName(LinkerPackageName, InstancingName.ToString());
+					InstancingContext.AddPackageMapping(InstancingName, InstancedName);
 				}
 			}
 		}
diff --git a/Engine/Source/Runtime/CoreUObject/Public/UObject/LinkerInstancingContext.h b/Engine/Source/Runtime/CoreUObject/Public/UObject/LinkerInstancingContext.h
--- a/Engine/Source/Runtime/CoreUObject/Public/UObject/LinkerInstancingContext.h
+++ b/Engine/Source/Runtime/CoreUObject/Public/UObject/LinkerInstancingContext.h
@@ -86,99 +86,35 @@ private:
 class FLinkerInstancingContext
 {
 public:
-	FLinkerInstancingContext() = default;
-	explicit FLinkerInstancingContext(TSet<FName> InTags)
-		: Tags(MoveTemp(InTags))
-	{
-	}
-	explicit FLinkerInstancingContext(bool bInSoftObjectPathRemappingEnabled)
-		: bSoftObjectPathRemappingEnabled(bInSoftObjectPathRemappingEnabled)
-	{
-	}
+	COREUOBJECT_API FLinkerInstancingContext();
+	COREUOBJECT_API explicit FLinkerInstancingContext(TSet<FName> InTags);
+	COREUOBJECT_API explicit FLinkerInstancingContext(bool bInSoftObjectPathRemappingEnabled);
 
-	bool IsInstanced() const
-	{
-		return InstancedPackageMap.IsInstanced() || InstancedPackageMapFunc || PathMapping.Num() > 0;
-	}
+	COREUOBJECT_API bool IsInstanced() const;
 
 	/** Remap the package name from the import table to its instanced counterpart, otherwise return the name unmodified. */
-	FName RemapPackage(const FName& PackageName) const
-	{
-		FName RemappedPackageName = InstancedPackageMap.RemapPackage(PackageName);
-
-		if ((RemappedPackageName == PackageName) && InstancedPackageMapFunc)
-		{
-			RemappedPackageName = InstancedPackageMapFunc(PackageName);
-
-			if (RemappedPackageName != PackageName)
-			{
-				InstancedPackageMap.AddPackageMapping(PackageName, RemappedPackageName);
-			}
-		}
-
-		return RemappedPackageName;
-	}
+	COREUOBJECT_API FName RemapPackage(const FName& PackageName) const;
 
 	/**
 	 * Remap the top level asset part of the path name to its instanced counterpart, otherwise return the name unmodified. 
 	 * i.e. remaps /Path/To/Package.AssetName:Inner to /NewPath/To/NewPackage.NewAssetName:Inner 
 	 */
-	FSoftObjectPath RemapPath(const FSoftObjectPath& Path) const
-	{
-		if (const FTopLevelAssetPath* Remapped = PathMapping.Find(Path.GetAssetPath()))
-		{
-			return FSoftObjectPath(*Remapped, Path.GetSubPathString());
-		}
-		return Path;
-	}
+	COREUOBJECT_API FSoftObjectPath RemapPath(const FSoftObjectPath& Path) const;
 
 	/** Add a mapping from a package name to a new package name. There should be no separators (. or :) in these strings. */
-	void AddPackageMapping(FName Original, FName Instanced)
-	{
-		InstancedPackageMap.AddPackageMapping(Original, Instanced);
-	}
-
-	/** Add a mapping function from a package name to a new package name. This function should be thread-safe, is it can be invoked from ALT. */
-	void AddPackageMappingFunc(TFunction<FName(FName)> InInstancedPackageMapFunc)
-	{
-		InstancedPackageMapFunc = MoveTemp(InInstancedPackageMapFunc);
-	}
+	COREUOBJECT_API void AddPackageMapping(FName Original, FName Instanced);
 
-	/** Add a mapping from a top level asset path (/Path/To/Package.AssetName) to another. */
-	void AddPathMapping(FSoftObjectPath Original, FSoftObjectPath Instanced)
-	{
-		ensureAlwaysMsgf(Original.GetSubPathString().IsEmpty(), 
-			TEXT("Linker instance remap paths should be top-level assets only: %s->"), *Original.ToString());
-		ensureAlwaysMsgf(Instanced.GetSubPathString().IsEmpty(), 
-			TEXT("Linker instance remap paths should be top-level assets only: ->%s"), *Instanced.ToString());
+	/** Add a mapping function from a package name to a new package name. This function should be thread-safe, as it can be invoked from ALT. */
+	COREUOBJECT_API void AddPackageMappingFunc(TFunction<FName(FName)> InInstancedPackageMapFunc);
 	
-		PathMapping.Emplace(Original.GetAssetPath(), Instanced.GetAssetPath());
-	}
-
-	void AddTag(FName NewTag)
-	{
-		Tags.Add(NewTag);
-	}
-
-	void AppendTags(const TSet<FName>& NewTags)
-	{
-		Tags.Append(NewTags);
-	}
-
-	bool HasTag(FName Tag) const
-	{
-		return Tags.Contains(Tag);
-	}
-
-	void SetSoftObjectPathRemappingEnabled(bool bInSoftObjectPathRemappingEnabled)
-	{
-		bSoftObjectPathRemappingEnabled = bInSoftObjectPathRemappingEnabled;
-	}
+	/** Add a mapping from a top level asset path (/Path/To/Package.AssetName) to another. */
+	COREUOBJECT_API void AddPathMapping(FSoftObjectPath Original, FSoftObjectPath Instanced);
 
-	bool GetSoftObjectPathRemappingEnabled() const 
-	{ 
-		return bSoftObjectPathRemappingEnabled; 
-	}
+	COREUOBJECT_API void AddTag(FName NewTag);
+	COREUOBJECT_API void AppendTags(const TSet<FName>& NewTags);
+	COREUOBJECT_API bool HasTag(FName Tag) const;
+	COREUOBJECT_API void SetSoftObjectPathRemappingEnabled(bool bInSoftObjectPathRemappingEnabled);
+	COREUOBJECT_API bool GetSoftObjectPathRemappingEnabled() const;
 
 	UE_DEPRECATED(5.2, "No longer used, pass ELoadFlags::LOAD_RegenerateBulkDataGuids to LoadPackage instead")
 	void SetRegenerateUniqueBulkDataGuids(bool bFlag) { }
@@ -192,33 +128,17 @@ public:
 		return FString::Printf(TEXT("%s_InstanceOf_%s"), *InOuterPackageName, *InPackageName);
 	}
 
-	void FixupSoftObjectPath(FSoftObjectPath& InOutSoftObjectPath) const;
+	COREUOBJECT_API void FixupSoftObjectPath(FSoftObjectPath& InOutSoftObjectPath) const;
 
 private:
-	void EnableAutomationTest() { InstancedPackageMap.EnableAutomationTest(); }
-
-	void BuildPackageMapping(FName Original, FName Instanced)
-	{
-		InstancedPackageMap.BuildPackageMapping(Original, Instanced, GetSoftObjectPathRemappingEnabled());
-	}
-
-	FName& FindOrAddPackageMapping(FName Original)
-	{
-		return InstancedPackageMap.InstancedPackageMapping.FindOrAdd(Original);
-	}
+	void EnableAutomationTest();
+	void BuildPackageMapping(FName Original, FName Instanced);
+	FName FindPackageMapping(FName Original) const;
 
 	friend class FLinkerLoad;
 	friend struct FAsyncPackage2;
 	friend class FLinkerInstancingContextTests;
+	class FSharedLinkerInstancingContextData;
 
-	/** Map of original package name to their instance counterpart. */
-	mutable FLinkerInstancedPackageMap InstancedPackageMap;
-	/** Optional function to map original package name to their instance counterpart. The result of this function should be immutable, as it will be cached. */
-	TFunction<FName(FName)> InstancedPackageMapFunc;
-	/** Map of original top level asset path to their instance counterpart. */
-	TMap<FTopLevelAssetPath, FTopLevelAssetPath> PathMapping;
-	/** Tags can be used to determine some loading behavior. */
-	TSet<FName> Tags;
-	/** Remap soft object paths */
-	bool bSoftObjectPathRemappingEnabled = true;
-};
+	TSharedPtr<FSharedLinkerInstancingContextData> SharedData;
+};
\ No newline at end of file
