Flush the log before forcing a crash in UObject::AbortInsideMemberFunction().

diff --git a/Engine/Source/Runtime/CoreUObject/Private/UObject/Obj.cpp b/Engine/Source/Runtime/CoreUObject/Private/UObject/Obj.cpp
--- a/Engine/Source/Runtime/CoreUObject/Private/UObject/Obj.cpp
+++ b/Engine/Source/Runtime/CoreUObject/Private/UObject/Obj.cpp
@@ -5172,6 +5172,12 @@ void UObject::AbortInsideMemberFunction() const
 {
 	//put a trace of this in the log to help diagnostics at a glance.
 	UE_LOG(LogObj, Warning, TEXT("UObject::AbortInsideMemberFunction called on object %s."), *GetFullName());
+
+	if (GLog)
+	{
+		GLog->Flush();
+	}
+
 	//a bit more ideally, we could set GIsCriticalError = true and call FPlatformMisc::RequestExit. however, not all platforms would generate a dump as a result of this.
 	//as such, we commit an access violation right here. we explicitly want to avoid the standard platform error/AssertFailed paths as they are likely to pollute the
 	//callstack. this in turn is more likely to prevent useful (e.g. this object) memory from making its way into a minidump.
