Truncation warning fix in PropertyOptional

diff --git a/Engine/Source/Runtime/CoreUObject/Private/UObject/PropertyOptional.cpp b/Engine/Source/Runtime/CoreUObject/Private/UObject/PropertyOptional.cpp
--- a/Engine/Source/Runtime/CoreUObject/Private/UObject/PropertyOptional.cpp
+++ b/Engine/Source/Runtime/CoreUObject/Private/UObject/PropertyOptional.cpp
@@ -258,7 +258,7 @@ void FOptionalProperty::SerializeItem(FStructuredArchive::FSlot Slot, void* Data
 			
 			// If saving to a non-text archive, save the size of the serialized value so it can be skipped over on load if it's the wrong type.
 			const int64 ValueEndOffset = UnderlyingArchive.Tell();
-			ValueTag.Size = ValueEndOffset - ValueStartOffset;
+			ValueTag.Size = IntCastChecked<int32>(ValueEndOffset - ValueStartOffset);
 			if (ValueTag.Size > 0 && !UnderlyingArchive.IsTextFormat())
 			{
 				// The tag serialization set ValueTag.SizeOffset to the archive offset where a placeholder size field was serialized.
