Truncation warning fixes for CoreUObject

diff --git a/Engine/Source/Runtime/CoreUObject/CoreUObject.Build.cs b/Engine/Source/Runtime/CoreUObject/CoreUObject.Build.cs
--- a/Engine/Source/Runtime/CoreUObject/CoreUObject.Build.cs
+++ b/Engine/Source/Runtime/CoreUObject/CoreUObject.Build.cs
@@ -23,8 +23,7 @@ public class CoreUObject : ModuleRules
 		PrivateDependencyModuleNames.Add("Projects");
         PrivateDependencyModuleNames.Add("Json");
 
-		//@TODO: UE-127233
-		// UnsafeTypeCastWarningLevel = WarningLevel.Error;
+		UnsafeTypeCastWarningLevel = WarningLevel.Error;
 
 		if (Target.bBuildWithEditorOnlyData)
 		{
diff --git a/Engine/Source/Runtime/CoreUObject/Private/AssetRegistry/AssetData.cpp b/Engine/Source/Runtime/CoreUObject/Private/AssetRegistry/AssetData.cpp
--- a/Engine/Source/Runtime/CoreUObject/Private/AssetRegistry/AssetData.cpp
+++ b/Engine/Source/Runtime/CoreUObject/Private/AssetRegistry/AssetData.cpp
@@ -761,7 +761,7 @@ bool FAssetData::TryNetworkRead(FCbFieldView Field, bool bReadPackageName, FName
 	if (!TagsField.HasError()) // Ok if it does not exist
 	{
 		FAssetDataTagMap Tags;
-		Tags.Reserve(TagsArray.Num());
+		Tags.Reserve(IntCastChecked<int32>(TagsArray.Num()));
 		for (FCbFieldView TagField : TagsArray)
 		{
 			FName TagName;
diff --git a/Engine/Source/Runtime/CoreUObject/Private/AssetRegistry/AssetDataTagMap.cpp b/Engine/Source/Runtime/CoreUObject/Private/AssetRegistry/AssetDataTagMap.cpp
--- a/Engine/Source/Runtime/CoreUObject/Private/AssetRegistry/AssetDataTagMap.cpp
+++ b/Engine/Source/Runtime/CoreUObject/Private/AssetRegistry/AssetDataTagMap.cpp
@@ -204,7 +204,7 @@ namespace FixedTagPrivate
 
 	uint32 HashCaseSensitive(const TCHAR* Str, int32 Len)		
 	{
-		return CityHash64(reinterpret_cast<const char*>(Str), sizeof(TCHAR) * Len);
+		return CityHash32(reinterpret_cast<const char*>(Str), sizeof(TCHAR) * Len);
 	}
 
 	uint32 HashCombineQuick(uint32 A, uint32 B)
@@ -705,7 +705,7 @@ namespace FixedTagPrivate
 		FMapHandle Out;
 		Out.IsValid = 1;
 		Out.HasNumberlessKeys = bHasNumberlessKeys;
-		Out.StoreIndex = StoreIndex;
+		Out.StoreIndex = static_cast<uint16>(StoreIndex);
 		Out.Num = Num;
 		Out.PairBegin = PairBegin;
 
@@ -760,7 +760,7 @@ namespace FixedTagPrivate
 		for (uint32 Idx = 0; Idx < Num; ++Idx)
 		{
 			check(Src[Idx] <= std::numeric_limits<CharType>::max());
-			Dst[Idx] = Src[Idx];
+			Dst[Idx] = static_cast<CharType>(Src[Idx]);
 		}
 	}
 
diff --git a/Engine/Source/Runtime/CoreUObject/Private/Serialization/AsyncLoading.cpp b/Engine/Source/Runtime/CoreUObject/Private/Serialization/AsyncLoading.cpp
--- a/Engine/Source/Runtime/CoreUObject/Private/Serialization/AsyncLoading.cpp
+++ b/Engine/Source/Runtime/CoreUObject/Private/Serialization/AsyncLoading.cpp
@@ -7411,11 +7411,11 @@ void FAsyncPackage::UpdateLoadPercentage()
 	if (Linker)
 	{
 		const int32 PostLoadCount = FMath::Max(DeferredPostLoadObjects.Num(), Linker->ImportMap.Num());
-		NewLoadPercentage = 100.f * (LoadImportIndex + ExportIndex + DeferredPostLoadIndex) / (Linker->ExportMap.Num() + Linker->ImportMap.Num() + PostLoadCount);		
+		NewLoadPercentage = 100.f * static_cast<float>(LoadImportIndex + ExportIndex + DeferredPostLoadIndex) / static_cast<float>(Linker->ExportMap.Num() + Linker->ImportMap.Num() + PostLoadCount);
 	}
 	else if (DeferredPostLoadObjects.Num() > 0)
 	{
-		NewLoadPercentage = static_cast<float>(DeferredPostLoadIndex) / DeferredPostLoadObjects.Num();
+		NewLoadPercentage = static_cast<float>(DeferredPostLoadIndex) / static_cast<float>(DeferredPostLoadObjects.Num());
 	}
 	// It's also possible that we got so many objects to PostLoad that LoadPercantage will actually drop
 	LoadPercentage = FMath::Max(NewLoadPercentage, LoadPercentage);
diff --git a/Engine/Source/Runtime/CoreUObject/Private/Serialization/Formatters/JsonArchiveOutputFormatter.cpp b/Engine/Source/Runtime/CoreUObject/Private/Serialization/Formatters/JsonArchiveOutputFormatter.cpp
--- a/Engine/Source/Runtime/CoreUObject/Private/Serialization/Formatters/JsonArchiveOutputFormatter.cpp
+++ b/Engine/Source/Runtime/CoreUObject/Private/Serialization/Formatters/JsonArchiveOutputFormatter.cpp
@@ -380,7 +380,7 @@ void FJsonArchiveOutputFormatter::Serialize(void* Data, uint64 DataSize)
 	if(DataSize < MaxLineBytes)
 	{
 		// Encode the data on a single line. No need for hashing; intra-line merge conflicts are rare.
-		WriteValue(FString::Printf(TEXT("\"Base64:%s\""), *FBase64::Encode((const uint8*)Data, DataSize)));
+		WriteValue(FString::Printf(TEXT("\"Base64:%s\""), *FBase64::Encode((const uint8*)Data, static_cast<uint32>(DataSize))));
 	}
 	else
 	{
@@ -417,7 +417,7 @@ void FJsonArchiveOutputFormatter::Serialize(void* Data, uint64 DataSize)
 			Write("\t\"");
 
 			ANSICHAR LineData[MaxLineChars + 1];
-			uint64 NumLineChars = FBase64::Encode((const uint8*)Data + DataPos, FMath::Min<uint64>(DataSize - DataPos, MaxLineBytes), LineData);
+			uint64 NumLineChars = FBase64::Encode((const uint8*)Data + DataPos, FMath::Min<uint32>(IntCastChecked<uint32>(DataSize - DataPos), MaxLineBytes), LineData);
 			Inner.Serialize(LineData, NumLineChars);
 
 			Write("\"");
diff --git a/Engine/Source/Runtime/CoreUObject/Private/Serialization/LoadTimeTrace.cpp b/Engine/Source/Runtime/CoreUObject/Private/Serialization/LoadTimeTrace.cpp
--- a/Engine/Source/Runtime/CoreUObject/Private/Serialization/LoadTimeTrace.cpp
+++ b/Engine/Source/Runtime/CoreUObject/Private/Serialization/LoadTimeTrace.cpp
@@ -220,7 +220,7 @@ void FLoadTimeProfilerTracePrivate::OutputAsyncPackageImportDependency(const voi
 void FLoadTimeProfilerTracePrivate::OutputClassInfo(const UClass* Class, const FName& Name)
 {
 	TCHAR Buffer[FName::StringBufferSize];
-	uint16 NameLen = Name.ToString(Buffer);
+	uint16 NameLen = static_cast<uint16>(Name.ToString(Buffer));
 	UE_TRACE_LOG(LoadTime, ClassInfo, LoadTimeChannel, NameLen * sizeof(ANSICHAR))
 		<< ClassInfo.Class(Class)
 		<< ClassInfo.Name(Buffer, NameLen);
diff --git a/Engine/Source/Runtime/CoreUObject/Private/Serialization/UnversionedPropertySerialization.cpp b/Engine/Source/Runtime/CoreUObject/Private/Serialization/UnversionedPropertySerialization.cpp
--- a/Engine/Source/Runtime/CoreUObject/Private/Serialization/UnversionedPropertySerialization.cpp
+++ b/Engine/Source/Runtime/CoreUObject/Private/Serialization/UnversionedPropertySerialization.cpp
@@ -578,7 +578,7 @@ protected:
 
 		uint16 Pack() const
 		{
-			return SkipNum | (bHasAnyZeroes ? HasZeroMask : 0) | (uint16)ValueNum << ValueNumShift | (bIsLast ? IsLastMask : 0);
+			return SkipNum | (uint16)(bHasAnyZeroes ? HasZeroMask : 0) | (uint16)(ValueNum << ValueNumShift) | (uint16)(bIsLast ? IsLastMask : 0);
 		}
 
 		static FFragment Unpack(uint16 Int)
diff --git a/Engine/Source/Runtime/CoreUObject/Private/Serialization/UnversionedPropertySerializationTest.cpp b/Engine/Source/Runtime/CoreUObject/Private/Serialization/UnversionedPropertySerializationTest.cpp
--- a/Engine/Source/Runtime/CoreUObject/Private/Serialization/UnversionedPropertySerializationTest.cpp
+++ b/Engine/Source/Runtime/CoreUObject/Private/Serialization/UnversionedPropertySerializationTest.cpp
@@ -23,7 +23,7 @@ namespace PropertySerializationStats
 					TEXT("SavedStructs"), Structs.Load(),
 					TEXT("SavedMB"), static_cast<uint32>(UnversionedBytes.Load() >> 20),
 					TEXT("EquivalentTaggedMB"), static_cast<uint32>(VersionedBytes.Load() >> 20),
-					TEXT("CompressionRatio"), static_cast<float>(VersionedBytes.Load()) / UnversionedBytes.Load(),
+					TEXT("CompressionRatio"), static_cast<float>(VersionedBytes.Load()) / static_cast<float>(UnversionedBytes.Load()),
 					TEXT("BitfieldWasteKB"), static_cast<uint32>(UselessBytes.Load()) >> 10));
 	});
 #endif
diff --git a/Engine/Source/Runtime/CoreUObject/Private/UObject/Class.cpp b/Engine/Source/Runtime/CoreUObject/Private/UObject/Class.cpp
--- a/Engine/Source/Runtime/CoreUObject/Private/UObject/Class.cpp
+++ b/Engine/Source/Runtime/CoreUObject/Private/UObject/Class.cpp
@@ -629,8 +629,8 @@ UStruct::UStruct(UStruct* InSuperStruct, SIZE_T ParamsSize, SIZE_T Alignment)
 	, SuperStruct(InSuperStruct)
 	, Children(nullptr)
 	, ChildProperties(nullptr)
-	, PropertiesSize(ParamsSize ? ParamsSize : (InSuperStruct ? InSuperStruct->GetPropertiesSize() : 0))
-	, MinAlignment(Alignment ? Alignment : (FMath::Max(InSuperStruct ? InSuperStruct->GetMinAlignment() : 1, 1)))
+	, PropertiesSize(ParamsSize ? IntCastChecked<int32>(ParamsSize) : (InSuperStruct ? InSuperStruct->GetPropertiesSize() : 0))
+	, MinAlignment(Alignment ? IntCastChecked<int32>(Alignment) : (FMath::Max(InSuperStruct ? InSuperStruct->GetMinAlignment() : 1, 1)))
 	, PropertyLink(nullptr)
 	, RefLink(nullptr)
 	, DestructorLink(nullptr)
@@ -651,8 +651,8 @@ UStruct::UStruct(const FObjectInitializer& ObjectInitializer, UStruct* InSuperSt
 	, SuperStruct(InSuperStruct)
 	, Children(nullptr)
 	, ChildProperties(nullptr)
-	, PropertiesSize(ParamsSize ? ParamsSize : (InSuperStruct ? InSuperStruct->GetPropertiesSize() : 0))
-	, MinAlignment(Alignment ? Alignment : (FMath::Max(InSuperStruct ? InSuperStruct->GetMinAlignment() : 1, 1)))
+	, PropertiesSize(ParamsSize ? IntCastChecked<int32>(ParamsSize) : (InSuperStruct ? InSuperStruct->GetPropertiesSize() : 0))
+	, MinAlignment(Alignment ? IntCastChecked<int32>(Alignment) : (FMath::Max(InSuperStruct ? InSuperStruct->GetMinAlignment() : 1, 1)))
 	, PropertyLink(nullptr)
 	, RefLink(nullptr)
 	, DestructorLink(nullptr)
@@ -1611,7 +1611,7 @@ void UStruct::SerializeVersionedTaggedProperties(FStructuredArchive::FSlot Slot,
 						}
 
 						// set the tag's size
-						Tag.Size = UnderlyingArchive.Tell() - DataOffset;
+						Tag.Size = IntCastChecked<int32>(UnderlyingArchive.Tell() - DataOffset);
 
 						if (Tag.Size > 0 && !UnderlyingArchive.IsTextFormat())
 						{
@@ -1979,7 +1979,7 @@ void UStruct::Serialize(FArchive& Ar)
 
 				// go back and write on-disk size
 				Ar.Seek(ScriptStorageSizeOffset);
-				int32 ScriptStorageSize = BytecodeEndOffset - BytecodeStartOffset;
+				int32 ScriptStorageSize = IntCastChecked<int32>(BytecodeEndOffset - BytecodeStartOffset);
 				Ar << ScriptStorageSize;
 
 				// back to where we were
@@ -6630,10 +6630,10 @@ void UFunction::InitializeDerivedMembers()
 		if (Property->PropertyFlags & CPF_Parm)
 		{
 			NumParms++;
-			ParmsSize = Property->GetOffset_ForUFunction() + Property->GetSize();
+			ParmsSize = IntCastChecked<uint16>(Property->GetOffset_ForUFunction() + Property->GetSize());
 			if (Property->PropertyFlags & CPF_ReturnParm)
 			{
-				ReturnValueOffset = Property->GetOffset_ForUFunction();
+				ReturnValueOffset = IntCastChecked<uint16>(Property->GetOffset_ForUFunction());
 			}
 		}
 		else if ((FunctionFlags & FUNC_HasDefaults) == 0)
diff --git a/Engine/Source/Runtime/CoreUObject/Private/UObject/GarbageCollectionSchema.cpp b/Engine/Source/Runtime/CoreUObject/Private/UObject/GarbageCollectionSchema.cpp
--- a/Engine/Source/Runtime/CoreUObject/Private/UObject/GarbageCollectionSchema.cpp
+++ b/Engine/Source/Runtime/CoreUObject/Private/UObject/GarbageCollectionSchema.cpp
@@ -801,8 +801,8 @@ static void DumpDistribution(FOutputDevice& Out, const TCHAR* Title, const FLog2
 		uint32 Num = Bins[Idx];
 		Sum += Num;
 		uint32 Max = Idx == 0 ? 0 : 1u << (Idx - 1);
-		float Percentage = (100.f * Num) / Total;
-		float Cumulative = (100.f * Sum) / Total;
+		float Percentage = (100.f * (float)Num) / (float)Total;
+		float Cumulative = (100.f * (float)Sum) / (float)Total;
 		if (Unit == EDumpUnit::Number)
 		{
 			Out.Logf(TEXT(" %4d %5.1f%% %5.1f%%  %d"), Max, Percentage, Cumulative, Num);
@@ -824,8 +824,8 @@ static void DumpSchemaStats(FOutputDevice& Out, const FSchemaStats& Stats)
 		if (uint32 Num = Stats.Counters[Idx])
 		{
 			Sum += Num;
-			float Percentage = (100.f * Num) / Stats.NumMembers;
-			float Cumulative = (100.f * Sum) / Stats.NumMembers;
+			float Percentage = (100.f * (float)Num) / (float)Stats.NumMembers;
+			float Cumulative = (100.f * (float)Sum) / (float)Stats.NumMembers;
 			Out.Logf(FormatStr, *ToName(EMemberType(Idx)).ToString(), Percentage, Cumulative, Num);
 		}
 	}
diff --git a/Engine/Source/Runtime/CoreUObject/Private/UObject/LinkerLoad.cpp b/Engine/Source/Runtime/CoreUObject/Private/UObject/LinkerLoad.cpp
--- a/Engine/Source/Runtime/CoreUObject/Private/UObject/LinkerLoad.cpp
+++ b/Engine/Source/Runtime/CoreUObject/Private/UObject/LinkerLoad.cpp
@@ -2798,7 +2798,7 @@ void FLinkerLoad::Verify()
 			if (ShouldCreateThrottledSlowTask())
 			{
 				static const FText LoadingImportsText = NSLOCTEXT("Core", "LinkerLoad_Imports", "Loading Imports");
-				SlowTask.Emplace(Summary.ImportCount, LoadingImportsText);
+				SlowTask.Emplace(static_cast<float>(Summary.ImportCount), LoadingImportsText);
 			}
 #endif
 			UE_TRACK_REFERENCING_PACKAGE_SCOPED(LinkerRoot->GetFName(), PackageAccessTrackingOps::NAME_Load);
@@ -3474,7 +3474,7 @@ bool FLinkerLoad::VerifyImportInner(const int32 ImportIndex, FString& WarningSuf
 	if (ShouldCreateThrottledSlowTask())
 	{
 		static const FTextFormat VerifyingTextFormat = NSLOCTEXT("Core", "VerifyPackage_Scope", "Verifying '{0}'");
-		SlowTask.Emplace(100, FText::Format(VerifyingTextFormat, FText::FromName(Import.ObjectName)));
+		SlowTask.Emplace(100.0f, FText::Format(VerifyingTextFormat, FText::FromName(Import.ObjectName)));
 	}
 #endif
 
@@ -3968,7 +3968,7 @@ void FLinkerLoad::LoadAllObjects(bool bForcePreload)
 	if (ShouldCreateThrottledSlowTask())
 	{
 		static const FText LoadingObjectText = NSLOCTEXT("Core", "LinkerLoad_LoadingObjects", "Loading Objects");
-		SlowTask.Emplace(ExportMap.Num(), LoadingObjectText);
+		SlowTask.Emplace(static_cast<float>(ExportMap.Num()), LoadingObjectText);
 		SlowTask->Visibility = ESlowTaskVisibility::Invisible;
 	}
 #endif
diff --git a/Engine/Source/Runtime/CoreUObject/Private/UObject/Property.cpp b/Engine/Source/Runtime/CoreUObject/Private/UObject/Property.cpp
--- a/Engine/Source/Runtime/CoreUObject/Private/UObject/Property.cpp
+++ b/Engine/Source/Runtime/CoreUObject/Private/UObject/Property.cpp
@@ -1553,7 +1553,7 @@ static const int32 ReadArrayIndex(UStruct* ObjectStruct, const TCHAR*& Str, FOut
 					if (IndexTokenName != NAME_None)
 					{
 						// Search for the enum in question.
-						Index = UEnum::LookupEnumName(FName(), IndexTokenName, EFindFirstObjectOptions::NativeFirst /* Only native enums can be used as array indices */);
+						Index = IntCastChecked<int32>(UEnum::LookupEnumName(FName(), IndexTokenName, EFindFirstObjectOptions::NativeFirst /* Only native enums can be used as array indices */));
 						if (Index == INDEX_NONE)
 						{
 							Index = 0;
diff --git a/Engine/Source/Runtime/CoreUObject/Private/UObject/PropertyArray.cpp b/Engine/Source/Runtime/CoreUObject/Private/UObject/PropertyArray.cpp
--- a/Engine/Source/Runtime/CoreUObject/Private/UObject/PropertyArray.cpp
+++ b/Engine/Source/Runtime/CoreUObject/Private/UObject/PropertyArray.cpp
@@ -377,7 +377,7 @@ void FArrayProperty::SerializeItem(FStructuredArchive::FSlot Slot, void* Value,
 		FPropertyTag& InnerTag = MaybeInnerTag.GetValue();
 
 		// set the tag's size
-		InnerTag.Size = UnderlyingArchive.Tell() - DataOffset;
+		InnerTag.Size = IntCastChecked<int32>(UnderlyingArchive.Tell() - DataOffset);
 
 		if (InnerTag.Size > 0)
 		{
diff --git a/Engine/Source/Runtime/CoreUObject/Private/UObject/PropertyByte.cpp b/Engine/Source/Runtime/CoreUObject/Private/UObject/PropertyByte.cpp
--- a/Engine/Source/Runtime/CoreUObject/Private/UObject/PropertyByte.cpp
+++ b/Engine/Source/Runtime/CoreUObject/Private/UObject/PropertyByte.cpp
@@ -74,11 +74,11 @@ void FByteProperty::SerializeItem( FStructuredArchive::FSlot Slot, void* Value,
 		int32 EnumIndex = Enum->GetIndexByName(EnumValueName, EGetByNameFlags::ErrorIfNotFound);
 		if (EnumIndex == INDEX_NONE)
 		{
-			*(uint8*)Value = Enum->GetMaxEnumValue();
+			*(uint8*)Value = IntCastChecked<uint8>(Enum->GetMaxEnumValue());
 		}
 		else
 		{
-			*(uint8*)Value = Enum->GetValueByIndex(EnumIndex);
+			*(uint8*)Value = IntCastChecked<uint8>(Enum->GetValueByIndex(EnumIndex));
 		}
 	}
 	// Saving
@@ -106,7 +106,7 @@ bool FByteProperty::NetSerializeItem( FArchive& Ar, UPackageMap* Map, void* Data
 
 	if (Ar.EngineNetVer() < FEngineNetworkCustomVersion::EnumSerializationCompat)
 	{
-		Ar.SerializeBits(Data, Enum ? FMath::CeilLogTwo(Enum->GetMaxEnumValue()) : 8);
+		Ar.SerializeBits(Data, Enum ? FMath::CeilLogTwo64(Enum->GetMaxEnumValue()) : 8);
 	}
 	else
 	{
@@ -206,7 +206,7 @@ struct TConvertIntToEnumProperty
 				*Enum->GetNameByValue(Enum->GetMaxEnumValue()).ToString()
 				);
 
-			NewValue = Enum->GetMaxEnumValue();
+			NewValue = IntCastChecked<uint8>(Enum->GetMaxEnumValue());
 		}
 
 		Property->SetPropertyValue_InContainer(Obj, NewValue, Tag.ArrayIndex);
@@ -425,7 +425,7 @@ const TCHAR* FByteProperty::ImportText_Internal( const TCHAR* InBuffer, void* Co
 			}
 			if (EnumIndex != INDEX_NONE)
 			{
-				int64 EnumValue = Enum->GetValueByIndex(EnumIndex);
+				uint8 EnumValue = IntCastChecked<uint8>(Enum->GetValueByIndex(EnumIndex));
 				if (PropertyPointerType == EPropertyPointerType::Container && HasSetter())
 				{
 					SetValue_InContainer(ContainerOrPropertyPtr, EnumValue);
@@ -466,7 +466,7 @@ const TCHAR* FByteProperty::ImportText_Internal( const TCHAR* InBuffer, void* Co
 				uint64 TrueValue = 1ull;
 				if (PropertyPointerType == EPropertyPointerType::Container && HasSetter())
 				{
-					SetValue_InContainer(ContainerOrPropertyPtr, TrueValue);
+					SetValue_InContainer(ContainerOrPropertyPtr, static_cast<uint8>(TrueValue));
 				}
 				else
 				{
@@ -479,7 +479,7 @@ const TCHAR* FByteProperty::ImportText_Internal( const TCHAR* InBuffer, void* Co
 				uint64 FalseValue = 0ull;
 				if (PropertyPointerType == EPropertyPointerType::Container && HasSetter())
 				{
-					SetValue_InContainer(ContainerOrPropertyPtr, FalseValue);
+					SetValue_InContainer(ContainerOrPropertyPtr, static_cast<uint8>(FalseValue));
 				}
 				else
 				{
diff --git a/Engine/Source/Runtime/CoreUObject/Private/UObject/ReferenceChainSearch.cpp b/Engine/Source/Runtime/CoreUObject/Private/UObject/ReferenceChainSearch.cpp
--- a/Engine/Source/Runtime/CoreUObject/Private/UObject/ReferenceChainSearch.cpp
+++ b/Engine/Source/Runtime/CoreUObject/Private/UObject/ReferenceChainSearch.cpp
@@ -640,7 +640,7 @@ namespace UE::ReferenceChainSearch
 			Search.MergeGraph(MergedEdgeLists);
 
 			OutGraph = UE::Graph::ConstructTransposeGraph(MergedEdgeLists);
-			UE_LOG(LogReferenceChain, Display, TEXT("InitialGather memory usage: %.2f"), (OutGraph.GetAllocatedSize() + Search.GetAllocatedSize()) / 1024.0 / 1024.0);
+			UE_LOG(LogReferenceChain, Display, TEXT("InitialGather memory usage: %.2f"), static_cast<double>(OutGraph.GetAllocatedSize() + Search.GetAllocatedSize()) / 1024.0 / 1024.0);
 		}
 		else
 		{
@@ -656,7 +656,7 @@ namespace UE::ReferenceChainSearch
 			Search.MergeGraph(MergedEdgeLists);
 
 			OutGraph = UE::Graph::ConstructTransposeGraph(MergedEdgeLists);
-			UE_LOG(LogReferenceChain, Display, TEXT("InitialGather memory usage: %.2f"), (OutGraph.GetAllocatedSize() + Search.GetAllocatedSize()) / 1024.0 / 1024.0);
+			UE_LOG(LogReferenceChain, Display, TEXT("InitialGather memory usage: %.2f"), static_cast<double>(OutGraph.GetAllocatedSize() + Search.GetAllocatedSize()) / 1024.0 / 1024.0);
 		}
 	}
 
@@ -1523,7 +1523,7 @@ FReferenceChainSearch::FReferenceChainSearch(TConstArrayView<UObject*> InObjects
 		PrintResults(!!(Mode & EReferenceChainSearchMode::PrintAllResults));
 	}
 
-	UE_LOG(LogReferenceChain, Display, TEXT("Post-search memory usage: %.2f"), (ReferenceGraph.GetAllocatedSize() + Paths.GetAllocatedSize() + GetAllocatedSize()) / 1024.0 / 1024.0);
+	UE_LOG(LogReferenceChain, Display, TEXT("Post-search memory usage: %.2f"), static_cast<double>(ReferenceGraph.GetAllocatedSize() + Paths.GetAllocatedSize() + GetAllocatedSize()) / 1024.0 / 1024.0);
 
 	GUObjectArray.UnlockInternalArray();
 }
diff --git a/Engine/Source/Runtime/CoreUObject/Private/UObject/ScriptStackTracker.cpp b/Engine/Source/Runtime/CoreUObject/Private/UObject/ScriptStackTracker.cpp
--- a/Engine/Source/Runtime/CoreUObject/Private/UObject/ScriptStackTracker.cpp
+++ b/Engine/Source/Runtime/CoreUObject/Private/UObject/ScriptStackTracker.cpp
@@ -87,7 +87,7 @@ void FScriptStackTracker::DumpStackTraces( int32 StackThreshold, FOutputDevice&
 
 	// Log quick summary as we don't log each individual so totals in CSV won't represent real totals.
 	Ar.Logf(TEXT("Captured %i unique callstacks totalling %" UINT64_FMT " function calls over %" UINT64_FMT " frames, averaging %5.2f calls/frame"), 
-		SortedCallStacks.Num(), TotalStackCount, FramesCaptured, (float)TotalStackCount / FramesCaptured);
+		SortedCallStacks.Num(), TotalStackCount, FramesCaptured, (float)TotalStackCount / (float)FramesCaptured);
 
 	// Iterate over each callstack and write out info in human readable form in CSV format
 	for( int32 CallStackIndex=0; CallStackIndex<SortedCallStacks.Num(); CallStackIndex++ )
diff --git a/Engine/Source/Runtime/CoreUObject/Private/UObject/UObjectGlobals.cpp b/Engine/Source/Runtime/CoreUObject/Private/UObject/UObjectGlobals.cpp
--- a/Engine/Source/Runtime/CoreUObject/Private/UObject/UObjectGlobals.cpp
+++ b/Engine/Source/Runtime/CoreUObject/Private/UObject/UObjectGlobals.cpp
@@ -1446,7 +1446,7 @@ public:
 
 		if (DiffArchive && !bDisable)
 		{
-			TArray<uint8> Data;
+			TArray64<uint8> Data;
 			Data.AddUninitialized(Length);
 			DiffArchive->Seek(Pos);
 			DiffArchive->Serialize(Data.GetData(), Length);
@@ -2087,7 +2087,7 @@ void EndLoad(FUObjectSerializeContext* LoadContext, TArray<UPackage*>* OutLoaded
 	if (ShouldCreateThrottledSlowTask())
 	{
 		static const FText PostLoadText = NSLOCTEXT("Core", "PerformingPostLoad", "Performing post-load...");
-		SlowTask.Emplace(0, PostLoadText);
+		SlowTask.Emplace(0.0f, PostLoadText);
 	}
 
 	int32 NumObjectsLoaded = 0, NumObjectsFound = 0;
@@ -2159,7 +2159,7 @@ void EndLoad(FUObjectSerializeContext* LoadContext, TArray<UPackage*>* OutLoaded
 			if (SlowTask)
 			{
 				SlowTask->CompletedWork = SlowTask->TotalAmountOfWork;
-				SlowTask->TotalAmountOfWork += ObjLoaded.Num();
+				SlowTask->TotalAmountOfWork += static_cast<float>(ObjLoaded.Num());
 				SlowTask->CurrentFrameScope = 0;
 			}
 #endif
@@ -2493,7 +2493,7 @@ namespace NameReuse
 			Entry->Store(reinterpret_cast<UPTRINT>(Parent), UsedName);
 
 			// Shift this entry to the end of the array as it's now the most recently used
-			int32 Index = Entry - Entries.GetData();
+			int32 Index = UE_PTRDIFF_TO_INT32(Entry - Entries.GetData());
 			if (Index != Entries.Num() - 1)
 			{
 				FEntry Removed = MoveTemp(*Entry);
diff --git a/Engine/Source/Runtime/CoreUObject/Tests/Serialization/BulkDataTests.cpp b/Engine/Source/Runtime/CoreUObject/Tests/Serialization/BulkDataTests.cpp
--- a/Engine/Source/Runtime/CoreUObject/Tests/Serialization/BulkDataTests.cpp
+++ b/Engine/Source/Runtime/CoreUObject/Tests/Serialization/BulkDataTests.cpp
@@ -17,7 +17,7 @@ namespace UE::Serialization::BulkDataTest
 static FUniqueBuffer CreatePayload(uint64 Size, uint64 Seed = 0)
 {
 	FUniqueBuffer Buffer = FUniqueBuffer::Alloc(Size);
-	TArrayView<uint64> Values(static_cast<uint64*>(Buffer.GetData()), Size / sizeof(uint64));
+	TArrayView<uint64> Values(static_cast<uint64*>(Buffer.GetData()), IntCastChecked<int32>(Size / sizeof(uint64)));
 
 	uint64 Index = Seed;
 	for (uint64& Value : Values)
diff --git a/Engine/Source/Runtime/CoreUObject/Tests/Serialization/DerivedDataTest.cpp b/Engine/Source/Runtime/CoreUObject/Tests/Serialization/DerivedDataTest.cpp
--- a/Engine/Source/Runtime/CoreUObject/Tests/Serialization/DerivedDataTest.cpp
+++ b/Engine/Source/Runtime/CoreUObject/Tests/Serialization/DerivedDataTest.cpp
@@ -25,7 +25,7 @@ namespace UE::DerivedDataTest::Private
 static FSharedBuffer MakeSharedBuffer(uint64& Counter, uint64 Size)
 {
 	FUniqueBuffer Unique = FUniqueBuffer::Alloc(Size);
-	TArrayView<uint64> Values(static_cast<uint64*>(Unique.GetData()), Size / sizeof(uint64));
+	TArrayView<uint64> Values(static_cast<uint64*>(Unique.GetData()), IntCastChecked<int32>(Size / sizeof(uint64)));
 	uint64 Index = 0;
 	for (uint64& Value : Values)
 	{
