Implement more granular progress tracking of async package loading in zenloader
  - Allows callees of LoadPackageAsync to be notified when serialization is done so they can schedule the next load right from the async loading thread without having to pay for a round-trip to the game-thread
  - Add UnattendedAsyncLoading test as a sample of how this feature can be used
  - Add thread-safety tests for LoadPackageAsync
  - Can be exercised using a cooked client build with the following command-line
        -ExecCmds="Automation RunTest System.Engine.Loading.UnattendedAsyncLoadingTest+System.Engine.Loading.ThreadSafeAsyncLoadingTest;Quit"

diff --git a/Engine/Source/Runtime/CoreUObject/Private/Serialization/AsyncLoading2.cpp b/Engine/Source/Runtime/CoreUObject/Private/Serialization/AsyncLoading2.cpp
--- a/Engine/Source/Runtime/CoreUObject/Private/Serialization/AsyncLoading2.cpp
+++ b/Engine/Source/Runtime/CoreUObject/Private/Serialization/AsyncLoading2.cpp
@@ -493,6 +493,7 @@ struct FPackageRequest
 	FName CustomName;
 	FPackagePath PackagePath;
 	TUniquePtr<FLoadPackageAsyncDelegate> PackageLoadedDelegate;
+	TUniquePtr<FLoadPackageAsyncProgressDelegate> PackageProgressDelegate;
 	FPackageRequest* Next = nullptr;
 
 	FLinkerInstancingContext* GetInstancingContext()
@@ -504,7 +505,7 @@ struct FPackageRequest
 #endif
 	}
 
-	static FPackageRequest Create(int32 RequestId, EPackageFlags PackageFlags, uint32 LoadFlags, int32 PIEInstanceID, int32 Priority, const FLinkerInstancingContext* InstancingContext, const FPackagePath& PackagePath, FName CustomName, TUniquePtr<FLoadPackageAsyncDelegate> PackageLoadedDelegate)
+	static FPackageRequest Create(int32 RequestId, EPackageFlags PackageFlags, uint32 LoadFlags, int32 PIEInstanceID, int32 Priority, const FLinkerInstancingContext* InstancingContext, const FPackagePath& PackagePath, FName CustomName, TUniquePtr<FLoadPackageAsyncDelegate> PackageLoadedDelegate, TUniquePtr<FLoadPackageAsyncProgressDelegate> PackageProgressDelegate)
 	{
 		return FPackageRequest
 		{
@@ -521,6 +522,7 @@ struct FPackageRequest
 			CustomName,
 			PackagePath,
 			MoveTemp(PackageLoadedDelegate),
+			MoveTemp(PackageProgressDelegate),
 			nullptr
 		};
 	}
@@ -2599,6 +2601,7 @@ struct FAsyncPackage2
 	double GetLoadStartTime() const;
 
 	void AddCompletionCallback(TUniquePtr<FLoadPackageAsyncDelegate>&& Callback);
+	void AddProgressCallback(TUniquePtr<FLoadPackageAsyncProgressDelegate>&& Callback);
 
 	FORCEINLINE UPackage* GetLinkerRoot() const
 	{
@@ -2818,6 +2821,8 @@ private:
 	TArray<FExternalReadCallback> ExternalReadDependencies;
 	/** Callbacks called when we finished loading this package */
 	TArray<TUniquePtr<FLoadPackageAsyncDelegate>, TInlineAllocator<2>> CompletionCallbacks;
+	/** Callbacks called for the different loading phase of this package */
+	TArray<TUniquePtr<FLoadPackageAsyncProgressDelegate>, TInlineAllocator<2>> ProgressCallbacks;
 
 	/** Set when the package is being loaded as an instance; null otherwise. */
 	TUniquePtr<FLinkerInstancingContext> InstanceContext;
@@ -3016,14 +3021,15 @@ private:
 		UPackage* UPackage = nullptr;
 		FAsyncPackage2* AsyncPackage = nullptr;
 		TArray<TUniquePtr<FLoadPackageAsyncDelegate>, TInlineAllocator<2>> CompletionCallbacks;
+		TArray<TUniquePtr<FLoadPackageAsyncProgressDelegate>, TInlineAllocator<2>> ProgressCallbacks;
 		TArray<int32, TInlineAllocator<2>> RequestIDs;
 
 		static FCompletedPackageRequest FromMissingPackage(
 			const FAsyncPackageDesc2& Desc,
-			TUniquePtr<FLoadPackageAsyncDelegate>&& Callback)
+			TUniquePtr<FLoadPackageAsyncDelegate>&& CompletionCallback)
 		{
 			FCompletedPackageRequest Res = {Desc.UPackageName, EAsyncLoadingResult::Failed};
-			Res.CompletionCallbacks.Add(MoveTemp(Callback));
+			Res.CompletionCallbacks.Add(MoveTemp(CompletionCallback));
 			Res.RequestIDs.Add(Desc.RequestID);
 			return Res;
 		}
@@ -3038,6 +3044,7 @@ private:
 				Package->LinkerRoot,
 				Package,
 				MoveTemp(Package->CompletionCallbacks),
+				MoveTemp(Package->ProgressCallbacks),
 				Package->RequestIDs
 			};
 		}
@@ -3046,17 +3053,34 @@ private:
 		{
 			checkSlow(IsInGameThread());
 			
-			if (CompletionCallbacks.Num() == 0)
+			if (CompletionCallbacks.Num() != 0)
 			{
-				return;
+				TRACE_CPUPROFILER_EVENT_SCOPE(PackageCompletionCallbacks);
+				for (TUniquePtr<FLoadPackageAsyncDelegate>& CompletionCallback : CompletionCallbacks)
+				{
+					CompletionCallback->ExecuteIfBound(PackageName, UPackage, Result);
+				}
+				CompletionCallbacks.Empty();
 			}
 
-			TRACE_CPUPROFILER_EVENT_SCOPE(PackageCompletionCallbacks);
-			for (TUniquePtr<FLoadPackageAsyncDelegate>& CompletionCallback : CompletionCallbacks)
+			if (ProgressCallbacks.Num() != 0)
 			{
-				CompletionCallback->ExecuteIfBound(PackageName, UPackage, Result);
+				TRACE_CPUPROFILER_EVENT_SCOPE(PackageProgressCallbacks_Completion);
+
+				const FLoadPackageAsyncProgressParams Params
+				{
+					.PackageName = PackageName,
+					.LoadedPackage = UPackage,
+					.ProgressType = Result == EAsyncLoadingResult::Succeeded ? EAsyncLoadingProgress::FullyLoaded : EAsyncLoadingProgress::Failed
+				};
+				
+				for (TUniquePtr<FLoadPackageAsyncProgressDelegate>& ProgressCallback : ProgressCallbacks)
+				{
+					ProgressCallback->ExecuteIfBound(Params);
+				}
+				ProgressCallbacks.Empty();
 			}
-			CompletionCallbacks.Empty();
+
 		}
 	};
 	TArray<FCompletedPackageRequest> CompletedPackageRequests;
@@ -3219,8 +3243,8 @@ public:
 	void UpdatePackagePriority(FAsyncLoadingThreadState2& ThreadState, FAsyncPackage2* Package);
 	void UpdatePackagePriorityRecursive(FAsyncLoadingThreadState2& ThreadState, FAsyncPackage2* Package, int32 NewPriority);
 
-	FAsyncPackage2* FindOrInsertPackage(FAsyncLoadingThreadState2& ThreadState, FAsyncPackageDesc2& InDesc, bool& bInserted, FAsyncPackage2* ImportedByPackage, TUniquePtr<FLoadPackageAsyncDelegate>&& PackageLoadedDelegate = TUniquePtr<FLoadPackageAsyncDelegate>());
-	void QueueMissingPackage(FAsyncPackageDesc2& PackageDesc, TUniquePtr<FLoadPackageAsyncDelegate>&& LoadPackageAsyncDelegate);
+	FAsyncPackage2* FindOrInsertPackage(FAsyncLoadingThreadState2& ThreadState, FAsyncPackageDesc2& InDesc, bool& bInserted, FAsyncPackage2* ImportedByPackage, TUniquePtr<FLoadPackageAsyncDelegate>&& PackageLoadedDelegate = TUniquePtr<FLoadPackageAsyncDelegate>(), TUniquePtr<FLoadPackageAsyncProgressDelegate>&& PackageProgressDelegate = TUniquePtr<FLoadPackageAsyncProgressDelegate>());
+	void QueueMissingPackage(FAsyncPackageDesc2& PackageDesc, TUniquePtr<FLoadPackageAsyncDelegate>&& LoadPackageAsyncDelegate, TUniquePtr<FLoadPackageAsyncProgressDelegate>&& LoadPackageAsyncProgressDelegate);
 
 	/**
 	* [ASYNC* THREAD] Loads all packages
@@ -3261,6 +3285,8 @@ public:
 		return true;
 	}
 
+	int32 LoadPackageInternal(const FPackagePath& InPackagePath, FName InCustomName, TUniquePtr<FLoadPackageAsyncDelegate>&& InCompletionDelegate, TUniquePtr<FLoadPackageAsyncProgressDelegate>&& InProgressDelegate, EPackageFlags InPackageFlags, int32 InPIEInstanceID, int32 InPackagePriority, const FLinkerInstancingContext* InInstancingContext, uint32 InLoadFlags);
+
 	virtual int32 LoadPackage(
 		const FPackagePath& InPackagePath,
 		FName InCustomName,
@@ -3271,6 +3297,10 @@ public:
 		const FLinkerInstancingContext* InstancingContext = nullptr,
 		uint32 InLoadFlags = LOAD_None) override;
 
+	virtual int32 LoadPackage(
+		const FPackagePath& InPackagePath,
+		FLoadPackageAsyncOptionalParams Params) override;
+
 	EAsyncPackageState::Type ProcessLoadingFromGameThread(FAsyncLoadingThreadState2& ThreadState, bool bUseTimeLimit, bool bUseFullTimeLimit, double TimeLimit);
 
 	inline virtual EAsyncPackageState::Type ProcessLoading(bool bUseTimeLimit, bool bUseFullTimeLimit, double TimeLimit) override
@@ -3909,7 +3939,7 @@ void FAsyncLoadingThread2::MergePostLoadGroups(FAsyncLoadingThreadState2& Thread
 	delete Source;
 }
 
-FAsyncPackage2* FAsyncLoadingThread2::FindOrInsertPackage(FAsyncLoadingThreadState2& ThreadState, FAsyncPackageDesc2& Desc, bool& bInserted, FAsyncPackage2* ImportedByPackage, TUniquePtr<FLoadPackageAsyncDelegate>&& PackageLoadedDelegate)
+FAsyncPackage2* FAsyncLoadingThread2::FindOrInsertPackage(FAsyncLoadingThreadState2& ThreadState, FAsyncPackageDesc2& Desc, bool& bInserted, FAsyncPackage2* ImportedByPackage, TUniquePtr<FLoadPackageAsyncDelegate>&& PackageLoadedDelegate, TUniquePtr<FLoadPackageAsyncProgressDelegate>&& PackageProgressDelegate)
 {
 	TRACE_CPUPROFILER_EVENT_SCOPE(FindOrInsertPackage);
 	check(ThreadState.bCanAccessAsyncLoadingThreadData);
@@ -3944,6 +3974,10 @@ FAsyncPackage2* FAsyncLoadingThread2::FindOrInsertPackage(FAsyncLoadingThreadSta
 		{
 			Package->AddCompletionCallback(MoveTemp(PackageLoadedDelegate));
 		}
+		if (PackageProgressDelegate.IsValid())
+		{
+			Package->AddProgressCallback(MoveTemp(PackageProgressDelegate));
+		}
 	}
 
 	// PostLoadGroup are to protect some race conditions happening in runtime when ALT is active but
@@ -4155,12 +4189,12 @@ bool FAsyncLoadingThread2::CreateAsyncPackagesFromQueue(FAsyncLoadingThreadState
 			FAsyncPackageDesc2 PackageDesc = FAsyncPackageDesc2::FromPackageRequest(Request, UPackageName, PackageIdToLoad);
 			if (PackageStatus == EPackageStoreEntryStatus::Missing)
 			{
-				QueueMissingPackage(PackageDesc, MoveTemp(Request.PackageLoadedDelegate));
+				QueueMissingPackage(PackageDesc, MoveTemp(Request.PackageLoadedDelegate), MoveTemp(Request.PackageProgressDelegate));
 			}
 			else
 			{
 				bool bInserted;
-				FAsyncPackage2* Package = FindOrInsertPackage(ThreadState, PackageDesc, bInserted, nullptr, MoveTemp(Request.PackageLoadedDelegate));
+				FAsyncPackage2* Package = FindOrInsertPackage(ThreadState, PackageDesc, bInserted, nullptr, MoveTemp(Request.PackageLoadedDelegate), MoveTemp(Request.PackageProgressDelegate));
 				checkf(Package, TEXT("Failed to find or insert package %s"), *PackageDesc.UPackageName.ToString());
 
 				if (bInserted)
@@ -4196,7 +4230,7 @@ bool FAsyncLoadingThread2::CreateAsyncPackagesFromQueue(FAsyncLoadingThreadState
 				else
 				{
 					UE_ASYNC_PACKAGE_LOG_VERBOSE(Verbose, PackageDesc, TEXT("CreateAsyncPackages: UpdatePackage"),
-						TEXT("Package is alreay being loaded."));
+						TEXT("Package is already being loaded."));
 					--PackagesWithRemainingWorkCounter;
 					TRACE_COUNTER_SET(AsyncLoadingPackagesWithRemainingWork, PackagesWithRemainingWorkCounter);
 				}
@@ -4996,6 +5030,23 @@ void FAsyncPackage2::StartLoading(FAsyncLoadingThreadState2& ThreadState, FIoBat
 
 	LoadStartTime = FPlatformTime::Seconds();
 
+	if (ProgressCallbacks.Num() != 0)
+	{
+		TRACE_CPUPROFILER_EVENT_SCOPE(PackageProgressCallbacks_Started);
+
+		const FLoadPackageAsyncProgressParams Params
+		{
+			.PackageName = Desc.UPackageName,
+			.LoadedPackage = GetLinkerRoot(),
+			.ProgressType = EAsyncLoadingProgress::Started
+		};
+
+		for (TUniquePtr<FLoadPackageAsyncProgressDelegate>& ProgressCallback : ProgressCallbacks)
+		{
+			ProgressCallback->ExecuteIfBound(Params);
+		}
+	}
+
 	AsyncPackageLoadingState = EAsyncPackageLoadingState2::WaitingForIo;
 
 	FIoReadOptions ReadOptions;
@@ -6435,6 +6486,23 @@ EEventLoadNodeExecutionResult FAsyncPackage2::Event_ExportsDone(FAsyncLoadingThr
 		FCoreDelegates::ReleasePreloadedPackageShaderMaps.ExecuteIfBound(Package->Data.ShaderMapHashes);
 	}
 
+	if (Package->ProgressCallbacks.Num())
+	{
+		TRACE_CPUPROFILER_EVENT_SCOPE(PackageProgressCallbacks_Serialized);
+
+		const FLoadPackageAsyncProgressParams Params
+		{
+			.PackageName = Package->Desc.UPackageName,
+			.LoadedPackage = Package->GetLinkerRoot(),
+			.ProgressType = EAsyncLoadingProgress::Serialized
+		};
+
+		for (TUniquePtr<FLoadPackageAsyncProgressDelegate>& ProgressCallback : Package->ProgressCallbacks)
+		{
+			ProgressCallback->ExecuteIfBound(Params);
+		}
+	}
+
 	FAsyncLoadingPostLoadGroup* PostLoadGroup = Package->PostLoadGroup;
 	check(PostLoadGroup);
 	check(PostLoadGroup->PackagesWithExportsToSerializeCount > 0);
@@ -7316,7 +7384,7 @@ EAsyncPackageState::Type FAsyncLoadingThread2::ProcessLoadedPackagesFromGameThre
 			Package->AsyncPackageLoadingState = EAsyncPackageLoadingState2::DeferredDelete;
 			Package->ClearImportedPackages();
 
-			if (Package->CompletionCallbacks.IsEmpty())
+			if (Package->CompletionCallbacks.IsEmpty() && Package->ProgressCallbacks.IsEmpty())
 			{
 				RemovePendingRequests(Package->RequestIDs);
 				Package->ReleaseRef();
@@ -8483,7 +8551,12 @@ void FAsyncPackage2::AddCompletionCallback(TUniquePtr<FLoadPackageAsyncDelegate>
 	CompletionCallbacks.Emplace(MoveTemp(Callback));
 }
 
-int32 FAsyncLoadingThread2::LoadPackage(const FPackagePath& InPackagePath, FName InCustomName, FLoadPackageAsyncDelegate InCompletionDelegate, EPackageFlags InPackageFlags, int32 InPIEInstanceID, int32 InPackagePriority, const FLinkerInstancingContext* InInstancingContext, uint32 InLoadFlags)
+void FAsyncPackage2::AddProgressCallback(TUniquePtr<FLoadPackageAsyncProgressDelegate>&& Callback)
+{
+	ProgressCallbacks.Emplace(MoveTemp(Callback));
+}
+
+int32 FAsyncLoadingThread2::LoadPackageInternal(const FPackagePath& InPackagePath, FName InCustomName, TUniquePtr<FLoadPackageAsyncDelegate>&& InCompletionDelegate, TUniquePtr<FLoadPackageAsyncProgressDelegate>&& InProgressDelegate, EPackageFlags InPackageFlags, int32 InPIEInstanceID, int32 InPackagePriority, const FLinkerInstancingContext* InInstancingContext, uint32 InLoadFlags)
 {
 	TRACE_CPUPROFILER_EVENT_SCOPE(LoadPackage);
 
@@ -8514,12 +8587,7 @@ PRAGMA_ENABLE_DEPRECATION_WARNINGS
 	TRACE_LOADTIME_BEGIN_REQUEST(RequestId);
 	AddPendingRequest(RequestId);
 
-	// Allocate delegate before going async, it is not safe to copy delegates by value on other threads
-	TUniquePtr<FLoadPackageAsyncDelegate> CompletionDelegate = InCompletionDelegate.IsBound()
-		? MakeUnique<FLoadPackageAsyncDelegate>(MoveTemp(InCompletionDelegate))
-		: TUniquePtr<FLoadPackageAsyncDelegate>();
-
-	PackageRequestQueue.Enqueue(FPackageRequest::Create(RequestId, InPackageFlags, InLoadFlags, InPIEInstanceID, InPackagePriority, InInstancingContext, InPackagePath, InCustomName, MoveTemp(CompletionDelegate)));
+	PackageRequestQueue.Enqueue(FPackageRequest::Create(RequestId, InPackageFlags, InLoadFlags, InPIEInstanceID, InPackagePriority, InInstancingContext, InPackagePath, InCustomName, MoveTemp(InCompletionDelegate), MoveTemp(InProgressDelegate)));
 	++QueuedPackagesCounter;
 	++PackagesWithRemainingWorkCounter;
 
@@ -8531,7 +8599,22 @@ PRAGMA_ENABLE_DEPRECATION_WARNINGS
 	return RequestId;
 }
 
-void FAsyncLoadingThread2::QueueMissingPackage(FAsyncPackageDesc2& PackageDesc, TUniquePtr<FLoadPackageAsyncDelegate>&& PackageLoadedDelegate)
+int32 FAsyncLoadingThread2::LoadPackage(const FPackagePath& InPackagePath, FLoadPackageAsyncOptionalParams InParams)
+{
+	return LoadPackageInternal(InPackagePath, InParams.CustomPackageName, MoveTemp(InParams.CompletionDelegate), MoveTemp(InParams.ProgressDelegate), InParams.PackageFlags, InParams.PIEInstanceID, InParams.PackagePriority, InParams.InstancingContext, InParams.LoadFlags);
+}
+
+int32 FAsyncLoadingThread2::LoadPackage(const FPackagePath& InPackagePath, FName InCustomName, FLoadPackageAsyncDelegate InCompletionDelegate, EPackageFlags InPackageFlags, int32 InPIEInstanceID, int32 InPackagePriority, const FLinkerInstancingContext* InInstancingContext, uint32 InLoadFlags)
+{
+	// Allocate delegate before going async, it is not safe to copy delegates by value on other threads
+	TUniquePtr<FLoadPackageAsyncDelegate> CompletionDelegate = InCompletionDelegate.IsBound()
+		? MakeUnique<FLoadPackageAsyncDelegate>(MoveTemp(InCompletionDelegate))
+		: TUniquePtr<FLoadPackageAsyncDelegate>();
+
+	return LoadPackageInternal(InPackagePath, InCustomName, MoveTemp(CompletionDelegate), TUniquePtr<FLoadPackageAsyncProgressDelegate>(), InPackageFlags, InPIEInstanceID, InPackagePriority, InInstancingContext, InLoadFlags);
+}
+
+void FAsyncLoadingThread2::QueueMissingPackage(FAsyncPackageDesc2& PackageDesc, TUniquePtr<FLoadPackageAsyncDelegate>&& PackageLoadedDelegate, TUniquePtr<FLoadPackageAsyncProgressDelegate>&& PackageProgressDelegate)
 {
 	const FName FailedPackageName = PackageDesc.UPackageName;
 
@@ -8554,6 +8637,20 @@ void FAsyncLoadingThread2::QueueMissingPackage(FAsyncPackageDesc2& PackageDesc,
 		TEXT("LoadPackage: SkipPackage: %s (0x%llX) - The package to load does not exist on disk or in the loader"),
 		*FailedPackageName.ToString(), PackageDesc.PackageIdToLoad.ValueForDebugging());
 
+	if (PackageProgressDelegate.IsValid())
+	{
+		TRACE_CPUPROFILER_EVENT_SCOPE(PackageProgressCallback_Failed);
+
+		const FLoadPackageAsyncProgressParams Params
+		{
+			.PackageName = FailedPackageName,
+			.LoadedPackage = nullptr,
+			.ProgressType = EAsyncLoadingProgress::Failed
+		};
+
+		PackageProgressDelegate->ExecuteIfBound(Params);
+	}
+
 	if (PackageLoadedDelegate.IsValid())
 	{
 		FScopeLock _(&FailedPackageRequestsCritical);
diff --git a/Engine/Source/Runtime/CoreUObject/Private/Serialization/AsyncPackageLoader.cpp b/Engine/Source/Runtime/CoreUObject/Private/Serialization/AsyncPackageLoader.cpp
--- a/Engine/Source/Runtime/CoreUObject/Private/Serialization/AsyncPackageLoader.cpp
+++ b/Engine/Source/Runtime/CoreUObject/Private/Serialization/AsyncPackageLoader.cpp
@@ -584,6 +584,16 @@ bool ShouldAlwaysLoadPackageAsync(const FPackagePath& InPackagePath)
 	return GPackageLoader->ShouldAlwaysLoadPackageAsync(InPackagePath);
 }
 
+int32 LoadPackageAsync(const FPackagePath& InPackagePath, FLoadPackageAsyncOptionalParams InOptionalParams)
+{
+	LLM_SCOPE(ELLMTag::AsyncLoading);
+	UE_CLOG(!GAsyncLoadingAllowed && !IsInAsyncLoadingThread(), LogStreaming, Fatal, TEXT("Requesting async load of \"%s\" when async loading is not allowed (after shutdown). Please fix higher level code."), *InPackagePath.GetDebugName());
+#if DO_TRACK_ASYNC_LOAD_REQUESTS
+	FTrackAsyncLoadRequests::Get().TrackRequest(InPackagePath.GetDebugName(), nullptr, InOptionalParams.PackagePriority);
+#endif
+	return GetAsyncPackageLoader().LoadPackage(InPackagePath, MoveTemp(InOptionalParams));
+}
+
 int32 LoadPackageAsync(const FPackagePath& InPackagePath,
 		FName InPackageNameToCreate /* = NAME_None*/,
 		FLoadPackageAsyncDelegate InCompletionDelegate /*= FLoadPackageAsyncDelegate()*/,
@@ -615,6 +625,13 @@ int32 LoadPackageAsync(const FString& InName, FLoadPackageAsyncDelegate Completi
 	return LoadPackageAsync(PackagePath, NAME_None /* InPackageNameToCreate */, CompletionDelegate, InPackageFlags, InPIEInstanceID, InPackagePriority, nullptr);
 }
 
+int32 LoadPackageAsync(const FString& InName, FLoadPackageAsyncOptionalParams InOptionalParams)
+{
+	LLM_SCOPE(ELLMTag::AsyncLoading);
+	FPackagePath PackagePath = GetLoadPackageAsyncPackagePath(InName);
+	return LoadPackageAsync(PackagePath, MoveTemp(InOptionalParams));
+}
+
 int32 LoadPackageAsync(const FString& InName, const FGuid* InGuid /*= nullptr*/, const TCHAR* InPackageToLoadFrom, FLoadPackageAsyncDelegate InCompletionDelegate /*= FLoadPackageAsyncDelegate()*/, EPackageFlags InPackageFlags /*= PKG_None*/, int32 InPIEInstanceID /*= INDEX_NONE*/, int32 InPackagePriority /*= 0*/, const FLinkerInstancingContext* InstancingContext /*=nullptr*/)
 {
 	LLM_SCOPE(ELLMTag::AsyncLoading);
@@ -629,7 +646,7 @@ int32 LoadPackageAsync(const FString& InName, const FGuid* InGuid /*= nullptr*/,
 			InPackageNameToCreate = PackagePathToCreate.GetPackageFName();
 		}
 	}
-	return LoadPackageAsync(PackagePath, InPackageNameToCreate, InCompletionDelegate, InPackageFlags, InPIEInstanceID, InPackagePriority, InstancingContext);
+	return LoadPackageAsync(PackagePath, InPackageNameToCreate, MoveTemp(InCompletionDelegate), InPackageFlags, InPIEInstanceID, InPackagePriority, InstancingContext);
 }
 
 void CancelAsyncLoading()
diff --git a/Engine/Source/Runtime/CoreUObject/Public/Serialization/AsyncPackageLoader.h b/Engine/Source/Runtime/CoreUObject/Public/Serialization/AsyncPackageLoader.h
--- a/Engine/Source/Runtime/CoreUObject/Public/Serialization/AsyncPackageLoader.h
+++ b/Engine/Source/Runtime/CoreUObject/Public/Serialization/AsyncPackageLoader.h
@@ -123,6 +123,19 @@ public:
 			const FLinkerInstancingContext* InInstancingContext,
 			uint32 InLoadFlags) = 0;
 
+	/**
+	 * Asynchronously load a package.
+	 *
+	 * @param  PackagePath         PackagePath to load. Must be a mounted path. The package is created if it does not already exist.
+	 * @param  OptionalParams      Struct containing all the parameters required to load the package.
+	 * @return Unique ID associated with this load request (the same package can be associated with multiple IDs).
+	 */
+	virtual int32 LoadPackage(const FPackagePath& PackagePath, FLoadPackageAsyncOptionalParams OptionalParams)
+	{
+		checkf(false, TEXT("Not Implemented"));
+		return INDEX_NONE;
+	};
+
 	/**
 	 * Process all currently loading package requests.
 	 *
diff --git a/Engine/Source/Runtime/CoreUObject/Public/UObject/UObjectGlobals.h b/Engine/Source/Runtime/CoreUObject/Public/UObject/UObjectGlobals.h
--- a/Engine/Source/Runtime/CoreUObject/Public/UObject/UObjectGlobals.h
+++ b/Engine/Source/Runtime/CoreUObject/Public/UObject/UObjectGlobals.h
@@ -644,6 +644,21 @@ namespace EAsyncLoadingResult
 	};
 }
 
+/** Async package loading result */
+enum class EAsyncLoadingProgress : uint32
+{
+	/** Package failed to load */
+	Failed,
+	/** Package has started loading. */
+	Started,
+	/** Package has finished its serialization phase. */
+	Serialized,
+	/** Package has finished all loading phase successfully */
+	FullyLoaded,
+	/** Async loading was canceled */
+	Canceled
+};
+
 /** The type that represents an async loading priority */
 typedef int32 TAsyncLoadPriority;
 
@@ -655,6 +670,59 @@ typedef int32 TAsyncLoadPriority;
  */
 DECLARE_DELEGATE_ThreeParams(FLoadPackageAsyncDelegate, const FName& /*PackageName*/, UPackage* /*LoadedPackage*/, EAsyncLoadingResult::Type /*Result*/)
 
+/**
+ * Parameters passed to the FLoadPackageAsyncProgressDelegate callback.
+ */
+struct FLoadPackageAsyncProgressParams
+{
+	/* Name of the package. */
+	FName PackageName { NAME_None };
+	/* Pointer to UPackage being loaded, can be nullptr depending on async loading progress. */
+	UPackage* LoadedPackage { nullptr };
+	/* Progress of async loading. */
+	EAsyncLoadingProgress ProgressType { EAsyncLoadingProgress::Failed };
+};
+
+/**
+ * Thread-safe delegate called on progress of async package loading.
+ * @param	Params        Struct containing the parameters for the callback.
+ */
+using FLoadPackageAsyncProgressDelegate = TTSDelegate<void(const FLoadPackageAsyncProgressParams& Params)>;
+
+/**
+ * Optional parameters passed to the LoadPackageAsync function.
+ */
+struct FLoadPackageAsyncOptionalParams
+{
+	/** If not none, this is the name of the package to load into (and create if not yet existing). If none, the name is take from PackagePath. **/
+	FName CustomPackageName { NAME_None };
+	/** Non Thread-safe delegate to be invoked from game-thread on completion. **/
+	TUniquePtr<FLoadPackageAsyncDelegate> CompletionDelegate;
+	/** Thread-safe delegate to be invoked at different state of progress for the given package. **/
+	TUniquePtr<FLoadPackageAsyncProgressDelegate> ProgressDelegate;
+	/** Package flags used to construct loaded package in memory. **/
+	EPackageFlags PackageFlags { PKG_None };
+	/** Play in Editor instance ID. **/
+	int32 PIEInstanceID { INDEX_NONE };
+	/** Loading priority. **/
+	int32 PackagePriority { 0 };
+	/** Additional context to map object names to their instanced counterpart when loading an instanced package. **/
+	const FLinkerInstancingContext* InstancingContext { nullptr };
+	/** Flags controlling loading behavior, from the ELoadFlags enum. */
+	uint32 LoadFlags { LOAD_None };
+};
+
+/**
+ * Asynchronously load a package and all contained objects that match context flags. Non-blocking.
+ * Use this version to specify the PackagePath rather than having the other versions internally convert the InName to a PackagePath by searching the current package mount points.
+ * Use this version if you need to specify a packagename that is different from the packagename on disk; this is useful when loading multiple copies of the same package.
+ *
+  * @param	InPackagePath           PackagePath to load. Must be a mounted path. The package is created if it does not already exist.
+  * @param	InOptionalParams        Optional parameters.
+ * @return Unique ID associated with this load request (the same package can be associated with multiple IDs).
+ */
+COREUOBJECT_API int32 LoadPackageAsync(const FPackagePath& InPackagePath, FLoadPackageAsyncOptionalParams InOptionalParams);
+
 /**
  * Asynchronously load a package and all contained objects that match context flags. Non-blocking.
  * Use this version to specify the PackagePath rather than having the other versions internally convert the InName to a PackagePath by searching the current package mount points.
@@ -695,6 +763,16 @@ COREUOBJECT_API int32 LoadPackageAsync(const FString& InName, const FGuid* InGui
  */
 COREUOBJECT_API int32 LoadPackageAsync(const FString& InName, FLoadPackageAsyncDelegate InCompletionDelegate, TAsyncLoadPriority InPackagePriority = 0, EPackageFlags InPackageFlags = PKG_None, int32 InPIEInstanceID = INDEX_NONE);
 
+/**
+ * Asynchronously load a package and all contained objects that match context flags. Non-blocking.
+ * Use this version when you need to load a package with default behavior from a packagename/filename, and need to be notified when it is loaded.
+ *
+ * @param	InName                  PackageName or LocalFilePath of package to load. Must be a mounted name/path. The package is created if it does not already exist.
+ * @param	InOptionalParams        Optional parameters.
+ * @return Unique ID associated with this load request (the same package can be associated with multiple IDs).
+ */
+COREUOBJECT_API int32 LoadPackageAsync(const FString& InName, FLoadPackageAsyncOptionalParams InOptionalParams);
+
 /**
 * Cancels all async package loading requests.
 */
diff --git a/Engine/Source/Runtime/Engine/Private/Tests/Loading/AsyncLoadingTests.cpp b/Engine/Source/Runtime/Engine/Private/Tests/Loading/AsyncLoadingTests.cpp
new file mode 100644
--- /dev/null
+++ b/Engine/Source/Runtime/Engine/Private/Tests/Loading/AsyncLoadingTests.cpp
@@ -0,0 +1,60 @@
+// Copyright Epic Games, Inc. All Rights Reserved.
+
+#include "Misc/AutomationTest.h"
+#include "AssetRegistry/AssetData.h"
+#include "AssetRegistry/AssetRegistryModule.h"
+#include "AssetRegistry/AssetRegistryState.h"
+#include "AssetRegistry/IAssetRegistry.h"
+#include "Async/ParallelFor.h"
+
+#if WITH_DEV_AUTOMATION_TESTS
+
+#undef TEST_NAME_ROOT
+#define TEST_NAME_ROOT "System.Engine.Loading"
+
+/**
+ * This test demonstrate that LoadPackageAsync is thread-safe and can be called from multiple workers at the same time.
+ */
+IMPLEMENT_SIMPLE_AUTOMATION_TEST(FThreadSafeAsyncLoadingTest, TEXT(TEST_NAME_ROOT ".ThreadSafeAsyncLoadingTest"), EAutomationTestFlags::ClientContext | EAutomationTestFlags::EngineFilter)
+bool FThreadSafeAsyncLoadingTest::RunTest(const FString& Parameters)
+{
+	// We use the asset registry to get a list of asset to load. 
+	IAssetRegistry& AssetRegistry = FModuleManager::LoadModuleChecked<FAssetRegistryModule>(FName(TEXT("AssetRegistry"))).Get();
+	AssetRegistry.WaitForCompletion();
+
+	// Limit the number of packages we're going to load for the test in case the project is very big.
+	constexpr int32 MaxPackageCount = 5000;
+
+	TSet<FName> UniquePackages;
+	AssetRegistry.EnumerateAllAssets(
+		[&UniquePackages](const FAssetData& AssetData)
+		{
+			if (UniquePackages.Num() < MaxPackageCount)
+			{
+				UniquePackages.FindOrAdd(AssetData.PackageName);
+				return true;
+			}
+			
+			return false;
+		},
+		true /* bIncludeOnlyOnDiskAssets */
+	);
+
+	TArray<FName> PackagesToLoad(UniquePackages.Array());
+	TArray<int32> RequestIDs;
+	RequestIDs.SetNum(PackagesToLoad.Num());
+
+	ParallelFor(PackagesToLoad.Num(),
+		[&PackagesToLoad, &RequestIDs](int32 Index)
+		{
+			RequestIDs[Index] = LoadPackageAsync(PackagesToLoad[Index].ToString());
+		}
+	);
+
+	FlushAsyncLoading(RequestIDs);
+
+	return true;
+}
+
+#undef TEST_NAME_ROOT
+#endif // WITH_DEV_AUTOMATION_TESTS
\ No newline at end of file
diff --git a/Engine/Source/Runtime/Engine/Private/Tests/Loading/UnattendedAsyncLoadingTests.cpp b/Engine/Source/Runtime/Engine/Private/Tests/Loading/UnattendedAsyncLoadingTests.cpp
new file mode 100644
--- /dev/null
+++ b/Engine/Source/Runtime/Engine/Private/Tests/Loading/UnattendedAsyncLoadingTests.cpp
@@ -0,0 +1,227 @@
+// Copyright Epic Games, Inc. All Rights Reserved.
+
+#include "Misc/AutomationTest.h"
+#include "AssetRegistry/AssetData.h"
+#include "AssetRegistry/AssetRegistryModule.h"
+#include "AssetRegistry/AssetRegistryState.h"
+#include "AssetRegistry/IAssetRegistry.h"
+#include "Tasks/Task.h"
+
+#if WITH_DEV_AUTOMATION_TESTS
+
+#undef TEST_NAME_ROOT
+#define TEST_NAME_ROOT "System.Engine.Loading"
+
+namespace UnattendedLoadTestImpl
+{
+	// We're still missing a simple multi-consumer thread-safe queue in the engine...
+	// So make a custom one for our purpose.
+	template <typename T>
+	class TThreadSafeQueue
+	{
+		mutable FRWLock Lock;
+		TArray<T> Array;
+	public:
+		TThreadSafeQueue(TArray<T>&& InArray)
+			: Array(MoveTemp(InArray))
+		{
+		}
+
+		bool Dequeue(T& OutName)
+		{
+			FWriteScopeLock Scope(Lock);
+			if (Array.IsEmpty())
+			{
+				return false;
+			}
+
+			OutName = Array.Pop(false /* bAllowShrinking */);
+			return true;
+		}
+
+		int32 Num() const
+		{
+			FReadScopeLock Scope(Lock);
+			return Array.Num();
+		}
+	};
+
+	// We use this to make sure the Lambda is only moved and never copied
+	// and gets called at least one time before being deleted.
+	// This is for testing purpose and is not required for the unattended method to work.
+	class FExecutionCheck
+	{
+		enum class EExecutionCheck
+		{
+			NotExecuted,
+			Executed,
+			MovedFrom
+		};
+
+		EExecutionCheck ExecutionState = EExecutionCheck::NotExecuted;
+	public:
+		FExecutionCheck() = default;
+		~FExecutionCheck()
+		{
+			check(ExecutionState == EExecutionCheck::Executed || ExecutionState == EExecutionCheck::MovedFrom);
+		}
+		FExecutionCheck(const FExecutionCheck&)
+		{
+			checkf(false, TEXT("Delegate should be moved, not copied"));
+		}
+		FExecutionCheck& operator= (const FExecutionCheck&)
+		{
+			checkf(false, TEXT("Delegate should be moved, not copied"));
+			return *this;
+		}
+		FExecutionCheck(FExecutionCheck&& Other)
+		{
+			checkf(ExecutionState != EExecutionCheck::Executed, TEXT("Probably don't want to move from an executed object"));
+			ExecutionState = Other.ExecutionState;
+			Other.ExecutionState = EExecutionCheck::MovedFrom;
+		}
+		FExecutionCheck& operator=(FExecutionCheck&& Other)
+		{
+			FExecutionCheck Temp = MoveTemp(Other);
+			Swap(Temp, *this);
+			return *this;
+		}
+		void SetHasBeenExecuted() 
+		{
+			check(ExecutionState != EExecutionCheck::MovedFrom);
+			ExecutionState = EExecutionCheck::Executed;
+		}
+		bool HasBeenExecuted() const { return ExecutionState == EExecutionCheck::Executed; }
+	};
+
+	// The UnattendedLoader will queue the next async loading as soon as a package finishes its serialization phase
+	// or finishes loading for any reason. 
+	class FUnattendedLoader
+	{
+		bool bFinished = false;
+		TThreadSafeQueue<FName> PackagesToLoad;
+
+		void LoadNextPackage()
+		{
+			TRACE_CPUPROFILER_EVENT_SCOPE(LoadNextPackage);
+
+			FName PackageName;
+			if (PackagesToLoad.Dequeue(PackageName))
+			{
+				UE_LOG(LogTemp, Display, TEXT("Loading %s (Remaining: %d)"), *PackageName.ToString(), PackagesToLoad.Num());
+				LoadPackageAsync(PackageName.ToString(),
+					{
+						.ProgressDelegate = MakeUnique<FLoadPackageAsyncProgressDelegate>(
+							FLoadPackageAsyncProgressDelegate::CreateLambda(
+								[this, ExecutionCheck = FExecutionCheck()](const FLoadPackageAsyncProgressParams& Params) mutable
+								{
+									// If we already triggered the next load from this callback, don't do anything else
+									if (ExecutionCheck.HasBeenExecuted())
+									{
+										return;
+									}
+							
+									switch (Params.ProgressType)
+									{
+										// Filter out some progress types as we wait until later to start loading the next package
+										case EAsyncLoadingProgress::Started:
+											return;
+
+										// Whenever we receive a progress type that indicate the request has finished being processed
+										// we schedule the next package to be loaded.
+										// LoadNextPackage is thread-safe when using zenloader, so we can run directly from the
+										// callback even if it comes from the ALT thread.
+										case EAsyncLoadingProgress::Failed:
+											[[fallthrough]];
+										case EAsyncLoadingProgress::Canceled:
+											[[fallthrough]];
+										case EAsyncLoadingProgress::FullyLoaded:
+											[[fallthrough]];
+										case EAsyncLoadingProgress::Serialized:
+											// Serialized is an intermediate state that happens just before postloading but it generally
+											// indicate that what's left to be processed needs to happen on the game-thread.
+											// This event is expected to be called directly from ALT, so this is a perfect place to
+											// schedule the next load in our unattended test. 
+											LoadNextPackage();
+											// We expect further progress to happen on the same package, so we set this
+											// to make sure only one package gets scheduled per load request.
+											ExecutionCheck.SetHasBeenExecuted();
+											break;
+										default:
+											checkNoEntry();
+											break;
+									}
+								}))
+					}
+				);
+			}
+			else
+			{
+				UE_LOG(LogTemp, Display, TEXT("Loading finished"));
+				bFinished = true;
+			}
+		}
+	public:
+		FUnattendedLoader(TArray<FName>&& InPackagesToLoad)
+			: PackagesToLoad(MoveTemp(InPackagesToLoad))
+		{
+			LoadNextPackage();
+		}
+
+		bool HasFinished() const
+		{
+			return bFinished;
+		}
+	};
+}
+
+/**
+ * This test demonstrate how to consume a list of asset that needs loading directly from the async loading thread
+ * without involving the game thread by monitoring asset loading progress and starting the next async load once the
+ * asset currently being loaded reach a certain state.
+ */
+IMPLEMENT_SIMPLE_AUTOMATION_TEST(FUnattendedLoadingTest, TEXT(TEST_NAME_ROOT ".UnattendedAsyncLoadingTest"), EAutomationTestFlags::ClientContext | EAutomationTestFlags::EngineFilter)
+bool FUnattendedLoadingTest::RunTest(const FString& Parameters)
+{
+	using namespace UnattendedLoadTestImpl;
+
+	// We use the asset registry to get a list of asset to load. 
+	IAssetRegistry& AssetRegistry = FModuleManager::LoadModuleChecked<FAssetRegistryModule>(FName(TEXT("AssetRegistry"))).Get();
+	AssetRegistry.WaitForCompletion();
+
+	// Limit the number of packages we're going to load for the test in case the project is very big.
+	constexpr int32 MaxPackageCount = 5000;
+
+	TSet<FName> UniquePackages;
+	AssetRegistry.EnumerateAllAssets(
+		[&UniquePackages](const FAssetData& AssetData)
+		{
+			if (UniquePackages.Num() < MaxPackageCount)
+			{
+				UniquePackages.FindOrAdd(AssetData.PackageName);
+				return true;
+			}
+			
+			return false;
+		},
+		true /* bIncludeOnlyOnDiskAssets */
+	);
+
+	// Use the unattended loader to load all unique packages gathered in the list.
+	// The loader will self-feed package to load directly from the callback which
+	// should happen on the ALT when zenloader and async loading thread are both
+	// active.
+	FUnattendedLoader Loader(UniquePackages.Array());
+
+	// For the purpose of this test, we manually tick the GT so it can process postloads
+	// and finalize some package loading steps.
+	while (GetNumAsyncPackages() != 0 || !Loader.HasFinished())
+	{
+		ProcessAsyncLoading(true/*bUseTimeLimit*/, false/*bUseFullTimeLimit*/, 1.0f);
+	}
+	
+	return true;
+}
+
+#undef TEST_NAME_ROOT
+#endif // WITH_DEV_AUTOMATION_TESTS
\ No newline at end of file
