SavePackage: Ignore and warn about private subobjects on imported class default objects.
#rnx

diff --git a/Engine/Source/Runtime/CoreUObject/Private/UObject/SavePackage/PackageHarvester.cpp b/Engine/Source/Runtime/CoreUObject/Private/UObject/SavePackage/PackageHarvester.cpp
--- a/Engine/Source/Runtime/CoreUObject/Private/UObject/SavePackage/PackageHarvester.cpp
+++ b/Engine/Source/Runtime/CoreUObject/Private/UObject/SavePackage/PackageHarvester.cpp
@@ -640,7 +640,7 @@ void FPackageHarvester::ProcessImport(TObjectPtr<UObject> InObject)
 	FName ObjName = InObject.GetFName();
 	if (SaveContext.IsCooking())
 	{
-		// The ignore dependencies check is is necessary not to have infinite recursive calls
+		// The ignore dependencies check is necessary not to have infinite recursive calls
 		if (!bIsNative && !CurrentExportDependencies.bIgnoreDependencies)
 		{
 			UClass* ClassObj = Cast<UClass>(InObject);
@@ -656,7 +656,18 @@ void FPackageHarvester::ProcessImport(TObjectPtr<UObject> InObject)
 				for (UObject* ObjTemplate : ObjectTemplates)
 				{
 					// Recurse into templates
-					*this << ObjTemplate;
+					if (ObjTemplate->HasAnyFlags(RF_Public))
+					{
+						*this << ObjTemplate;
+					}
+					else
+					{
+						// CDO Subobjects are supposed to be public; we rely on that because otherwise they could be garbage collected by
+						// SoftGC during cooking and then not saved out, causing a missing export.
+						UE_LOG(LogSavePackage, Warning,
+							TEXT("Invalid subobject on a CDO; we will skip importing it. Found when saving package %s which imported the CDO containing subobject %s."),
+							*SaveContext.GetPackage()->GetName(), *ObjTemplate->GetPathName());
+					}
 				}
 			}
 		}
