Revert behavior introduced in 12669682, but avoid discarding InstancedSubobjects that have been loaded by FLinkerLoad, this ensures that objects with deferred intitializers have their instanced references fixed up by FObjectInitializer::InstanceSubobjects

diff --git a/Engine/Source/Runtime/CoreUObject/Private/UObject/CoreNative.cpp b/Engine/Source/Runtime/CoreUObject/Private/UObject/CoreNative.cpp
--- a/Engine/Source/Runtime/CoreUObject/Private/UObject/CoreNative.cpp
+++ b/Engine/Source/Runtime/CoreUObject/Private/UObject/CoreNative.cpp
@@ -220,7 +220,7 @@ UObject* FObjectInstancingGraph::GetInstancedSubobject( UObject* SourceSubobject
 								InstancedSubobject = StaticFindObjectFast(nullptr, SubobjectOuter, SubobjectName);
 							}
 
-							if (InstancedSubobject && IsCreatingArchetype())
+							if (InstancedSubobject && IsCreatingArchetype() && !InstancedSubobject->HasAnyFlags(RF_LoadCompleted))
 							{
 								// since we are updating an archetype, this needs to reconstruct as that is the mechanism used to copy properties
 								// it will destroy the existing object and overwrite it
diff --git a/Engine/Source/Runtime/CoreUObject/Private/UObject/UObjectGlobals.cpp b/Engine/Source/Runtime/CoreUObject/Private/UObject/UObjectGlobals.cpp
--- a/Engine/Source/Runtime/CoreUObject/Private/UObject/UObjectGlobals.cpp
+++ b/Engine/Source/Runtime/CoreUObject/Private/UObject/UObjectGlobals.cpp
@@ -3817,11 +3817,7 @@ void FObjectInitializer::PostConstructInit()
 		InitProperties(Obj, BaseClass, Defaults, bCopyTransientsFromClassDefaults);
 	}
 
-#if USE_CIRCULAR_DEPENDENCY_LOAD_DEFERRING
-	const bool bAllowInstancing = IsInstancingAllowed() && !bIsDeferredInitializer;
-#else
 	const bool bAllowInstancing = IsInstancingAllowed();
-#endif // USE_CIRCULAR_DEPENDENCY_LOAD_DEFERRING
 	bool bNeedSubobjectInstancing = InitSubobjectProperties(bAllowInstancing);
 
 	// Restore class information if replacing native class.
