Fix serialization of FP specials with the UE structured archive JSON formatter:
Like JS, JSON does not directly support FP specials in its numeric literal syntax. However, JS has Infinity and NaN global variables, and JSON does not have anything like that.
To work around the omission, this change makes the JSON formatter encode FP specials as a tagged string in the form:
    "Number:+inf"
    "Number:-inf"
    "Number:-nan:0x123"
    "Number:+nan:0xabc"
This change also gets rid of the special case that tried to print floats as integers as it invoked UB to cast an out-of-range float to an int.

diff --git a/Engine/Source/Runtime/CoreUObject/Private/Serialization/Formatters/JsonArchiveInputFormatter.cpp b/Engine/Source/Runtime/CoreUObject/Private/Serialization/Formatters/JsonArchiveInputFormatter.cpp
--- a/Engine/Source/Runtime/CoreUObject/Private/Serialization/Formatters/JsonArchiveInputFormatter.cpp
+++ b/Engine/Source/Runtime/CoreUObject/Private/Serialization/Formatters/JsonArchiveInputFormatter.cpp
@@ -2,15 +2,19 @@
 
 #include "Serialization/Formatters/JsonArchiveInputFormatter.h"
 #include "Serialization/JsonReader.h"
+#include "Containers/StringView.h"
 #include "Dom/JsonObject.h"
-#include "Serialization/JsonSerializer.h"
 #include "Misc/Base64.h"
+#include "Misc/CString.h"
 #include "Misc/SecureHash.h"
+#include "Serialization/JsonSerializer.h"
+#include "Templates/TypeCompatibleBytes.h"
 #include "UObject/SoftObjectPath.h"
 #include "UObject/SoftObjectPtr.h"
 #include "UObject/WeakObjectPtr.h"
 #include "UObject/LazyObjectPtr.h"
 #include "UObject/ObjectPtr.h"
+#include <limits>
 
 #if WITH_TEXT_ARCHIVE_SUPPORT
 
@@ -306,14 +310,102 @@ void FJsonArchiveInputFormatter::Serialize(int64& Value)
 	verify(ValueStack.Top()->TryGetNumber(Value));
 }
 
+namespace
+{
+template<typename FloatType, typename IntType>
+bool ParseFPSpecial(const FString& String, FloatType& OutValue, IntType SignMask, IntType ExponentMask, IntType SignificandMask)
+{
+	FStringView ResidualString{String};
+	
+	// Parse the Number: prefix
+	FStringView SpecialStringPrefix = TEXT("Number:");
+	if (!ResidualString.StartsWith(SpecialStringPrefix))
+	{
+		return false;
+	}
+	ResidualString.MidInline(SpecialStringPrefix.Len());
+
+	// Parse the sign.
+	if (!ResidualString.Len())
+	{
+		return false;
+	}
+	if (ResidualString[0] != TEXT('+') && ResidualString[0] != TEXT('-'))
+	{
+		return false;
+	}
+	const bool bIsNegative = ResidualString[0] == TEXT('-');
+	ResidualString.MidInline(1);
+
+	if(ResidualString == TEXT("inf"))
+	{
+		OutValue = bIsNegative
+			? -std::numeric_limits<FloatType>::infinity()
+			: +std::numeric_limits<FloatType>::infinity();
+		return true;
+	}
+	else if(ResidualString.StartsWith(TEXT("nan:0x")))
+	{
+		ResidualString.MidInline(6);
+		
+		// Assume that the string view is null terminated for the sake of FCString::Strtoui64.
+		check(*ResidualString.end() == 0);
+		
+		// Parse the significand.
+		TCHAR* SignificandEnd = nullptr;
+		uint64 Significand64 = FCString::Strtoui64(ResidualString.begin(), &SignificandEnd, 16);
+		if (SignificandEnd != ResidualString.end()
+			|| Significand64 == 0
+			|| Significand64 > static_cast<uint64>(SignificandMask))
+		{
+			return false;
+		}
+		const IntType Significand = static_cast<IntType>(Significand64);
+
+		// Construct the output NaN from its components.
+		IntType ValueAsInt = ExponentMask | Significand;
+		if (bIsNegative)
+		{
+			ValueAsInt |= SignMask;
+		}
+		OutValue = BitCast<FloatType>(ValueAsInt);
+		return true;
+	}
+	else
+	{
+		return false;
+	}
+}
+}
+
 void FJsonArchiveInputFormatter::Serialize(float& Value)
 {
-	verify(ValueStack.Top()->TryGetNumber(Value));
+	if(!ValueStack.Top()->TryGetNumber(Value))
+	{
+		FString SpecialString;
+		verify(ValueStack.Top()->TryGetString(SpecialString));
+
+		if(!ParseFPSpecial<float, uint32>(SpecialString, Value, 0x80000000, 0x7f800000, 0x007fffff))
+		{
+			Value = 0.0f;
+			checkf(false, TEXT("Expected float special but got string \"%s\""), *SpecialString);
+		}
+	}
 }
 
 void FJsonArchiveInputFormatter::Serialize(double& Value)
 {
-	verify(ValueStack.Top()->TryGetNumber(Value));
+	if(!ValueStack.Top()->TryGetNumber(Value))
+	{
+		FString SpecialString;
+		verify(ValueStack.Top()->TryGetString(SpecialString));
+		
+		if(!ParseFPSpecial<double, uint64>(SpecialString, Value, 0x8000000000000000, 0x7ff0000000000000, 0x000fffffffffffff))
+		{
+			Value = 0.0f;
+			checkf(false, TEXT("Expected double special but got string \"%s\""), *SpecialString);
+		}
+	}
 }
 
 void FJsonArchiveInputFormatter::Serialize(bool& Value)
diff --git a/Engine/Source/Runtime/CoreUObject/Private/Serialization/Formatters/JsonArchiveOutputFormatter.cpp b/Engine/Source/Runtime/CoreUObject/Private/Serialization/Formatters/JsonArchiveOutputFormatter.cpp
--- a/Engine/Source/Runtime/CoreUObject/Private/Serialization/Formatters/JsonArchiveOutputFormatter.cpp
+++ b/Engine/Source/Runtime/CoreUObject/Private/Serialization/Formatters/JsonArchiveOutputFormatter.cpp
@@ -10,6 +10,7 @@
 #include "UObject/SoftObjectPath.h"
 #include "UObject/LazyObjectPtr.h"
 #include "UObject/ObjectPtr.h"
+#include <cinttypes>
 
 #if WITH_TEXT_ARCHIVE_SUPPORT
 
@@ -248,11 +249,7 @@ void FJsonArchiveOutputFormatter::Serialize(int64& Value)
 
 void FJsonArchiveOutputFormatter::Serialize(float& Value)
 {
-	if((float)(int)Value == Value)
-	{
-		WriteValue(LexToString((int)Value));
-	}
-	else
+	if(FPlatformMath::IsFinite(Value))
 	{
 		FString String = FString::Printf(TEXT("%.17g"), Value);
 #if DO_GUARD_SLOW
@@ -262,15 +259,22 @@ void FJsonArchiveOutputFormatter::Serialize(float& Value)
 #endif
 		WriteValue(String);
 	}
+	else if(FPlatformMath::IsNaN(Value))
+	{
+		const uint32 ValueAsInt = BitCast<uint32>(Value);
+		const bool bIsNegative = !!(ValueAsInt & 0x80000000);
+		const uint32 Significand = ValueAsInt & 0x007fffff;
+		WriteValue(FString::Printf(TEXT("\"Number:%snan:0x%" PRIx32 "\""), bIsNegative ? TEXT("-") : TEXT("+"), Significand));
+	}
+	else
+	{
+		WriteValue(Value < 0.0f ? TEXT("\"Number:-inf\"") : TEXT("\"Number:+inf\""));
+	}
 }
 
 void FJsonArchiveOutputFormatter::Serialize(double& Value)
 {
-	if((double)(int)Value == Value)
-	{
-		WriteValue(LexToString((int)Value));
-	}
-	else
+	if(FPlatformMath::IsFinite(Value))
 	{
 		FString String = FString::Printf(TEXT("%.17g"), Value);
 #if DO_GUARD_SLOW
@@ -280,6 +284,17 @@ void FJsonArchiveOutputFormatter::Serialize(double& Value)
 #endif
 		WriteValue(String);
 	}
+	else if(FPlatformMath::IsNaN(Value))
+	{
+		const uint64 ValueAsInt = BitCast<uint64>(Value);
+		const bool bIsNegative = !!(ValueAsInt & 0x8000000000000000);
+		const uint64 Significand = ValueAsInt & 0x000fffffffffffff;
+		WriteValue(FString::Printf(TEXT("\"Number:%snan:0x%" PRIx64 "\""), bIsNegative ? TEXT("-") : TEXT("+"), Significand));
+	}
+	else
+	{
+		WriteValue(Value < 0.0 ? TEXT("\"Number:-inf\"") : TEXT("\"Number:+inf\""));
+	}
 }
 
 void FJsonArchiveOutputFormatter::Serialize(bool& Value)
