Fix for unversioned serialization of arrays of optionals (or other cases where an optional property might be nested in a container).
Also fix a memory leak when doing unversioned serialization of properties with destructors and default values: "loading as zero" doesn't call the destructor on the value being overwritten, and so isn't safe to use for properties that require destruction.

diff --git a/Engine/Source/Runtime/CoreUObject/Private/Serialization/UnversionedPropertySerialization.cpp b/Engine/Source/Runtime/CoreUObject/Private/Serialization/UnversionedPropertySerialization.cpp
--- a/Engine/Source/Runtime/CoreUObject/Private/Serialization/UnversionedPropertySerialization.cpp
+++ b/Engine/Source/Runtime/CoreUObject/Private/Serialization/UnversionedPropertySerialization.cpp
@@ -151,7 +151,7 @@ public:
 
 		if (bIsOptional)
 		{
-			return !static_cast<const FOptionalProperty*>(Property)->IsSet(GetValue(Data));
+			return FastZeroIntNum > 0 && !static_cast<const FOptionalProperty*>(Property)->IsSet(GetValue(Data));
 		}
 		else if (FastZeroIntNum == 1) // Can likely be simplified and faster using a switch() statement like LoadZero()
 		{
@@ -191,7 +191,11 @@ private:
 		{
 			checkf(GetIntNum(Property, IntType) < MaxZeroComparisons || IsOptional(CastFlags),  TEXT("Unexpectedly large property type encountered %s"), *Property->GetName());
 
-			return true;
+			// We can only zero-serialize properties that:
+			// - Don't need a destructor
+			// - And can be zero-initialized
+			// This is because if the property is zero-serialized, loading it will simply memzero the previous value.
+			return Property->HasAllPropertyFlags(CPF_ZeroConstructor | CPF_NoDestructor);
 		}
 		else if ((CastFlags & CASTCLASS_FBoolProperty) != 0)
 		{
diff --git a/Engine/Source/Runtime/CoreUObject/Private/UObject/PropertyOptional.cpp b/Engine/Source/Runtime/CoreUObject/Private/UObject/PropertyOptional.cpp
--- a/Engine/Source/Runtime/CoreUObject/Private/UObject/PropertyOptional.cpp
+++ b/Engine/Source/Runtime/CoreUObject/Private/UObject/PropertyOptional.cpp
@@ -136,8 +136,11 @@ void FOptionalProperty::LinkInternal(FArchive& Ar)
 	// After ValueProperty's size has been computed, compute the size of this property.
 	ElementSize = CalcSize();
 
-	// Propagate CPF_ZeroConstructor, CPF_NoDestructor, CPF_IsPlainOldData, and CPF_HasGetValueTypeHash from the value property.
-	PropertyFlags |= (ValueProperty->PropertyFlags & (CPF_ZeroConstructor|CPF_NoDestructor|CPF_IsPlainOldData|CPF_HasGetValueTypeHash));
+	// Optional properties can always be initialized by zeroing memory.
+	PropertyFlags |= CPF_ZeroConstructor;
+
+	// Propagate CPF_NoDestructor, CPF_IsPlainOldData, and CPF_HasGetValueTypeHash from the value property.
+	PropertyFlags |= (ValueProperty->PropertyFlags & (CPF_NoDestructor|CPF_IsPlainOldData|CPF_HasGetValueTypeHash));
 	if (ValueProperty->ContainsInstancedObjectProperty())
 	{
 		PropertyFlags |= CPF_ContainsInstancedReference;
@@ -174,19 +177,6 @@ void FOptionalProperty::SerializeItem(FStructuredArchive::FSlot Slot, void* Data
 
 	FStructuredArchive::FRecord Record = Slot.EnterRecord();
 
-	if (UnderlyingArchive.UseUnversionedPropertySerialization())
-	{
-		// UPS only calls SerializeItem on set optionals, just need to serialize the inner property
-		checkf(bIsLoading || IsSet(Data), TEXT("UPS should never call SerializeItem to save an empty Optional"));
-
-		FStructuredArchive::FSlot ValueSlot = Record.EnterField(TEXT("Value"));
-		const void* ValueDefaults = Defaults ? GetValuePointerForReadIfSet(Defaults) : nullptr;
-		void* ValueData = bIsLoading ? MarkSetAndGetInitializedValuePointerToReplace(Data) : GetValuePointerForReadOrReplace(Data);
-		GetValueProperty()->SerializeItem(ValueSlot, ValueData, ValueDefaults);
-
-		return;
-	}
-
 	// Use an optional field slot to encode whether the optional was set or not.
 	TOptional<FStructuredArchiveSlot> MaybeValueSlot = Record.TryEnterField(TEXT("Value"), IsSet(Data));
 	if (MaybeValueSlot.IsSet())
@@ -195,7 +185,15 @@ void FOptionalProperty::SerializeItem(FStructuredArchive::FSlot Slot, void* Data
 		
 		const void* ValueDefaults = Defaults ? GetValuePointerForReadIfSet(Defaults) : nullptr;
 
-		if (bIsLoading)
+		if (Slot.GetArchiveState().UseUnversionedPropertySerialization())
+		{
+			// Simply serialize the inner value if using unversioned property serialization.
+			void* ValueData = bIsLoading
+				? MarkSetAndGetInitializedValuePointerToReplace(Data)
+				: GetValuePointerForReadOrReplace(Data);
+			GetValueProperty()->SerializeItem(ValueSlot, ValueData, ValueDefaults);
+		}
+		else if (bIsLoading)
 		{
 			FPropertyTag ValueTag;
 
