Cooker:  Object handle tracking related RW lock and critical sections modified to use C++ atomics to avoid calls into ntdll for the system library implementations of that functionality.  Calls into ntdll are extremely slow on cook farm VMs, running 50x slower than they do on physical machines, becoming a major bottleneck for cook times.  For a shader invalidation cook of a large project, getting rid of the system library RW lock saves 22% of total cook time, while removing the critical section saves another 5%, for 27% overall cook time savings on a VM.  Saves 0.5% of cook time on a physical machine.

#rnx

diff --git a/Engine/Source/Editor/UnrealEd/Private/Cooker/PackageBuildDependencyTracker.cpp b/Engine/Source/Editor/UnrealEd/Private/Cooker/PackageBuildDependencyTracker.cpp
--- a/Engine/Source/Editor/UnrealEd/Private/Cooker/PackageBuildDependencyTracker.cpp
+++ b/Engine/Source/Editor/UnrealEd/Private/Cooker/PackageBuildDependencyTracker.cpp
@@ -2,6 +2,7 @@
 
 #include "PackageBuildDependencyTracker.h"
 
+#include "Async/UniqueLock.h"
 #include "HAL/Platform.h"
 #include "Logging/LogMacros.h"
 #include "Misc/PackageAccessTrackingOps.h"
@@ -36,7 +37,7 @@ void FPackageBuildDependencyTracker::DumpStats() const
 		return;
 	}
 
-	FScopeLock RecordsScopeLock(&RecordsLock);
+	UE::TUniqueLock RecordsScopeLock(RecordsLock);
 	uint64 ReferencingPackageCount = 0;
 	uint64 ReferenceCount = 0;
 	for (const TPair<FName, TSet<FBuildDependencyAccessData>>& PackageAccessRecord : Records)
@@ -66,7 +67,7 @@ void FPackageBuildDependencyTracker::DumpStats() const
 
 TArray<FBuildDependencyAccessData> FPackageBuildDependencyTracker::GetAccessDatas(FName ReferencerPackage) const
 {
-	FScopeLock RecordsScopeLock(&Singleton.RecordsLock);
+	UE::TUniqueLock RecordsScopeLock(Singleton.RecordsLock);
 	const TSet<FBuildDependencyAccessData>* ReferencerSet = Records.Find(ReferencerPackage);
 	if (!ReferencerSet)
 	{
@@ -137,7 +138,7 @@ void FPackageBuildDependencyTracker::StaticOnObjectHandleRead(TArrayView<const U
 		LLM_SCOPE_BYNAME(TEXTVIEW("PackageBuildDependencyTracker"));
 
 		FBuildDependencyAccessData AccessData{ Referenced, AccumulatedScopeData->TargetPlatform };
-		FScopeLock RecordsScopeLock(&Singleton.RecordsLock);
+		UE::TUniqueLock RecordsScopeLock(Singleton.RecordsLock);
 		if (Referencer == Singleton.LastReferencer)
 		{
 			if (AccessData != Singleton.LastAccessData)
diff --git a/Engine/Source/Editor/UnrealEd/Private/Cooker/PackageBuildDependencyTracker.h b/Engine/Source/Editor/UnrealEd/Private/Cooker/PackageBuildDependencyTracker.h
--- a/Engine/Source/Editor/UnrealEd/Private/Cooker/PackageBuildDependencyTracker.h
+++ b/Engine/Source/Editor/UnrealEd/Private/Cooker/PackageBuildDependencyTracker.h
@@ -2,7 +2,7 @@
 
 #pragma once
 
-#include "HAL/CriticalSection.h"
+#include "Async/Mutex.h"
 #include "Misc/PackageAccessTracking.h"
 #include "Templates/UnrealTemplate.h"
 #include "UObject/ObjectHandle.h"
@@ -44,7 +44,11 @@ private:
 	/** Track object reference reads */
 	static void StaticOnObjectHandleRead(TArrayView<const UObject* const> Objects);
 
-	mutable FCriticalSection RecordsLock;
+	// Use a mutex rather than a critical section for synchronization.  Calls into system libraries, such as windows critical section
+	// functions, are 50 times more expensive on build farm VMs, radically affecting cook times, which this avoids.  Saves 5% of total
+	// cook time for shader invalidation on a large project.
+	mutable UE::FMutex RecordsLock;
+
 	TMap<FName, TSet<FBuildDependencyAccessData>> Records;
 	FName LastReferencer = NAME_None;
 	FBuildDependencyAccessData LastAccessData{ NAME_None, nullptr };
diff --git a/Engine/Source/Runtime/CoreUObject/Private/UObject/ObjectHandleTracking.cpp b/Engine/Source/Runtime/CoreUObject/Private/UObject/ObjectHandleTracking.cpp
--- a/Engine/Source/Runtime/CoreUObject/Private/UObject/ObjectHandleTracking.cpp
+++ b/Engine/Source/Runtime/CoreUObject/Private/UObject/ObjectHandleTracking.cpp
@@ -6,6 +6,77 @@
 
 #if UE_WITH_OBJECT_HANDLE_TRACKING
 
+// Simple lock class used to avoid needing to call into ntdll, which is very high overhead on VM cook machines.  Saves
+// 22% of overall cook time for a large project with shader invalidation on the VM.  This class can support nesting read
+// locks but not write locks, but there's no use case where write locks could nest for this code, so it shouldn't be a
+// problem.  Eventually may be replaced with UE::FSharedMutex when that comes online.
+struct FLightweightRWLock
+{
+	static const int32 Unlocked = 0;
+	static const int32 LockedOffset = 0x7fffffff;		// Subtracted for write lock, producing a large negative number
+
+	std::atomic_int32_t Mutex = 0;
+};
+
+class FLightweightReadScopeLock
+{
+public:
+	FLightweightReadScopeLock(FLightweightRWLock& InLock)
+		: Lock(InLock)
+	{
+		int32 LockValue = Lock.Mutex.fetch_add(1, std::memory_order_acquire);
+
+		// If we attempt a read lock while there is a write lock, LockValue will be negative -- spin until it goes positive.
+		// The write lock class sets the mutex value to a large negative value, and then adds that value back in for the
+		// unlock, so any queued read lock increments will already be baked into the mutex value when it goes positive again.
+		// Any threads with pending read locks will then be allowed to finish before another write lock can be obtained from
+		// the FLightweightRWLock::Unlocked = 0 state.
+		// 
+		// This approach avoids the need for the initial modification of the Mutex to do a compare/exchange, making it
+		// "wait free" for multiple readers, versus "lock free", where a thread may need to retry the compare/exchange if
+		// another thread updated the value.
+		while (LockValue < 0)
+		{
+			FPlatformProcess::Yield();
+			LockValue = Lock.Mutex.load(std::memory_order_acquire);
+		}
+	}
+
+	~FLightweightReadScopeLock()
+	{
+		Lock.Mutex.fetch_sub(1, std::memory_order_release);
+	}
+
+private:
+	FLightweightRWLock& Lock;
+};
+
+class FLightweightWriteScopeLock
+{
+public:
+	FLightweightWriteScopeLock(FLightweightRWLock& InLock)
+		: Lock(InLock)
+	{
+		// Spin until we successfully transition from Unlocked to -LockedOffset state
+		int32 ExpectedUnlocked = FLightweightRWLock::Unlocked;
+		while (!InLock.Mutex.compare_exchange_weak(ExpectedUnlocked, -FLightweightRWLock::LockedOffset, std::memory_order_acquire))
+		{
+			FPlatformProcess::Yield();
+			ExpectedUnlocked = FLightweightRWLock::Unlocked;
+		}
+	}
+
+	~FLightweightWriteScopeLock()
+	{
+		// See comment above in FLightweightRWLockReadScope constructor -- we add LockedOffset rather than setting back to Unlocked,
+		// as read lock attempts may have incremented the negative value mutex while the write lock was held.
+		Lock.Mutex.fetch_add(FLightweightRWLock::LockedOffset, std::memory_order_release);
+	}
+
+private:
+	FLightweightRWLock& Lock;
+};
+
 namespace UE::CoreUObject
 {
 	namespace Private
@@ -24,7 +95,7 @@ namespace UE::CoreUObject
 			{
 				UE_AUTORTFM_OPEN(
 				{
-					FReadScopeLock _(HandleLock);
+					FLightweightReadScopeLock _(HandleLock);
 					for (auto&& Pair : ReadHandleCallbacks)
 					{
 						Pair.Value(Objects);
@@ -36,7 +107,7 @@ namespace UE::CoreUObject
 			{
 				UE_AUTORTFM_OPEN(
 				{
-					FReadScopeLock _(HandleLock);
+					FLightweightReadScopeLock _(HandleLock);
 					for (auto&& Pair : ClassResolvedCallbacks)
 					{
 						Pair.Value(ObjectRef, Package, Class);
@@ -48,7 +119,7 @@ namespace UE::CoreUObject
 			{
 				UE_AUTORTFM_OPEN(
 				{
-					FReadScopeLock _(HandleLock);
+					FLightweightReadScopeLock _(HandleLock);
 					for (auto&& Pair : HandleResolvedCallbacks)
 					{
 						Pair.Value(ObjectRef, Package, Object);
@@ -60,7 +131,7 @@ namespace UE::CoreUObject
 			{
 				UE_AUTORTFM_OPEN(
 				{
-					FReadScopeLock _(HandleLock);
+					FLightweightReadScopeLock _(HandleLock);
 					for (auto&& Pair : HandleLoadedCallbacks)
 					{
 						Pair.Value(ObjectRef, Package, Object);
@@ -74,7 +145,7 @@ namespace UE::CoreUObject
 				UE_AUTORTFM_OPEN(
 				{
 					HandleReadCallbackQuantity.fetch_add(1, std::memory_order_release);
-					FWriteScopeLock _(HandleLock);
+					FLightweightWriteScopeLock _(HandleLock);
 					NextHandleId++;
 					ReadHandleCallbacks.Add({ NextHandleId, Func });
 					Result = FObjectHandleTrackingCallbackId{ NextHandleId };
@@ -87,7 +158,7 @@ namespace UE::CoreUObject
 				UE_AUTORTFM_OPEN(
 				{
 					HandleReadCallbackQuantity.fetch_sub(1, std::memory_order_release);
-					FWriteScopeLock _(HandleLock);
+					FLightweightWriteScopeLock _(HandleLock);
 					for (int32 i = ReadHandleCallbacks.Num() - 1; i >= 0; --i)
 					{
 						auto& Pair = ReadHandleCallbacks[i];
@@ -104,7 +175,7 @@ namespace UE::CoreUObject
 				FObjectHandleTrackingCallbackId Result;
 				UE_AUTORTFM_OPEN(
 				{
-					FWriteScopeLock _(HandleLock);
+					FLightweightWriteScopeLock _(HandleLock);
 					NextHandleId++;
 					ClassResolvedCallbacks.Add({ NextHandleId, Func });
 					Result = FObjectHandleTrackingCallbackId{ NextHandleId };
@@ -116,7 +187,7 @@ namespace UE::CoreUObject
 			{
 				UE_AUTORTFM_OPEN(
 				{
-					FWriteScopeLock _(HandleLock);
+					FLightweightWriteScopeLock _(HandleLock);
 					for (int32 i = ClassResolvedCallbacks.Num() - 1; i >= 0; --i)
 					{
 						auto& Pair = ClassResolvedCallbacks[i];
@@ -133,7 +204,7 @@ namespace UE::CoreUObject
 				FObjectHandleTrackingCallbackId Result;
 				UE_AUTORTFM_OPEN(
 				{
-					FWriteScopeLock _(HandleLock);
+					FLightweightWriteScopeLock _(HandleLock);
 					NextHandleId++;
 					HandleResolvedCallbacks.Add({ NextHandleId, Func });
 					Result = FObjectHandleTrackingCallbackId{ NextHandleId };
@@ -145,7 +216,7 @@ namespace UE::CoreUObject
 			{
 				UE_AUTORTFM_OPEN(
 				{
-					FWriteScopeLock _(HandleLock);
+					FLightweightWriteScopeLock _(HandleLock);
 					for (int32 i = HandleResolvedCallbacks.Num() - 1; i >= 0; --i)
 					{
 						auto& Pair = HandleResolvedCallbacks[i];
@@ -162,7 +233,7 @@ namespace UE::CoreUObject
 				FObjectHandleTrackingCallbackId Result;
 				UE_AUTORTFM_OPEN(
 				{
-					FWriteScopeLock _(HandleLock);
+					FLightweightWriteScopeLock _(HandleLock);
 					NextHandleId++;
 					HandleLoadedCallbacks.Add({ NextHandleId, Func });
 					Result = FObjectHandleTrackingCallbackId{ NextHandleId };
@@ -174,7 +245,7 @@ namespace UE::CoreUObject
 			{
 				UE_AUTORTFM_OPEN(
 				{
-					FWriteScopeLock _(HandleLock);
+					FLightweightWriteScopeLock _(HandleLock);
 					for (int32 i = HandleLoadedCallbacks.Num() - 1; i >= 0; --i)
 					{
 						auto& Pair = HandleLoadedCallbacks[i];
@@ -194,7 +265,7 @@ namespace UE::CoreUObject
 			TArray<TTuple<int32, FObjectHandleReferenceLoadedFunc>> HandleLoadedCallbacks;
 
 			//using a single lock as currently there is not any contention. could add separate locks for each list. 
-			FRWLock HandleLock;
+			FLightweightRWLock HandleLock;
 		};
 
 		void OnHandleReadInternal(TArrayView<const UObject* const> Objects)
