Non-unity compile error fix (need PlatformProcess.h include).

#rnx

diff --git a/Engine/Source/Runtime/CoreUObject/Private/UObject/ObjectHandleTracking.cpp b/Engine/Source/Runtime/CoreUObject/Private/UObject/ObjectHandleTracking.cpp
--- a/Engine/Source/Runtime/CoreUObject/Private/UObject/ObjectHandleTracking.cpp
+++ b/Engine/Source/Runtime/CoreUObject/Private/UObject/ObjectHandleTracking.cpp
@@ -1,6 +1,7 @@
 // Copyright Epic Games, Inc. All Rights Reserved.
 
 #include "UObject/ObjectHandleTracking.h"
+#include "Hal/PlatformProcess.h"
 #include "Misc/ScopeRWLock.h"
 #include "AutoRTFM/AutoRTFM.h"
 
