Fixed FObjectPropertyBase::GetExportPath with PPF_Delimited quoting only the object part of the path

diff --git a/Engine/Source/Runtime/CoreUObject/Private/UObject/PropertyBaseObject.cpp b/Engine/Source/Runtime/CoreUObject/Private/UObject/PropertyBaseObject.cpp
--- a/Engine/Source/Runtime/CoreUObject/Private/UObject/PropertyBaseObject.cpp
+++ b/Engine/Source/Runtime/CoreUObject/Private/UObject/PropertyBaseObject.cpp
@@ -227,13 +227,14 @@ FString FObjectPropertyBase::GetExportPath(const TObjectPtr<const UObject>& Obje
 	// Take the path name relative to the stopping point outermost ptr.
 	// This is so that cases like a component referencing a component in another actor work correctly when pasted
 	FString PathName = StopOuter ? Object->GetPathName(StopOuter) : Object.GetPathName();
+	const FTopLevelAssetPath ClassPathName = Object.GetClass()->GetClassPathName();
+	FString ExportPath = GetExportPath(ClassPathName, *PathName);
 	// Object names that contain invalid characters and paths that contain spaces must be put into quotes to be handled correctly
 	if (PortFlags & PPF_Delimited)
 	{
-		PathName = FString::Printf(TEXT("\"%s\""), *PathName.ReplaceQuotesWithEscapedQuotes());
+		ExportPath = FString::Printf(TEXT("\"%s\""), *ExportPath.ReplaceQuotesWithEscapedQuotes());
 	}
-	const FTopLevelAssetPath ClassPathName = Object.GetClass()->GetClassPathName();
-	return GetExportPath(ClassPathName, *PathName);
+	return ExportPath;
 }
 
 void FObjectPropertyBase::ExportText_Internal( FString& ValueStr, const void* PropertyValueOrContainer, EPropertyPointerType PropertyPointerType, const void* DefaultValue, UObject* Parent, int32 PortFlags, UObject* ExportRootScope ) const
