Make sure that reference chain search processors skip weak references from blueprint persistent ubergraph frame

diff --git a/Engine/Source/Runtime/CoreUObject/Private/UObject/ReferenceChainSearch.cpp b/Engine/Source/Runtime/CoreUObject/Private/UObject/ReferenceChainSearch.cpp
--- a/Engine/Source/Runtime/CoreUObject/Private/UObject/ReferenceChainSearch.cpp
+++ b/Engine/Source/Runtime/CoreUObject/Private/UObject/ReferenceChainSearch.cpp
@@ -261,6 +261,7 @@ namespace UE::ReferenceChainSearch
 			}
 		};
 
+		template<bool bNeedsPropertyReferencer = false>
 		class FCollector final: public FReferenceCollector
 		{
 		protected:
@@ -291,7 +292,7 @@ namespace UE::ReferenceChainSearch
 
 			virtual bool NeedsPropertyReferencer() const
 			{
-				return false;
+				return bNeedsPropertyReferencer;
 			}
 			virtual bool IsIgnoringArchetypeRef() const override
 			{
@@ -404,7 +405,7 @@ namespace UE::ReferenceChainSearch
 
 						UE::GC::FWorkerContext Context;
 						Context.SetInitialObjectsUnpadded(ObjectsToSerialize);
-						CollectReferences<UE::GC::TDefaultCollector<FProcessor>>(Processor, Context);
+						CollectReferences<FCollector<>>(Processor, Context);
 					});
 				return;
 			}
@@ -432,7 +433,7 @@ namespace UE::ReferenceChainSearch
 				UE::GC::FWorkerContext Context;
 				ObjectsToProcess = { Object };
 				Context.SetInitialObjectsUnpadded(ObjectsToProcess);
-				CollectReferences<UE::GC::TDefaultCollector<FProcessor>>(Processor, Context);
+				CollectReferences<FCollector<>>(Processor, Context);
 			}
 		}
 
@@ -445,7 +446,7 @@ namespace UE::ReferenceChainSearch
 			TArray<UObject*> ObjectsToProcess;
 			ObjectsToProcess = { FromObject };
 			Context.SetInitialObjectsUnpadded(ObjectsToProcess);
-			CollectReferences<UE::GC::TDefaultCollector<FProcessor>>(Processor, Context);
+			CollectReferences<FCollector<true /* detailed property info */>>(Processor, Context);
 		}
 	};
 
