Fix bug skipping serialization of bulk data with custom element serialization

diff --git a/Engine/Source/Runtime/CoreUObject/Private/Serialization/BulkData.cpp b/Engine/Source/Runtime/CoreUObject/Private/Serialization/BulkData.cpp
--- a/Engine/Source/Runtime/CoreUObject/Private/Serialization/BulkData.cpp
+++ b/Engine/Source/Runtime/CoreUObject/Private/Serialization/BulkData.cpp
@@ -1353,6 +1353,12 @@ void FBulkData::Serialize(FArchive& Ar, UObject* Owner, bool bAttemptFileMapping
 			SerializedMeta.ElementCount = GetBulkDataSize() / ElementSize;
 			SerializedMeta.SizeOnDisk = GetBulkDataSize();
 
+			if (SerializeBulkDataElements != nullptr)
+			{
+				// Force 64 bit precision when using custom element serialization
+				SetBulkDataFlagsOn(SerializedMeta.Flags, static_cast<EBulkDataFlags>(BULKDATA_Size64Bit));
+			}
+
 			const EBulkDataFlags FlagsToClear = static_cast<EBulkDataFlags>(BULKDATA_PayloadAtEndOfFile | BULKDATA_PayloadInSeperateFile | BULKDATA_WorkspaceDomainPayload | BULKDATA_ForceSingleElementSerialization);
 			FBulkData::ClearBulkDataFlagsOn(SerializedMeta.Flags, FlagsToClear);
 
@@ -1366,6 +1372,12 @@ void FBulkData::Serialize(FArchive& Ar, UObject* Owner, bool bAttemptFileMapping
 			}
 			SerializedMeta.Offset = Ar.Tell();
 			SerializedMeta.SizeOnDisk = SerializePayload(Ar, SerializedMeta.Flags, FileRegionTypeOptional);
+
+			if (SerializeBulkDataElements != nullptr)
+			{
+				SerializedMeta.ElementCount = SerializedMeta.SizeOnDisk / ElementSize;
+			}
+
 			{
 				FArchive::FScopeSeekTo _(Ar, MetaOffset);
 				Ar << SerializedMeta;
@@ -1549,7 +1561,11 @@ int64 FBulkData::SerializePayload(FArchive& Ar, EBulkDataFlags SerializationFlag
 
 	const int64 PayloadStart = Ar.Tell();
 
-	if (int64 PayloadSize = GetBulkDataSize(); PayloadSize > 0)
+	if (int64 PayloadSize = GetBulkDataSize(); PayloadSize > 0
+#if !USE_RUNTIME_BULKDATA
+	|| (SerializeBulkDataElements != nullptr)
+#endif
+	)
 	{
 		if (RegionType)
 		{
diff --git a/Engine/Source/Runtime/CoreUObject/Private/UObject/LinkerSave.cpp b/Engine/Source/Runtime/CoreUObject/Private/UObject/LinkerSave.cpp
--- a/Engine/Source/Runtime/CoreUObject/Private/UObject/LinkerSave.cpp
+++ b/Engine/Source/Runtime/CoreUObject/Private/UObject/LinkerSave.cpp
@@ -599,11 +599,17 @@ bool FLinkerSave::SerializeBulkData(FBulkData& BulkData, const FBulkDataSerializ
 
 	const EBulkDataFlags BulkDataFlags	= static_cast<EBulkDataFlags>(BulkData.GetBulkDataFlags());
 	int32 ResourceIndex					= DataResourceMap.Num();
-	const int64 PayloadSize				= BulkData.GetBulkDataSize();
-	TOptional<EFileRegionType> RegionToUse;
+	int64 PayloadSize					= BulkData.GetBulkDataSize();
 	const bool bSupportsMemoryMapping	= IsCooking() && MemoryMappingAlignment >= 0;
 	const bool bSaveAsResourceIndex		= IsCooking();
+
+#if USE_RUNTIME_BULKDATA
+	const bool bCustomElementSerialization = false;
+#else
+	const bool bCustomElementSerialization = BulkData.SerializeBulkDataElements != nullptr;
+#endif
 	
+	TOptional<EFileRegionType> RegionToUse;
 	if (bFileRegionsEnabled)
 	{
 		if (IsCooking())
@@ -618,7 +624,13 @@ bool FLinkerSave::SerializeBulkData(FBulkData& BulkData, const FBulkDataSerializ
 	FBulkMetaResource SerializedMeta;
 	SerializedMeta.Flags = BulkDataFlags;
 	SerializedMeta.ElementCount = PayloadSize / Params.ElementSize;
-	SerializedMeta.SizeOnDisk = PayloadSize; 
+	SerializedMeta.SizeOnDisk = PayloadSize;
+
+	if (bCustomElementSerialization)
+	{
+		// Force 64 bit precision when using custom element serialization
+		FBulkData::SetBulkDataFlagsOn(SerializedMeta.Flags, static_cast<EBulkDataFlags>(BULKDATA_Size64Bit));
+	}
 
 	EBulkDataFlags FlagsToClear = static_cast<EBulkDataFlags>(BULKDATA_PayloadAtEndOfFile | BULKDATA_PayloadInSeperateFile | BULKDATA_WorkspaceDomainPayload | BULKDATA_ForceSingleElementSerialization);
 	if (IsCooking())
@@ -649,6 +661,11 @@ bool FLinkerSave::SerializeBulkData(FBulkData& BulkData, const FBulkDataSerializ
 
 		SerializedMeta.Offset = Tell();
 		SerializedMeta.SizeOnDisk = BulkData.SerializePayload(Ar, SerializedMeta.Flags, RegionToUse);
+		if (bCustomElementSerialization)
+		{
+			PayloadSize = SerializedMeta.SizeOnDisk;
+			SerializedMeta.ElementCount = PayloadSize / Params.ElementSize; 
+		}
 
 		if (bSaveAsResourceIndex == false)
 		{
@@ -712,7 +729,13 @@ bool FLinkerSave::SerializeBulkData(FBulkData& BulkData, const FBulkDataSerializ
 				SerializedMeta.SizeOnDisk = BulkData.SerializePayload(BulkDataAr, SerializedMeta.Flags, RegionToUse);
 			}
 		}
-		
+
+		if (bCustomElementSerialization)
+		{
+			PayloadSize = SerializedMeta.SizeOnDisk;
+			SerializedMeta.ElementCount = PayloadSize / Params.ElementSize; 
+		}
+
 		FArchive& Ar = *this;
 		if (bSaveAsResourceIndex)
 		{
