PackageName: Logging: Downgrade the mountpoint added log message to verbose to make cookcommandlet logs easier to read. Keep the message at Log verbosity for mountpoints added after startup.
#rnx

diff --git a/Engine/Source/Runtime/CoreUObject/Private/Misc/PackageName.cpp b/Engine/Source/Runtime/CoreUObject/Private/Misc/PackageName.cpp
--- a/Engine/Source/Runtime/CoreUObject/Private/Misc/PackageName.cpp
+++ b/Engine/Source/Runtime/CoreUObject/Private/Misc/PackageName.cpp
@@ -302,9 +302,20 @@ struct FLongPackagePathsSingleton
 			ContentPathTree.FindOrAdd(MountPoint->ContentPathAbsolute) = MountPoint;
 		}		
 
-		UE_LOG(LogPackageName, Display, TEXT("FPackageName: Mount point added: '%s' mounted to '%s'"),
+		// Downgrade the log messages on startup to Verbose to reduce startup spam
+#if !NO_LOGGING
+		FString Message = FString::Printf(TEXT("FPackageName: Mount point added: '%s' mounted to '%s'"),
 			*MountPoint->ContentPathRelative, *MountPoint->RootPath);
-		
+		if (IsEngineStartupModuleLoadingComplete())
+		{
+			UE_LOG(LogPackageName, Log, TEXT("%s"), *Message);
+		}
+		else
+		{
+			UE_LOG(LogPackageName, Verbose, TEXT("%s"), *Message);
+		}
+#endif
+
 		// Let subscribers know that a new content path was mounted
 		FPackageName::OnContentPathMounted().Broadcast( RootPath, MountPoint->ContentPathRelative);
 	}
