Add assertions for bad object indices in reference chain search.

diff --git a/Engine/Source/Runtime/CoreUObject/Private/UObject/ReferenceChainSearch.cpp b/Engine/Source/Runtime/CoreUObject/Private/UObject/ReferenceChainSearch.cpp
--- a/Engine/Source/Runtime/CoreUObject/Private/UObject/ReferenceChainSearch.cpp
+++ b/Engine/Source/Runtime/CoreUObject/Private/UObject/ReferenceChainSearch.cpp
@@ -56,7 +56,9 @@ namespace UE::ReferenceChainSearch
 		}
 		FVertex ObjectToVertex(ConstObjectType Object) const
 		{
-			return GUObjectArray.ObjectToIndex(Object);
+			FVertex Vertex = GUObjectArray.ObjectToIndex(Object);
+			UE_CLOG(Vertex < 0 || Vertex >= GetNumVertices(), LogReferenceChain, Fatal, TEXT("Invalid object index in reference chain search %d"), Vertex);
+			return Vertex;
 		}
 		ObjectType VertexToObject(FVertex Vertex) const
 		{
@@ -163,7 +165,9 @@ namespace UE::ReferenceChainSearch
 
 		FVertex ObjectToVertex(ObjectType Object) const
 		{
-			return ObjectInfoToVertex.FindChecked(Object);
+			const FVertex* Vertex = ObjectInfoToVertex.Find(Object);
+			UE_CLOG(!Vertex || *Vertex < 0 || *Vertex >= GetNumVertices(), LogReferenceChain, Fatal, TEXT("Invalid object index in reference chain search %d"), Vertex ? *Vertex : -1);
+			return *Vertex;
 		}
 
 		ObjectType VertexToObject(FVertex Vertex) const
