Update FPackageName::IsValidPath and FPackageName::IsValidObjectPath to use FStringView

diff --git a/Engine/Source/Runtime/CoreUObject/Private/Misc/PackageName.cpp b/Engine/Source/Runtime/CoreUObject/Private/Misc/PackageName.cpp
--- a/Engine/Source/Runtime/CoreUObject/Private/Misc/PackageName.cpp
+++ b/Engine/Source/Runtime/CoreUObject/Private/Misc/PackageName.cpp
@@ -1271,10 +1271,10 @@ bool FPackageName::IsValidLongPackageName(FStringView InLongPackageName, bool bI
 	return false;
 }
 
-bool FPackageName::IsValidObjectPath(const FString& InObjectPath, FText* OutReason)
+bool FPackageName::IsValidObjectPath(FStringView InObjectPath, FText* OutReason)
 {
-	FString PackageName;
-	FString RemainingObjectPath;
+	FStringView PackageName;
+	FStringView RemainingObjectPath;
 
 	// Check for package delimiter
 	int32 ObjectDelimiterIdx;
@@ -1335,7 +1335,7 @@ bool FPackageName::IsValidObjectPath(const FString& InObjectPath, FText* OutReas
 	return true;
 }
 
-bool FPackageName::IsValidPath(const FString& InPath)
+bool FPackageName::IsValidPath(FStringView InPath)
 {
 	const FLongPackagePathsSingleton& Paths = FLongPackagePathsSingleton::Get();
 	FReadScopeLock ScopeLock(Paths.MountLock);
diff --git a/Engine/Source/Runtime/CoreUObject/Public/Misc/PackageName.h b/Engine/Source/Runtime/CoreUObject/Public/Misc/PackageName.h
--- a/Engine/Source/Runtime/CoreUObject/Public/Misc/PackageName.h
+++ b/Engine/Source/Runtime/CoreUObject/Public/Misc/PackageName.h
@@ -224,7 +224,7 @@ public:
 	 * @param OutReason					When returning false, this will provide a description of what was wrong with the name.
 	 * @return							true if a valid object path
 	 */
-	static COREUOBJECT_API bool IsValidObjectPath(const FString& InObjectPath, FText* OutReason = nullptr);
+	static COREUOBJECT_API bool IsValidObjectPath(FStringView InObjectPath, FText* OutReason = nullptr);
 
 	/**
 	 * Returns true if the path starts with a valid root (i.e. /Game/, /Engine/, etc).
@@ -232,7 +232,7 @@ public:
 	 * @param InObjectPath				The object path to test
 	 * @return							true if a valid object path
 	 */
-	static COREUOBJECT_API bool IsValidPath(const FString& InPath);
+	static COREUOBJECT_API bool IsValidPath(FStringView InPath);
 
 	/**
 	 * Checks if the string is a ShortPackageName. A ShortPackageName is the leaf name after the last
