Added an overload of ParseObject that allows the referenced object to be loaded if required

This can also be disabled globally via the CVar "s.AllowParseObjectLoading"

diff --git a/Engine/Source/Runtime/CoreUObject/Private/UObject/UObjectGlobals.cpp b/Engine/Source/Runtime/CoreUObject/Private/UObject/UObjectGlobals.cpp
--- a/Engine/Source/Runtime/CoreUObject/Private/UObject/UObjectGlobals.cpp
+++ b/Engine/Source/Runtime/CoreUObject/Private/UObject/UObjectGlobals.cpp
@@ -84,7 +84,6 @@
 DEFINE_LOG_CATEGORY(LogUObjectGlobals);
 
 int32 GAllowUnversionedContentInEditor = 0;
-
 static FAutoConsoleVariableRef CVarAllowUnversionedContentInEditor(
 	TEXT("s.AllowUnversionedContentInEditor"),
 	GAllowUnversionedContentInEditor,
@@ -92,6 +91,14 @@ static FAutoConsoleVariableRef CVarAllowUnversionedContentInEditor(
 	ECVF_Default
 );
 
+bool GAllowParseObjectLoading = true;
+static FAutoConsoleVariableRef CVarAllowParseObjectLoading(
+	TEXT("s.AllowParseObjectLoading"),
+	GAllowParseObjectLoading,
+	TEXT("If true, allows ParseObject to load fully qualified objects if needed and requested."),
+	ECVF_Default
+);
+
 void EndLoad(FUObjectSerializeContext* LoadContext, TArray<UPackage*>* OutLoadedPackages);
 
 COREUOBJECT_API bool GetAllowNativeComponentClassOverrides()
@@ -1199,58 +1206,76 @@ bool ResolveName(UObject*& InPackage, FString& InOutName, bool Create, bool Thro
 	}
 }
 
-bool ParseObject( const TCHAR* Stream, const TCHAR* Match, UClass* Class, UObject*& DestRes, UObject* InParent, bool* bInvalidObject )
+bool ParseObject( const TCHAR* Stream, const TCHAR* Match, UClass* Class, UObject*& DestRes, UObject* InParent, EParseObjectLoadingPolicy LoadingPolicy, bool* bInvalidObject )
 {
+	if (!GAllowParseObjectLoading)
+	{
+		LoadingPolicy = EParseObjectLoadingPolicy::Find;
+	}
+
 	TCHAR TempStr[1024];
-	if( !FParse::Value( Stream, Match, TempStr, UE_ARRAY_COUNT(TempStr) ) )
+	if (!FParse::Value(Stream, Match, TempStr, UE_ARRAY_COUNT(TempStr)))
 	{
 		// Match not found
-		return 0;
+		return false;
 	}
-	else if( FCString::Stricmp(TempStr,TEXT("NONE"))==0 )
+	else if (FCString::Stricmp(TempStr, TEXT("NONE")) == 0)
 	{
 		// Match found, object explicit set to be None
 		DestRes = nullptr;
-		return 1;
+		return true;
 	}
 	else
 	{
+		auto ResolveObjectImpl = [Class, InParent, LoadingPolicy](const TCHAR* ObjNameOrPathName)
+		{
+			if (FPackageName::IsValidObjectPath(ObjNameOrPathName))
+			{
+				// A fully qualified object path can be resolved with no parent
+				return LoadingPolicy == EParseObjectLoadingPolicy::FindOrLoad
+					? StaticLoadObject(Class, nullptr, ObjNameOrPathName)
+					: StaticFindObject(Class, nullptr, ObjNameOrPathName);
+			}
+			else if (InParent && InParent != ANY_PACKAGE_DEPRECATED)
+			{
+				// Try to find the object within its parent
+				return StaticFindObject(Class, InParent, ObjNameOrPathName);
+			}
+			else
+			{
+				// Try to find first object matching the provided name
+				return StaticFindFirstObject(Class, ObjNameOrPathName, EFindFirstObjectOptions::EnsureIfAmbiguous);
+			}
+		};
+
 		UObject* Res = nullptr;
 		// Look this object up.
-		if ((!InParent || InParent == ANY_PACKAGE_DEPRECATED) && !FPackageName::IsValidObjectPath(TempStr))
-		{
-			// Try to find first object matching the provided name
-			Res = StaticFindFirstObject(Class, TempStr, EFindFirstObjectOptions::EnsureIfAmbiguous);
-		}
-		else
-		{
-			Res = StaticFindObject(Class, InParent, TempStr);
-		}
-		if( !Res )
+		Res = ResolveObjectImpl(TempStr);
+		if (!Res)
 		{
 			if (Class->IsChildOf<UClass>())
 			{
 				FString RedirectedObjectName = FLinkerLoad::FindNewPathNameForClass(TempStr, false);
 				if (!RedirectedObjectName.IsEmpty())
 				{
-					Res = StaticFindObject(Class, InParent, *RedirectedObjectName);
+					Res = ResolveObjectImpl(*RedirectedObjectName);
 				}
 			}
 
 			if (!Res)
-			{ 
+			{
 				// Match found, object not found
 				if (bInvalidObject)
 				{
 					*bInvalidObject = true;
 				}
-				return 0;
+				return false;
 			}
 		}
 
 		// Match found, object found
 		DestRes = Res;
-		return 1;
+		return true;
 	}
 }
 
diff --git a/Engine/Source/Runtime/CoreUObject/Public/UObject/UObjectGlobals.h b/Engine/Source/Runtime/CoreUObject/Public/UObject/UObjectGlobals.h
--- a/Engine/Source/Runtime/CoreUObject/Public/UObject/UObjectGlobals.h
+++ b/Engine/Source/Runtime/CoreUObject/Public/UObject/UObjectGlobals.h
@@ -501,19 +501,33 @@ COREUOBJECT_API UObject* StaticFindFirstObject(UClass* Class, const TCHAR* Name,
  */
 COREUOBJECT_API UObject* StaticFindFirstObjectSafe(UClass* Class, const TCHAR* Name, EFindFirstObjectOptions Options = EFindFirstObjectOptions::None, ELogVerbosity::Type AmbiguousMessageVerbosity = ELogVerbosity::NoLogging, const TCHAR* InCurrentOperation = nullptr);
 
+/** Loading policy to use with ParseObject */
+enum class EParseObjectLoadingPolicy : uint8
+{
+	/** Try and find the object, but do not attempt to load it */
+	Find,
+	/** Try and find the object, or attempt of load it if it cannot be found (note: loading may be globally disabled via the CVar "s.AllowParseObjectLoading") */
+	FindOrLoad,
+};
+
 /**
  * Parse a reference to an object from a text representation
  *
- * @param Stream		String containing text to parse
+ * @param Stream			String containing text to parse
  * @param Match				Tag to search for object representation within string
- * @param Class			The class of the object to be loaded.
+ * @param Class				The class of the object to be found.
  * @param DestRes			Returned object pointer
- * @param InParent		Outer to search
+ * @param InParent			Outer to search
+ * @oaran LoadingPolicy		Controls whether the parse will attempt to load a fully qualified object reference, if needed.
  * @param bInvalidObject	[opt] Optional output.  If true, Tag was matched but the specified object wasn't found.
  *
  * @return True if the object parsed successfully, even if object was not found
  */
-COREUOBJECT_API bool ParseObject( const TCHAR* Stream, const TCHAR* Match, UClass* Class, UObject*& DestRes, UObject* InParent, bool* bInvalidObject=nullptr );
+COREUOBJECT_API bool ParseObject( const TCHAR* Stream, const TCHAR* Match, UClass* Class, UObject*& DestRes, UObject* InParent, EParseObjectLoadingPolicy LoadingPolicy, bool* bInvalidObject=nullptr );
+inline bool ParseObject( const TCHAR* Stream, const TCHAR* Match, UClass* Class, UObject*& DestRes, UObject* InParent, bool* bInvalidObject=nullptr )
+{
+	return ParseObject( Stream, Match, Class, DestRes, InParent, EParseObjectLoadingPolicy::Find, bInvalidObject );
+}
 
 /**
  * Find or load an object by string name with optional outer and filename specifications.
@@ -1792,9 +1806,14 @@ COREUOBJECT_API FString GetConfigFilename( UObject* SourceObject );
 
 /** Parse a reference to an object from the input stream. */
 template< class T > 
+inline bool ParseObject( const TCHAR* Stream, const TCHAR* Match, T*& Obj, UObject* Outer, EParseObjectLoadingPolicy LoadingPolicy, bool* bInvalidObject=nullptr )
+{
+	return ParseObject( Stream, Match, T::StaticClass(), (UObject*&)Obj, Outer, LoadingPolicy, bInvalidObject );
+}
+template< class T > 
 inline bool ParseObject( const TCHAR* Stream, const TCHAR* Match, T*& Obj, UObject* Outer, bool* bInvalidObject=nullptr )
 {
-	return ParseObject( Stream, Match, T::StaticClass(), (UObject*&)Obj, Outer, bInvalidObject );
+	return ParseObject( Stream, Match, T::StaticClass(), (UObject*&)Obj, Outer, EParseObjectLoadingPolicy::Find, bInvalidObject );
 }
 
 /** 
