Incremental Reachability Analysis fixes:

- FastReferenceCollector will now flush struct references before suspending to avoid storing pointers to struct owned UObject references between reachability iterations
- Persistent frame weak reference clearing will no longer attempt to clear the stored reference pointer when running incremental reachability and instead will re-run its owner's AddReferencedObjects with a special collector to clear any references to unreachable objects
- Objects marked by GC barrier will now be processed immediately in the next reachability iteration instead of after reachability analsysis was complete

diff --git a/Engine/Plugins/Runtime/DataRegistry/Source/DataRegistry/Private/DataRegistrySubsystem.cpp b/Engine/Plugins/Runtime/DataRegistry/Source/DataRegistry/Private/DataRegistrySubsystem.cpp
--- a/Engine/Plugins/Runtime/DataRegistry/Source/DataRegistry/Private/DataRegistrySubsystem.cpp
+++ b/Engine/Plugins/Runtime/DataRegistry/Source/DataRegistry/Private/DataRegistrySubsystem.cpp
@@ -904,7 +904,7 @@ void UDataRegistrySubsystem::AddReferencedObjects(UObject* InThis, FReferenceCol
 		for (FRegistryMapPair& RegistryPair : This->RegistryMap)
 		{
 			// In editor builds we mark this as a weak reference so it can be deleted properly
-			Collector.MarkWeakObjectReferenceForClearing(reinterpret_cast<UObject**>(&RegistryPair.Value));
+			Collector.MarkWeakObjectReferenceForClearing(reinterpret_cast<UObject**>(&RegistryPair.Value), This);
 		}
 	}
 	else
diff --git a/Engine/Source/Runtime/CoreUObject/Private/UObject/GarbageCollection.cpp b/Engine/Source/Runtime/CoreUObject/Private/UObject/GarbageCollection.cpp
--- a/Engine/Source/Runtime/CoreUObject/Private/UObject/GarbageCollection.cpp
+++ b/Engine/Source/Runtime/CoreUObject/Private/UObject/GarbageCollection.cpp
@@ -3159,9 +3159,9 @@ public:
 		Dispatcher.QueueSet(Dispatcher.Context.GetReferencingObject(), UE::Core::Private::Unsafe::Decay(*Objects), ETokenlessId::Collector, MayKill());
 	}
 	
-	virtual bool MarkWeakObjectReferenceForClearing(UObject** WeakReference) override
+	virtual bool MarkWeakObjectReferenceForClearing(UObject** WeakReference, UObject* ReferenceOwner) override
 	{
-		Dispatcher.Context.WeakReferences.Add(WeakReference);
+		Dispatcher.Context.WeakReferences.Add({ *WeakReference, WeakReference, ReferenceOwner });
 		return true;
 	}
 };
@@ -3376,9 +3376,9 @@ public:
 		}
 	}
 
-	virtual bool MarkWeakObjectReferenceForClearing(UObject** WeakReference) override
+	virtual bool MarkWeakObjectReferenceForClearing(UObject** WeakReference, UObject* ReferenceOwner) override
 	{
-		Context.WeakReferences.Add(WeakReference);
+		Context.WeakReferences.Add({ *WeakReference, WeakReference, ReferenceOwner });
 		return true;
 	}
 };
@@ -3913,26 +3913,26 @@ private:
 		if (GReachabilityState.IsSuspended())
 		{
 			Context = GReachabilityState.GetContextArray()[0];
+			InitialObjects.Reset();
 		}
 		else
 		{
 			Context = Pool.AllocateFromPool();
+		}
 
-			if (!Private::GReachableObjects.IsEmpty())
-			{
-				InitialObjects.Reset();
-				Private::GReachableObjects.PopAllAndEmpty(InitialObjects);
-				UE_LOG(LogGarbage, Log, TEXT("GC restarting reachability due to Reachable Objects list not being empty (%d objects to process)"), InitialObjects.Num());
-				ConditionallyAddBarrierReferencesToHistory(*Context);
-			}
-			else
-			{
-				Context->InitialNativeReferences = GetInitialReferences(Options);
-			}
-
-			Context->SetInitialObjectsUnpadded(InitialObjects);
+		if (!Private::GReachableObjects.IsEmpty())
+		{
+			Private::GReachableObjects.PopAllAndEmpty(InitialObjects);
+			UE_LOG(LogGarbage, Verbose, TEXT("Adding %d object(s) marker by GC barrier to the list of objects to process"), InitialObjects.Num());
+			ConditionallyAddBarrierReferencesToHistory(*Context);
+		}
+		else if (GReachabilityState.GetNumIterations() == 0 || (Stats.bFoundGarbageRef && !GReachabilityState.IsSuspended()))
+		{
+			Context->InitialNativeReferences = GetInitialReferences(Options);
 		}
 
+		Context->SetInitialObjectsUnpadded(InitialObjects);
+
 		PerformReachabilityAnalysisOnObjects(Context, Options);
 
 		if (!GReachabilityState.CheckIfAnyContextIsSuspended())
@@ -4576,21 +4576,68 @@ void GatherUnreachableObjects(bool bForceSingleThreaded)
 namespace UE::GC
 {
 
+class FWeakReferenceEliminator final : public TReachabilityCollectorBase<EGCOptions::None>
+{
+
+public:
+	FWeakReferenceEliminator() = default;
+		
+	virtual void HandleObjectReference(UObject*& InObject, const UObject* InReferencingObject, const FProperty* InReferencingProperty) override 
+	{
+		if (InObject && InObject->IsUnreachable())
+		{
+			InObject = nullptr;
+		}
+	}
+	
+	virtual bool MarkWeakObjectReferenceForClearing(UObject** WeakReference, UObject* ReferenceOwner) override
+	{
+		if (WeakReference && *WeakReference && (*WeakReference)->IsUnreachable())
+		{
+			*WeakReference = nullptr;
+		}
+		return true;
+	}
+};
+
+template <bool bGatheredWithIncrementalReachability = true>
 static void ClearWeakReferences(TConstArrayView<TUniquePtr<FWorkerContext>> Contexts)
 {
 	TRACE_CPUPROFILER_EVENT_SCOPE(ClearWeakReferences);
+	TSet<UObject*> ObjectsThatNeedWeakReferenceClearing;
 	for (const TUniquePtr<FWorkerContext>& Context : Contexts)
 	{
-		for (UObject** WeakReference : Context->WeakReferences)
+		for (FWeakReferenceInfo& ReferenceInfo : Context->WeakReferences)
 		{
-			UObject*& ReferencedObject = *WeakReference;
+			UObject* ReferencedObject = ReferenceInfo.ReferencedObject;
 			if (ReferencedObject && ReferencedObject->IsUnreachable())
 			{
-				ReferencedObject = nullptr;
+				if constexpr (bGatheredWithIncrementalReachability)
+				{
+					// When running incremental reachability we can't assume the Reference pointer is still valid
+					// so instead collect all referencing objects and run AddReferencedObjects on them again with
+					// a special collector that will null out references to unrachable objects
+					ObjectsThatNeedWeakReferenceClearing.Add(ReferenceInfo.ReferenceOwner);
+				}
+				else
+				{
+					*ReferenceInfo.Reference = nullptr;
+				}
 			}
 		}
 		Context->WeakReferences.Reset();
 	}
+	if constexpr (bGatheredWithIncrementalReachability)
+	{
+		FWeakReferenceEliminator ReferenceEliminator;
+		for (UObject* Object : ObjectsThatNeedWeakReferenceClearing)
+		{
+			if (Object && !Object->IsUnreachable())
+			{
+				Object->GetClass()->CallAddReferencedObjects(Object, ReferenceEliminator);
+			}
+		}
+	}
 }
 
 static void DumpGarbageReferencers(TConstArrayView<TUniquePtr<FWorkerContext>> Contexts)
@@ -4859,7 +4906,14 @@ void PostCollectGarbageImpl(EObjectFlags KeepFlags)
 		GatherUnreachableObjects(!(GetReferenceCollectorOptions(bPerformFullPurge) & EGCOptions::Parallel));
 
 		// This needs to happen after GatherUnreachableObjects since it can mark more objects as unreachable
-		ClearWeakReferences(AllContexts);
+		if (GReachabilityState.GetNumIterations() > 1)
+		{
+			ClearWeakReferences<true>(AllContexts);
+		}
+		else
+		{
+			ClearWeakReferences<false>(AllContexts);
+		}
 
 		if (bPerformFullPurge)
 		{
@@ -5079,6 +5133,11 @@ void FReachabilityAnalysisState::PerformReachabilityAnalysisAndConditionallyPurg
 
 void FReachabilityAnalysisState::PerformReachabilityAnalysis()
 {
+	if (!bIsSuspended)
+	{
+		Init();
+	}
+
 	if (bPerformFullPurge)
 	{
 		UE::GC::CollectGarbageFull(ObjectKeepFlags);
@@ -5087,6 +5146,8 @@ void FReachabilityAnalysisState::PerformReachabilityAnalysis()
 	{
 		UE::GC::CollectGarbageIncremental(ObjectKeepFlags);
 	}
+
+	FinishIteration();
 }
 
 } // namespace UE::GC
@@ -5894,43 +5955,58 @@ ELoot StealWork(FWorkerContext& Context, FReferenceCollector& Collector, FWorkBl
 
 TArrayView<FWorkerContext*> InitializeAsyncProcessingContexts(FWorkerContext& InContext)
 {
-	FContextPoolScope ContextPool;
+	TArrayView<FWorkerContext*> Contexts;
+
+	if (!GReachabilityState.IsSuspended())
+	{
+		FContextPoolScope ContextPool;
 
-	checkf(InContext.ObjectsToSerialize.IsUnused(), TEXT("Use InitialObjects instead, ObjectsToSerialize.Add() may only be called during reference processing"));
-	check(InContext.Stats.NumObjects == 0 && InContext.Stats.NumReferences == 0 && InContext.Stats.bFoundGarbageRef == false);
+		checkf(InContext.ObjectsToSerialize.IsUnused(), TEXT("Use InitialObjects instead, ObjectsToSerialize.Add() may only be called during reference processing"));
+		check(InContext.Stats.NumObjects == 0 && InContext.Stats.NumReferences == 0 && InContext.Stats.bFoundGarbageRef == false);
 
-	TConstArrayView<UObject*> InitialObjects = InContext.GetInitialObjects();
-	TConstArrayView<UObject**> InitialReferences = InContext.InitialNativeReferences;
+		const int32 NumTaskgraphWorkers = FTaskGraphInterface::Get().GetNumWorkerThreads();
+		const int32 NumWorkers = FMath::Clamp(NumTaskgraphWorkers, 1, MaxWorkers);
 
-	int32 NumTaskgraphWorkers = FTaskGraphInterface::Get().GetNumWorkerThreads();
-	const int32 NumWorkers = FMath::Clamp(NumTaskgraphWorkers, 1, MaxWorkers);
-	const int32 ObjPerWorker = (InitialObjects.Num() + NumWorkers - 1) / NumWorkers;
-	const int32 RefPerWorker = (InitialReferences.Num() + NumWorkers - 1) / NumWorkers;
+		GReachabilityState.SetupWorkers(NumWorkers);
 
-	GReachabilityState.InitReachabilityAnalysis(NumWorkers);
+		// Allocate contexts	
+		checkf(ContextPool.NumAllocated() == 1, TEXT("Other contexts forbidden during parallel reference collection. Work-stealing from all live contexts. "));
 
-	// Allocate contexts	
-	checkf(ContextPool.NumAllocated() == 1, TEXT("Other contexts forbidden during parallel reference collection. Work-stealing from all live contexts. "));
+		Contexts = MakeArrayView(GReachabilityState.GetContextArray(), GReachabilityState.GetNumWorkers());
+		Contexts[0] = &InContext;
+		for (FWorkerContext*& Context : Contexts.RightChop(1))
+		{
+			Context = ContextPool.AllocateFromPool();
+		}
 
-	TArrayView<FWorkerContext*> Contexts = MakeArrayView(GReachabilityState.GetContextArray(), GReachabilityState.GetNumWorkers());
-	Contexts[0] = &InContext;
-	for (FWorkerContext*& Context : Contexts.RightChop(1))
+		GSlowARO.GetPostInit().SetupWorkerQueues(NumWorkers);
+
+		// Setup work-stealing queues
+		for (FWorkerContext* Context : Contexts)
+		{
+			const int32 Idx = Context->GetWorkerIndex();
+			check(Idx >= 0 && Idx < NumWorkers);
+			Context->ObjectsToSerialize.SetAsyncQueue(GWorkstealingManager.Queues[Idx]);
+			checkf(!Context->IncrementalStructs.ContainsBatchData(), TEXT("Reachability analysis is done but worker context holds suspended dispatcher state"));
+		}
+	}
+	else
 	{
-		Context = ContextPool.AllocateFromPool();
+		Contexts = MakeArrayView(GReachabilityState.GetContextArray(), GReachabilityState.GetNumWorkers());
 	}
 
-	GSlowARO.GetPostInit().SetupWorkerQueues(NumWorkers);
+	TConstArrayView<UObject*> InitialObjects = InContext.GetInitialObjects();
+	TConstArrayView<UObject**> InitialReferences = InContext.InitialNativeReferences;
+	const int32 ObjPerWorker = (InitialObjects.Num() + GReachabilityState.GetNumWorkers() - 1) / GReachabilityState.GetNumWorkers();
+	const int32 RefPerWorker = (InitialReferences.Num() + GReachabilityState.GetNumWorkers() - 1) / GReachabilityState.GetNumWorkers();
 
-	// Setup work-stealing queues and distribute initial workload across worker contexts
+	// Distribute initial workload across worker contexts
 	for (FWorkerContext* Context : Contexts)
 	{
-		int32 Idx = Context->GetWorkerIndex();
-		check(Idx >= 0 && Idx < NumWorkers);
-		Context->ObjectsToSerialize.SetAsyncQueue(GWorkstealingManager.Queues[Idx]);
+		const int32 Idx = Context->GetWorkerIndex();
 		// Initial objects is already padded at the end and its safe to prefetch in the middle too
 		Context->SetInitialObjectsPrepadded(InitialObjects.Mid(Idx * ObjPerWorker, ObjPerWorker));
 		Context->InitialNativeReferences = InitialReferences.Mid(Idx * RefPerWorker, RefPerWorker);
-		checkf(!Context->IncrementalStructs.ContainsBatchData(), TEXT("Reachability analysis is done but worker context holds suspended dispatcher state"));
 	}
 
 	return Contexts;
@@ -5956,27 +6032,19 @@ void ReleaseAsyncProcessingContexts(FWorkerContext& InContext, TArrayView<FWorke
 		ContextPool.ReturnToPool(Context);
 	}
 
-	GReachabilityState.FinishReachabilityAnalysis();
+	GReachabilityState.ResetWorkers();
 }
 
 void ProcessAsync(void (*ProcessSync)(void*, FWorkerContext&), void* Processor, FWorkerContext& InContext)
 {
 	using namespace UE::GC;
 
-	TArrayView<FWorkerContext*> Contexts;
-	if (!GReachabilityState.IsSuspended())
-	{
-		Contexts = InitializeAsyncProcessingContexts(InContext);
-	}
-	else
-	{
-		Contexts = MakeArrayView(GReachabilityState.GetContextArray(), GReachabilityState.GetNumWorkers());
-	}
-	
+	TArrayView<FWorkerContext*> Contexts = InitializeAsyncProcessingContexts(InContext);
+
 	TSharedRef<FWorkCoordinator> WorkCoordinator = MakeShared<FWorkCoordinator>(Contexts, FTaskGraphInterface::Get().GetNumWorkerThreads());
 	for (FWorkerContext* Context : Contexts)
 	{
-		Context->bIsSuspended = false;
+		Context->bDidWork = false;
 		Context->Coordinator = &WorkCoordinator.Get();
 	}
 
@@ -6001,12 +6069,37 @@ void ProcessAsync(void (*ProcessSync)(void*, FWorkerContext&), void* Processor,
 	// Wait until all work is complete. Current thread can steal and complete everything 
 	// alone if task workers are busy with long-running tasks.
 	WorkCoordinator->SpinUntilAllStopped();
-
-	// Reset initial object sets so that we don't process again them in the next iteration or when we're done processing references
+	
 	for (FWorkerContext* Context : Contexts)
 	{
+		// Reset initial object sets so that we don't process again them in the next iteration or when we're done processing references
 		Context->ResetInitialObjects();
 		Context->InitialNativeReferences = TConstArrayView<UObject**>();
+
+		// Update contexts' suspended state. This is necessary because some context may have their work stolen and never started (resumed) work.
+		if (Context->ObjectsToSerialize.HasWork())
+		{
+			if (!Context->bDidWork)
+			{
+				if (GReachabilityState.IsTimeLimitExceeded())
+				{
+					Context->bIsSuspended = true;
+				}
+				else
+				{
+					// Rare case where this context's work has been completely dropped because incremental reachability does not support context stealing
+					WorkCoordinator->TryStartWorking(Context->GetWorkerIndex());
+					ProcessSync(Processor, *Context);
+					checkf(!Context->ObjectsToSerialize.HasWork(), TEXT("GC Context %d was processed but it stil has unfinished work"), Context->GetWorkerIndex());
+				}
+			}
+			check(GReachabilityState.IsTimeLimitExceeded() || !Context->ObjectsToSerialize.HasWork());
+		}
+		else if (!GReachabilityState.IsTimeLimitExceeded()) // !WithTimeLimit
+		{
+			// This context's work has been stolen by another context so it never had a chance to spin and clear its bIsSuspended flag
+			Context->bIsSuspended = false;
+		}
 	}
 
 	if (!GReachabilityState.CheckIfAnyContextIsSuspended())
diff --git a/Engine/Source/Runtime/CoreUObject/Private/UObject/ReachabilityAnalysisState.cpp b/Engine/Source/Runtime/CoreUObject/Private/UObject/ReachabilityAnalysisState.cpp
--- a/Engine/Source/Runtime/CoreUObject/Private/UObject/ReachabilityAnalysisState.cpp
+++ b/Engine/Source/Runtime/CoreUObject/Private/UObject/ReachabilityAnalysisState.cpp
@@ -13,7 +13,12 @@
 namespace UE::GC
 {
 
-void FReachabilityAnalysisState::InitReachabilityAnalysis(int32 InNumWorkers)
+void FReachabilityAnalysisState::Init()
+{
+	NumIterations = 0;
+}
+
+void FReachabilityAnalysisState::SetupWorkers(int32 InNumWorkers)
 {
 	NumWorkers = InNumWorkers;
 	Stats = {};
@@ -25,11 +30,16 @@ void FReachabilityAnalysisState::UpdateStats(const FProcessorStats& InStats)
 	Stats.AddStats(InStats);
 }
 
-void FReachabilityAnalysisState::FinishReachabilityAnalysis()
+void FReachabilityAnalysisState::ResetWorkers()
 {
 	NumWorkers = 0;
 }
 
+void FReachabilityAnalysisState::FinishIteration()
+{
+	NumIterations++;
+}
+
 bool FReachabilityAnalysisState::CheckIfAnyContextIsSuspended()
 {
 	bIsSuspended = false;
diff --git a/Engine/Source/Runtime/CoreUObject/Private/UObject/ReachabilityAnalysisState.h b/Engine/Source/Runtime/CoreUObject/Private/UObject/ReachabilityAnalysisState.h
--- a/Engine/Source/Runtime/CoreUObject/Private/UObject/ReachabilityAnalysisState.h
+++ b/Engine/Source/Runtime/CoreUObject/Private/UObject/ReachabilityAnalysisState.h
@@ -46,6 +46,9 @@ private:
 	/** Total time of the actual reference traversal */
 	double ReferenceProcessingTotalTime = 0.0;
 
+	/** Number of reachability analysis iterations performed during reachability analysis */
+	int32 NumIterations = 0;
+
 	alignas (PLATFORM_CACHE_LINE_SIZE) double IterationStartTime = 0.0;
 	alignas (PLATFORM_CACHE_LINE_SIZE) double IterationTimeLimit = 0.0;
 
@@ -95,8 +98,16 @@ public:
 		return Contexts;
 	}
 
-	/** Initializes internal variables and sets the expected number of worker threads */
-	void InitReachabilityAnalysis(int32 InNumWorkers);
+	FORCEINLINE int32 GetNumIterations() const
+	{
+		return NumIterations;
+	}
+
+	/** Initializes reachability analysis */
+	void Init();
+
+	/** Initializes expected number of worker threads */
+	void SetupWorkers(int32 InNumWorkers);
 
 	/** Main Garbage Collection function executed on Reachability Analysis Thread when StartEvent has been triggered */
 	void PerformReachabilityAnalysis();
@@ -108,7 +119,10 @@ public:
 	bool CheckIfAnyContextIsSuspended();
 
 	/** Resets workers after reachability analysis is fully complete */
-	void FinishReachabilityAnalysis();
+	void ResetWorkers();
+
+	/** Marks the end of reachability iteration */
+	void FinishIteration();
 };
 
 } // namespace UE::GC
\ No newline at end of file
diff --git a/Engine/Source/Runtime/CoreUObject/Private/UObject/ReferenceChainSearch.cpp b/Engine/Source/Runtime/CoreUObject/Private/UObject/ReferenceChainSearch.cpp
--- a/Engine/Source/Runtime/CoreUObject/Private/UObject/ReferenceChainSearch.cpp
+++ b/Engine/Source/Runtime/CoreUObject/Private/UObject/ReferenceChainSearch.cpp
@@ -306,7 +306,7 @@ namespace UE::ReferenceChainSearch
 			{
 				return false;
 			}
-			virtual bool MarkWeakObjectReferenceForClearing(UObject** WeakReference)
+			virtual bool MarkWeakObjectReferenceForClearing(UObject** WeakReference, UObject* ReferenceOwner)
 			{
 				// To avoid false positives we need to implement this method just like GC does
 				// as these references will be treated as weak and should not be reported
diff --git a/Engine/Source/Runtime/CoreUObject/Private/UObject/ReferencerFinder.cpp b/Engine/Source/Runtime/CoreUObject/Private/UObject/ReferencerFinder.cpp
--- a/Engine/Source/Runtime/CoreUObject/Private/UObject/ReferencerFinder.cpp
+++ b/Engine/Source/Runtime/CoreUObject/Private/UObject/ReferencerFinder.cpp
@@ -43,7 +43,7 @@ public:
 		}
 	}
 
-	bool MarkWeakObjectReferenceForClearing(UObject** WeakReference)
+	bool MarkWeakObjectReferenceForClearing(UObject** WeakReference, UObject* ReferenceOwner)
 	{
 		return (Flags & EReferencerFinderFlags::SkipWeakReferences) == EReferencerFinderFlags::SkipWeakReferences;
 	}
@@ -55,9 +55,9 @@ class FAllReferencesCollector : public UE::GC::TDefaultCollector<FAllReferencesP
 public:
 	using Super::Super;
 
-	virtual bool MarkWeakObjectReferenceForClearing(UObject** WeakReference) override
+	virtual bool MarkWeakObjectReferenceForClearing(UObject** WeakReference, UObject* ReferenceOwner) override
 	{
-		return Processor.MarkWeakObjectReferenceForClearing(WeakReference);
+		return Processor.MarkWeakObjectReferenceForClearing(WeakReference, ReferenceOwner);
 	}
 };
 
diff --git a/Engine/Source/Runtime/CoreUObject/Public/UObject/FastReferenceCollector.h b/Engine/Source/Runtime/CoreUObject/Public/UObject/FastReferenceCollector.h
--- a/Engine/Source/Runtime/CoreUObject/Public/UObject/FastReferenceCollector.h
+++ b/Engine/Source/Runtime/CoreUObject/Public/UObject/FastReferenceCollector.h
@@ -179,7 +179,11 @@ public:
 	void SetWorkerIndex(int32 Idx) { WorkerIndex = Idx; }
 	int32 GetWorkerIndex() const { return WorkerIndex; }
 
-	
+	bool HasWork() const
+	{
+		return PartialNum() != 0;
+	}
+
 private:
 	UObject** WipIt; // Wip->Objects cursor
 	FWorkBlock* Wip;
@@ -262,6 +266,12 @@ struct FSuspendedStructBatch
 	}
 };
 
+struct FWeakReferenceInfo
+{
+	UObject* ReferencedObject = nullptr;
+	UObject** Reference = nullptr;
+	UObject* ReferenceOwner = nullptr;
+};
 
 /** Thread-local context containing initial objects and references to collect */
 struct alignas(PLATFORM_CACHE_LINE_SIZE) FWorkerContext
@@ -282,7 +292,7 @@ public:
 	FWorkBlockifier ObjectsToSerialize;
 	TConstArrayView<UObject**> InitialNativeReferences;
 	FWorkCoordinator* Coordinator = nullptr;
-	TArray<UObject**> WeakReferences;
+	TArray<FWeakReferenceInfo> WeakReferences;
 	FProcessorStats Stats;
 
 #if !UE_BUILD_SHIPPING
@@ -294,6 +304,7 @@ public:
 
 	FSuspendedStructBatch IncrementalStructs;
 	bool bIsSuspended = false;
+	bool bDidWork = false;
 
 	FORCEINLINE UObject* GetReferencingObject()	{ return ReferencingObject;	}
 
@@ -691,6 +702,8 @@ public:
 
 	void ProcessObjectArray(FWorkerContext& Context)
 	{
+		Context.bDidWork = true;
+		Context.bIsSuspended = false;
 		static_assert(!EnumHasAllFlags(Options, EGCOptions::Parallel | EGCOptions::AutogenerateSchemas), "Can't assemble token streams in parallel");
 		
 		CollectorType Collector(Processor, Context);
@@ -720,8 +733,9 @@ StoleContext:
 
 			if (Processor.IsTimeLimitExceeded())
 			{
+				FlushWork(Dispatcher);
 				Dispatcher.Suspend();
-				SuspendWork(Context);				
+				SuspendWork(Context);
 				return;
 			}
 
@@ -736,15 +750,7 @@ StoleContext:
 				}
 
 StoleARO:
-				if constexpr (DispatcherType::bBatching)
-				{
-					if (Dispatcher.FlushToStructBlocks())
-					{
-						ProcessStructs(Dispatcher);
-					}
-					
-					Dispatcher.FlushQueuedReferences();
-				}
+				FlushWork(Dispatcher);
 
 				if (	 Block = RemainingObjects.PopFullBlock<Options>(); Block);
 				else if (Block = RemainingObjects.PopPartialBlock(/* out if successful */ BlockSize); Block);
@@ -768,6 +774,7 @@ StoleARO:
 			CurrentObjects = MakeArrayView(Block->Objects, BlockSize);
 		} // while (true)
 		
+		check(!Context.ObjectsToSerialize.HasWork());
 		Processor.LogDetailedStatsSummary();
 	}
 
@@ -811,6 +818,19 @@ private:
 		}
 	}
 
+	FORCEINLINE_DEBUGGABLE void FlushWork(DispatcherType& Dispatcher)
+	{
+		if constexpr (DispatcherType::bBatching)
+		{
+			if (Dispatcher.FlushToStructBlocks())
+			{
+				ProcessStructs(Dispatcher);
+			}
+
+			Dispatcher.FlushQueuedReferences();
+		}
+	}
+
 	FORCENOINLINE void ProcessStructs(DispatcherType& Dispatcher);
 };
 
diff --git a/Engine/Source/Runtime/CoreUObject/Public/UObject/UObjectGlobals.h b/Engine/Source/Runtime/CoreUObject/Public/UObject/UObjectGlobals.h
--- a/Engine/Source/Runtime/CoreUObject/Public/UObject/UObjectGlobals.h
+++ b/Engine/Source/Runtime/CoreUObject/Public/UObject/UObjectGlobals.h
@@ -2684,7 +2684,7 @@ public:
 	 * Marks a specific object reference as a weak reference. This does not affect GC but will be freed at a later point
 	 * The default behavior returns false as weak references must be explicitly supported
 	 */
-	virtual bool MarkWeakObjectReferenceForClearing(UObject** WeakReference) { return false; }
+	virtual bool MarkWeakObjectReferenceForClearing(UObject** WeakReference, UObject* ReferenceOwner) { return false; }
 	/**
 	 * Sets whether this collector is currently processing native references or not.
 	 */
diff --git a/Engine/Source/Runtime/Engine/Private/BlueprintGeneratedClass.cpp b/Engine/Source/Runtime/Engine/Private/BlueprintGeneratedClass.cpp
--- a/Engine/Source/Runtime/Engine/Private/BlueprintGeneratedClass.cpp
+++ b/Engine/Source/Runtime/Engine/Private/BlueprintGeneratedClass.cpp
@@ -2120,16 +2120,16 @@ class FPersistentFrameCollector final : public FReferenceCollector //-V1052
 	FReferenceCollector& InnerCollector;
 #if UE_CHECK_BLUEPRINT_REFERENCES
 	const uint8* PersistentFrameDataAddr;
-	const UObject* Blueprint;
 #endif
+	UObject* Blueprint;
 
 public:
 	FPersistentFrameCollector(FReferenceCollector& Collector, uint8* Instance, UObject* InBlueprint)
 	: InnerCollector(Collector)
 #if UE_CHECK_BLUEPRINT_REFERENCES
 	, PersistentFrameDataAddr(Instance)
-	, Blueprint(InBlueprint)
 #endif
+	, Blueprint(InBlueprint)
 	{}
 
 
@@ -2204,7 +2204,7 @@ public:
 				if (const FObjectProperty* ObjectProperty = CastField<const FObjectProperty>(ReferencingProperty))
 				{
 					// This was a raw UObject* serialized by FObjectProperty, so just save the address
-					if (InnerCollector.MarkWeakObjectReferenceForClearing(&Object))
+					if (InnerCollector.MarkWeakObjectReferenceForClearing(&Object, Blueprint))
 					{
 						return;
 					}
