Wrap uses of UE_TRACE_LOG_SCOPED_T in CPUPROFILERTRACE_ENABLED and LOADTIMEPROFILERTRACE_ENABLED blocks

diff --git a/Engine/Source/Runtime/AssetRegistry/Private/AssetDataGatherer.cpp b/Engine/Source/Runtime/AssetRegistry/Private/AssetDataGatherer.cpp
--- a/Engine/Source/Runtime/AssetRegistry/Private/AssetDataGatherer.cpp
+++ b/Engine/Source/Runtime/AssetRegistry/Private/AssetDataGatherer.cpp
@@ -4356,9 +4356,11 @@ int64 FAssetDataGatherer::SaveCacheFileInternal(const FString& CacheFilename, co
 	}
 	CHECK_IS_NOT_LOCKED_CURRENT_THREAD(TickLock);
 	CHECK_IS_NOT_LOCKED_CURRENT_THREAD(ResultsLock);
+#if CPUPROFILERTRACE_ENABLED
 	UE_TRACE_LOG_SCOPED_T(Cpu, SaveCacheFile, CpuChannel)
 		<< SaveCacheFile.Path(*CacheFilename)
 		<< SaveCacheFile.Count(AssetsToSave.Num());
+#endif // CPUPROFILERTRACE_ENABLED
 	
 	// Save to a temp file first, then move to the destination to avoid corruption
 	FString CacheFilenameStr(CacheFilename);
diff --git a/Engine/Source/Runtime/CoreUObject/Private/UObject/LinkerLoad.cpp b/Engine/Source/Runtime/CoreUObject/Private/UObject/LinkerLoad.cpp
--- a/Engine/Source/Runtime/CoreUObject/Private/UObject/LinkerLoad.cpp
+++ b/Engine/Source/Runtime/CoreUObject/Private/UObject/LinkerLoad.cpp
@@ -4314,7 +4314,9 @@ void FLinkerLoad::Preload( UObject* Object )
 		if (Object->GetLinker() == this)
 		{
 			TRACE_CPUPROFILER_EVENT_SCOPE(FLinkerLoad::Preload);
+#if LOADTIMEPROFILERTRACE_ENABLED
 			TRACE_CPUPROFILER_EVENT_SCOPE_TEXT_ON_CHANNEL(*WriteToString<256>(TEXT("FLinkerLoad::Preload "), Object->GetFName()), AssetLoadTimeChannel);
+#endif // LOADTIMEPROFILERTRACE_ENABLED
 
 			UClass* Cls = Cast<UClass>(Object);
 			checkf(!GEventDrivenLoaderEnabled || !bLockoutLegacyOperations || !EVENT_DRIVEN_ASYNC_LOAD_ACTIVE_AT_RUNTIME, TEXT("Invalid call to FLinkerLoad::Preload while using the EDL. '%d' should have been reported via GetPreloadDependencies instead."), *Object->GetPathName());
diff --git a/Engine/Source/Runtime/CoreUObject/Private/UObject/UObjectGlobals.cpp b/Engine/Source/Runtime/CoreUObject/Private/UObject/UObjectGlobals.cpp
--- a/Engine/Source/Runtime/CoreUObject/Private/UObject/UObjectGlobals.cpp
+++ b/Engine/Source/Runtime/CoreUObject/Private/UObject/UObjectGlobals.cpp
@@ -1285,8 +1285,10 @@ UE_TRACE_EVENT_END()
 
 UObject* StaticLoadObjectInternal(UClass* ObjectClass, UObject* InOuter, const TCHAR* InName, const TCHAR* Filename, uint32 LoadFlags, UPackageMap* Sandbox, bool bAllowObjectReconciliation, const FLinkerInstancingContext* InstancingContext)
 {
+#if CPUPROFILERTRACE_ENABLED
 	UE_TRACE_LOG_SCOPED_T(Cpu, LoadObject, CpuChannel)
 		<< LoadObject.AssetPath(InName);
+#endif // CPUPROFILERTRACE_ENABLED
 	SCOPED_NAMED_EVENT(LoadObject, FColor::Red);
 	check(InName);
 
diff --git a/Engine/Source/Runtime/Engine/Private/StreamableManager.cpp b/Engine/Source/Runtime/Engine/Private/StreamableManager.cpp
--- a/Engine/Source/Runtime/Engine/Private/StreamableManager.cpp
+++ b/Engine/Source/Runtime/Engine/Private/StreamableManager.cpp
@@ -1585,8 +1585,10 @@ UE_TRACE_EVENT_END()
 
 UObject* FStreamableManager::LoadSynchronous(const FSoftObjectPath& Target, bool bManageActiveHandle, TSharedPtr<FStreamableHandle>* RequestHandlePointer)
 {
+#if CPUPROFILERTRACE_ENABLED
 	UE_TRACE_LOG_SCOPED_T(Cpu, StreamableManager_LoadSynchronous, CpuChannel)
 		<< StreamableManager_LoadSynchronous.AssetPath(*WriteToWideString<FName::StringBufferSize>(Target));
+#endif // CPUPROFILERTRACE_ENABLED
 	TSharedPtr<FStreamableHandle> Request = RequestSyncLoad(Target, bManageActiveHandle, FString::Printf(TEXT("LoadSynchronous of %s"), *Target.ToString()));
 
 	if (RequestHandlePointer)
