Fixing incremental reachability crashes when running single-threaded GC

diff --git a/Engine/Source/Runtime/CoreUObject/Private/UObject/GarbageCollection.cpp b/Engine/Source/Runtime/CoreUObject/Private/UObject/GarbageCollection.cpp
--- a/Engine/Source/Runtime/CoreUObject/Private/UObject/GarbageCollection.cpp
+++ b/Engine/Source/Runtime/CoreUObject/Private/UObject/GarbageCollection.cpp
@@ -3909,15 +3909,22 @@ private:
 	{
 		FContextPoolScope Pool;
 		FWorkerContext* Context = nullptr;
+		const bool bIsSingleThreaded = !(Options & EGCOptions::Parallel);
 
 		if (GReachabilityState.IsSuspended())
 		{
 			Context = GReachabilityState.GetContextArray()[0];
+			Context->bDidWork = false;
 			InitialObjects.Reset();
 		}
 		else
 		{
 			Context = Pool.AllocateFromPool();
+			if (bIsSingleThreaded)
+			{
+				GReachabilityState.SetupWorkers(1);
+				GReachabilityState.GetContextArray()[0] = Context;
+			}
 		}
 
 		if (!Private::GReachableObjects.IsEmpty())
@@ -3937,10 +3944,16 @@ private:
 
 		if (!GReachabilityState.CheckIfAnyContextIsSuspended())
 		{
+			GReachabilityState.ResetWorkers();
 			Stats = Context->Stats;
 			Pool.ReturnToPool(Context);
 			GReachabilityState.UpdateStats(Stats);
 		}
+		else if (bIsSingleThreaded)
+		{
+			Context->ResetInitialObjects();
+			Context->InitialNativeReferences = TConstArrayView<UObject**>();
+		}
 	}
 
 public:
@@ -5904,7 +5917,10 @@ public:
 void SuspendWork(FWorkerContext& Context)
 {
 	Context.bIsSuspended = true;
-	Context.Coordinator->Suspend();
+	if (Context.Coordinator)
+	{
+		Context.Coordinator->Suspend();
+	}
 }
 
 ELoot StealWork(FWorkerContext& Context, FReferenceCollector& Collector, FWorkBlock*& OutBlock, EGCOptions Options)
@@ -6031,8 +6047,6 @@ void ReleaseAsyncProcessingContexts(FWorkerContext& InContext, TArrayView<FWorke
 		InContext.Stats.AddStats(Context->Stats);
 		ContextPool.ReturnToPool(Context);
 	}
-
-	GReachabilityState.ResetWorkers();
 }
 
 void ProcessAsync(void (*ProcessSync)(void*, FWorkerContext&), void* Processor, FWorkerContext& InContext)
