#UE Add -CookSkipSave to cook commandline options. Using this will skip saving any assets the cook processes. Useful for validating that the non-saving parts of a cook without having to wait for the save.

diff --git a/Engine/Source/Editor/UnrealEd/Classes/CookOnTheSide/CookOnTheFlyServer.h b/Engine/Source/Editor/UnrealEd/Classes/CookOnTheSide/CookOnTheFlyServer.h
--- a/Engine/Source/Editor/UnrealEd/Classes/CookOnTheSide/CookOnTheFlyServer.h
+++ b/Engine/Source/Editor/UnrealEd/Classes/CookOnTheSide/CookOnTheFlyServer.h
@@ -1482,6 +1482,8 @@ private:
 	bool bIterativeCalculateExe = true;
 	/** If true this will ignore cooking unsolicited packages.  This is only useful if launching the COTF server to only process shader recompile requests. */
 	bool bIgnoreUnsolicitedPackages = false;
+	/** Whether to skip saving packages that are cooked. When true, the cook will only load and process packages but not write them to disk */
+	bool bSkipSave = false;
 	/** Timers for tracking how long we have been busy, to manage retries and warnings of deadlock */
 	double SaveBusyStartTimeSeconds = MAX_flt;
 	double SaveBusyRetryTimeSeconds = MAX_flt;
diff --git a/Engine/Source/Editor/UnrealEd/Private/CookOnTheFlyServer.cpp b/Engine/Source/Editor/UnrealEd/Private/CookOnTheFlyServer.cpp
--- a/Engine/Source/Editor/UnrealEd/Private/CookOnTheFlyServer.cpp
+++ b/Engine/Source/Editor/UnrealEd/Private/CookOnTheFlyServer.cpp
@@ -5847,7 +5847,14 @@ void UCookOnTheFlyServer::SaveCookedPackage(UE::Cook::FSaveCookedPackageContext&
 				{
 					LLM_SCOPE_BYTAG(Cooker_SavePackage);
 					FScopedActivePackage ScopedActivePackage(*this, Context.PackageData.GetPackageName(), NAME_None);
-					Context.SavePackageResult = GEditor->Save(Package, Context.World, *Context.PlatFilename, SaveArgs);
+					if (bSkipSave)
+					{
+						Context.SavePackageResult = ESavePackageResult::Success;
+					}
+					else
+					{
+						Context.SavePackageResult = GEditor->Save(Package, Context.World, *Context.PlatFilename, SaveArgs);
+					}
 				}
 				catch (std::exception&)
 				{
@@ -6097,10 +6104,15 @@ void FSaveCookedPackageContext::FinishPlatform()
 		}
 		// TODO: Reenable BuildDefinitionList once FCbPackage support for empty FCbObjects is in
 		//Info.Attachments.Add({ "BuildDefinitionList", BuildDefinitionList });
-		Info.WriteOptions = IPackageWriter::EWriteOptions::Write;
-		if (COTFS.IsDirectorCookByTheBook())
+		Info.WriteOptions = IPackageWriter::EWriteOptions::None;
+		if (!COTFS.bSkipSave)
 		{
-			Info.WriteOptions |= IPackageWriter::EWriteOptions::ComputeHash;
+			Info.WriteOptions |= IPackageWriter::EWriteOptions::Write;
+
+			if (COTFS.IsDirectorCookByTheBook())
+			{
+				Info.WriteOptions |= IPackageWriter::EWriteOptions::ComputeHash;
+			}
 		}
 
 		PackageWriter->CommitPackage(MoveTemp(Info));
@@ -6781,6 +6793,7 @@ void UCookOnTheFlyServer::SetInitializeConfigSettings(UE::Cook::FInitializeConfi
 	bIterativeCalculateExe = !bIterativeIgnoreExe || !bConfigSettingSetIterativeIgnoreExe;
 
 	bIgnoreUnsolicitedPackages = FParse::Param(FCommandLine::Get(), TEXT("odsc"));
+	bSkipSave = FParse::Param(FCommandLine::Get(), TEXT("CookSkipSave"));
 }
 
 void UCookOnTheFlyServer::ParseCookFilters()
diff --git a/Engine/Source/Runtime/CoreUObject/Private/Serialization/PackageWriterToSharedBuffer.cpp b/Engine/Source/Runtime/CoreUObject/Private/Serialization/PackageWriterToSharedBuffer.cpp
--- a/Engine/Source/Runtime/CoreUObject/Private/Serialization/PackageWriterToSharedBuffer.cpp
+++ b/Engine/Source/Runtime/CoreUObject/Private/Serialization/PackageWriterToSharedBuffer.cpp
@@ -133,15 +133,18 @@ TUniquePtr<FPackageWriterRecords::FPackage> FPackageWriterRecords::FindAndRemove
 
 void FPackageWriterRecords::ValidateCommit(FPackage& Record, const IPackageWriter::FCommitPackageInfo& Info) const
 {
-	checkf(Info.Status != IPackageWriter::ECommitStatus::Success || Record.Packages.Num() > 0,
-		TEXT("IPackageWriter->WritePackageData must be called before Commit if the Package save was successful."));
-	checkf(Info.Status != IPackageWriter::ECommitStatus::Success || Record.Packages.FindByPredicate([](const FPackageWriterRecords::FWritePackage& Package) { return Package.Info.MultiOutputIndex == 0; }),
-		TEXT("SavePackage must provide output 0 when saving multioutput packages."));
-	uint8 HasBulkDataType[IPackageWriter::FBulkDataInfo::NumTypes]{};
-	for (FBulkData& BulkRecord : Record.BulkDatas)
+	if (EnumHasAnyFlags(Info.WriteOptions, IPackageWriter::EWriteOptions::Write | IPackageWriter::EWriteOptions::ComputeHash))
 	{
-		checkf((HasBulkDataType[(int32)BulkRecord.Info.BulkDataType] & (1 << BulkRecord.Info.MultiOutputIndex)) == 0,
-			TEXT("IPackageWriter->WriteBulkData must not be called with more than one BulkData of the same type."));
-		HasBulkDataType[(int32)BulkRecord.Info.BulkDataType] |= 1 << BulkRecord.Info.MultiOutputIndex;
+		checkf(Info.Status != IPackageWriter::ECommitStatus::Success || Record.Packages.Num() > 0,
+			TEXT("IPackageWriter->WritePackageData must be called before Commit if the Package save was successful."));
+		checkf(Info.Status != IPackageWriter::ECommitStatus::Success || Record.Packages.FindByPredicate([](const FPackageWriterRecords::FWritePackage& Package) { return Package.Info.MultiOutputIndex == 0; }),
+			TEXT("SavePackage must provide output 0 when saving multioutput packages."));
+		uint8 HasBulkDataType[IPackageWriter::FBulkDataInfo::NumTypes]{};
+		for (FBulkData& BulkRecord : Record.BulkDatas)
+		{
+			checkf((HasBulkDataType[(int32)BulkRecord.Info.BulkDataType] & (1 << BulkRecord.Info.MultiOutputIndex)) == 0,
+				TEXT("IPackageWriter->WriteBulkData must not be called with more than one BulkData of the same type."));
+			HasBulkDataType[(int32)BulkRecord.Info.BulkDataType] |= 1 << BulkRecord.Info.MultiOutputIndex;
+		}
 	}
 }
