Improve the documentation for EBulkDataFlags to help people understand how to use them.

#rnx

diff --git a/Engine/Source/Runtime/CoreUObject/Public/Serialization/BulkData.h b/Engine/Source/Runtime/CoreUObject/Public/Serialization/BulkData.h
--- a/Engine/Source/Runtime/CoreUObject/Public/Serialization/BulkData.h
+++ b/Engine/Source/Runtime/CoreUObject/Public/Serialization/BulkData.h
@@ -105,48 +105,49 @@ enum EBulkDataFlags : uint32
 {
 	/** Empty flag set. */
 	BULKDATA_None = 0,
-	/**
-	 * INTERNAL SET ONLY - callers of bulkdata should not set this flag on the bulk data
-	 * It is overwritten according to global configuration by Serialize.
-	 * If set, payload is stored not inline; it is stored either at the end of the file
-	 * or in a separate file.
-	 */
+	/** Indicates that when the payload was saved/cooked as part of a package was stored at the end of the package file rather than inline. */
 	BULKDATA_PayloadAtEndOfFile = 1 << 0,
 	/** If set, payload should be [un]compressed using ZLIB during serialization. */
 	BULKDATA_SerializeCompressedZLIB = 1 << 1,
 	/** Force usage of SerializeElement over bulk serialization. */
 	BULKDATA_ForceSingleElementSerialization = 1 << 2,
-	/** Bulk data is only used once at runtime in the game. */
+	/** When set the payload will be unloaded from memory after it has been accessed via ::Lock/::Unlock. Other forms of accessing the payload are unaffected by this flag. */
 	BULKDATA_SingleUse = 1 << 3,
-	/** Bulk data won't be used and doesn't need to be loaded. */
+	/** DEPRECATED */
 	BULKDATA_Unused UE_DEPRECATED(5.3, "This feature is being removed") = 1 << 5,
-	/** Forces the payload to be saved inline, regardless of its size. */
+	/** Should be set before saing/cooking to force the internal payload to be stored inline */
 	BULKDATA_ForceInlinePayload = 1 << 6,
-	/** Flag to check if either compression mode is specified. */
+	/** @see BULKDATA_SerializeCompressedZLIB */
 	BULKDATA_SerializeCompressed = (BULKDATA_SerializeCompressedZLIB),
 	/** Forces the payload to be always streamed, regardless of its size. */
 	BULKDATA_ForceStreamPayload UE_DEPRECATED(5.3, "This flag has had no purpose for sometime") = 1 << 7,
-	/**
-	 * INTERNAL SET ONLY - callers of bulkdata should not set this flag on the bulk data
-	 * It is overwritten according to global configuration by Serialize.
-	 * If set, payload is stored in a separate file such as .ubulk.
-	 * */
+	/** Set when the bulkdata saved/cooked as part of a package indicates that the payload is stored in its own file (.ubulk, .uptnl or .m.ubulk depending on other flags) */
 	BULKDATA_PayloadInSeperateFile = 1 << 8,
-	/** DEPRECATED: If set, payload is compressed using platform specific bit window. */
+	/** DEPRECATED */
 	BULKDATA_SerializeCompressedBitWindow UE_DEPRECATED(5.3, "This flag has had no purpose for sometime") = 1 << 9,
-	/** There is a new default to inline unless you opt out. */
+	/** Should be set before being cooked as part of a package to force the cooked payload to before stored at the end of the package file either than being inline*/
 	BULKDATA_Force_NOT_InlinePayload = 1 << 10,
-	/** This payload is optional and may not be on device. */
+	/** 
+	 * During cooking, this flag indicates that the the payload is optional at runtime and should be stored in an .uptnl file, unless the payload is also set to be inline in
+	 * which case this flag is ignored. This flag will be preserved during cooking and and at runtime and can be used to identify where the payload should be loaded from.
+	 */
 	BULKDATA_OptionalPayload = 1 << 11,
-	/** This payload will be memory mapped, this requires alignment, no compression etc. */
+	/** 
+	 * During cooking this flag indicates that the payload should work with memory mapping at runtime if the target cooking platform supports it so the payload should be stored in
+	 * a .m.ubulk file, unless the payload is also set to be inline in which case this flag is ignored. This flag will be preserved during cooking and and at runtime and can be
+	 * used to identify where the payload should be loaded from.
+	 */
 	BULKDATA_MemoryMappedPayload = 1 << 12,
-	/** Bulk data size is 64 bits long. */
+	/** Set during serialization to indicate if the size and offset values were serialized as int64 types rather than the default int32 */
 	BULKDATA_Size64Bit = 1 << 13,
-	/** Duplicate non-optional payload in optional bulk data. */
+	/** During cooking this flag indicates that although the payload is NOT optional it should be stored in both the .ubulk file and the .uptnl file as duplicate non-optional data,
+	 * unless the payload is also set to be inline in which case this flag is ignored. This flag will be preserved during cooking and and at runtime and can be used to identify
+	 * where the payload should be loaded from.
+	 */
 	BULKDATA_DuplicateNonOptionalPayload = 1 << 14,
-	/** Indicates that an old ID is present in the data, at some point when the DDCs are flushed we can remove this. */
+	/** DEPRECATED */
 	BULKDATA_BadDataVersion UE_DEPRECATED(5.3, "This flag has had no purpose for sometime") = 1 << 15,
-	/** BulkData did not have it's offset changed during the cook and does not need the fix up at load time */
+	/** Set during saving and indicates that the payload offset value is correct and does not need adjusting via an additional offset stored in the FLinker (older legacy behavior) */
 	BULKDATA_NoOffsetFixUp = 1 << 16,
 	/**
 	 * INTERNAL SET ONLY - callers of bulkdata should not set this flag on the bulk data
