Enabling TObjectPtr GC barrier

diff --git a/Engine/Source/Runtime/CoreUObject/Public/UObject/ObjectPtr.h b/Engine/Source/Runtime/CoreUObject/Public/UObject/ObjectPtr.h
--- a/Engine/Source/Runtime/CoreUObject/Public/UObject/ObjectPtr.h
+++ b/Engine/Source/Runtime/CoreUObject/Public/UObject/ObjectPtr.h
@@ -19,7 +19,7 @@
 #endif
 
 #ifndef UE_OBJECT_PTR_GC_BARRIER
-	#define UE_OBJECT_PTR_GC_BARRIER 0
+	#define UE_OBJECT_PTR_GC_BARRIER 1
 #endif
 
 /** 
