Linker Instancing Context: only consider the context as instanced if there's at least one valid remapping added to the context, ignoring all null entries (which serves as a way to filter our exports loading).

diff --git a/Engine/Source/Runtime/CoreUObject/Private/UObject/LinkerInstancingContext.cpp b/Engine/Source/Runtime/CoreUObject/Private/UObject/LinkerInstancingContext.cpp
--- a/Engine/Source/Runtime/CoreUObject/Private/UObject/LinkerInstancingContext.cpp
+++ b/Engine/Source/Runtime/CoreUObject/Private/UObject/LinkerInstancingContext.cpp
@@ -284,11 +284,13 @@ void FLinkerInstancedPackageMap::AddPackageMapping(FName Original, FName Instanc
 	if (InstanceMappingDirection == EInstanceMappingDirection::OriginalToInstanced)
 	{
 		InstancedPackageMapping.Add(Original, Instanced);
+		bIsInstanced |= !Instanced.IsNone();
 	}
 	else
 	{
 		check(InstanceMappingDirection == EInstanceMappingDirection::InstancedToOriginal);
 		InstancedPackageMapping.Add(Instanced, Original);
+		bIsInstanced |= !Original.IsNone();
 	}
 }
 
diff --git a/Engine/Source/Runtime/CoreUObject/Public/UObject/LinkerInstancingContext.h b/Engine/Source/Runtime/CoreUObject/Public/UObject/LinkerInstancingContext.h
--- a/Engine/Source/Runtime/CoreUObject/Public/UObject/LinkerInstancingContext.h
+++ b/Engine/Source/Runtime/CoreUObject/Public/UObject/LinkerInstancingContext.h
@@ -21,6 +21,7 @@ public:
 
 	explicit FLinkerInstancedPackageMap(EInstanceMappingDirection MappingDirection)
 		: InstanceMappingDirection(MappingDirection)
+		, bIsInstanced(false)
 #if WITH_EDITOR
 		, bEnableNonEditorPath(false)
 #else
@@ -31,7 +32,7 @@ public:
 
 	bool IsInstanced() const
 	{
-		return InstancedPackageMapping.Num() > 0;
+		return bIsInstanced;
 	}
 
 	/** Remap the package name from the import table to its instanced counterpart, otherwise return the name unmodified. */
@@ -73,6 +74,9 @@ private:
 	FString InstancedPackagePrefix;
 	FString InstancedPackageSuffix;
 
+	/** Whether InstancedPackageMapping contains remapping data other that none */
+	bool bIsInstanced;
+
 	/** Allows tests to run non editor path from editor build. */
 	bool bEnableNonEditorPath;
 };
