Improved assert when replicating actor class without PushModel properly implemented

#rnx

diff --git a/Engine/Source/Runtime/CoreUObject/Public/UObject/Object.h b/Engine/Source/Runtime/CoreUObject/Public/UObject/Object.h
--- a/Engine/Source/Runtime/CoreUObject/Public/UObject/Object.h
+++ b/Engine/Source/Runtime/CoreUObject/Public/UObject/Object.h
@@ -1713,7 +1713,7 @@ private:
 	{
 		// This method should only be called on Objects that are networked, and those should
 		// always have this implemented (by UHT).
-		check(false);
+		checkf(false, TEXT("REPLICATED_BASE_CLASS macro is not implemented in %s. Something went very wrong."), *GetFullName());
 	}
 
 public:
