Create csv stat for GC count

diff --git a/Engine/Source/Runtime/CoreUObject/Private/UObject/GarbageCollection.cpp b/Engine/Source/Runtime/CoreUObject/Private/UObject/GarbageCollection.cpp
--- a/Engine/Source/Runtime/CoreUObject/Private/UObject/GarbageCollection.cpp
+++ b/Engine/Source/Runtime/CoreUObject/Private/UObject/GarbageCollection.cpp
@@ -53,6 +53,7 @@
 #include "UObject/FastReferenceCollector.h"
 
 DEFINE_LOG_CATEGORY(LogGarbage);
+CSV_DEFINE_CATEGORY(GC, true);
 
 #define PERF_DETAILED_PER_CLASS_GC_STATS				(LOOKING_FOR_PERF_ISSUES || 0) 
 
@@ -4764,6 +4765,8 @@ FORCEINLINE void CollectGarbageInternal(EObjectFlags KeepFlags, bool bPerformFul
 	GReachabilityState.CollectGarbage(KeepFlags, bPerformFullPurge);
 
 	GTimingInfo.LastGCDuration = FPlatformTime::Seconds() - StartTime;
+
+	CSV_CUSTOM_STAT(GC, Count, 1, ECsvCustomStatOp::Accumulate);
 }
 
 EGCOptions GetReferenceCollectorOptions(bool bPerformFullPurge)
