UPS roundtrip test fixes

diff --git a/Engine/Source/Runtime/CoreUObject/Private/Serialization/UnversionedPropertySerializationTest.cpp b/Engine/Source/Runtime/CoreUObject/Private/Serialization/UnversionedPropertySerializationTest.cpp
--- a/Engine/Source/Runtime/CoreUObject/Private/Serialization/UnversionedPropertySerializationTest.cpp
+++ b/Engine/Source/Runtime/CoreUObject/Private/Serialization/UnversionedPropertySerializationTest.cpp
@@ -54,9 +54,17 @@ struct FUnversionedPropertyTest : public FUnversionedPropertyTestInput
 			return *this;
 		}
 
+		
+		virtual FArchive& operator<<(FObjectPtr& Value) override
+		{
+			InnerArchive << reinterpret_cast<UPTRINT&>(Value);
+			return *this;
+		}
+
 		virtual FArchive& operator<<(UObject*& Value) override
 		{
-			return InnerArchive << reinterpret_cast<UPTRINT&>(Value);
+			InnerArchive << reinterpret_cast<UPTRINT&>(Value);
+			return *this;
 		}
 
 		virtual FArchive& operator<<(FLazyObjectPtr& Value) override { return FArchiveUObject::SerializeLazyObjectPtr(*this, Value); }
@@ -426,7 +434,7 @@ struct FUnversionedPropertyTest : public FUnversionedPropertyTestInput
 		FSaveResult VersionedSaved = Save(EPath::Versioned);
 		FSaveResult UnversionedSaved = Save(EPath::Unversioned);	
 
-		check(ExcludeEditorOnlyProperties(VersionedSaved.Properties) == UnversionedSaved.Properties);
+		check(VersionedSaved.Properties == UnversionedSaved.Properties || ExcludeEditorOnlyProperties(VersionedSaved.Properties) == UnversionedSaved.Properties);
 
 		FTestInstance VersionedLoaded = Load(VersionedSaved);
 		FTestInstance UnversionedLoaded = Load(UnversionedSaved);
