Allow verification of GC assumption on full GC purges in Test builds.

diff --git a/Engine/Source/Runtime/CoreUObject/Private/UObject/GarbageCollection.cpp b/Engine/Source/Runtime/CoreUObject/Private/UObject/GarbageCollection.cpp
--- a/Engine/Source/Runtime/CoreUObject/Private/UObject/GarbageCollection.cpp
+++ b/Engine/Source/Runtime/CoreUObject/Private/UObject/GarbageCollection.cpp
@@ -57,6 +57,9 @@ CSV_DEFINE_CATEGORY_MODULE(COREUOBJECT_API, GC, true);
 
 #define PERF_DETAILED_PER_CLASS_GC_STATS				(LOOKING_FOR_PERF_ISSUES || 0) 
 
+/** Allows release builds to override not verifying GC assumptions. Useful for profiling as it's hitchy. */
+extern COREUOBJECT_API bool GShouldVerifyGCAssumptionsOnFullPurge;
+
 /** Object count during last mark phase																				*/
 FThreadSafeCounter		GObjectCountDuringLastMarkPhase;
 /** Whether UObject hash tables are locked by GC */
@@ -4862,7 +4865,7 @@ void PreCollectGarbageImpl(EObjectFlags KeepFlags)
 
 #if VERIFY_DISREGARD_GC_ASSUMPTIONS
 			// Only verify assumptions if option is enabled. This avoids false positives in the Editor or commandlets.
-			if (GShouldVerifyGCAssumptions)
+			if (GShouldVerifyGCAssumptions || (bPerformFullPurge && GShouldVerifyGCAssumptionsOnFullPurge))
 			{
 				DECLARE_SCOPE_CYCLE_COUNTER(TEXT("CollectGarbageInternal.VerifyGCAssumptions"), STAT_CollectGarbageInternal_VerifyGCAssumptions, STATGROUP_GC);
 				const double StartTime = FPlatformTime::Seconds();
diff --git a/Engine/Source/Runtime/CoreUObject/Private/UObject/GarbageCollectionVerification.cpp b/Engine/Source/Runtime/CoreUObject/Private/UObject/GarbageCollectionVerification.cpp
--- a/Engine/Source/Runtime/CoreUObject/Private/UObject/GarbageCollectionVerification.cpp
+++ b/Engine/Source/Runtime/CoreUObject/Private/UObject/GarbageCollectionVerification.cpp
@@ -30,9 +30,11 @@
 -----------------------------------------------------------------------------*/
 
 /**
-* If set and VERIFY_DISREGARD_GC_ASSUMPTIONS is true, we verify GC assumptions about "Disregard For GC" objects.
+* If set and VERIFY_DISREGARD_GC_ASSUMPTIONS is true, we verify GC assumptions about "Disregard For GC" objects and clusters.
 */
-COREUOBJECT_API bool	GShouldVerifyGCAssumptions = !(UE_BUILD_SHIPPING != 0 && WITH_EDITOR != 0);
+COREUOBJECT_API bool	GShouldVerifyGCAssumptions = !UE_BUILD_SHIPPING && !UE_BUILD_TEST && !WITH_EDITOR;
+/** If set and VERIFY_DISREGARD_GC_ASSUMPTIONS is set, we verify GC assumptions when performing a full (blocking) purge */
+COREUOBJECT_API bool	GShouldVerifyGCAssumptionsOnFullPurge = !UE_BUILD_SHIPPING && !WITH_EDITOR;
 
 #if VERIFY_DISREGARD_GC_ASSUMPTIONS
 
@@ -320,9 +322,9 @@ void VerifyClustersAssumptions()
 				UE::GC::FWorkerContext Context;
 				Context.SetInitialObjectsUnpadded(ObjectsToSerialize);
 				CollectReferences(Processor, Context);
-				NumErrors.Add(Processor.GetErrorCount());
 			}			
 		}		
+		NumErrors.Add(Processor.GetErrorCount());
 	});
 
 
diff --git a/Engine/Source/Runtime/CoreUObject/Private/UObject/GarbageCollectionVerification.h b/Engine/Source/Runtime/CoreUObject/Private/UObject/GarbageCollectionVerification.h
--- a/Engine/Source/Runtime/CoreUObject/Private/UObject/GarbageCollectionVerification.h
+++ b/Engine/Source/Runtime/CoreUObject/Private/UObject/GarbageCollectionVerification.h
@@ -12,7 +12,7 @@
 #include "UObject/Class.h"
 
 // UE_BUILD_SHIPPING has GShouldVerifyGCAssumptions=false by default
-#define VERIFY_DISREGARD_GC_ASSUMPTIONS			!(UE_BUILD_SHIPPING || UE_BUILD_TEST)
+#define VERIFY_DISREGARD_GC_ASSUMPTIONS			!(UE_BUILD_SHIPPING)
 
 #if VERIFY_DISREGARD_GC_ASSUMPTIONS
 
