Explicity handle references from FGCObject::GGCObjectReferencer in thread data even if it's disregard for GC.

diff --git a/Engine/Source/Runtime/CoreUObject/Private/UObject/ReferenceChainSearch.cpp b/Engine/Source/Runtime/CoreUObject/Private/UObject/ReferenceChainSearch.cpp
--- a/Engine/Source/Runtime/CoreUObject/Private/UObject/ReferenceChainSearch.cpp
+++ b/Engine/Source/Runtime/CoreUObject/Private/UObject/ReferenceChainSearch.cpp
@@ -314,6 +314,7 @@ namespace UE::ReferenceChainSearch
 			}
 		};
 
+		// Stores a set of edge lists for objects starting at StartVertex.
 		struct FThreadData
 		{
 			UObject* PreviousReferencingObject = nullptr;
@@ -344,7 +345,8 @@ namespace UE::ReferenceChainSearch
 				}
 			}
 		};
-		TArray<FThreadData> AllThreadData;
+		FPolicyUObjectHeap& Policy;
+		TArray<FThreadData> AllThreadData; // Order of threads does not matter for MergeGraph
 
 		SIZE_T GetAllocatedSize() const
 		{
@@ -356,7 +358,7 @@ namespace UE::ReferenceChainSearch
 			return Size + AllThreadData.GetAllocatedSize();
 		}
 
-		void SetNumThreads(int32 InNumThreads, int32 NumberOfObjectsPerThread, int32 GlobalStartIndex, int32 MaxNumberOfObjects)
+		void SetNumThreads(int32 InNumThreads, int32 NumberOfObjectsPerThread, int32 GlobalStartIndex, int32 MaxNumberOfObjects, int32 ObjectReferencerIndex)
 		{
 			AllThreadData.Reset();
 			AllThreadData.AddDefaulted(InNumThreads);
@@ -369,6 +371,14 @@ namespace UE::ReferenceChainSearch
 				Thread.EdgeLists.Reserve(NumObjects);
 				Thread.EdgeLists.SetNumZeroed(NumObjects);
 			}
+
+			if (ObjectReferencerIndex != INDEX_NONE)
+			{
+				FThreadData& ObjectReferencerThreadData = AllThreadData.AddDefaulted_GetRef();
+				ObjectReferencerThreadData.StartVertex = ObjectReferencerIndex;
+				ObjectReferencerThreadData.EdgeLists.Reserve(1);
+				ObjectReferencerThreadData.EdgeLists.SetNumZeroed(1);
+			}
 		}
 
 		void CollectAllReferences(bool bGCOnly)
@@ -376,19 +386,26 @@ namespace UE::ReferenceChainSearch
 			constexpr bool bParallel = Derived::bParallel;
 			Derived* This = static_cast<Derived*>(this);
 			const int32 GlobalStartObjectIndex = bGCOnly ? GUObjectArray.GetFirstGCIndex() : 0;
+			int32 ObjectReferencerVertex = bGCOnly && FGCObject::GGCObjectReferencer ? Policy.ObjectToVertex(FGCObject::GGCObjectReferencer) : INDEX_NONE;
+			if (ObjectReferencerVertex >= GlobalStartObjectIndex)
+			{
+				ObjectReferencerVertex = INDEX_NONE;
+			}
+
 			const int32 MaxNumberOfObjects = GUObjectArray.GetObjectArrayNum() - GlobalStartObjectIndex;
+			
 			if constexpr (bParallel)
 			{
 				const int32 NumThreads = GetNumCollectReferenceWorkers();
 				const int32 NumberOfObjectsPerThread = (MaxNumberOfObjects / NumThreads) + 1;
 
-				SetNumThreads(NumThreads, NumberOfObjectsPerThread, GlobalStartObjectIndex, MaxNumberOfObjects);
+				SetNumThreads(NumThreads, NumberOfObjectsPerThread, GlobalStartObjectIndex, MaxNumberOfObjects, ObjectReferencerVertex);
 
 				ParallelFor(NumThreads,
-					[this, This, GlobalStartObjectIndex, MaxNumberOfObjects, NumThreads, NumberOfObjectsPerThread](int32 ThreadIndex) {
+					[this, bGCOnly, This, GlobalStartObjectIndex, MaxNumberOfObjects, NumThreads, NumberOfObjectsPerThread](int32 ThreadIndex) {
 						FProcessor Processor{ ThreadIndex, This };
-						TSet<UObject*> ThreadResult;
 						TArray<UObject*> ObjectsToSerialize;
+						
 						ObjectsToSerialize.Reserve(NumberOfObjectsPerThread);
 
 						const int32 FirstObjectIndex = GlobalStartObjectIndex + ThreadIndex * NumberOfObjectsPerThread;
@@ -411,31 +428,43 @@ namespace UE::ReferenceChainSearch
 						Context.SetInitialObjectsUnpadded(ObjectsToSerialize);
 						CollectReferences<FCollector<>>(Processor, Context);
 					});
-				return;
 			}
-
-			SetNumThreads(1, MaxNumberOfObjects, GlobalStartObjectIndex, MaxNumberOfObjects);
-			FProcessor Processor{ 0, This };
-			TArray<UObject*> ObjectsToProcess;
-			for (FRawObjectIterator It; It; ++It)
+			else 
 			{
-				FUObjectItem* ObjItem = *It;
-				UObject* Object = static_cast<UObject*>(ObjItem->Object);
-
-				// We can't ask the iterator for only GC objects because that would skip the GC Object referencer
-				if (bGCOnly && GUObjectArray.IsDisregardForGC(Object) && (Object != FGCObject::GGCObjectReferencer))
+				SetNumThreads(1, MaxNumberOfObjects, GlobalStartObjectIndex, MaxNumberOfObjects, ObjectReferencerVertex);
+				FProcessor Processor{ 0, This };
+				TArray<UObject*> ObjectsToProcess;
+				for (FRawObjectIterator It; It; ++It)
 				{
-					continue;
-				}
+					FUObjectItem* ObjItem = *It;
+					UObject* Object = static_cast<UObject*>(ObjItem->Object);
 
-				if (This->ShouldSkipReferencer(0, Object))
-				{
-					continue;
-				}
+					// // We can't ask the iterator for only GC objects because that would skip the GC Object referencer
+					// if (bGCOnly && GUObjectArray.IsDisregardForGC(Object) && (Object != FGCObject::GGCObjectReferencer))
+					// {
+					// 	continue;
+					// }
+
+					if (This->ShouldSkipReferencer(0, Object))
+					{
+						continue;
+					}
 
+					// Find direct references
+					UE::GC::FWorkerContext Context;
+					ObjectsToProcess = { Object };
+					Context.SetInitialObjectsUnpadded(ObjectsToProcess);
+					CollectReferences<FCollector<>>(Processor, Context);
+				}
+			}
+			
+			if (ObjectReferencerVertex != INDEX_NONE)
+			{
+				FProcessor Processor{ AllThreadData.Num() - 1, This };
+				TArray<UObject*> ObjectsToProcess;
 				// Find direct references
 				UE::GC::FWorkerContext Context;
-				ObjectsToProcess = { Object };
+				ObjectsToProcess = { FGCObject::GGCObjectReferencer };
 				Context.SetInitialObjectsUnpadded(ObjectsToProcess);
 				CollectReferences<FCollector<>>(Processor, Context);
 			}
@@ -452,16 +481,18 @@ namespace UE::ReferenceChainSearch
 			Context.SetInitialObjectsUnpadded(ObjectsToProcess);
 			CollectReferences<FCollector<true /* detailed property info */>>(Processor, Context);
 		}
+
+		TReferenceSearchBase(FPolicyUObjectHeap& InPolicy)
+			: Policy(InPolicy) { }
 	};
 
 	struct FDirectReferenceSearch : public TReferenceSearchBase<FDirectReferenceSearch>
 	{
 		using Super = TReferenceSearchBase<FDirectReferenceSearch>;
 		static constexpr bool bParallel = true;
-		FPolicyUObjectHeap& Policy;
 
 		FDirectReferenceSearch(FPolicyUObjectHeap& InPolicy)
-			: Policy(InPolicy) { }
+			: Super(InPolicy) { }
 
 		bool ShouldSkipReferencer(int32 ThreadIndex, UObject* Object)
 		{
@@ -504,12 +535,11 @@ namespace UE::ReferenceChainSearch
 	{
 		using Super = TReferenceSearchBase<FMinimalReferenceSearch>;
 		static constexpr bool bParallel = true;
-		FPolicyUObjectHeap& Policy;
 		TSet<const UObject*> TargetObjects;
 		TSet<UObject*> FoundTargets;
 
 		FMinimalReferenceSearch(FPolicyUObjectHeap& InPolicy)
-			: Policy(InPolicy) { }
+			: Super(InPolicy) { }
 
 		SIZE_T GetAllocatedSize() const
 		{
@@ -1011,13 +1041,13 @@ namespace UE::ReferenceChainSearch
 
 	struct FReferenceInfoSearch : public TReferenceSearchBase<FReferenceInfoSearch>
 	{
-		FPolicyUObjectHeap& Policy;
+		using Super = TReferenceSearchBase<FReferenceInfoSearch>;
 		const TMap<const UObject*, FGCObjectInfo*>& ObjectToInfoMap;
 
 		TMap<FVertex, FReferenceChainSearch::FObjectReferenceInfo>* ReferenceInfoMap;
 
 		FReferenceInfoSearch(FPolicyUObjectHeap& InPolicy, const TMap<const UObject*, FGCObjectInfo*>& InObjectToInfoMap)
-			: Policy(InPolicy)
+			: Super(InPolicy)
 			, ObjectToInfoMap(InObjectToInfoMap)
 		{
 		}
