Fix variable shadowing issue.

diff --git a/Engine/Source/Runtime/CoreUObject/Private/UObject/ReferenceChainSearch.cpp b/Engine/Source/Runtime/CoreUObject/Private/UObject/ReferenceChainSearch.cpp
--- a/Engine/Source/Runtime/CoreUObject/Private/UObject/ReferenceChainSearch.cpp
+++ b/Engine/Source/Runtime/CoreUObject/Private/UObject/ReferenceChainSearch.cpp
@@ -329,12 +329,12 @@ namespace UE::ReferenceChainSearch
 				return CurrentEdgeList.GetAllocatedSize() + GraphBuffer.GetAllocatedSize() + EdgeLists.GetAllocatedSize();
 			}
 
-			void FlushEdgeList(const FPolicyUObjectHeap& Policy)
+			void FlushEdgeList(const FPolicyUObjectHeap& InPolicy)
 			{
 				if (PreviousReferencingObject && CurrentEdgeList.Num() > 0)
 				{
 					// Flush current edge list to graph buffer
-					FVertex SourceVertex = Policy.ObjectToVertex(PreviousReferencingObject);
+					FVertex SourceVertex = InPolicy.ObjectToVertex(PreviousReferencingObject);
 					TArray<FVertex> EdgeList = CurrentEdgeList.Array();
 					const int64 BufferStartIndex = GraphBuffer.Num();
 					// During construction store edge lists relative to null to be fixed up later
