#UE Lower ensure to warning related to having external objects without a pointer to an external package, for now

diff --git a/Engine/Source/Runtime/CoreUObject/Private/UObject/UObjectBase.cpp b/Engine/Source/Runtime/CoreUObject/Private/UObject/UObjectBase.cpp
--- a/Engine/Source/Runtime/CoreUObject/Private/UObject/UObjectBase.cpp
+++ b/Engine/Source/Runtime/CoreUObject/Private/UObject/UObjectBase.cpp
@@ -261,8 +261,10 @@ UPackage* UObjectBase::GetExternalPackage() const
 	{
 		ExternalPackage = GetObjectExternalPackageThreadSafe(this);
 		// if the flag is set there should be an override set.
-		ensureMsgf(ExternalPackage, TEXT("Object %s ExternalPackage is invalid: RF_ExternalPackage is set, but ExternalPackage is nullptr."),
-			*static_cast<const UObjectBaseUtility*>(this)->GetPathName());
+		if (!ExternalPackage)
+		{
+			UE_LOG(LogUObjectBase, Warning, TEXT("Object %s ExternalPackage is invalid: RF_ExternalPackage is set, but ExternalPackage is nullptr."), *static_cast<const UObjectBaseUtility*>(this)->GetPathName());
+		}
 	}
 	return ExternalPackage;
 }
