Make UPS roundtrip test ignore CPF_NonNullable flag used by some non-UE properties

TPS loading calls FObjectProperty::PostSerializeObjectItem and converts loaded nullptrs to a default value. UPS by design avoids virtual function calls for most individual properties and will not convert nullptrs to defaults.

#rnx

diff --git a/Engine/Source/Runtime/CoreUObject/Private/Serialization/UnversionedPropertySerializationTest.cpp b/Engine/Source/Runtime/CoreUObject/Private/Serialization/UnversionedPropertySerializationTest.cpp
--- a/Engine/Source/Runtime/CoreUObject/Private/Serialization/UnversionedPropertySerializationTest.cpp
+++ b/Engine/Source/Runtime/CoreUObject/Private/Serialization/UnversionedPropertySerializationTest.cpp
@@ -181,7 +181,7 @@ struct FUnversionedPropertyTest : public FUnversionedPropertyTestInput
 	// FProperty::Identical() flavor suited to comparing loaded instances
 	static bool Equals(const FProperty* Property, const void* A, const void* B, FPropertyDiff& OutDiff)
 	{
-		if (Property->GetPropertyFlags() & (CPF_EditorOnly | CPF_Transient))
+		if (Property->GetPropertyFlags() & (CPF_EditorOnly | CPF_Transient | CPF_NonNullable))
 		{
 			return true;
 		}
