Make UObject.Class AttemptToFindUninitializedScriptStructMembers test runnable in all application type

#rnx

diff --git a/Engine/Source/Runtime/CoreUObject/Private/UObject/Class.cpp b/Engine/Source/Runtime/CoreUObject/Private/UObject/Class.cpp
--- a/Engine/Source/Runtime/CoreUObject/Private/UObject/Class.cpp
+++ b/Engine/Source/Runtime/CoreUObject/Private/UObject/Class.cpp
@@ -3736,7 +3736,7 @@ public:
 	virtual bool ElevateLogWarningsToErrors() override { return false; }
 };
 
-IMPLEMENT_CUSTOM_SIMPLE_AUTOMATION_TEST(FAutomationTestAttemptToFindUninitializedScriptStructMembers, FAutomationTestUObjectClassBase, "UObject.Class AttemptToFindUninitializedScriptStructMembers", EAutomationTestFlags::EditorContext | EAutomationTestFlags::ClientContext | EAutomationTestFlags::ServerContext | EAutomationTestFlags::SmokeFilter)
+IMPLEMENT_CUSTOM_SIMPLE_AUTOMATION_TEST(FAutomationTestAttemptToFindUninitializedScriptStructMembers, FAutomationTestUObjectClassBase, "UObject.Class AttemptToFindUninitializedScriptStructMembers", EAutomationTestFlags::ApplicationContextMask | EAutomationTestFlags::SmokeFilter)
 bool FAutomationTestAttemptToFindUninitializedScriptStructMembers::RunTest(const FString& Parameters)
 {
 	// This test fails when running tests under UHT because there is no TestUninitializedScriptStructMembersTest, so just skip it in that config.
