Fix performance and consistency issue where PostLoadAssetRegistryTagProperty was converting null FAssetData object tags in the editor from None to /Script/Engine.ClassName'None'
This was allocating extra memory and does not match the behavior on Save or in packaged builds

diff --git a/Engine/Source/Runtime/CoreUObject/Private/UObject/Obj.cpp b/Engine/Source/Runtime/CoreUObject/Private/UObject/Obj.cpp
--- a/Engine/Source/Runtime/CoreUObject/Private/UObject/Obj.cpp
+++ b/Engine/Source/Runtime/CoreUObject/Private/UObject/Obj.cpp
@@ -2254,7 +2254,8 @@ const FName& UObject::AssetVersePathTagName()
 #if WITH_EDITOR
 
 static void PostLoadAssetRegistryTagProperty(FProperty* Prop, const FAssetData& AssetData, TArray<UObject::FAssetRegistryTag>& OutTagsAndValuesToUpdate)
-{	
+{
+	// This TagType is ignored by the asset registry
 	UObject::FAssetRegistryTag::ETagType TagType = UObject::FAssetRegistryTag::ETagType::TT_Alphabetical;
 
 	if (Prop->HasAnyPropertyFlags(CPF_AssetRegistrySearchable))
@@ -2274,9 +2275,10 @@ static void PostLoadAssetRegistryTagProperty(FProperty* Prop, const FAssetData&
 		}
 		else if (Prop->IsA<FObjectPropertyBase>())
 		{
+			// Update the export path for short class names, but leave None alone to match save behavior
 			FObjectPropertyBase* PropertyObject = CastFieldChecked<FObjectPropertyBase>(Prop);
 			FString ExportPath = AssetData.GetTagValueRef<FString>(Prop->GetFName());
-			if (!ExportPath.IsEmpty() && ExportPath[0] != '/')
+			if (!ExportPath.IsEmpty() && ExportPath[0] != '/' && ExportPath != TEXT("None"))
 			{
 				FString ObjectPath = FPackageName::ExportTextPathToObjectPath(ExportPath);
 				ExportPath = FObjectPropertyBase::GetExportPath(PropertyObject->PropertyClass->GetClassPathName(), ObjectPath);
