Add AddReferencedObjects implementation for StructOnScope which correctly handles the weak pointer.
TWeakPointers must be dereferenced with GetEvenIfUnreachable during GC because GC initially marks all possibly-collectable objects as Unreachable before scanning the heap.

#localization none
#tests Stepped through new implementation in the debugger in the editor.

diff --git a/Engine/Source/Runtime/CoreUObject/Private/UObject/StructOnScope.cpp b/Engine/Source/Runtime/CoreUObject/Private/UObject/StructOnScope.cpp
--- a/Engine/Source/Runtime/CoreUObject/Private/UObject/StructOnScope.cpp
+++ b/Engine/Source/Runtime/CoreUObject/Private/UObject/StructOnScope.cpp
@@ -49,3 +49,19 @@ void FStructOnScope::SetPackage(UPackage* InPackage)
 {
 	Package = InPackage;
 }
+
+void FStructOnScope::AddReferencedObjects(FReferenceCollector& Collector)
+{
+	if (!SampleStructMemory)
+	{
+		return;
+	}
+	
+	// Unused result of cast just to confirm type
+	// Must use GetEvenIfUnreachable to resolve not-yet-scanned structures during GC where all objects are initially marked unreachable.
+	if (const UScriptStruct* ResolvedPtr = CastChecked<UScriptStruct>(ScriptStruct.GetEvenIfUnreachable()))
+	{
+		TWeakObjectPtr<const UScriptStruct>& StructPointer = reinterpret_cast<TWeakObjectPtr<const UScriptStruct>&>(ScriptStruct);
+		Collector.AddReferencedObjects(StructPointer, SampleStructMemory);	
+	}
+}
\ No newline at end of file
diff --git a/Engine/Source/Runtime/CoreUObject/Public/UObject/StructOnScope.h b/Engine/Source/Runtime/CoreUObject/Public/UObject/StructOnScope.h
--- a/Engine/Source/Runtime/CoreUObject/Public/UObject/StructOnScope.h
+++ b/Engine/Source/Runtime/CoreUObject/Public/UObject/StructOnScope.h
@@ -7,6 +7,8 @@
 #include "Templates/Casts.h"
 #include "Templates/RemoveReference.h"
 
+class FReferenceCollector;
+
 class FStructOnScope
 {
 protected:
@@ -126,6 +128,10 @@ public:
 		ScriptStruct = InScriptStruct;
 		Initialize();
 	}
+	
+	// If the struct definition is still available and there is struct memory to read from,
+	// add any object references in the struct data to the collector
+	COREUOBJECT_API void AddReferencedObjects(FReferenceCollector& Collector);
 };
 
 /**
