Additional debug options for garbage collection.
CVars gc.VerifyAssumptions, gc.VerifyAssumptionsOnFullPurge and gc.VerifyAssumptionsChance to control frequency of verifying GC assumptions.
CVar gc.ForceEnableGCProcessor to force GC to use the debug processor to allow GC to crash in the context of a memory error / assumption violation rather than later when batchers are handled in the optimized processor.

diff --git a/Engine/Source/Runtime/CoreUObject/Private/UObject/GarbageCollection.cpp b/Engine/Source/Runtime/CoreUObject/Private/UObject/GarbageCollection.cpp
--- a/Engine/Source/Runtime/CoreUObject/Private/UObject/GarbageCollection.cpp
+++ b/Engine/Source/Runtime/CoreUObject/Private/UObject/GarbageCollection.cpp
@@ -59,6 +59,7 @@ CSV_DEFINE_CATEGORY_MODULE(COREUOBJECT_API, GC, true);
 
 /** Allows release builds to override not verifying GC assumptions. Useful for profiling as it's hitchy. */
 extern COREUOBJECT_API bool GShouldVerifyGCAssumptionsOnFullPurge;
+extern COREUOBJECT_API float GVerifyGCAssumptionsChance;
 
 /** Object count during last mark phase																				*/
 FThreadSafeCounter		GObjectCountDuringLastMarkPhase;
@@ -260,6 +261,14 @@ static FAutoConsoleVariableRef CVarVerifyNoUnreachableObjects(
 );
 #endif // VERIFY_DISREGARD_GC_ASSUMPTIONS
 
+static bool GForceEnableDebugGCProcessor = 0;
+static FAutoConsoleVariableRef CVarForceEnableDebugGCProcessor(
+	TEXT("gc.ForceEnableGCProcessor"),
+	GForceEnableDebugGCProcessor,
+	TEXT("Force garbage collection to use the debug processor which may provide additional information during GC crashes."),
+	ECVF_Default
+);
+
 namespace UE::GC
 {
 
@@ -3211,10 +3220,12 @@ public:
 	TDebugReachabilityProcessor()
 	: bTrackGarbage(GGarbageReferenceTrackingEnabled != 0)
 	, bTrackHistory(FGCHistory::Get().IsActive())
+	, bForceEnable(GForceEnableDebugGCProcessor  != 0)
 	{}
 	
 	bool TracksHistory() const { return bTrackHistory; }
 	bool TracksGarbage() const { return bTrackGarbage; }
+	bool IsForceEnabled() const { return bForceEnable; }
 
 	FORCENOINLINE void HandleTokenStreamObjectReference(FWorkerContext& Context, const UObject* ReferencingObject, UObject*& Object, FMemberId MemberId, EOrigin Origin, bool bAllowReferenceElimination)
 	{
@@ -3251,6 +3262,7 @@ private:
 	const FPermanentObjectPoolExtents PermanentPool;
 	const bool bTrackGarbage;
 	const bool bTrackHistory;
+	const bool bForceEnable;
 
 	FORCENOINLINE static void HandleGarbageReference(FWorkerContext& Context, const UObject* ReferencingObject, UObject*& Object, FMemberId MemberId)
 	{
@@ -3621,7 +3633,8 @@ class FRealtimeGC : public FGarbageCollectionTracer
 
 #if !UE_BUILD_SHIPPING
 		TDebugReachabilityProcessor<Options> DebugProcessor;
-		if (DebugProcessor.TracksHistory() | 
+		if (DebugProcessor.IsForceEnabled() |
+			DebugProcessor.TracksHistory() | 
 			DebugProcessor.TracksGarbage() & Stats.bFoundGarbageRef)
 		{
 			CollectReferencesForGC<TDebugReachabilityCollector<Options>>(DebugProcessor, Context);
@@ -4865,7 +4878,8 @@ void PreCollectGarbageImpl(EObjectFlags KeepFlags)
 
 #if VERIFY_DISREGARD_GC_ASSUMPTIONS
 			// Only verify assumptions if option is enabled. This avoids false positives in the Editor or commandlets.
-			if (GShouldVerifyGCAssumptions || (bPerformFullPurge && GShouldVerifyGCAssumptionsOnFullPurge))
+			bool bShouldRandomlyVerifyGCAssumptions = GVerifyGCAssumptionsChance != 0.0 && FMath::FRand() < GVerifyGCAssumptionsChance;
+			if (GShouldVerifyGCAssumptions || (bPerformFullPurge && GShouldVerifyGCAssumptionsOnFullPurge) || bShouldRandomlyVerifyGCAssumptions)
 			{
 				DECLARE_SCOPE_CYCLE_COUNTER(TEXT("CollectGarbageInternal.VerifyGCAssumptions"), STAT_CollectGarbageInternal_VerifyGCAssumptions, STATGROUP_GC);
 				const double StartTime = FPlatformTime::Seconds();
diff --git a/Engine/Source/Runtime/CoreUObject/Private/UObject/GarbageCollectionVerification.cpp b/Engine/Source/Runtime/CoreUObject/Private/UObject/GarbageCollectionVerification.cpp
--- a/Engine/Source/Runtime/CoreUObject/Private/UObject/GarbageCollectionVerification.cpp
+++ b/Engine/Source/Runtime/CoreUObject/Private/UObject/GarbageCollectionVerification.cpp
@@ -33,8 +33,30 @@
 * If set and VERIFY_DISREGARD_GC_ASSUMPTIONS is true, we verify GC assumptions about "Disregard For GC" objects and clusters.
 */
 COREUOBJECT_API bool	GShouldVerifyGCAssumptions = !UE_BUILD_SHIPPING && !UE_BUILD_TEST && !WITH_EDITOR;
+static FAutoConsoleVariableRef CVarShouldVerifyGCAssumptions(
+	TEXT("gc.VerifyAssumptions"),
+	GShouldVerifyGCAssumptions,
+	TEXT("Whether to verify GC assumptions (disregard for GC, clustering) on each GC."),
+	ECVF_Default
+);
+
 /** If set and VERIFY_DISREGARD_GC_ASSUMPTIONS is set, we verify GC assumptions when performing a full (blocking) purge */
 COREUOBJECT_API bool	GShouldVerifyGCAssumptionsOnFullPurge = !UE_BUILD_SHIPPING && !WITH_EDITOR;
+static FAutoConsoleVariableRef CVarShouldVerifyGCAssumptionsOnFullPurge(
+	TEXT("gc.VerifyAssumptionsOnFullPurge"),
+	GShouldVerifyGCAssumptions,
+	TEXT("Whether to verify GC assumptions (disregard for GC, clustering) on full purge GCs."),
+	ECVF_Default
+);
+
+/** If > 0 and VERIFY_DISREGARD_GC_ASSUMPTIONS is set, we verify GC assumptions on that fraction of GCs. */
+COREUOBJECT_API float	GVerifyGCAssumptionsChance = 0.0f;
+static FAutoConsoleVariableRef CVarVerifyGCAssumptionsChance (
+	TEXT("gc.VerifyAssumptionsChance"),
+	GVerifyGCAssumptionsChance,
+	TEXT("Chance (0-1) to randomly verify GC assumptions on each GC."),
+	ECVF_Default
+);
 
 #if VERIFY_DISREGARD_GC_ASSUMPTIONS
 
diff --git a/Engine/Source/Runtime/CoreUObject/Private/UObject/Obj.cpp b/Engine/Source/Runtime/CoreUObject/Private/UObject/Obj.cpp
--- a/Engine/Source/Runtime/CoreUObject/Private/UObject/Obj.cpp
+++ b/Engine/Source/Runtime/CoreUObject/Private/UObject/Obj.cpp
@@ -5008,13 +5008,16 @@ void StaticUObjectInit()
 	GObjTransientPkg = NewObject<UPackage>(nullptr, TEXT("/Engine/Transient"), RF_Transient);
 	GObjTransientPkg->AddToRoot();
 
-	if( FParse::Param( FCommandLine::Get(), TEXT("VERIFYGC") ) )
+	if (IConsoleVariable* CVarVerifyGCAssumptions = IConsoleManager::Get().FindConsoleVariable(TEXT("gc.VerifyAssumptions")))
 	{
-		GShouldVerifyGCAssumptions = true;
-	}
-	if( FParse::Param( FCommandLine::Get(), TEXT("NOVERIFYGC") ) )
-	{
-		GShouldVerifyGCAssumptions = false;
+		if( FParse::Param( FCommandLine::Get(), TEXT("VERIFYGC") ) )
+		{
+			CVarVerifyGCAssumptions->Set(true, ECVF_SetByCommandline);
+		}
+		if( FParse::Param( FCommandLine::Get(), TEXT("NOVERIFYGC") ) )
+		{
+			CVarVerifyGCAssumptions->Set(false, ECVF_SetByCommandline);
+		}
 	}
 
 	UE_LOG(LogInit, Log, TEXT("Object subsystem initialized") );
