FStructUtils::AreSameLayout now returns true for nullptr, nullptr - importantly this fixes comparison of class layouts when the class does not have any sparse class data

diff --git a/Engine/Source/Runtime/CoreUObject/Private/UObject/Class.cpp b/Engine/Source/Runtime/CoreUObject/Private/UObject/Class.cpp
--- a/Engine/Source/Runtime/CoreUObject/Private/UObject/Class.cpp
+++ b/Engine/Source/Runtime/CoreUObject/Private/UObject/Class.cpp
@@ -6826,6 +6826,11 @@ bool FStructUtils::ArePropertiesTheSame(const FProperty* A, const FProperty* B,
 
 bool FStructUtils::TheSameLayout(const UStruct* StructA, const UStruct* StructB, bool bCheckPropertiesNames)
 {
+	if (StructA == StructB)
+	{
+		return true;
+	}
+
 	bool bResult = false;
 	if (StructA 
 		&& StructB 
