Add missing FORCEINLINE to ObjectPtrDecay path

diff --git a/Engine/Source/Runtime/CoreUObject/Public/UObject/ObjectPtr.h b/Engine/Source/Runtime/CoreUObject/Public/UObject/ObjectPtr.h
--- a/Engine/Source/Runtime/CoreUObject/Public/UObject/ObjectPtr.h
+++ b/Engine/Source/Runtime/CoreUObject/Public/UObject/ObjectPtr.h
@@ -1175,7 +1175,7 @@ namespace UE::Core::Private // private facilities; not for direct use
 	struct TObjectPtrDecayTypeOf
 	{
 		using Type = T;
-		static void PerformDecayActions(const T&) { /* nb: intentionally empty */ }
+		static FORCEINLINE void PerformDecayActions(const T&) { /* nb: intentionally empty */ }
 	};
 
 	template <typename T>
@@ -1183,7 +1183,7 @@ namespace UE::Core::Private // private facilities; not for direct use
 	{
 		using Type = T*;
 
-		static void PerformDecayActions(const TObjectPtr<T>& Value)
+		static FORCEINLINE void PerformDecayActions(const TObjectPtr<T>& Value)
 		{
 #if UE_WITH_OBJECT_HANDLE_LATE_RESOLVE || UE_WITH_OBJECT_HANDLE_TRACKING
 			Value.Get();
@@ -1196,7 +1196,7 @@ namespace UE::Core::Private // private facilities; not for direct use
 	{
 		using Type = TSet<typename TObjectPtrDecayTypeOf<T>::Type>;
 
-		static void PerformDecayActions(const TSet<T>& Value)
+		static FORCEINLINE void PerformDecayActions(const TSet<T>& Value)
 		{
 			for (const auto& V : Value)
 			{
@@ -1210,7 +1210,7 @@ namespace UE::Core::Private // private facilities; not for direct use
 	{
 		using Type = TMap<typename TObjectPtrDecayTypeOf<K>::Type, typename TObjectPtrDecayTypeOf<V>::Type>;
 			
-		static void PerformDecayActions(const TMap<K, V>& Value)
+		static FORCEINLINE void PerformDecayActions(const TMap<K, V>& Value)
 		{
 			for (const auto& KV : Value)
 			{
@@ -1225,7 +1225,7 @@ namespace UE::Core::Private // private facilities; not for direct use
 	{
 		using Type = TArray<typename TObjectPtrDecayTypeOf<T>::Type>;
 			
-		static void PerformDecayActions(const TArray<T>& Value)
+		static FORCEINLINE void PerformDecayActions(const TArray<T>& Value)
 		{
 			for (const auto& V : Value)
 			{
@@ -1239,7 +1239,7 @@ namespace UE::Core::Private // private facilities; not for direct use
 	{
 		using Type = TNonNullPtr<T>;
 
-		static void PerformDecayActions(const TNonNullPtr<TObjectPtr<T>>& Value)
+		static FORCEINLINE void PerformDecayActions(const TNonNullPtr<TObjectPtr<T>>& Value)
 		{
 #if UE_WITH_OBJECT_HANDLE_LATE_RESOLVE || UE_WITH_OBJECT_HANDLE_TRACKING
 			Value.GetRef().Get();
