Added InterfaceType alias to TScriptInterface

diff --git a/Engine/Source/Runtime/CoreUObject/Public/UObject/ScriptInterface.h b/Engine/Source/Runtime/CoreUObject/Public/UObject/ScriptInterface.h
--- a/Engine/Source/Runtime/CoreUObject/Public/UObject/ScriptInterface.h
+++ b/Engine/Source/Runtime/CoreUObject/Public/UObject/ScriptInterface.h
@@ -131,10 +131,12 @@ template<> struct TIsZeroConstructType<class FScriptInterface> { enum { Value =
 /**
  * Templated version of FScriptInterface, which provides accessors and operators for referencing the interface portion of a UObject that implements a native interface.
  */
-template <typename InterfaceType>
+template <typename InInterfaceType>
 class TScriptInterface : public FScriptInterface
 {
 public:
+	using InterfaceType = InInterfaceType;
+	
 	/**
 	 * Default constructor
 	 */
@@ -157,7 +159,7 @@ public:
 		UObject* SourceObject = ImplicitConv<UObject*>(Source);
 		SetObject(SourceObject);
 
-		InterfaceType* SourceInterface = Cast<InterfaceType>(SourceObject);
+		InInterfaceType* SourceInterface = Cast<InInterfaceType>(SourceObject);
 		SetInterface(SourceInterface);
 	}
 
@@ -166,13 +168,13 @@ public:
 	 */
 	template <
 		typename OtherInterfaceType,
-		decltype(ImplicitConv<InterfaceType*>(std::declval<OtherInterfaceType>()))* = nullptr
+		decltype(ImplicitConv<InInterfaceType*>(std::declval<OtherInterfaceType>()))* = nullptr
 	>
 	FORCEINLINE TScriptInterface(const TScriptInterface<OtherInterfaceType>& Other)
 	{
 		SetObject(Other.GetObject());
 
-		InterfaceType* SourceInterface = Other.GetInterface();
+		InInterfaceType* SourceInterface = Other.GetInterface();
 		SetInterface(SourceInterface);
 	}
 
@@ -184,7 +186,7 @@ public:
 	{
 		SetObject(SourceObject);
 
-		InterfaceType* SourceInterface = Cast<InterfaceType>(ToRawPtr(SourceObject));
+		InInterfaceType* SourceInterface = Cast<InInterfaceType>(ToRawPtr(SourceObject));
 		SetInterface(SourceInterface);
 	}
 
@@ -227,7 +229,7 @@ public:
 	 */
 	template <
 		typename OtherInterfaceType,
-		decltype(ImplicitConv<InterfaceType*>(std::declval<OtherInterfaceType>()))* = nullptr
+		decltype(ImplicitConv<InInterfaceType*>(std::declval<OtherInterfaceType>()))* = nullptr
 	>
 	TScriptInterface& operator=(const TScriptInterface<OtherInterfaceType>& Other)
 	{
@@ -248,12 +250,12 @@ public:
 	/**
 	 * Comparison operator, taking a pointer to InterfaceType
 	 */
-	template <typename OtherInterface, typename = decltype(ImplicitConv<InterfaceType*>((OtherInterface*)nullptr))>
+	template <typename OtherInterface, typename = decltype(ImplicitConv<InInterfaceType*>((OtherInterface*)nullptr))>
 	FORCEINLINE bool operator==( const OtherInterface* Other ) const
 	{
 		return GetInterface() == Other;
 	}
-	template <typename OtherInterface, typename = decltype(ImplicitConv<InterfaceType*>((OtherInterface*)nullptr))>
+	template <typename OtherInterface, typename = decltype(ImplicitConv<InInterfaceType*>((OtherInterface*)nullptr))>
 	FORCEINLINE bool operator!=( const OtherInterface* Other ) const
 	{
 		return GetInterface() != Other;
@@ -286,7 +288,7 @@ public:
 	/**
 	 * Member access operator.  Provides transparent access to the interface pointer contained by this TScriptInterface
 	 */
-	FORCEINLINE InterfaceType* operator->() const
+	FORCEINLINE InInterfaceType* operator->() const
 	{
 		return GetInterface();
 	}
@@ -296,7 +298,7 @@ public:
 	 *
 	 * @return	a reference (of type InterfaceType) to the object pointed to by InterfacePointer
 	 */
-	FORCEINLINE InterfaceType& operator*() const
+	FORCEINLINE InInterfaceType& operator*() const
 	{
 		return *GetInterface();
 	}
@@ -304,15 +306,15 @@ public:
 	/**
 	 * Returns the pointer to the interface
 	 */
-	FORCEINLINE InterfaceType* GetInterface() const
+	FORCEINLINE InInterfaceType* GetInterface() const
 	{
-		return (InterfaceType*)FScriptInterface::GetInterface();
+		return (InInterfaceType*)FScriptInterface::GetInterface();
 	}
 
 	/**
 	 * Sets the value of the InterfacePointer for this TScriptInterface
 	 */
-	FORCEINLINE void SetInterface(InterfaceType* InInterfacePointer)
+	FORCEINLINE void SetInterface(InInterfaceType* InInterfacePointer)
 	{
 		FScriptInterface::SetInterface(InInterfacePointer);
 	}
@@ -329,7 +331,7 @@ public:
 
 	friend FArchive& operator<<( FArchive& Ar, TScriptInterface& Interface )
 	{
-		return Interface.Serialize(Ar, InterfaceType::UClassType::StaticClass());
+		return Interface.Serialize(Ar, InInterfaceType::UClassType::StaticClass());
 	}
 };
 
