Fix for crash when trying to create a vpath from an object with a non-standard object path.

diff --git a/Engine/Source/Runtime/CoreUObject/Private/UObject/UObjectBaseUtility.cpp b/Engine/Source/Runtime/CoreUObject/Private/UObject/UObjectBaseUtility.cpp
--- a/Engine/Source/Runtime/CoreUObject/Private/UObject/UObjectBaseUtility.cpp
+++ b/Engine/Source/Runtime/CoreUObject/Private/UObject/UObjectBaseUtility.cpp
@@ -236,7 +236,12 @@ UE::Core::FVersePath UObjectBaseUtility::GetVersePath() const
 
 	FString PackageName = Outermost->GetPathName();
 	const FStringView MountPointName = FPathViews::GetMountPointNameFromPath(PackageName);
-	check(PackageName.StartsWith(TEXT("/") + FString(MountPointName)));
+
+	// If the mount point is of an unexpected form, assume we can't create a vpath from it
+	if (!PackageName.StartsWith(TEXT("/") + FString(MountPointName)))
+	{
+		return {};
+	}
 
 	IPluginManager& PluginManager = IPluginManager::Get();
 
