Prevent packages that have been saved when the editor domain is enabled from showing up as corrupted and allow the editor to load them correctly.

#rnx

### Problem
- When the bulkdata reference is written out it is either relative to a value stored in the package summary or absolute for the file that it is being written to. We only store the absolute value when writing an editor domain package or a cooked package, in which case  the bulkdata will be marked with the 'BULKDATA_NoOffsetFixUp' flag so we know not to convert the offset from relative to absolute on package load.
- If a package was loaded from the editor domain and then saved to the workspace domain it would retain the 'BULKDATA_NoOffsetFixUp' as it was not being cleared during serialization. This caused problems as we do not expect to encounter the flag when loading from the workspace domain.
-- In 5.0 and 5.1 packages saved to the workspace domain with 'BULKDATA_NoOffsetFixUp' were doing so with absolute offsets.
-- In 5.2 due to a major refactor on how we store data, packages being saved to the workspace domain with 'BULKDATA_NoOffsetFixUp' were doing so with relative offsets.

### Fix
- When saving a package the 'BULKDATA_NoOffsetFixUp' flag is cleared and then only reapplied to the package if required.
- To allow us to load packages that have been already saved out incorrectly the loading logic has been modified.
-- The existing check was to see if the package was from the editor domain, if so we assumed absolute offset and for all other cases relative. This is because the linkerloading code assumes we are loading from uncooked data in the editor.
-- We now also check the 'BULKDATA_NoOffsetFixUp' flag, and if it is set AND the package has an older version than EUnrealEngineObjectUE5Version::DATA_RESOURCES we assume the offsets are absolute, if the package is of the newer format, we assume they are relative and fix them up accordingly.

- Added a cvar "Serialization.ApplyBulkDataOffsetFix" which can be set to false to disable the loading fix and restore the previous behaviour. This is provided to de-risk the submit.

diff --git a/Engine/Source/Runtime/CoreUObject/Private/UObject/LinkerLoad.cpp b/Engine/Source/Runtime/CoreUObject/Private/UObject/LinkerLoad.cpp
--- a/Engine/Source/Runtime/CoreUObject/Private/UObject/LinkerLoad.cpp
+++ b/Engine/Source/Runtime/CoreUObject/Private/UObject/LinkerLoad.cpp
@@ -181,6 +181,15 @@ static FAutoConsoleVariableRef CEnforcePackageCompatibleVersionCheck(
 	ECVF_Default
 );
 
+/** 
+ * Required to load packages saved from the editor domain between UE 5.0 and 5.2, the cvar is only provided in case the fix causes
+ * unintended problems so that it can be disabled quickly.
+ */
+static TAutoConsoleVariable<bool> CVarApplyBulkDataFix(
+	TEXT("Serialization.ApplyBulkDataOffsetFix"),
+	true,
+	TEXT("When true, we will try to fix potentially bad bulkdata offsets"));
+
 /**
 * Test whether the given package index is a valid import or export in this package
 */
@@ -6741,7 +6750,34 @@ bool FLinkerLoad::SerializeBulkData(FBulkData& BulkData, const FBulkDataSerializ
 		// to be adjusted to the start of non-inline bulk data in the .uasset file. 
 		if (Meta.HasAnyFlags(BULKDATA_WorkspaceDomainPayload) == false)
 		{
-			Meta.SetOffset(Meta.GetOffset() + Summary.BulkDataStartOffset);
+			if (CVarApplyBulkDataFix.GetValueOnAnyThread())
+			{
+				// In theory we should never see the 'BULKDATA_NoOffsetFixUp' flag at this point, but for a time there was a bug that allowed
+				// packages saved to the workspace domain to have the flag so we cannot assume that the offset is relative and need to check.
+				// The outcome of this bug actually changed in the 'EUnrealEngineObjectUE5Version::DATA_RESOURCES' refactor which makes the 
+				// following checks more involved.
+
+				// If 'BULKDATA_NoOffsetFixUp' is not set then we know that the offset is always relative and needs to be converted to absolute
+				if (!Meta.HasAnyFlags(BULKDATA_NoOffsetFixUp))
+				{
+					Meta.SetOffset(Meta.GetOffset() + Summary.BulkDataStartOffset);
+				}
+				else
+				{
+					// If 'BULKDATA_NoOffsetFixUp' is set and the package was written after the 'EUnrealEngineObjectUE5Version::DATA_RESOURCES'
+					// refactor then we know the offset is relative and needs to be converted. If the package was written before the refactor
+					// and has the flag then we know the offset is already in absolute format and can be left unmodified.
+					if (Summary.GetFileVersionUE() >= EUnrealEngineObjectUE5Version::DATA_RESOURCES)
+					{
+						Meta.SetOffset(Meta.GetOffset() + Summary.BulkDataStartOffset);
+					}
+				}
+			}
+			else
+			{
+				// Previous behavior before an attempted fix for bad data was added.
+				Meta.SetOffset(Meta.GetOffset() + Summary.BulkDataStartOffset);
+			}
 		}
 
 		if (bLazyLoadable == false)
diff --git a/Engine/Source/Runtime/CoreUObject/Private/UObject/LinkerSave.cpp b/Engine/Source/Runtime/CoreUObject/Private/UObject/LinkerSave.cpp
--- a/Engine/Source/Runtime/CoreUObject/Private/UObject/LinkerSave.cpp
+++ b/Engine/Source/Runtime/CoreUObject/Private/UObject/LinkerSave.cpp
@@ -632,7 +632,7 @@ bool FLinkerSave::SerializeBulkData(FBulkData& BulkData, const FBulkDataSerializ
 		FBulkData::SetBulkDataFlagsOn(SerializedMeta.Flags, static_cast<EBulkDataFlags>(BULKDATA_Size64Bit));
 	}
 
-	EBulkDataFlags FlagsToClear = static_cast<EBulkDataFlags>(BULKDATA_PayloadAtEndOfFile | BULKDATA_PayloadInSeperateFile | BULKDATA_WorkspaceDomainPayload | BULKDATA_ForceSingleElementSerialization);
+	EBulkDataFlags FlagsToClear = static_cast<EBulkDataFlags>(BULKDATA_PayloadAtEndOfFile | BULKDATA_PayloadInSeperateFile | BULKDATA_WorkspaceDomainPayload | BULKDATA_ForceSingleElementSerialization | BULKDATA_NoOffsetFixUp);
 	if (IsCooking())
 	{
 		FBulkData::SetBulkDataFlagsOn(FlagsToClear, static_cast<EBulkDataFlags>(BULKDATA_SerializeCompressed));
