fix for renamed gamefeature plugin cvars
added cvars with old names to keep backwards compatibility
https://p4-swarm.epicgames.net/reviews/26618963

diff --git a/Engine/Source/Runtime/CoreUObject/Private/Misc/CoreUObjectPluginManager.cpp b/Engine/Source/Runtime/CoreUObject/Private/Misc/CoreUObjectPluginManager.cpp
--- a/Engine/Source/Runtime/CoreUObject/Private/Misc/CoreUObjectPluginManager.cpp
+++ b/Engine/Source/Runtime/CoreUObject/Private/Misc/CoreUObjectPluginManager.cpp
@@ -23,6 +23,11 @@ namespace UE::CoreUObject::Private
 		TEXT("Verify plugin assets are no longer in memory when unloading."),
 		ECVF_Default);
 
+	static FAutoConsoleVariableRef CVarVerifyPluginUnloadOldName(TEXT("GameFeaturePlugin.VerifyUnload"),
+		GVerifyUnload,
+		TEXT("Verify plugin assets are no longer in memory when unloading. Deprecated use PluginManager.VerifyUnload instead"),
+		ECVF_Default);
+
 #if WITH_LOW_LEVEL_TESTS
 	bool bEnsureOnLeakedPackages = false;
 #endif
@@ -42,6 +47,21 @@ namespace UE::CoreUObject::Private
 		ECVF_Default
 	);
 
+	static FAutoConsoleVariableRef CVarLeakedAssetTrace_SeverityOldName(
+#if UE_BUILD_SHIPPING
+		TEXT("GameFeaturePlugin.LeakedAssetTrace.Severity.Shipping"),
+#else
+		TEXT("GameFeaturePlugin.LeakedAssetTrace.Severity"),
+#endif
+		GLeakedAssetTrace_Severity,
+		TEXT("Controls severity of logging when the engine detects that assets from an Game Feature Plugin were leaked during unloading or unmounting. . Deprecated use GameFeaturePlugin.LeakedAssetTrace instead\n")
+		TEXT("0 - all reference tracing and logging is disabled\n")
+		TEXT("1 - logs an error\n")
+		TEXT("2 - ensure\n")
+		TEXT("3 - fatal error\n"),
+		ECVF_Default
+	);
+
 	static bool GRenameLeakedPackages = true;
 	static FAutoConsoleVariableRef CVarRenameLeakedPackages(
 		TEXT("PluginManager.LeakedAssetTrace.RenameLeakedPackages"),
@@ -50,6 +70,13 @@ namespace UE::CoreUObject::Private
 		ECVF_Default
 	);
 
+	static FAutoConsoleVariableRef CVarRenameLeakedPackagesOldName(
+		TEXT("GameFeaturePlugin.LeakedAssetTrace.RenameLeakedPackages. Deprecated used PluginManager.LeakedAssetTrace.RenameLeakedPackages instead"),
+		GRenameLeakedPackages,
+		TEXT("Should packages which are leaked after the Game Feature Plugin is unloaded or unmounted."),
+		ECVF_Default
+	);
+
 	static int32 GLeakedAssetTrace_TraceMode = (UE_BUILD_SHIPPING ? 0 : 1);
 	static FAutoConsoleVariableRef CVarLeakedAssetTrace_TraceMode(
 #if UE_BUILD_SHIPPING
@@ -64,6 +91,19 @@ namespace UE::CoreUObject::Private
 		ECVF_Default
 	);
 
+	static FAutoConsoleVariableRef CVarLeakedAssetTrace_TraceModeOldName(
+#if UE_BUILD_SHIPPING
+		TEXT("GameFeaturePlugin.LeakedAssetTrace.TraceMode.Shipping"),
+#else
+		TEXT("GameFeaturePlugin.LeakedAssetTrace.TraceMode"),
+#endif
+		GLeakedAssetTrace_TraceMode,
+		TEXT("Controls detail level of reference tracing when the engine detects that assets from a Game Feature Plugin were leaked during unloading or unmounting. Deprecated used PluginManager.LeakedAssetTrace.TraceMode instead\n")
+		TEXT("0 - direct references only\n")
+		TEXT("1 - full reference trace"),
+		ECVF_Default
+	);
+
 	static int32 GLeakedAssetTrace_MaxReportCount = 10;
 	static FAutoConsoleVariableRef CVarLeakedAssetTrace_MaxReportCount(
 		TEXT("PluginManager.LeakedAssetTrace.MaxReportCount"),
@@ -72,6 +112,13 @@ namespace UE::CoreUObject::Private
 		ECVF_Default
 	);
 
+	static FAutoConsoleVariableRef CVarLeakedAssetTrace_MaxReportCountOldName(
+		TEXT("GameFeaturePlugin.LeakedAssetTrace.MaxReportCount. Deprecated use PluginManager.LeakedAssetTrace.MaxReportCount instead"),
+		GLeakedAssetTrace_MaxReportCount,
+		TEXT("Max number of assets to report when we find leaked assets.\n"),
+		ECVF_Default
+	);
+
 	DEFINE_LOG_CATEGORY_STATIC(PluginHandlerLog, Log, All);
 
 	PluginHandler GPluginHandler;
