PackageReload: Fixed package reloading issue observed on MetaSounds. This also re-enables the experimental fast package reload codepath.

#rnx

diff --git a/Engine/Source/Runtime/CoreUObject/Private/UObject/PackageReload.cpp b/Engine/Source/Runtime/CoreUObject/Private/UObject/PackageReload.cpp
--- a/Engine/Source/Runtime/CoreUObject/Private/UObject/PackageReload.cpp
+++ b/Engine/Source/Runtime/CoreUObject/Private/UObject/PackageReload.cpp
@@ -18,7 +18,7 @@
 #if WITH_EDITOR
 TAutoConsoleVariable<bool> CVarPackageReloadEnableFastPath(
 	TEXT("PackageReload.EnableFastPath"),
-	false,
+	true,
 	TEXT("When 'true', an optimized codepath is used to speed up reloading packages (experimental)."),
 	ECVF_Default);
 #endif
@@ -65,6 +65,34 @@ class FPackageReferencersHelper
 		TWeakObjectPtr<UObject> Object = MakeWeakObjectPtr<UObject>(InObject);
 		PotentialReferencerObjects.Add(Object);
 	}
+
+	/**
+	 * Recursively retrieve all objects that reference any object in the InReferencedObjects.
+	 */
+	static void RecursiveRetrieveReferencers(const TArray<UObject*>& InReferencedObjects, TSet<TWeakObjectPtr<UObject>>& OutReferencingObjects)
+	{
+		TSet<UObject*> ReferencerObjects;
+
+		// Use the fast reference collector to recursively find referencers until no more new ones are found.
+		int32 LastObjectCount = 0;
+
+		TArray<UObject*> FoundReferencerObjects = FReferencerFinder::GetAllReferencers(InReferencedObjects, nullptr, EReferencerFinderFlags::None);
+		do
+		{
+			LastObjectCount = ReferencerObjects.Num();
+			ReferencerObjects.Append(FoundReferencerObjects);
+
+			FoundReferencerObjects = FReferencerFinder::GetAllReferencers(FoundReferencerObjects, nullptr, EReferencerFinderFlags::SkipInnerReferences);
+
+		} while (LastObjectCount != ReferencerObjects.Num());
+
+		// Convert them into weak pointers to deal with objects that get GC'd during the reload operation.
+		OutReferencingObjects.Reserve(ReferencerObjects.Num());
+		for (UObject* ReferencerObject : ReferencerObjects)
+		{
+			OutReferencingObjects.Emplace(MakeWeakObjectPtr(ReferencerObject));
+		}
+	}
 #endif
 
 public:
@@ -85,15 +113,8 @@ public:
 				ObjectsToReload.Add(PackageToReloadData.PackageToReload);
 			}
 
-			// Retrieve all other UObjects that refer to any of them.
-			TArray<UObject*> ReferencerObjects = FReferencerFinder::GetAllReferencers(ObjectsToReload, nullptr, EReferencerFinderFlags::None);
-
-			// Convert them into weak pointers to deal with objects that get GC'd during the reload operation.
-			PotentialReferencerObjects.Reserve(ReferencerObjects.Num());
-			for (UObject* ReferencerObject : ReferencerObjects)
-			{
-				PotentialReferencerObjects.Emplace(MakeWeakObjectPtr(ReferencerObject));
-			}
+			// Collect all other objects that reference any of them.
+			RecursiveRetrieveReferencers(ObjectsToReload, PotentialReferencerObjects);
 		}
 #endif
 	}
