Only use the object flag RF_HasExternalPackage as a hint as if the object has an external package or not, and only report a warning if the flag is set but no external package is hashed for this object.

This fixes a race condition between UStaticMesh::GetPlatformStaticMeshRenderData calling GetPackage from a task thread to know if the package is already cooked for editor, while we are moving an HLOD mesh from an external actor package into a streaming level during cooking, resulting in a discrepency betwen the package RF_HasExternalPackage and the content of the external package hash map. The situation still exists, but we know that it's not going to cause problems in that specific case.

diff --git a/Engine/Source/Runtime/CoreUObject/Private/UObject/UObjectBase.cpp b/Engine/Source/Runtime/CoreUObject/Private/UObject/UObjectBase.cpp
--- a/Engine/Source/Runtime/CoreUObject/Private/UObject/UObjectBase.cpp
+++ b/Engine/Source/Runtime/CoreUObject/Private/UObject/UObjectBase.cpp
@@ -260,11 +260,6 @@ UPackage* UObjectBase::GetExternalPackage() const
 	if ((GetFlags() & RF_HasExternalPackage) != 0)
 	{
 		ExternalPackage = GetObjectExternalPackageThreadSafe(this);
-		// if the flag is set there should be an override set.
-		if (!ExternalPackage)
-		{
-			UE_LOG(LogUObjectBase, Warning, TEXT("Object %s ExternalPackage is invalid: RF_ExternalPackage is set, but ExternalPackage is nullptr."), *static_cast<const UObjectBaseUtility*>(this)->GetPathName());
-		}
 	}
 	return ExternalPackage;
 }
@@ -291,11 +286,9 @@ void UObjectBase::SetExternalPackage(UPackage* InPackage)
 	if (InPackage)
 	{
 		HashObjectExternalPackage(this, InPackage);
-		AtomicallySetFlags(RF_HasExternalPackage);
 	}
 	else
 	{
-		AtomicallyClearFlags(RF_HasExternalPackage);
 		UnhashObjectExternalPackage(this);
 	}
 }
diff --git a/Engine/Source/Runtime/CoreUObject/Private/UObject/UObjectHash.cpp b/Engine/Source/Runtime/CoreUObject/Private/UObject/UObjectHash.cpp
--- a/Engine/Source/Runtime/CoreUObject/Private/UObject/UObjectHash.cpp
+++ b/Engine/Source/Runtime/CoreUObject/Private/UObject/UObjectHash.cpp
@@ -1620,6 +1620,7 @@ void HashObjectExternalPackage(UObjectBase* Object, UPackage* Package)
 			}
 			AddToPackageMap(ThreadHash, Object, Package);
 		}
+		Object->AtomicallySetFlags(RF_HasExternalPackage);
 	}
 	else
 	{
@@ -1632,8 +1633,8 @@ void UnhashObjectExternalPackage(class UObjectBase* Object)
 	LLM_SCOPE_BYTAG(UObjectHash);
 	FUObjectHashTables& ThreadHash = FUObjectHashTables::Get();
 	FHashTableLock LockHash(ThreadHash);
-	UPackage* Package = UnassignExternalPackageFromObject(ThreadHash, Object);
-	if (Package)
+	Object->AtomicallyClearFlags(RF_HasExternalPackage);
+	if (UPackage* Package = UnassignExternalPackageFromObject(ThreadHash, Object))
 	{
 		RemoveFromPackageMap(ThreadHash, Object, Package);
 	}
@@ -1643,7 +1644,9 @@ UPackage* GetObjectExternalPackageThreadSafe(const UObjectBase* Object)
 {
 	FUObjectHashTables& ThreadHash = FUObjectHashTables::Get();
 	FHashTableLock LockHash(ThreadHash);
-	return ThreadHash.ObjectToPackageMap.FindRef(Object);
+	UPackage* ExternalPackage = ThreadHash.ObjectToPackageMap.FindRef(Object);
+	UE_CLOG(!ExternalPackage && ((Object->GetFlags() & RF_HasExternalPackage) != 0), LogUObjectHash, Warning, TEXT("Object %s ExternalPackage is invalid: RF_HasExternalPackage is set, but ExternalPackage is nullptr."), *static_cast<const UObjectBaseUtility*>(Object)->GetPathName());
+	return ExternalPackage;
 }
 
 UPackage* GetObjectExternalPackageInternal(const UObjectBase* Object)
