Improve LinkerLoad instancing context population to catch object in external package outered to object in instanced packages

diff --git a/Engine/Source/Runtime/CoreUObject/Private/UObject/LinkerLoad.cpp b/Engine/Source/Runtime/CoreUObject/Private/UObject/LinkerLoad.cpp
--- a/Engine/Source/Runtime/CoreUObject/Private/UObject/LinkerLoad.cpp
+++ b/Engine/Source/Runtime/CoreUObject/Private/UObject/LinkerLoad.cpp
@@ -2053,6 +2053,7 @@ FLinkerLoad::ELinkerStatus FLinkerLoad::FixupImportMap()
 
 FLinkerLoad::ELinkerStatus FLinkerLoad::PopulateInstancingContext()
 {
+	TRACE_CPUPROFILER_EVENT_SCOPE(FLinkerLoad::PopulateInstancingContext);
 	DECLARE_SCOPE_CYCLE_COUNTER(TEXT("FLinkerLoad::PopulateInstancingContext"), STAT_LinkerLoad_PopulateInstancingContext, STATGROUP_LinkerLoad);
 
 	if (!bHasPopulatedInstancingContext)
@@ -2061,7 +2062,19 @@ FLinkerLoad::ELinkerStatus FLinkerLoad::PopulateInstancingContext()
 		// Generate Instance Remapping if needed
 		if (IsContextInstanced())
 		{
-			TSet<FName> InstancingPackageName;
+			auto AddInstancedMapping = [this](const FString& OuterPackageName, FName InstancingPackageName) -> bool
+			{
+				FName InstancedName = InstancingContext.FindPackageMapping(InstancingPackageName);
+				// if there's isn't already a remapping for that package, create one
+				if (InstancedName.IsNone())
+				{
+					InstancedName = *FLinkerInstancingContext::GetInstancedPackageName(OuterPackageName, InstancingPackageName.ToString());
+					InstancingContext.AddPackageMapping(InstancingPackageName, InstancedName);
+					return true;
+				}
+				return false;
+			};
+
 			FString LinkerPackageName = LinkerRoot->GetName();
 
 			// Add import package we should instantiate since object in this instanced linker are outered to them
@@ -2074,42 +2087,57 @@ FLinkerLoad::ELinkerStatus FLinkerLoad::PopulateInstancingContext()
 					{
 						if (Import->HasPackageName())
 						{
-							InstancingPackageName.Add(Import->PackageName);
+							AddInstancedMapping(LinkerPackageName, Import->PackageName);
 						}
 						Import = &Imp(Import->OuterIndex);
 					}
 					check(Import->OuterIndex.IsNull() && !Import->HasPackageName());
-					InstancingPackageName.Add(Import->ObjectName);
+					AddInstancedMapping(LinkerPackageName, Import->ObjectName);
 				}
 			}
 
-			// Also add import package, we should instantiate as their are outered to object in this package
-			auto HasExportOuterChain = [this](const FObjectImport* InImport) -> bool
+			// Also add import package, we should instantiate as their are outered to object in this package or one of their outer is already instanced
+			auto HasInstancedOuterChain = [this](const FObjectImport* InImport) -> FName
 			{
 				while (InImport->OuterIndex.IsImport())
 				{
 					InImport = &Imp(InImport->OuterIndex);
+					FName ImportPackageName = InImport->HasPackageName() ? InImport->GetPackageName() : (InImport->OuterIndex.IsNull() ? InImport->ObjectName : NAME_None);
+					if (!ImportPackageName.IsNone())
+					{
+						FName InstancedRemap = InstancingContext.RemapPackage(ImportPackageName);
+						if (InstancedRemap != ImportPackageName)
+						{
+							return InstancedRemap;
+						}
+					}
 				}
-				return InImport->OuterIndex.IsExport();
+				// return if the import outer is an export or not if we didn't find an instanced import
+				return InImport->OuterIndex.IsExport() ? NAME_TRUE : NAME_FALSE;
 			};
 
-			for (const FObjectImport& Import : ImportMap)
+			for (int32 ImportIndex = 0; ImportIndex < ImportMap.Num(); ++ImportIndex)
 			{
-				if (Import.HasPackageName() && HasExportOuterChain(&Import))
+				const FObjectImport& Import = ImportMap[ImportIndex];
+				if (Import.HasPackageName())
 				{
-					InstancingPackageName.Add(Import.PackageName);
-				}
-			}
-
-			// add remapping for all the packages that should be instantiated along with this one
-			for (const FName& InstancingName : InstancingPackageName)
-			{
-				FName InstancedName = InstancingContext.FindPackageMapping(InstancingName);
-				// if there's isn't already a remapping for that package, create one
-				if (InstancedName.IsNone())
-				{
-					InstancedName = *FLinkerInstancingContext::GetInstancedPackageName(LinkerPackageName, InstancingName.ToString());
-					InstancingContext.AddPackageMapping(InstancingName, InstancedName);
+					FName Result = HasInstancedOuterChain(&Import);
+					// Outer chain has an export
+					if (Result == NAME_TRUE)
+					{
+						AddInstancedMapping(LinkerPackageName, Import.PackageName);
+					}
+					// Outer chain has an instanced import
+					else if (!Result.IsNone() && Result != NAME_FALSE)
+					{
+						const FString InstancedOuterNameStr = Result.ToString();
+						const bool bAdded = AddInstancedMapping(InstancedOuterNameStr, Import.GetPackageName());
+						if (bAdded)
+						{
+							UE_LOG(LogLinker, Warning, TEXT("Mapping for '%s' with external package '%s' not provided while outer '%s' is instanced.")
+								, *GetImportPathName(ImportIndex), *Import.GetPackageName().ToString(), *InstancedOuterNameStr);
+						}
+					}
 				}
 			}
 		}
