GC barrier support for FFieldVariant

diff --git a/Engine/Plugins/Runtime/ModelViewViewModel/Source/ModelViewViewModelEditor/Private/Widgets/SMVVMViewModelPanel.cpp b/Engine/Plugins/Runtime/ModelViewViewModel/Source/ModelViewViewModelEditor/Private/Widgets/SMVVMViewModelPanel.cpp
--- a/Engine/Plugins/Runtime/ModelViewViewModel/Source/ModelViewViewModelEditor/Private/Widgets/SMVVMViewModelPanel.cpp
+++ b/Engine/Plugins/Runtime/ModelViewViewModel/Source/ModelViewViewModelEditor/Private/Widgets/SMVVMViewModelPanel.cpp
@@ -532,7 +532,7 @@ FReply SMVVMViewModelPanel::HandleDragDetected(const FGeometry& MyGeometry, cons
 			if (const FGuid* Id = PropertyViewerHandles.Find(ContainerHandle))
 			{
 				TArray<FFieldVariant> FieldsArray;
-				for (const FFieldVariant Field : Fields)
+				for (const FFieldVariant& Field : Fields)
 				{
 					FieldsArray.Add(Field);
 				}
diff --git a/Engine/Source/Runtime/AdvancedWidgets/Private/Widgets/PropertyViewer/SPropertyViewerImpl.cpp b/Engine/Source/Runtime/AdvancedWidgets/Private/Widgets/PropertyViewer/SPropertyViewerImpl.cpp
--- a/Engine/Source/Runtime/AdvancedWidgets/Private/Widgets/PropertyViewer/SPropertyViewerImpl.cpp
+++ b/Engine/Source/Runtime/AdvancedWidgets/Private/Widgets/PropertyViewer/SPropertyViewerImpl.cpp
@@ -320,7 +320,7 @@ void FTreeNode::BuildChildNodesRecursive(IFieldIterator& FieldIterator, IFieldEx
 
 	if (ChildStructType)
 	{
-		for (const FFieldVariant FieldIt : FieldIterator.GetFields(ChildStructType))
+		for (const FFieldVariant& FieldIt : FieldIterator.GetFields(ChildStructType))
 		{
 			if (const FProperty* PropertyIt = FieldIt.Get<FProperty>())
 			{
diff --git a/Engine/Source/Runtime/CoreUObject/Private/UObject/Field.cpp b/Engine/Source/Runtime/CoreUObject/Private/UObject/Field.cpp
--- a/Engine/Source/Runtime/CoreUObject/Private/UObject/Field.cpp
+++ b/Engine/Source/Runtime/CoreUObject/Private/UObject/Field.cpp
@@ -459,6 +459,10 @@ void FField::AddReferencedObjects(FReferenceCollector& Collector)
 	if (OwnerUObject)
 	{
 		Collector.AddReferencedObject(OwnerUObject);
+		if (!OwnerUObject)
+		{
+			Owner = FFieldVariant{};
+		}
 	}
 }
 
diff --git a/Engine/Source/Runtime/CoreUObject/Public/UObject/Field.h b/Engine/Source/Runtime/CoreUObject/Public/UObject/Field.h
--- a/Engine/Source/Runtime/CoreUObject/Public/UObject/Field.h
+++ b/Engine/Source/Runtime/CoreUObject/Public/UObject/Field.h
@@ -271,6 +271,14 @@ class FFieldVariant
 
 	static constexpr uintptr_t UObjectMask = 0x1;
 
+	void ConditionallyMarkAsReachable()
+	{
+		if (IsUObject() && Container.Object && UE::GC::Private::GIsIncrementalReachabilityPending)
+		{
+			Container.Object->MarkAsReachable();
+		}
+	}
+	
 public:
 
 	FFieldVariant()
@@ -292,6 +300,7 @@ public:
 	{
 		Container.Object = const_cast<UObject*>(ImplicitConv<const UObject*>(InObject));
 		Container.Object = (UObject*)((uintptr_t)Container.Object | UObjectMask);
+		ConditionallyMarkAsReachable();
 	}
 
 	FFieldVariant(TYPE_OF_NULLPTR)
@@ -299,6 +308,32 @@ public:
 	{
 	}
 
+	FFieldVariant(const FFieldVariant& Other)
+		: Container(Other.Container)
+	{
+		ConditionallyMarkAsReachable();
+	}
+	
+	FFieldVariant& operator=(const FFieldVariant& Other)
+	{
+		Container = Other.Container;
+		ConditionallyMarkAsReachable();
+		return *this;
+	}
+	
+	FFieldVariant(FFieldVariant&& Other)
+		: Container(Other.Container)
+	{
+		ConditionallyMarkAsReachable();
+	}
+	
+	FFieldVariant& operator=(FFieldVariant&& Other)
+	{
+		Container = Other.Container;
+		ConditionallyMarkAsReachable();
+		return *this;
+	}
+
 	inline bool IsUObject() const
 	{
 		return (uintptr_t)Container.Object & UObjectMask;
