Cooker: DiffOnly: fix mismatched alloc/malloc calls on the previousbytes read from disk. They were being allocated with FMemory::Malloc but freed with delete. This mismatch is detected by ASAN and can cause crashes.
#rnx

diff --git a/Engine/Source/Editor/UnrealEd/Private/Cooker/LooseCookedPackageWriter.cpp b/Engine/Source/Editor/UnrealEd/Private/Cooker/LooseCookedPackageWriter.cpp
--- a/Engine/Source/Editor/UnrealEd/Private/Cooker/LooseCookedPackageWriter.cpp
+++ b/Engine/Source/Editor/UnrealEd/Private/Cooker/LooseCookedPackageWriter.cpp
@@ -523,10 +523,12 @@ void FLooseCookedPackageWriter::UpdateManifest(FRecord& Record)
 bool FLooseCookedPackageWriter::GetPreviousCookedBytes(const FPackageInfo& Info, FPreviousCookedBytesData& OutData)
 {
 	UE::ArchiveStackTrace::FPackageData ExistingPackageData;
-	UE::ArchiveStackTrace::LoadPackageIntoMemory(*Info.LooseFilePath, ExistingPackageData, OutData.Data);
+	TUniquePtr<uint8, UE::ArchiveStackTrace::FDeleteByFree> Bytes;
+	UE::ArchiveStackTrace::LoadPackageIntoMemory(*Info.LooseFilePath, ExistingPackageData, Bytes);
 	OutData.Size = ExistingPackageData.Size;
 	OutData.HeaderSize = ExistingPackageData.HeaderSize;
 	OutData.StartOffset = ExistingPackageData.StartOffset;
+	OutData.Data.Reset(Bytes.Release());
 	return OutData.Data.IsValid();
 }
 
diff --git a/Engine/Source/Editor/UnrealEd/Private/Tests/DiffWriterArchiveTests.cpp b/Engine/Source/Editor/UnrealEd/Private/Tests/DiffWriterArchiveTests.cpp
--- a/Engine/Source/Editor/UnrealEd/Private/Tests/DiffWriterArchiveTests.cpp
+++ b/Engine/Source/Editor/UnrealEd/Private/Tests/DiffWriterArchiveTests.cpp
@@ -32,7 +32,8 @@ ICookedPackageWriter::FPreviousCookedBytesData ToPackageData(FLargeMemoryWriter&
 	int64 Size = Ar.TotalSize(); // Cache this before calling ReleaseOwnership
 	return ICookedPackageWriter::FPreviousCookedBytesData
 	{
-		TUniquePtr<uint8>(Ar.ReleaseOwnership()), Size, 0 /* HeaderSize */, 0 /* StartOffset */
+		TUniquePtr<uint8, ICookedPackageWriter::FDeleteByFree>(Ar.ReleaseOwnership()),
+			Size, 0 /* HeaderSize */, 0 /* StartOffset */
 	};
 }
 
diff --git a/Engine/Source/Runtime/Core/Public/Serialization/PackageWriter.h b/Engine/Source/Runtime/Core/Public/Serialization/PackageWriter.h
--- a/Engine/Source/Runtime/Core/Public/Serialization/PackageWriter.h
+++ b/Engine/Source/Runtime/Core/Public/Serialization/PackageWriter.h
@@ -317,9 +317,17 @@ public:
 	virtual void UpdatePackageModificationStatus(FName PackageName, bool bIterativelyUnmodified,
 		bool& bInOutShouldIterativelySkip) = 0;
 
+	struct FDeleteByFree
+	{
+		void operator()(void* Ptr) const
+		{
+			FMemory::Free(Ptr);
+		}
+	};
+
 	struct FPreviousCookedBytesData
 	{
-		TUniquePtr<uint8> Data;
+		TUniquePtr<uint8, FDeleteByFree> Data;
 		int64 Size;
 		int64 HeaderSize;
 		int64 StartOffset;
diff --git a/Engine/Source/Runtime/CoreUObject/Private/Serialization/ArchiveStackTrace.cpp b/Engine/Source/Runtime/CoreUObject/Private/Serialization/ArchiveStackTrace.cpp
--- a/Engine/Source/Runtime/CoreUObject/Private/Serialization/ArchiveStackTrace.cpp
+++ b/Engine/Source/Runtime/CoreUObject/Private/Serialization/ArchiveStackTrace.cpp
@@ -116,7 +116,7 @@ namespace UE::ArchiveStackTrace
 {
 
 bool LoadPackageIntoMemory(const TCHAR* InFilename, UE::ArchiveStackTrace::FPackageData& OutPackageData,
-	TUniquePtr<uint8>& OutLoadedBytes)
+	TUniquePtr<uint8, FDeleteByFree>& OutLoadedBytes)
 {
 	TUniquePtr<FArchive> UAssetFileArchive(IFileManager::Get().CreateFileReader(InFilename));
 	if (!UAssetFileArchive || UAssetFileArchive->TotalSize() == 0)
@@ -145,7 +145,7 @@ bool LoadPackageIntoMemory(const TCHAR* InFilename, UE::ArchiveStackTrace::FPack
 				OutPackageData.Size += ExpFileArchive->TotalSize();
 			}
 		}
-		OutLoadedBytes.Reset(new uint8[OutPackageData.Size]);
+		OutLoadedBytes.Reset(reinterpret_cast<uint8*>(FMemory::Malloc(OutPackageData.Size)));
 		OutPackageData.Data = OutLoadedBytes.Get();
 		UAssetFileArchive->Serialize(OutPackageData.Data, UAssetFileArchive->TotalSize());
 
@@ -668,7 +668,8 @@ FUObjectSerializeContext* FArchiveStackTrace::GetSerializeContext()
 	return StackTraceWriter.GetSerializeContext();
 }
 
-bool FArchiveStackTrace::LoadPackageIntoMemory(const TCHAR* InFilename, FPackageData& OutPackageData, TUniquePtr<uint8>& OutLoadedBytes)
+bool FArchiveStackTrace::LoadPackageIntoMemory(const TCHAR* InFilename, FPackageData& OutPackageData,
+	TUniquePtr<uint8, UE::ArchiveStackTrace::FDeleteByFree>& OutLoadedBytes)
 {
 	UE::ArchiveStackTrace::FPackageData PackageData;
 	bool bResult = UE::ArchiveStackTrace::LoadPackageIntoMemory(InFilename, PackageData, OutLoadedBytes);
@@ -953,7 +954,7 @@ void FArchiveStackTraceWriter::Compare(
 
 void FArchiveStackTrace::CompareWith(const TCHAR* InFilename, const int64 TotalHeaderSize, const TCHAR* CallstackCutoffText, const int32 MaxDiffsToLog, TMap<FName, FArchiveDiffStats>& OutStats)
 {
-	TUniquePtr<uint8> SourcePackageBytes;
+	TUniquePtr<uint8, UE::ArchiveStackTrace::FDeleteByFree> SourcePackageBytes;
 	FPackageData SourcePackage;
 	LoadPackageIntoMemory(InFilename, SourcePackage, SourcePackageBytes);
 	CompareWith(SourcePackage, InFilename, TotalHeaderSize, CallstackCutoffText, MaxDiffsToLog, OutStats);
@@ -1150,7 +1151,7 @@ bool FArchiveStackTraceWriter::GenerateDiffMap(const FPackageData& SourcePackage
 
 bool FArchiveStackTrace::GenerateDiffMap(const TCHAR* InFilename, int64 TotalHeaderSize, int32 MaxDiffsToFind, FArchiveDiffMap& OutDiffMap)
 {
-	TUniquePtr<uint8> SourcePackageBytes;
+	TUniquePtr<uint8, UE::ArchiveStackTrace::FDeleteByFree> SourcePackageBytes;
 	FPackageData SourcePackage;
 	if (!LoadPackageIntoMemory(InFilename, SourcePackage, SourcePackageBytes))
 	{
@@ -1206,7 +1207,7 @@ bool FArchiveStackTrace::GenerateDiffMap(const FPackageData& SourcePackage, int6
 
 bool FArchiveStackTrace::IsIdentical(const TCHAR* InFilename, int64 BufferSize, const uint8* BufferData)
 {
-	TUniquePtr<uint8> SourcePackageBytes;
+	TUniquePtr<uint8, UE::ArchiveStackTrace::FDeleteByFree> SourcePackageBytes;
 	FPackageData SourcePackage;
 	if (!LoadPackageIntoMemory(InFilename, SourcePackage, SourcePackageBytes))
 	{
@@ -1741,7 +1742,7 @@ void FArchiveStackTraceReader::Serialize(void* OutData, int64 Num)
 FArchiveStackTraceReader* FArchiveStackTraceReader::CreateFromFile(const TCHAR* InFilename)
 {
 	FArchiveStackTraceReader* Reader = nullptr;
-	TUniquePtr<uint8> PackageBytes;
+	TUniquePtr<uint8, UE::ArchiveStackTrace::FDeleteByFree> PackageBytes;
 	UE::ArchiveStackTrace::FPackageData PackageData;
 	if (UE::ArchiveStackTrace::LoadPackageIntoMemory(InFilename, PackageData, PackageBytes))
 	{
diff --git a/Engine/Source/Runtime/CoreUObject/Public/Serialization/ArchiveStackTrace.h b/Engine/Source/Runtime/CoreUObject/Public/Serialization/ArchiveStackTrace.h
--- a/Engine/Source/Runtime/CoreUObject/Public/Serialization/ArchiveStackTrace.h
+++ b/Engine/Source/Runtime/CoreUObject/Public/Serialization/ArchiveStackTrace.h
@@ -73,9 +73,19 @@ struct FPackageData
 	int64 HeaderSize = 0;
 	int64 StartOffset = 0;
 };
+
+struct FDeleteByFree
+{
+	void operator()(void* Ptr) const
+	{
+		FMemory::Free(Ptr);
+	}
+};
+
 /** Helper function to load package contents into memory. Supports EDL packages. */
 COREUOBJECT_API bool LoadPackageIntoMemory(const TCHAR* InFilename,
-	FPackageData& OutPackageData, TUniquePtr<uint8>& OutLoadedBytes);
+	FPackageData& OutPackageData, TUniquePtr<uint8, FDeleteByFree>& OutLoadedBytes);
+
 COREUOBJECT_API void ForceKillPackageAndLinker(FLinkerLoad* Linker);
 COREUOBJECT_API bool ShouldIgnoreDiff();
 COREUOBJECT_API bool ShouldBypassDiff();
@@ -428,7 +438,7 @@ public:
 
 	/** Helper function to load package contents into memory. Supports EDL packages. */
 	static COREUOBJECT_API bool LoadPackageIntoMemory(const TCHAR* InFilename, FPackageData& OutPackageData,
-		TUniquePtr<uint8>& OutLoadedBytes);
+		TUniquePtr<uint8, UE::ArchiveStackTrace::FDeleteByFree>& OutLoadedBytes);
 
 private:
 	FArchiveCallstacks Callstacks;
