Added new PerPlatformConfig class flag indicating the config properties can be overridden in a Platform/PlatformEngine.ini file.
When a class is marked PerPlatformConfig, in the editor the class will be loaded from the platform inis of the current PreviewPlatform on load, and when the PreviewPlatform is changed, LoadConfig will be called to reload the settings from the platform ini file.

diff --git a/Engine/Source/Editor/UnrealEd/Private/EditorEngine.cpp b/Engine/Source/Editor/UnrealEd/Private/EditorEngine.cpp
--- a/Engine/Source/Editor/UnrealEd/Private/EditorEngine.cpp
+++ b/Engine/Source/Editor/UnrealEd/Private/EditorEngine.cpp
@@ -496,6 +496,9 @@ UEditorEngine::UEditorEngine(const FObjectInitializer& ObjectInitializer)
 
 	// The AssetRegistry module is needed early in initialization functions so load it here rather than in Init
 	FModuleManager::LoadModuleChecked<FAssetRegistryModule>("AssetRegistry");
+
+	// Callback to get the preview platform is used for PerPlatformConfig classes
+	UObject::OnGetPreviewPlatform.BindUObject(this, &UEditorEngine::GetPreviewPlatformName);
 }
 
 
@@ -1400,6 +1403,7 @@ void UEditorEngine::FinishDestroy()
 		}
 
 		// Unregister events
+		UObject::OnGetPreviewPlatform.Unbind();
 		FEditorDelegates::MapChange.RemoveAll(this);
 		FCoreDelegates::ModalMessageDialog.Unbind();
 		FCoreUObjectDelegates::ShouldLoadOnTop.Unbind();
@@ -7780,6 +7784,15 @@ void UEditorEngine::SetPreviewPlatform(const FPreviewPlatformInfo& NewPreviewPla
 		}
 	}
 
+	// Update any PerPlatformConfig class defaults or instances
+	for (FThreadSafeObjectIterator ObjIterator(UObject::StaticClass(), RF_NoFlags); ObjIterator; ++ObjIterator)
+	{
+		if ((*ObjIterator) && ObjIterator->GetClass()->HasAnyClassFlags(CLASS_PerPlatformConfig))
+		{
+			ObjIterator->LoadConfig();
+		}
+	}
+
 	constexpr bool bUpdateProgressDialog = true;
 	constexpr bool bCacheAllRemainingShaders = false;
 
@@ -7888,8 +7901,18 @@ void UEditorEngine::ToggleFeatureLevelPreview()
 		UDeviceProfileManager::Get().RestorePreviewDeviceProfile();
 	}
 
+	// Update any PerPlatformConfig class defaults or instances
+	for (FThreadSafeObjectIterator ObjIterator(UObject::StaticClass(), RF_NoFlags); ObjIterator; ++ObjIterator)
+	{
+		if ((*ObjIterator) && ObjIterator->GetClass()->HasAnyClassFlags(CLASS_PerPlatformConfig))
+		{
+			ObjIterator->LoadConfig();
+		}
+	}
+
 	Scalability::ApplyCachedQualityLevelForShaderPlatform(GetActiveShaderPlatform());
 	OnEffectivePreviewShaderPlatformChange();
+
 	PreviewPlatformChanged.Broadcast();
 
 	UStaticMesh::OnLodStrippingQualityLevelChanged(nullptr);
diff --git a/Engine/Source/Programs/Shared/EpicGames.Core/UnrealEngineTypes.cs b/Engine/Source/Programs/Shared/EpicGames.Core/UnrealEngineTypes.cs
--- a/Engine/Source/Programs/Shared/EpicGames.Core/UnrealEngineTypes.cs
+++ b/Engine/Source/Programs/Shared/EpicGames.Core/UnrealEngineTypes.cs
@@ -180,10 +180,9 @@ namespace EpicGames.Core
 		Interface = 0x00004000u,
 
 		/// <summary>
-		/// Do not export a constructor for this class, assuming it is in the cpptext
-		/// </summary>
-		[Obsolete("No longer used in the engine.")]
-		CustomConstructor = 0x00008000u,
+		/// Config for this class is overridden in platform inis, reload when previewing platforms
+		/// </summary>	
+		PerPlatformConfig = 0x00008000u,
 
 		/// <summary>
 		/// all properties and functions in this class are const and should be exported as const
@@ -289,7 +288,7 @@ namespace EpicGames.Core
 		/// <summary>
 		/// This is used as a mask for the flags put into generated code for "compiled in" classes.
 		/// </summary>
-		SaveInCompiledInClasses = Abstract | DefaultConfig | GlobalUserConfig | ProjectUserConfig | Config | Transient | Optional | Native | NotPlaceable | PerObjectConfig |
+		SaveInCompiledInClasses = Abstract | DefaultConfig | GlobalUserConfig | ProjectUserConfig | PerPlatformConfig | Config | Transient | Optional | Native | NotPlaceable | PerObjectConfig |
 			ConfigDoNotCheckDefaults | EditInlineNew | CollapseCategories | Interface | DefaultToInstanced | HasInstancedReference | Hidden | Deprecated |
 			HideDropDown | Intrinsic | Const | MinimalAPI | RequiredAPI | MatchedSerializers | NeedsDeferredDependencyLoading,
 	};
diff --git a/Engine/Source/Programs/Shared/EpicGames.UHT/Exporters/CodeGen/UhtHeaderCodeGeneratorHFile.cs b/Engine/Source/Programs/Shared/EpicGames.UHT/Exporters/CodeGen/UhtHeaderCodeGeneratorHFile.cs
--- a/Engine/Source/Programs/Shared/EpicGames.UHT/Exporters/CodeGen/UhtHeaderCodeGeneratorHFile.cs
+++ b/Engine/Source/Programs/Shared/EpicGames.UHT/Exporters/CodeGen/UhtHeaderCodeGeneratorHFile.cs
@@ -1405,6 +1405,10 @@ namespace EpicGames.UHT.Exporters.CodeGen
 			{
 				builder.Append(" | CLASS_ProjectUserConfig");
 			}
+			if (classObj.ClassFlags.HasAnyFlags(EClassFlags.PerPlatformConfig))
+			{
+				builder.Append(" | CLASS_PerPlatformConfig");
+			}		
 			if (classObj.ClassFlags.HasAnyFlags(EClassFlags.Config))
 			{
 				builder.Append(" | CLASS_Config");
diff --git a/Engine/Source/Programs/Shared/EpicGames.UHT/Specifiers/UhtClassSpecifiers.cs b/Engine/Source/Programs/Shared/EpicGames.UHT/Specifiers/UhtClassSpecifiers.cs
--- a/Engine/Source/Programs/Shared/EpicGames.UHT/Specifiers/UhtClassSpecifiers.cs
+++ b/Engine/Source/Programs/Shared/EpicGames.UHT/Specifiers/UhtClassSpecifiers.cs
@@ -120,6 +120,13 @@ namespace EpicGames.UHT.Parsers
 			classObj.AddClassFlags(EClassFlags.PerObjectConfig);
 		}
 
+		[UhtSpecifier(Extends = UhtTableNames.Class, ValueType = UhtSpecifierValueType.Legacy)]
+		private static void PerPlatformConfigSpecifier(UhtSpecifierContext specifierContext)
+		{
+			UhtClass classObj = (UhtClass)specifierContext.Type;
+			classObj.AddClassFlags(EClassFlags.PerPlatformConfig);
+		}
+
 		[UhtSpecifier(Extends = UhtTableNames.Class, ValueType = UhtSpecifierValueType.Legacy)]
 		private static void ConfigDoNotCheckDefaultsSpecifier(UhtSpecifierContext specifierContext)
 		{
diff --git a/Engine/Source/Runtime/CoreUObject/Private/UObject/Obj.cpp b/Engine/Source/Runtime/CoreUObject/Private/UObject/Obj.cpp
--- a/Engine/Source/Runtime/CoreUObject/Private/UObject/Obj.cpp
+++ b/Engine/Source/Runtime/CoreUObject/Private/UObject/Obj.cpp
@@ -130,6 +130,7 @@ static UPackage*			GObjTransientPkg								= NULL;
 	static FUObjectAnnotationSparse<FPropagatedEditChangeAnnotation, true> PropagatedEditChangeAnnotation;
 	UObject::FAssetRegistryTag::FOnGetObjectAssetRegistryTags UObject::FAssetRegistryTag::OnGetExtraObjectTags;
 	UObject::FAssetRegistryTag::FOnGetExtendedAssetRegistryTagsForSave UObject::FAssetRegistryTag::OnGetExtendedAssetRegistryTagsForSave;
+	UObject::FOnGetPreviewPlatform UObject::OnGetPreviewPlatform;
 #endif // WITH_EDITOR
 
 UObject::UObject( EStaticConstructor, EObjectFlags InFlags )
@@ -2664,16 +2665,29 @@ void UObject::LoadConfig( UClass* ConfigClass/*=NULL*/, const TCHAR* InFilename/
 
 	const bool bPerObject = UsesPerObjectConfig(this);
 
-	// does the class want to override the platform hierarchy (ignored if we passd in a specific ini file),
+	// does the class want to override the platform hierarchy (ignored if we passed in a specific ini file),
 	// and if the name isn't the current running platform (no need to load extra files if already in GConfig)
-	bool bUseConfigOverride = InFilename == nullptr && GetConfigOverridePlatform() != nullptr &&
-		FCString::Stricmp(GetConfigOverridePlatform(), ANSI_TO_TCHAR(FPlatformProperties::IniPlatformName())) != 0;
+	bool bUseConfigOverride = false;
 	FConfigFile* OverrideConfigFile = nullptr;
 	FConfigFile LocalOverrideConfig;
-	if (bUseConfigOverride)
+
+	if (InFilename == nullptr && GetConfigOverridePlatform() != nullptr &&
+		FCString::Stricmp(GetConfigOverridePlatform(), ANSI_TO_TCHAR(FPlatformProperties::IniPlatformName())) != 0)
 	{
 		OverrideConfigFile = FConfigCacheIni::FindOrLoadPlatformConfig(LocalOverrideConfig, *GetClass()->ClassConfigName.ToString(), GetConfigOverridePlatform());
+		bUseConfigOverride = true;
+	}
+#if WITH_EDITOR
+	else if (GetClass()->HasAnyClassFlags(CLASS_PerPlatformConfig) && UObject::OnGetPreviewPlatform.IsBound())
+	{
+		FName PreviewPlatform;
+		if (UObject::OnGetPreviewPlatform.Execute(PreviewPlatform))
+		{
+			OverrideConfigFile = FConfigCacheIni::FindOrLoadPlatformConfig(LocalOverrideConfig, *GetClass()->ClassConfigName.ToString(), *PreviewPlatform.ToString());
+			bUseConfigOverride = true;
+		}
 	}
+#endif
 
 	FString ClassSection;
 	FString ClassPathSection;
@@ -2801,6 +2815,8 @@ void UObject::LoadConfig( UClass* ConfigClass/*=NULL*/, const TCHAR* InFilename/
 			Key = CVarName;
 			PortFlags |= PPF_ConsoleVariable;
 		}
+
+		const bool bIsPerPlatformConfig = GetClass()->HasAnyClassFlags(CLASS_PerPlatformConfig);
 #endif // #if WITH_EDITOR
 
 		UE_LOG(LogConfig, Verbose, TEXT("   Loading value for %s from [%s]"), *Key, *ClassSection);
@@ -2841,6 +2857,18 @@ void UObject::LoadConfig( UClass* ConfigClass/*=NULL*/, const TCHAR* InFilename/
 		}
 		else
 		{
+			FScriptArrayHelper_InContainer ArrayHelper(Array, this);
+
+#if WITH_EDITOR
+			// Empty out any array properties if this is a PerPlatformConfig class
+			// as we are replacing the values with the Platform's version when entering
+			// a new preview platform.
+			if (bIsPerPlatformConfig)
+			{
+				ArrayHelper.EmptyValues();
+			}
+#endif
+
 			FConfigSection* Sec = GetConfigSection(*ClassSection, *PropFileName);
 			if (!Sec && bPerObject && ClassPathSection.Len())
 			{
@@ -2852,8 +2880,8 @@ void UObject::LoadConfig( UClass* ConfigClass/*=NULL*/, const TCHAR* InFilename/
 				const FName KeyName(*Key, FNAME_Find);
 				Sec->MultiFind(KeyName,List);
 
-				FScriptArrayHelper_InContainer ArrayHelper(Array, this);
 				const int32 Size = Array->Inner->ElementSize;
+
 				// Only override default properties if there is something to override them with.
 				if ( List.Num() > 0 )
 				{
diff --git a/Engine/Source/Runtime/CoreUObject/Public/UObject/Object.h b/Engine/Source/Runtime/CoreUObject/Public/UObject/Object.h
--- a/Engine/Source/Runtime/CoreUObject/Public/UObject/Object.h
+++ b/Engine/Source/Runtime/CoreUObject/Public/UObject/Object.h
@@ -1291,7 +1291,18 @@ public:
 	 */
 	COREUOBJECT_API FString GetProjectUserConfigFilename() const;
 
-	/** Returns the override config hierarchy platform (if NDAd platforms need defaults to not be in Base*.ini but still want editor to load them) */
+#if WITH_EDITOR
+	/** 
+	 * Delegate to return the current preview platform name, or NAME_None if no preview platform is currently active.
+	 * Used by LoadConfig for CLASS_PerPlatformConfig objects.
+	 */
+	DECLARE_DELEGATE_RetVal_OneParam(bool, FOnGetPreviewPlatform, FName&);
+	COREUOBJECT_API static FOnGetPreviewPlatform OnGetPreviewPlatform;
+#endif
+
+	/** 
+	 * Returns the override config hierarchy platform (if NDAd platforms need defaults to not be in Base*.ini but still want editor to load them) 
+	 */
 	virtual const TCHAR* GetConfigOverridePlatform() const { return nullptr; }
 
 	/**
diff --git a/Engine/Source/Runtime/CoreUObject/Public/UObject/ObjectMacros.h b/Engine/Source/Runtime/CoreUObject/Public/UObject/ObjectMacros.h
--- a/Engine/Source/Runtime/CoreUObject/Public/UObject/ObjectMacros.h
+++ b/Engine/Source/Runtime/CoreUObject/Public/UObject/ObjectMacros.h
@@ -219,8 +219,8 @@ enum EClassFlags
 	CLASS_CollapseCategories  = 0x00002000u,
 	/** Class is an interface **/
 	CLASS_Interface           = 0x00004000u,
-	/**  Do not export a constructor for this class, assuming it is in the cpptext **/
-	CLASS_CustomConstructor UE_DEPRECATED(5.1, "CLASS_CustomConstructor should no longer be used. It is no longer being set by engine code.") = 0x00008000u,
+	/**  Config for this class is overridden in platform inis, reload when previewing platforms **/
+	CLASS_PerPlatformConfig   = 0x00008000u,
 	/** all properties and functions in this class are const and should be exported as const */
 	CLASS_Const			      = 0x00010000u,
 
@@ -266,7 +266,7 @@ ENUM_CLASS_FLAGS(EClassFlags);
 
 /** Flags to inherit from base class */
 #define CLASS_Inherit ((EClassFlags)(CLASS_Transient | CLASS_Optional | CLASS_DefaultConfig | CLASS_Config | CLASS_PerObjectConfig | CLASS_ConfigDoNotCheckDefaults | CLASS_NotPlaceable \
-						| CLASS_Const | CLASS_HasInstancedReference | CLASS_Deprecated | CLASS_DefaultToInstanced | CLASS_GlobalUserConfig | CLASS_ProjectUserConfig | CLASS_NeedsDeferredDependencyLoading))
+						| CLASS_Const | CLASS_HasInstancedReference | CLASS_Deprecated | CLASS_DefaultToInstanced | CLASS_GlobalUserConfig | CLASS_ProjectUserConfig | CLASS_PerPlatformConfig | CLASS_NeedsDeferredDependencyLoading))
 
 /** These flags will be cleared by the compiler when the class is parsed during script compilation */
 #define CLASS_RecompilerClear ((EClassFlags)(CLASS_Inherit | CLASS_Abstract | CLASS_Native | CLASS_Intrinsic | CLASS_TokenStreamAssembled))
@@ -283,12 +283,14 @@ ENUM_CLASS_FLAGS(EClassFlags);
 	CLASS_DefaultConfig | \
 	CLASS_GlobalUserConfig | \
 	CLASS_ProjectUserConfig | \
+	CLASS_PerPlatformConfig | \
 	CLASS_Config | \
 	CLASS_Transient | \
 	CLASS_Optional | \
 	CLASS_Native | \
 	CLASS_NotPlaceable | \
 	CLASS_PerObjectConfig | \
+	CLASS_PerPlatformConfig | \
 	CLASS_ConfigDoNotCheckDefaults | \
 	CLASS_EditInlineNew | \
 	CLASS_CollapseCategories | \
diff --git a/Engine/Source/Runtime/Engine/Public/MeshDrawCommandStatsSettings.h b/Engine/Source/Runtime/Engine/Public/MeshDrawCommandStatsSettings.h
--- a/Engine/Source/Runtime/Engine/Public/MeshDrawCommandStatsSettings.h
+++ b/Engine/Source/Runtime/Engine/Public/MeshDrawCommandStatsSettings.h
@@ -25,7 +25,7 @@ struct FMeshDrawCommandStatsBudget
 };
 
 /** User settings used by the MeshDrawCommandStats system. */
-UCLASS(Config=Engine, defaultconfig, meta = (DisplayName = "Mesh Stats"))
+UCLASS(Config=Engine, defaultconfig, perplatformconfig, meta = (DisplayName = "Mesh Stats"))
 class ENGINE_API UMeshDrawCommandStatsSettings : public UDeveloperSettings
 {
 	GENERATED_BODY()
