Fix impossible Cast, UObjectRedirector is not a UPackage
#rnx

diff --git a/Engine/Source/Runtime/CoreUObject/Private/UObject/ObjectRef.cpp b/Engine/Source/Runtime/CoreUObject/Private/UObject/ObjectRef.cpp
--- a/Engine/Source/Runtime/CoreUObject/Private/UObject/ObjectRef.cpp
+++ b/Engine/Source/Runtime/CoreUObject/Private/UObject/ObjectRef.cpp
@@ -123,10 +123,6 @@ static inline UPackage* FindOrLoadPackage(FName PackageName, int32 LoadFlags, OU
 	// 		 for additional behavior that we're not handling here yet.
 	FName* ScriptPackageName = FPackageName::FindScriptPackageName(PackageName);
 	UPackage* TargetPackage = (UPackage*)StaticFindObjectFastInternal(UPackage::StaticClass(), nullptr, PackageName);
-	if (UObjectRedirector* Redirector = dynamic_cast<UObjectRedirector*>(TargetPackage))
-	{
-		TargetPackage = (UPackage*)Redirector->DestinationObject;
-	}
 	if (!ScriptPackageName && !TargetPackage)
 	{
 		// @TODO: OBJPTR: When using the "external package" feature, we will have objects that have a differing package path vs "outer hierarchy" path
@@ -157,14 +153,6 @@ UClass* FObjectRef::ResolveObjectRefClass(uint32 LoadFlags /*= LOAD_None*/) cons
 		if (!ClassName.IsNone())
 		{
 			ClassObject = (UClass*)StaticFindObjectFastInternal(UClass::StaticClass(), ClassPackage, ClassName);
-			if (UObjectRedirector* Redirector = dynamic_cast<UObjectRedirector*>(ClassObject))
-			{
-				ClassObject = (UClass*)Redirector->DestinationObject;
-				if (ClassObject != nullptr)
-				{
-					ClassPackage = ClassObject->GetPackage();
-				}
-			}
 
 			if (ClassObject)
 			{
