Deprecated some UnrealHeaderTool only FProperty API.

diff --git a/Engine/Source/Runtime/CoreUObject/Private/UObject/Property.cpp b/Engine/Source/Runtime/CoreUObject/Private/UObject/Property.cpp
--- a/Engine/Source/Runtime/CoreUObject/Private/UObject/Property.cpp
+++ b/Engine/Source/Runtime/CoreUObject/Private/UObject/Property.cpp
@@ -996,6 +996,7 @@ bool FProperty::PassCPPArgsByRef() const
 }
 
 
+PRAGMA_DISABLE_DEPRECATION_WARNINGS
 void FProperty::ExportCppDeclaration(FOutputDevice& Out, EExportedDeclaration::Type DeclarationType, const TCHAR* ArrayDimOverride, uint32 AdditionalExportCPPFlags
 	, bool bSkipParameterName, const FString* ActualCppType, const FString* ActualExtendedType, const FString* ActualParameterName) const
 {
@@ -1135,6 +1136,7 @@ void FProperty::ExportCppDeclaration(FOutputDevice& Out, EExportedDeclaration::T
 		}
 	}
 }
+PRAGMA_ENABLE_DEPRECATION_WARNINGS
 
 bool FProperty::ExportText_Direct
 	(
diff --git a/Engine/Source/Runtime/CoreUObject/Private/UObject/PropertyArray.cpp b/Engine/Source/Runtime/CoreUObject/Private/UObject/PropertyArray.cpp
--- a/Engine/Source/Runtime/CoreUObject/Private/UObject/PropertyArray.cpp
+++ b/Engine/Source/Runtime/CoreUObject/Private/UObject/PropertyArray.cpp
@@ -448,11 +448,14 @@ FString FArrayProperty::GetCPPType( FString* ExtendedTypeText/*=NULL*/, uint32 C
 	return GetCPPTypeCustom(ExtendedTypeText, CPPExportFlags, InnerTypeText, InnerExtendedTypeText);
 }
 
+PRAGMA_DISABLE_DEPRECATION_WARNINGS
 FString FArrayProperty::GetCPPTypeForwardDeclaration() const
 {
 	checkSlow(Inner);
 	return Inner->GetCPPTypeForwardDeclaration();
 }
+PRAGMA_ENABLE_DEPRECATION_WARNINGS
+
 FString FArrayProperty::GetCPPMacroType( FString& ExtendedTypeText ) const
 {
 	checkSlow(Inner);
diff --git a/Engine/Source/Runtime/CoreUObject/Private/UObject/PropertyMap.cpp b/Engine/Source/Runtime/CoreUObject/Private/UObject/PropertyMap.cpp
--- a/Engine/Source/Runtime/CoreUObject/Private/UObject/PropertyMap.cpp
+++ b/Engine/Source/Runtime/CoreUObject/Private/UObject/PropertyMap.cpp
@@ -617,6 +617,7 @@ FString FMapProperty::GetCPPType(FString* ExtendedTypeText, uint32 CPPExportFlag
 	return GetCPPTypeCustom(ExtendedTypeText, CPPExportFlags, KeyTypeText, KeyExtendedTypeText, ValueTypeText, ValueExtendedTypeText);
 }
 
+PRAGMA_DISABLE_DEPRECATION_WARNINGS
 FString FMapProperty::GetCPPTypeForwardDeclaration() const
 {
 	checkSlow(KeyProp);
@@ -624,6 +625,7 @@ FString FMapProperty::GetCPPTypeForwardDeclaration() const
 	// Generates a single ' ' when no forward declaration is needed. Purely an aesthetic concern at this time:
 	return FString::Printf( TEXT("%s %s"), *KeyProp->GetCPPTypeForwardDeclaration(), *ValueProp->GetCPPTypeForwardDeclaration());
 }
+PRAGMA_ENABLE_DEPRECATION_WARNINGS
 
 FString FMapProperty::GetCPPMacroType( FString& ExtendedTypeText ) const
 {
diff --git a/Engine/Source/Runtime/CoreUObject/Private/UObject/PropertyOptional.cpp b/Engine/Source/Runtime/CoreUObject/Private/UObject/PropertyOptional.cpp
--- a/Engine/Source/Runtime/CoreUObject/Private/UObject/PropertyOptional.cpp
+++ b/Engine/Source/Runtime/CoreUObject/Private/UObject/PropertyOptional.cpp
@@ -121,12 +121,14 @@ FString FOptionalProperty::GetCPPMacroType(FString& ExtendedTypeText) const
 	return TEXT("TOPTIONAL");
 }
 
+PRAGMA_DISABLE_DEPRECATION_WARNINGS
 FString FOptionalProperty::GetCPPTypeForwardDeclaration() const
 {
 	// We assume that TOptional<> is globally known already and that we just need to make the value type known
 	checkSlow(ValueProperty);
 	return ValueProperty->GetCPPTypeForwardDeclaration();
 }
+PRAGMA_ENABLE_DEPRECATION_WARNINGS
 
 void FOptionalProperty::LinkInternal(FArchive& Ar)
 {
diff --git a/Engine/Source/Runtime/CoreUObject/Private/UObject/PropertySet.cpp b/Engine/Source/Runtime/CoreUObject/Private/UObject/PropertySet.cpp
--- a/Engine/Source/Runtime/CoreUObject/Private/UObject/PropertySet.cpp
+++ b/Engine/Source/Runtime/CoreUObject/Private/UObject/PropertySet.cpp
@@ -541,11 +541,13 @@ FString FSetProperty::GetCPPType(FString* ExtendedTypeText, uint32 CPPExportFlag
 	return GetCPPTypeCustom(ExtendedTypeText, CPPExportFlags, ElementTypeText, ElementExtendedTypeText);
 }
 
+PRAGMA_DISABLE_DEPRECATION_WARNINGS
 FString FSetProperty::GetCPPTypeForwardDeclaration() const
 {
 	checkSlow(ElementProp);
 	return ElementProp->GetCPPTypeForwardDeclaration();
 }
+PRAGMA_ENABLE_DEPRECATION_WARNINGS
 
 void FSetProperty::ExportText_Internal(FString& ValueStr, const void* ContainerOrPropertyPtr, EPropertyPointerType PropertyPointerType, const void* DefaultValue, UObject* Parent, int32 PortFlags, UObject* ExportRootScope) const
 {
diff --git a/Engine/Source/Runtime/CoreUObject/Public/UObject/EnumProperty.h b/Engine/Source/Runtime/CoreUObject/Public/UObject/EnumProperty.h
--- a/Engine/Source/Runtime/CoreUObject/Public/UObject/EnumProperty.h
+++ b/Engine/Source/Runtime/CoreUObject/Public/UObject/EnumProperty.h
@@ -64,6 +64,7 @@ public:
 	// FProperty interface
 	COREUOBJECT_API virtual FString GetCPPMacroType( FString& ExtendedTypeText ) const  override;
 	COREUOBJECT_API virtual FString GetCPPType( FString* ExtendedTypeText, uint32 CPPExportFlags ) const override;
+	UE_DEPRECATED(5.4, "UnrealHeaderTool only API.  No replacement available.")
 	COREUOBJECT_API virtual FString GetCPPTypeForwardDeclaration() const override;
 	COREUOBJECT_API virtual void LinkInternal(FArchive& Ar) override;
 	COREUOBJECT_API virtual bool Identical( const void* A, const void* B, uint32 PortFlags ) const override;
diff --git a/Engine/Source/Runtime/CoreUObject/Public/UObject/FieldPathProperty.h b/Engine/Source/Runtime/CoreUObject/Public/UObject/FieldPathProperty.h
--- a/Engine/Source/Runtime/CoreUObject/Public/UObject/FieldPathProperty.h
+++ b/Engine/Source/Runtime/CoreUObject/Public/UObject/FieldPathProperty.h
@@ -63,6 +63,7 @@ PRAGMA_ENABLE_DEPRECATION_WARNINGS
 
 	// UHT interface
 	COREUOBJECT_API virtual FString GetCPPMacroType(FString& ExtendedTypeText) const  override;
+	UE_DEPRECATED(5.4, "UnrealHeaderTool only API.  No replacement available.")
 	COREUOBJECT_API virtual FString GetCPPTypeForwardDeclaration() const override;
 	COREUOBJECT_API virtual FString GetCPPType(FString* ExtendedTypeText = nullptr, uint32 CPPExportFlags = 0) const override;
 	// End of UHT interface
diff --git a/Engine/Source/Runtime/CoreUObject/Public/UObject/PropertyOptional.h b/Engine/Source/Runtime/CoreUObject/Public/UObject/PropertyOptional.h
--- a/Engine/Source/Runtime/CoreUObject/Public/UObject/PropertyOptional.h
+++ b/Engine/Source/Runtime/CoreUObject/Public/UObject/PropertyOptional.h
@@ -203,6 +203,7 @@ public:
 	// UHT interface
 	virtual FString GetCPPType(FString* ExtendedTypeText = NULL, uint32 CPPExportFlags = 0) const override;
 	virtual FString GetCPPMacroType(FString& ExtendedTypeText) const override;
+	UE_DEPRECATED(5.4, "UnrealHeaderTool only API.  No replacement available.")
 	virtual FString GetCPPTypeForwardDeclaration() const override;
 	// End of UHT interface
 
diff --git a/Engine/Source/Runtime/CoreUObject/Public/UObject/TextProperty.h b/Engine/Source/Runtime/CoreUObject/Public/UObject/TextProperty.h
--- a/Engine/Source/Runtime/CoreUObject/Public/UObject/TextProperty.h
+++ b/Engine/Source/Runtime/CoreUObject/Public/UObject/TextProperty.h
@@ -64,6 +64,7 @@ protected:
 	COREUOBJECT_API virtual void ExportText_Internal( FString& ValueStr, const void* PropertyValueOrContainer, EPropertyPointerType PropertyPointerType, const void* DefaultValue, UObject* Parent, int32 PortFlags, UObject* ExportRootScope ) const override;
 	COREUOBJECT_API virtual const TCHAR* ImportText_Internal(const TCHAR* Buffer, void* ContainerOrPropertyPtr, EPropertyPointerType PropertyPointerType, UObject* OwnerObject, int32 PortFlags, FOutputDevice* ErrorText) const override;
 public:
+	UE_DEPRECATED(5.4, "UnrealHeaderTool only API.  No replacement available.")
 	COREUOBJECT_API virtual FString GetCPPTypeForwardDeclaration() const override;
 	// End of FProperty interface
 
diff --git a/Engine/Source/Runtime/CoreUObject/Public/UObject/UnrealType.h b/Engine/Source/Runtime/CoreUObject/Public/UObject/UnrealType.h
--- a/Engine/Source/Runtime/CoreUObject/Public/UObject/UnrealType.h
+++ b/Engine/Source/Runtime/CoreUObject/Public/UObject/UnrealType.h
@@ -110,7 +110,7 @@ enum EPropertyExportCPPFlags
 
 namespace EExportedDeclaration
 {
-	enum Type
+	enum UE_DEPRECATED(5.4, "UnrealHeaderTool only API.  No replacement available.") Type
 	{
 		Local,
 		Member,
@@ -227,20 +227,24 @@ public:
 	/** Gets a redirected property name, will return NAME_None if no redirection was found */
 	static COREUOBJECT_API FName FindRedirectedPropertyName(UStruct* ObjectStruct, FName OldName);
 
-	// UHT interface
-	COREUOBJECT_API void ExportCppDeclaration(FOutputDevice& Out, EExportedDeclaration::Type DeclarationType, const TCHAR* ArrayDimOverride = NULL, uint32 AdditionalExportCPPFlags = 0
-		, bool bSkipParameterName = false, const FString* ActualCppType = nullptr, const FString* ActualExtendedType = nullptr, const FString* ActualParameterName = nullptr) const;
-	COREUOBJECT_API virtual FString GetCPPMacroType( FString& ExtendedTypeText ) const;
-	COREUOBJECT_API virtual bool PassCPPArgsByRef() const;
-
 	/**
-	 * Returns the C++ name of the property, including the _DEPRECATED suffix if the 
+	 * Returns the C++ name of the property, including the _DEPRECATED suffix if the
 	 * property is deprecated.
 	 *
 	 * @return C++ name of property
 	 */
 	COREUOBJECT_API FString GetNameCPP() const;
 
+	// UHT interface
+	PRAGMA_DISABLE_DEPRECATION_WARNINGS
+	UE_DEPRECATED(5.4, "UnrealHeaderTool only API.  No replacement available.")
+	COREUOBJECT_API void ExportCppDeclaration(FOutputDevice& Out, EExportedDeclaration::Type DeclarationType, const TCHAR* ArrayDimOverride = NULL, uint32 AdditionalExportCPPFlags = 0
+		, bool bSkipParameterName = false, const FString* ActualCppType = nullptr, const FString* ActualExtendedType = nullptr, const FString* ActualParameterName = nullptr) const;
+	PRAGMA_ENABLE_DEPRECATION_WARNINGS
+	COREUOBJECT_API virtual FString GetCPPMacroType( FString& ExtendedTypeText ) const;
+	UE_DEPRECATED(5.4, "UnrealHeaderTool only API.  No replacement available.")
+	COREUOBJECT_API virtual bool PassCPPArgsByRef() const;
+
 	/**
 	 * Returns the text to use for exporting this property to header file.
 	 *
@@ -249,6 +253,7 @@ public:
 	 */
 	COREUOBJECT_API virtual FString GetCPPType( FString* ExtendedTypeText=NULL, uint32 CPPExportFlags=0 ) const PURE_VIRTUAL(FProperty::GetCPPType,return TEXT(""););
 
+	UE_DEPRECATED(5.4, "UnrealHeaderTool only API.  No replacement available.")
 	COREUOBJECT_API virtual FString GetCPPTypeForwardDeclaration() const PURE_VIRTUAL(FProperty::GetCPPTypeForwardDeclaration, return TEXT(""););
 	// End of UHT interface
 
@@ -1337,6 +1342,7 @@ public:
 	{
 		return FString(TTypeFundamentals::GetTypeName());
 	}
+	UE_DEPRECATED(5.4, "UnrealHeaderTool only API.  No replacement available.")
 	virtual bool PassCPPArgsByRef() const override
 	{
 		// non-pod data is passed by reference
@@ -1668,6 +1674,7 @@ PRAGMA_ENABLE_DEPRECATION_WARNINGS
 	}
 #endif // WITH_EDITORONLY_DATA
 
+	UE_DEPRECATED(5.4, "UnrealHeaderTool only API.  No replacement available.")
 	virtual FString GetCPPTypeForwardDeclaration() const override
 	{
 		return FString();
@@ -2390,6 +2397,7 @@ public:
 	// UHT interface
 	virtual FString GetCPPType( FString* ExtendedTypeText, uint32 CPPExportFlags ) const override;
 	virtual FString GetCPPMacroType( FString& ExtendedTypeText ) const override;
+	UE_DEPRECATED(5.4, "UnrealHeaderTool only API.  No replacement available.")
 	virtual FString GetCPPTypeForwardDeclaration() const override;
 	// End of UHT interface
 
@@ -2867,6 +2875,7 @@ PRAGMA_ENABLE_DEPRECATION_WARNINGS
 
 	// UHT interface
 	virtual FString GetCPPMacroType( FString& ExtendedTypeText ) const  override;
+	UE_DEPRECATED(5.4, "UnrealHeaderTool only API.  No replacement available.")
 	virtual FString GetCPPTypeForwardDeclaration() const override;
 	// End of UHT interface
 
@@ -3009,6 +3018,7 @@ PRAGMA_ENABLE_DEPRECATION_WARNINGS
 	virtual FString GetCPPMacroType( FString& ExtendedTypeText ) const  override;
 	virtual FString GetCPPType( FString* ExtendedTypeText, uint32 CPPExportFlags ) const override;
 	virtual FString GetCPPTypeCustom(FString* ExtendedTypeText, uint32 CPPExportFlags, const FString& InnerNativeTypeName) const override;
+	UE_DEPRECATED(5.4, "UnrealHeaderTool only API.  No replacement available.")
 	virtual FString GetCPPTypeForwardDeclaration() const override;
 	// End of UHT interface
 
@@ -3074,6 +3084,7 @@ PRAGMA_ENABLE_DEPRECATION_WARNINGS
 	virtual FString GetCPPMacroType( FString& ExtendedTypeText ) const  override;
 	virtual FString GetCPPType( FString* ExtendedTypeText, uint32 CPPExportFlags ) const override;
 	virtual FString GetCPPTypeCustom(FString* ExtendedTypeText, uint32 CPPExportFlags, const FString& InnerNativeTypeName) const override;
+	UE_DEPRECATED(5.4, "UnrealHeaderTool only API.  No replacement available.")
 	virtual FString GetCPPTypeForwardDeclaration() const override;
 
 	// End of UHT interface
@@ -3147,6 +3158,7 @@ PRAGMA_ENABLE_DEPRECATION_WARNINGS
 
 	// UHT interface
 	virtual FString GetCPPMacroType( FString& ExtendedTypeText ) const  override;
+	UE_DEPRECATED(5.4, "UnrealHeaderTool only API.  No replacement available.")
 	virtual FString GetCPPTypeForwardDeclaration() const override;
 	// End of UHT interface
 
@@ -3246,6 +3258,7 @@ PRAGMA_ENABLE_DEPRECATION_WARNINGS
 	// UHT interface
 	virtual FString GetCPPType(FString* ExtendedTypeText, uint32 CPPExportFlags)  const override;
 	virtual FString GetCPPMacroType( FString& ExtendedTypeText ) const  override;
+	UE_DEPRECATED(5.4, "UnrealHeaderTool only API.  No replacement available.")
 	virtual FString GetCPPTypeForwardDeclaration() const override;
 	// End of UHT interface
 
@@ -3362,6 +3375,7 @@ PRAGMA_ENABLE_DEPRECATION_WARNINGS
 	// UHT interface
 	virtual FString GetCPPType(FString* ExtendedTypeText, uint32 CPPExportFlags) const override;
 	virtual FString GetCPPMacroType( FString& ExtendedTypeText ) const  override;
+	UE_DEPRECATED(5.4, "UnrealHeaderTool only API.  No replacement available.")
 	virtual FString GetCPPTypeForwardDeclaration() const override;
 	// End of UHT interface
 
@@ -3448,6 +3462,7 @@ PRAGMA_ENABLE_DEPRECATION_WARNINGS
 	// UHT interface
 	virtual FString GetCPPMacroType( FString& ExtendedTypeText ) const  override;
 	virtual FString GetCPPType( FString* ExtendedTypeText, uint32 CPPExportFlags ) const override;
+	UE_DEPRECATED(5.4, "UnrealHeaderTool only API.  No replacement available.")
 	virtual FString GetCPPTypeForwardDeclaration() const override;
 	// End of UHT interface
 
@@ -3545,6 +3560,7 @@ protected:
 	virtual const TCHAR* ImportText_Internal(const TCHAR* Buffer, void* ContainerOrPropertyPtr, EPropertyPointerType PropertyPointerType, UObject* OwnerObject, int32 PortFlags, FOutputDevice* ErrorText) const override;
 public:
 	virtual EConvertFromTypeResult ConvertFromType(const FPropertyTag& Tag, FStructuredArchive::FSlot Slot, uint8* Data, UStruct* DefaultsStruct, const uint8* Defaults) override;
+	UE_DEPRECATED(5.4, "UnrealHeaderTool only API.  No replacement available.")
 	virtual FString GetCPPTypeForwardDeclaration() const override;
 	uint32 GetValueTypeHashInternal(const void* Src) const override;
 	// End of FProperty interface
@@ -3601,6 +3617,7 @@ protected:
 	virtual const TCHAR* ImportText_Internal(const TCHAR* Buffer, void* ContainerOrPropertyPtr, EPropertyPointerType PropertyPointerType, UObject* OwnerObject, int32 PortFlags, FOutputDevice* ErrorText) const override;
 public:
 	virtual EConvertFromTypeResult ConvertFromType(const FPropertyTag& Tag, FStructuredArchive::FSlot Slot, uint8* Data, UStruct* DefaultsStruct, const uint8* Defaults) override;
+	UE_DEPRECATED(5.4, "UnrealHeaderTool only API.  No replacement available.")
 	virtual FString GetCPPTypeForwardDeclaration() const override;
 	uint32 GetValueTypeHashInternal(const void* Src) const override;
 	// End of FProperty interface
@@ -3696,6 +3713,7 @@ PRAGMA_ENABLE_DEPRECATION_WARNINGS
 	// FProperty interface
 	virtual FString GetCPPMacroType( FString& ExtendedTypeText ) const  override;
 	virtual FString GetCPPType( FString* ExtendedTypeText, uint32 CPPExportFlags ) const override;
+	UE_DEPRECATED(5.4, "UnrealHeaderTool only API.  No replacement available.")
 	virtual FString GetCPPTypeForwardDeclaration() const override;
 	virtual void LinkInternal(FArchive& Ar) override;
 	virtual bool Identical( const void* A, const void* B, uint32 PortFlags ) const override;
@@ -3727,6 +3745,7 @@ public:
 	virtual void CopyValuesInternal( void* Dest, void const* Src, int32 Count  ) const override;
 	virtual void ClearValueInternal( void* Data ) const override;
 	virtual void DestroyValueInternal( void* Dest ) const override;
+	UE_DEPRECATED(5.4, "UnrealHeaderTool only API.  No replacement available.")
 	virtual bool PassCPPArgsByRef() const override;
 	virtual void InstanceSubobjects( void* Data, void const* DefaultData, UObject* Owner, struct FObjectInstancingGraph* InstanceGraph ) override;
 	virtual bool ContainsObjectReference(TArray<const FStructProperty*>& EncounteredStructProps, EPropertyObjectReferenceType InReferenceType = EPropertyObjectReferenceType::Strong) const override;
@@ -3832,6 +3851,7 @@ public:
 	// FProperty interface
 	virtual FString GetCPPMacroType(FString& ExtendedTypeText) const  override;
 	virtual FString GetCPPType(FString* ExtendedTypeText, uint32 CPPExportFlags) const override;
+	UE_DEPRECATED(5.4, "UnrealHeaderTool only API.  No replacement available.")
 	virtual FString GetCPPTypeForwardDeclaration() const override;
 	virtual void LinkInternal(FArchive& Ar) override;
 	virtual bool Identical(const void* A, const void* B, uint32 PortFlags) const override;
@@ -3863,6 +3883,7 @@ public:
 	virtual void CopyValuesInternal(void* Dest, void const* Src, int32 Count) const override;
 	virtual void ClearValueInternal(void* Data) const override;
 	virtual void DestroyValueInternal(void* Dest) const override;
+	UE_DEPRECATED(5.4, "UnrealHeaderTool only API.  No replacement available.")
 	virtual bool PassCPPArgsByRef() const override;
 	virtual void InstanceSubobjects(void* Data, void const* DefaultData, UObject* Owner, struct FObjectInstancingGraph* InstanceGraph) override;
 	virtual bool ContainsObjectReference(TArray<const FStructProperty*>& EncounteredStructProps, EPropertyObjectReferenceType InReferenceType = EPropertyObjectReferenceType::Strong) const override;
@@ -3975,6 +3996,7 @@ public:
 	// FProperty interface
 	virtual FString GetCPPMacroType(FString& ExtendedTypeText) const  override;
 	virtual FString GetCPPType(FString* ExtendedTypeText, uint32 CPPExportFlags) const override;
+	UE_DEPRECATED(5.4, "UnrealHeaderTool only API.  No replacement available.")
 	virtual FString GetCPPTypeForwardDeclaration() const override;
 	virtual void LinkInternal(FArchive& Ar) override;
 	virtual bool Identical(const void* A, const void* B, uint32 PortFlags) const override;
@@ -3987,6 +4009,7 @@ public:
 	virtual void CopyValuesInternal(void* Dest, void const* Src, int32 Count) const override;
 	virtual void ClearValueInternal(void* Data) const override;
 	virtual void DestroyValueInternal(void* Dest) const override;
+	UE_DEPRECATED(5.4, "UnrealHeaderTool only API.  No replacement available.")
 	virtual bool PassCPPArgsByRef() const override;
 	virtual void InstanceSubobjects(void* Data, void const* DefaultData, UObject* Owner, struct FObjectInstancingGraph* InstanceGraph) override;
 	virtual bool ContainsObjectReference(TArray<const FStructProperty*>& EncounteredStructProps, EPropertyObjectReferenceType InReferenceType = EPropertyObjectReferenceType::Strong) const override;
@@ -5820,6 +5843,7 @@ public:
 	// FProperty interface
 	virtual FString GetCPPMacroType( FString& ExtendedTypeText ) const  override;
 	virtual FString GetCPPType( FString* ExtendedTypeText, uint32 CPPExportFlags ) const override;
+	UE_DEPRECATED(5.4, "UnrealHeaderTool only API.  No replacement available.")
 	virtual FString GetCPPTypeForwardDeclaration() const override;
 	virtual void LinkInternal(FArchive& Ar) override;
 	virtual bool Identical( const void* A, const void* B, uint32 PortFlags ) const override;
@@ -5910,6 +5934,7 @@ PRAGMA_ENABLE_DEPRECATION_WARNINGS
 
 	// FProperty interface
 	virtual FString GetCPPType( FString* ExtendedTypeText, uint32 CPPExportFlags ) const override;
+	UE_DEPRECATED(5.4, "UnrealHeaderTool only API.  No replacement available.")
 	virtual FString GetCPPTypeForwardDeclaration() const override;
 	virtual bool Identical( const void* A, const void* B, uint32 PortFlags ) const override;
 	virtual void SerializeItem(FStructuredArchive::FSlot Slot, void* Value, void const* Defaults) const override;
@@ -5981,6 +6006,7 @@ PRAGMA_ENABLE_DEPRECATION_WARNINGS
 
 	// FProperty interface
 	virtual FString GetCPPType( FString* ExtendedTypeText, uint32 CPPExportFlags ) const override;
+	UE_DEPRECATED(5.4, "UnrealHeaderTool only API.  No replacement available.")
 	virtual FString GetCPPTypeForwardDeclaration() const override;
 	virtual bool Identical( const void* A, const void* B, uint32 PortFlags ) const override;
 	virtual bool NetSerializeItem( FArchive& Ar, UPackageMap* Map, void* Data, TArray<uint8> * MetaData = NULL ) const override;
