BulkDataRegistry: Disable in commandlets other than the cook commandlet. Those commandlets do not need to query its api, and they do not tick it so if it is enabled it grows memory forever without releasing it.
#rnx

diff --git a/Engine/Source/Runtime/CoreUObject/Private/Serialization/BulkDataRegistry.cpp b/Engine/Source/Runtime/CoreUObject/Private/Serialization/BulkDataRegistry.cpp
--- a/Engine/Source/Runtime/CoreUObject/Private/Serialization/BulkDataRegistry.cpp
+++ b/Engine/Source/Runtime/CoreUObject/Private/Serialization/BulkDataRegistry.cpp
@@ -106,6 +106,14 @@ private:
 
 bool IBulkDataRegistry::IsEnabled()
 {
+	if (IsRunningCommandlet() && !IsRunningCookCommandlet())
+	{
+		// Disabled in other commandlets because they do not tick the BulkDataRegistry.
+		// When not ticked, the BulkDataRegistry will never free memory and will hold onto all memory
+		// from loaded BulkDatas until process shutdown; GC will not cause it to be released.
+		return false;
+	}
+
 	bool bEnabled = true;
 	GConfig->GetBool(TEXT("CookSettings"), TEXT("BulkDataRegistryEnabled"), bEnabled, GEditorIni);
 	return bEnabled;
