Added GetTypeHash to FDynamicallyTypedValue for when used as a key for TMaps.

diff --git a/Engine/Source/Runtime/CoreUObject/Public/UObject/DynamicallyTypedValue.h b/Engine/Source/Runtime/CoreUObject/Public/UObject/DynamicallyTypedValue.h
--- a/Engine/Source/Runtime/CoreUObject/Public/UObject/DynamicallyTypedValue.h
+++ b/Engine/Source/Runtime/CoreUObject/Public/UObject/DynamicallyTypedValue.h
@@ -109,6 +109,12 @@ namespace UE
 			Type->InitializeValue(GetDataPointer());
 		}
 
+		// Returns hash of the underlying FDynamicallyTypedValue's value. Added to allow for FDynamicallyTypedValue to be used as TMap keys.
+		friend uint32 GetTypeHash(const FDynamicallyTypedValue& DynamicallyTypedValue)
+		{
+			return DynamicallyTypedValue.GetType().GetValueHash(DynamicallyTypedValue.GetDataPointer());
+		}
+
 	private:
 
 		FDynamicallyTypedValueType* Type;
