[HeaderPatching]
Making PackageSummary symmetrical for load and save.

diff --git a/Engine/Source/Runtime/CoreUObject/Private/UObject/PackageFileSummary.cpp b/Engine/Source/Runtime/CoreUObject/Private/UObject/PackageFileSummary.cpp
--- a/Engine/Source/Runtime/CoreUObject/Private/UObject/PackageFileSummary.cpp
+++ b/Engine/Source/Runtime/CoreUObject/Private/UObject/PackageFileSummary.cpp
@@ -215,14 +215,16 @@ void operator<<(FStructuredArchive::FSlot Slot, FPackageFileSummary& Sum)
 
 		Record << SA_VALUE(TEXT("NameCount"), Sum.NameCount) << SA_VALUE(TEXT("NameOffset"), Sum.NameOffset);
 
-		if (BaseArchive.IsSaving() || Sum.FileVersionUE >= EUnrealEngineObjectUE5Version::ADD_SOFTOBJECTPATH_LIST)
+		// Sometimes its useful to be able to save out files with older versions.
+		// So this code needs to be Symmetrical for saving and loading unless an error is encountered.
+		if (Sum.FileVersionUE >= EUnrealEngineObjectUE5Version::ADD_SOFTOBJECTPATH_LIST)
 		{
 			Record << SA_VALUE(TEXT("SoftObjectPathsCount"), Sum.SoftObjectPathsCount) << SA_VALUE(TEXT("SoftObjectPathsOffset"), Sum.SoftObjectPathsOffset);
 		}
 
 		if (!BaseArchive.IsFilterEditorOnly())
 		{
-			if (BaseArchive.IsSaving() || Sum.FileVersionUE >= VER_UE4_ADDED_PACKAGE_SUMMARY_LOCALIZATION_ID)
+			if (Sum.FileVersionUE >= VER_UE4_ADDED_PACKAGE_SUMMARY_LOCALIZATION_ID)
 			{
 				Record << SA_VALUE(TEXT("LocalizationId"), Sum.LocalizationId);
 			}
@@ -241,12 +243,12 @@ void operator<<(FStructuredArchive::FSlot Slot, FPackageFileSummary& Sum)
 			return; // we can't safely load more than this because the below was different in older files.
 		}
 
-		if (BaseArchive.IsSaving() || Sum.FileVersionUE >= VER_UE4_ADD_STRING_ASSET_REFERENCES_MAP)
+		if (Sum.FileVersionUE >= VER_UE4_ADD_STRING_ASSET_REFERENCES_MAP)
 		{
 			Record << SA_VALUE(TEXT("SoftPackageReferencesCount"), Sum.SoftPackageReferencesCount) << SA_VALUE(TEXT("SoftPackageReferencesOffset"), Sum.SoftPackageReferencesOffset);
 		}
 
-		if (BaseArchive.IsSaving() || Sum.FileVersionUE >= VER_UE4_ADDED_SEARCHABLE_NAMES)
+		if (Sum.FileVersionUE >= VER_UE4_ADDED_SEARCHABLE_NAMES)
 		{
 			Record << SA_VALUE(TEXT("SearchableNamesOffset"), Sum.SearchableNamesOffset);
 		}
@@ -260,7 +262,7 @@ void operator<<(FStructuredArchive::FSlot Slot, FPackageFileSummary& Sum)
 #if WITH_EDITORONLY_DATA
 		if (!BaseArchive.IsFilterEditorOnly())
 		{
-			if (BaseArchive.IsSaving() || Sum.FileVersionUE >= VER_UE4_ADDED_PACKAGE_OWNER)
+			if (Sum.FileVersionUE >= VER_UE4_ADDED_PACKAGE_OWNER)
 			{
 				Record << SA_VALUE(TEXT("PersistentGuid"), Sum.PersistentGuid);
 			}
@@ -273,7 +275,7 @@ void operator<<(FStructuredArchive::FSlot Slot, FPackageFileSummary& Sum)
 			}
 
 			// The owner persistent guid was added in VER_UE4_ADDED_PACKAGE_OWNER but removed in the next version VER_UE4_NON_OUTER_PACKAGE_IMPORT
-			if (BaseArchive.IsLoading() && Sum.FileVersionUE >= VER_UE4_ADDED_PACKAGE_OWNER && Sum.FileVersionUE < VER_UE4_NON_OUTER_PACKAGE_IMPORT)
+			if (Sum.FileVersionUE >= VER_UE4_ADDED_PACKAGE_OWNER && Sum.FileVersionUE < VER_UE4_NON_OUTER_PACKAGE_IMPORT)
 			{
 				FGuid OwnerPersistentGuid;
 				Record << SA_VALUE(TEXT("OwnerPersistentGuid"), OwnerPersistentGuid);
@@ -405,7 +407,7 @@ void operator<<(FStructuredArchive::FSlot Slot, FPackageFileSummary& Sum)
 				}
 			}
 		}
-		if (BaseArchive.IsSaving() || Sum.FileVersionUE >= VER_UE4_PRELOAD_DEPENDENCIES_IN_COOKED_EXPORTS)
+		if (Sum.FileVersionUE >= VER_UE4_PRELOAD_DEPENDENCIES_IN_COOKED_EXPORTS)
 		{
 			Record << SA_VALUE(TEXT("PreloadDependencyCount"), Sum.PreloadDependencyCount) << SA_VALUE(TEXT("PreloadDependencyOffset"), Sum.PreloadDependencyOffset);
 		}
@@ -415,7 +417,7 @@ void operator<<(FStructuredArchive::FSlot Slot, FPackageFileSummary& Sum)
 			Sum.PreloadDependencyOffset = 0;
 		}
 
-		if (BaseArchive.IsSaving() || Sum.FileVersionUE >= EUnrealEngineObjectUE5Version::NAMES_REFERENCED_FROM_EXPORT_DATA)
+		if (Sum.FileVersionUE >= EUnrealEngineObjectUE5Version::NAMES_REFERENCED_FROM_EXPORT_DATA)
 		{
 			Record << SA_VALUE(TEXT("NamesReferencedFromExportDataCount"), Sum.NamesReferencedFromExportDataCount);
 		}
@@ -424,7 +426,7 @@ void operator<<(FStructuredArchive::FSlot Slot, FPackageFileSummary& Sum)
 			Sum.NamesReferencedFromExportDataCount = Sum.NameCount;
 		}
 
-		if (BaseArchive.IsSaving() || Sum.FileVersionUE >= EUnrealEngineObjectUE5Version::PAYLOAD_TOC)
+		if (Sum.FileVersionUE >= EUnrealEngineObjectUE5Version::PAYLOAD_TOC)
 		{
 			Record << SA_VALUE(TEXT("PayloadTocOffset"), Sum.PayloadTocOffset);
 		}
@@ -433,7 +435,7 @@ void operator<<(FStructuredArchive::FSlot Slot, FPackageFileSummary& Sum)
 			Sum.PayloadTocOffset = INDEX_NONE;
 		}
 		
-		if (BaseArchive.IsSaving() || Sum.GetFileVersionUE() >= EUnrealEngineObjectUE5Version::DATA_RESOURCES)
+		if (Sum.GetFileVersionUE() >= EUnrealEngineObjectUE5Version::DATA_RESOURCES)
 		{
 			Record << SA_VALUE(TEXT("DataResourceOffset"), Sum.DataResourceOffset);
 		}
