Seek back to the previous position on early exit.

diff --git a/Engine/Source/Runtime/CoreUObject/Private/UObject/LinkerLoad.cpp b/Engine/Source/Runtime/CoreUObject/Private/UObject/LinkerLoad.cpp
--- a/Engine/Source/Runtime/CoreUObject/Private/UObject/LinkerLoad.cpp
+++ b/Engine/Source/Runtime/CoreUObject/Private/UObject/LinkerLoad.cpp
@@ -4487,6 +4487,7 @@ void FLinkerLoad::Preload( UObject* Object )
 							// sure to force load this later)
 							check(!GEventDrivenLoaderEnabled || !EVENT_DRIVEN_ASYNC_LOAD_ACTIVE_AT_RUNTIME);
 							Object->SetFlags(RF_NeedLoad);
+							Seek(SavedPos);
 							return;
 						}
 #endif // USE_CIRCULAR_DEPENDENCY_LOAD_DEFERRING
