Check for nullptr on DefaultValue. This is a possible return value and was not caught on the logging.

diff --git a/Engine/Source/Runtime/CoreUObject/Private/UObject/PropertyObject.cpp b/Engine/Source/Runtime/CoreUObject/Private/UObject/PropertyObject.cpp
--- a/Engine/Source/Runtime/CoreUObject/Private/UObject/PropertyObject.cpp
+++ b/Engine/Source/Runtime/CoreUObject/Private/UObject/PropertyObject.cpp
@@ -139,7 +139,7 @@ void FObjectProperty::PostSerializeObjectItem(FArchive& SerializingArchive, void
 		UE_LOG(LogProperty, Warning,
 			TEXT("Failed to serialize value for non-nullable property %s. Reference will be defaulted to %s."),
 			*GetFullName(),
-			*DefaultValue->GetFullName()
+			DefaultValue ? *DefaultValue->GetFullName() : TEXT("None")
 		);
 
 		SetObjectPropertyValue(Value, DefaultValue);
