Validate Cast at compile time
- Casts between unrelated subclasses of UObject will produce a static warning (and an error when treat warning as error is true)
- Trivial casts will be resolved at compile time without calling IsA<>
All existing warnings have been fixed
This can be disabled with UE_ENABLE_UNRELATED_CAST_WARNINGS = 0

diff --git a/Engine/Source/Runtime/Core/Public/Misc/CoreMiscDefines.h b/Engine/Source/Runtime/Core/Public/Misc/CoreMiscDefines.h
--- a/Engine/Source/Runtime/Core/Public/Misc/CoreMiscDefines.h
+++ b/Engine/Source/Runtime/Core/Public/Misc/CoreMiscDefines.h
@@ -304,6 +304,17 @@ struct TStaticDeprecateExpression
 	}; \
 	enum class PREPROCESSOR_JOIN(EDeprecationMsg_, __LINE__) { Value = PREPROCESSOR_JOIN(FDeprecationMsg_, __LINE__)::condition(TStaticDeprecateExpression<!!(bExpression)>()) }
 
+/**
+ * Can be used in the same contexts as static_assert but gives a warning rather than an error
+ */
+#define UE_STATIC_ASSERT_WARN(bExpression, Message) \
+	struct PREPROCESSOR_JOIN(FStaticWarningMsg_, __LINE__) { \
+		[[deprecated(Message)]] \
+		static constexpr int condition(TStaticDeprecateExpression<true>) { return 1; } \
+		static constexpr int condition(TStaticDeprecateExpression<false>) { return 1; } \
+	}; \
+	enum class PREPROCESSOR_JOIN(EStaticWarningMsg_, __LINE__) { Value = PREPROCESSOR_JOIN(FStaticWarningMsg_, __LINE__)::condition(TStaticDeprecateExpression<!(bExpression)>()) }
+
 // These defines are used to mark a difference between two pointers as expected to fit into the specified range
 // while still leaving something searchable if the surrounding code is updated to work with a 64 bit count/range
 // in the future
diff --git a/Engine/Source/Runtime/CoreUObject/Public/Templates/Casts.h b/Engine/Source/Runtime/CoreUObject/Public/Templates/Casts.h
--- a/Engine/Source/Runtime/CoreUObject/Public/Templates/Casts.h
+++ b/Engine/Source/Runtime/CoreUObject/Public/Templates/Casts.h
@@ -11,11 +11,14 @@
 #include "UObject/ObjectHandle.h"
 #include "UObject/ObjectMacros.h"
 #include "UObject/ObjectPtr.h"
-#include "UObject/WeakObjectPtrTemplates.h"
 #include <type_traits>
 
 #define UE_USE_CAST_FLAGS (USTRUCT_FAST_ISCHILDOF_IMPL != USTRUCT_ISCHILDOF_STRUCTARRAY)
 
+#ifndef UE_ENABLE_UNRELATED_CAST_WARNINGS
+#define UE_ENABLE_UNRELATED_CAST_WARNINGS 1
+#endif
+
 class AActor;
 class APawn;
 class APlayerController;
@@ -105,10 +108,20 @@ FORCEINLINE To* Cast(From* Src)
 		}
 		else if constexpr (UE_USE_CAST_FLAGS && TCastFlags<To>::Value != CASTCLASS_None)
 		{
-			if (((const UObject*)Src)->GetClass()->HasAnyCastFlag(TCastFlags<To>::Value))
+			if constexpr (std::is_base_of_v<To, From>)
 			{
 				return (To*)Src;
 			}
+			else
+			{
+#if UE_ENABLE_UNRELATED_CAST_WARNINGS
+				UE_STATIC_ASSERT_WARN((std::is_base_of_v<From, To>), "Attempting to use Cast<> on types that are not related");
+#endif
+				if (((const UObject*)Src)->GetClass()->HasAnyCastFlag(TCastFlags<To>::Value))
+				{
+					return (To*)Src;
+				}
+			}
 		}
 		else
 		{
@@ -118,8 +131,15 @@ FORCEINLINE To* Cast(From* Src)
 			{
 				return (To*)((UObject*)Src)->GetInterfaceAddress(To::UClassType::StaticClass());
 			}
+			else if constexpr (std::is_base_of_v<To, From>)
+			{
+				return Src;
+			}
 			else
 			{
+#if UE_ENABLE_UNRELATED_CAST_WARNINGS
+				UE_STATIC_ASSERT_WARN((std::is_base_of_v<From, To>), "Attempting to use Cast<> on types that are not related");
+#endif
 				if (((const UObject*)Src)->IsA<To>())
 				{
 					return (To*)Src;
@@ -148,6 +168,10 @@ FORCEINLINE T* ExactCast( UObject* Src )
 		{
 			return GetFullNameSafe(InObjectOrInterface);
 		}
+		else if constexpr (std::is_base_of_v<UObject, T>)
+    	{
+    		return InObjectOrInterface->GetFullName();;
+    	}
 		else
 		{
 			return Cast<UObject>(InObjectOrInterface)->GetFullName();
@@ -252,15 +276,28 @@ template< class T, class U > FORCEINLINE T* CastChecked( const TWeakObjectPtr<U>
 template <typename To, typename From>
 FORCEINLINE To* Cast(const TObjectPtr<From>& InSrc)
 {
-	static_assert(sizeof(To) > 0, "Attempting to cast to an incomplete type");
+	static_assert(sizeof(To) > 0 && sizeof(From) > 0, "Attempting to cast between incomplete types");
 
 	const FObjectPtr& Src = (const FObjectPtr&)InSrc;
 
 	if constexpr (UE_USE_CAST_FLAGS && TCastFlags<To>::Value != CASTCLASS_None)
 	{
-		if (Src && Src.GetClass()->HasAnyCastFlag(TCastFlags<To>::Value))
+		if (Src)
 		{
-			return (To*)Src.Get();
+			if constexpr (std::is_base_of_v<To, From>)
+			{
+				return (To*)Src.Get();
+			}
+			else
+			{
+	#if UE_ENABLE_UNRELATED_CAST_WARNINGS
+				UE_STATIC_ASSERT_WARN((std::is_base_of_v<From, To>), "Attempting to use Cast<> on types that are not related");
+	#endif
+				if (Src.GetClass()->HasAnyCastFlag(TCastFlags<To>::Value))
+				{
+					return (To*)Src.Get();
+				}
+			}
 		}
 	}
 	else if constexpr (TIsIInterface<To>::Value)
@@ -272,8 +309,18 @@ FORCEINLINE To* Cast(const TObjectPtr<From>& InSrc)
 			return (To*)Src.Get()->GetInterfaceAddress(To::UClassType::StaticClass());
 		}
 	}
+	else if constexpr (std::is_base_of_v<To, From>)
+	{
+		if (Src)
+		{
+			return (To*)Src.Get();
+		}
+	}
 	else
 	{
+#if UE_ENABLE_UNRELATED_CAST_WARNINGS
+		UE_STATIC_ASSERT_WARN((std::is_base_of_v<From, To>), "Attempting to use Cast<> on types that are not related");
+#endif
 		if (Src && Src.IsA<To>())
 		{
 			return (To*)Src.Get();
