Don't pass true to bExactClass when looking up redirected classes for cooked packages, as we don't know the exact class of what we are looking for (could be UClass or UBlueprintGeneratedClass).

diff --git a/Engine/Source/Runtime/CoreUObject/Private/UObject/LinkerLoad.cpp b/Engine/Source/Runtime/CoreUObject/Private/UObject/LinkerLoad.cpp
--- a/Engine/Source/Runtime/CoreUObject/Private/UObject/LinkerLoad.cpp
+++ b/Engine/Source/Runtime/CoreUObject/Private/UObject/LinkerLoad.cpp
@@ -3297,7 +3297,9 @@ void StaticFindAllImportObjects(TArray<FObjectImport>& InOutImportMap, FPackageI
 				if (!ObjectImport.XObject)
 				{
 					UClass* Class = FindClass(ObjectImport);
-					if (UObject* InnerObject = StaticFindObjectFastInternal(Class, OuterObject, ObjectImport.ObjectName, true))
+					// Don't pass bExactClass=true when looking up redirected classes for cooked packages, as we don't know the exact class of what
+					// we are looking for (could be UClass or UBlueprintGeneratedClass, see CreateImportClassAndPackage).
+					if (UObject* InnerObject = StaticFindObjectFastInternal(Class, OuterObject, ObjectImport.ObjectName, /*bExactClass*/false))
 					{
 						ObjectImport.XObject = InnerObject;
 						Outers.Add(Inner);
