Improved ensure information to help track the use of abstract classes

diff --git a/Engine/Source/Runtime/CoreUObject/Private/UObject/UObjectGlobals.cpp b/Engine/Source/Runtime/CoreUObject/Private/UObject/UObjectGlobals.cpp
--- a/Engine/Source/Runtime/CoreUObject/Private/UObject/UObjectGlobals.cpp
+++ b/Engine/Source/Runtime/CoreUObject/Private/UObject/UObjectGlobals.cpp
@@ -3208,14 +3208,14 @@ bool StaticAllocateObjectErrorTests( const UClass* InClass, UObject* InOuter, FN
 	{
 		if ( GIsEditor )
 		{
-			const FString ErrorMsg = FString::Printf(TEXT("Class which was marked abstract was trying to be loaded.  It will be nulled out on save. %s %s"), *InName.ToString(), *InClass->GetName());
+			const FString ErrorMsg = FString::Printf(TEXT("Class which was marked abstract was trying to be loaded in Outer %s.  It will be nulled out on save. %s %s"), *GetPathNameSafe(InOuter), *InName.ToString(), *InClass->GetName());
 			// if we are trying instantiate an abstract class in the editor we'll warn the user that it will be nulled out on save
 			UE_LOG(LogUObjectGlobals, Warning, TEXT("%s"), *ErrorMsg);
 			ensureMsgf(false, TEXT("%s"), *ErrorMsg);
 		}
 		else
 		{
-			UE_LOG(LogUObjectGlobals, Fatal, TEXT("%s"), *FString::Printf( TEXT("Can't create object %s: class %s is abstract"), *InName.ToString(), *InClass->GetName()));
+			UE_LOG(LogUObjectGlobals, Fatal, TEXT("%s"), *FString::Printf( TEXT("Can't create object %s in Outer %s: class %s is abstract"), *InName.ToString(), *GetPathNameSafe(InOuter), *InClass->GetName()));
 			return true;
 		}
 	}
