Added missing TBaseStructure/TVariantStructure definitions for FRay/FSphere.

#rnx

diff --git a/Engine/Source/Runtime/CoreUObject/Private/UObject/Class.cpp b/Engine/Source/Runtime/CoreUObject/Private/UObject/Class.cpp
--- a/Engine/Source/Runtime/CoreUObject/Private/UObject/Class.cpp
+++ b/Engine/Source/Runtime/CoreUObject/Private/UObject/Class.cpp
@@ -7166,7 +7166,9 @@ UE_DEFINE_CORE_VARIANT_TYPE(Rotator3,	Rotator);
 UE_DEFINE_CORE_VARIANT_TYPE(Transform3,	Transform);
 UE_DEFINE_CORE_VARIANT_TYPE(Matrix44,	Matrix);
 UE_DEFINE_CORE_VARIANT_TYPE(Box2,		Box2D);
-
+UE_DEFINE_CORE_VARIANT_TYPE(Ray3,		Ray);
+UE_DEFINE_CORE_VARIANT_TYPE(Sphere3,	Sphere);
+	
 #undef UE_DEFINE_CORE_VARIANT_TYPE
 
 IMPLEMENT_CORE_INTRINSIC_CLASS(UFunction, UStruct,
