Fixing a crash when marking FFieldVariant referenced UObject as reachable

diff --git a/Engine/Source/Runtime/CoreUObject/Public/UObject/Field.h b/Engine/Source/Runtime/CoreUObject/Public/UObject/Field.h
--- a/Engine/Source/Runtime/CoreUObject/Public/UObject/Field.h
+++ b/Engine/Source/Runtime/CoreUObject/Public/UObject/Field.h
@@ -273,9 +273,9 @@ class FFieldVariant
 
 	void ConditionallyMarkAsReachable()
 	{
-		if (IsUObject() && Container.Object && UE::GC::Private::GIsIncrementalReachabilityPending)
+		if (IsUObject() && ToUObjectUnsafe() && UE::GC::Private::GIsIncrementalReachabilityPending)
 		{
-			Container.Object->MarkAsReachable();
+			ToUObjectUnsafe()->MarkAsReachable();
 		}
 	}
 	
