Adding missing scope tags for LLM statistics

diff --git a/Engine/Source/Runtime/CoreUObject/Private/UObject/UObjectArray.cpp b/Engine/Source/Runtime/CoreUObject/Private/UObject/UObjectArray.cpp
--- a/Engine/Source/Runtime/CoreUObject/Private/UObject/UObjectArray.cpp
+++ b/Engine/Source/Runtime/CoreUObject/Private/UObject/UObjectArray.cpp
@@ -189,6 +189,8 @@ void FUObjectArray::DisableDisregardForGC()
 
 void FUObjectArray::AllocateUObjectIndex(UObjectBase* Object, EInternalObjectFlags InitialFlags, int32 AlreadyAllocatedIndex, int32 SerialNumber)
 {
+	LLM_SCOPE(ELLMTag::UObject);
+
 	int32 Index = INDEX_NONE;
 	check(Object->InternalIndex == INDEX_NONE);
 
@@ -287,6 +289,8 @@ void FUObjectArray::RemoveObjectFromDeleteListeners(UObjectBase* Object)
  */
 void FUObjectArray::FreeUObjectIndex(UObjectBase* Object)
 {
+	LLM_SCOPE(ELLMTag::UObject);
+
 	// This should only be happening on the game thread (GC runs only on game thread when it's freeing objects)
 	check(IsInGameThread() || IsInGarbageCollectorThread());
 
diff --git a/Engine/Source/Runtime/PreLoadScreen/Private/PreLoadScreenManager.cpp b/Engine/Source/Runtime/PreLoadScreen/Private/PreLoadScreenManager.cpp
--- a/Engine/Source/Runtime/PreLoadScreen/Private/PreLoadScreenManager.cpp
+++ b/Engine/Source/Runtime/PreLoadScreen/Private/PreLoadScreenManager.cpp
@@ -353,6 +353,7 @@ void FPreLoadScreenManager::HandleCustomSplashScreenPlay()
 
 void FPreLoadScreenManager::StaticRenderTick_RenderThread()
 {
+	LLM_SCOPE(ELLMTag::RenderingThreadMemory);
 	check(IsInRenderingThread());
 
 	FScopeLock Lock(&AcquireCriticalSection);
diff --git a/Engine/Source/Runtime/RenderCore/Private/ShaderPipelineCache.cpp b/Engine/Source/Runtime/RenderCore/Private/ShaderPipelineCache.cpp
--- a/Engine/Source/Runtime/RenderCore/Private/ShaderPipelineCache.cpp
+++ b/Engine/Source/Runtime/RenderCore/Private/ShaderPipelineCache.cpp
@@ -1773,6 +1773,7 @@ bool FShaderPipelineCache::IsTickable() const
 
 void FShaderPipelineCache::Tick(float DeltaTime)
 {
+	LLM_SCOPE(ELLMTag::Shaders);
 	FScopeLock Lock(&Mutex);
 
 	FShaderPipelineCacheTask* CurrentCacheTask = CurrentPrecompilingPSOFCKey.IsEmpty() ? nullptr : ShaderCacheTasks.Find(CurrentPrecompilingPSOFCKey)->Get();
