Cast, ExactCast and CastChecked with TObjectPtr are const correct

diff --git a/Engine/Source/Runtime/CoreUObject/Public/Templates/Casts.h b/Engine/Source/Runtime/CoreUObject/Public/Templates/Casts.h
--- a/Engine/Source/Runtime/CoreUObject/Public/Templates/Casts.h
+++ b/Engine/Source/Runtime/CoreUObject/Public/Templates/Casts.h
@@ -99,9 +99,16 @@ FORCEINLINE To* Cast(From* Src)
 				}
 				else
 				{
-					if (Obj->IsA<To>())
+					if constexpr (std::is_same_v<To, UObject>)
 					{
-						return (To*)Obj;
+						return Obj;
+					}
+					else
+					{
+						if (Obj->IsA<To>())
+						{
+							return (To*)Obj;
+						}
 					}
 				}
 			}
@@ -274,7 +281,7 @@ template< class T, class U > FORCEINLINE T* CastChecked( const TWeakObjectPtr<U>
 
 // object ptr versions
 template <typename To, typename From>
-FORCEINLINE To* Cast(const TObjectPtr<From>& InSrc)
+FORCEINLINE typename TCopyQualifiersFromTo<From, To>::Type* Cast(const TObjectPtr<From>& InSrc)
 {
 	static_assert(sizeof(To) > 0 && sizeof(From) > 0, "Attempting to cast between incomplete types");
 
@@ -302,7 +309,7 @@ FORCEINLINE To* Cast(const TObjectPtr<From>& InSrc)
 	}
 	else if constexpr (TIsIInterface<To>::Value)
 	{
-		UObject* SrcObj = UE::CoreUObject::Private::ResolveObjectHandleNoRead(Src.GetHandleRef());
+		const UObject* SrcObj = UE::CoreUObject::Private::ResolveObjectHandleNoRead(Src.GetHandleRef());
 		if (SrcObj)
 		{
 			UE::CoreUObject::Private::OnHandleRead(SrcObj);
@@ -331,7 +338,7 @@ FORCEINLINE To* Cast(const TObjectPtr<From>& InSrc)
 }
 
 template <typename To, typename From>
-FORCEINLINE To* ExactCast(const TObjectPtr<From>& Src)
+FORCEINLINE typename TCopyQualifiersFromTo<From, To>::Type* ExactCast(const TObjectPtr<From>& Src)
 {
 	static_assert(sizeof(To) > 0, "Attempting to cast to an incomplete type");
 
@@ -345,14 +352,14 @@ FORCEINLINE To* ExactCast(const TObjectPtr<From>& Src)
 }
 
 template <typename To, typename From>
-FORCEINLINE To* CastChecked(const TObjectPtr<From>& Src, ECastCheckedType::Type CheckType = ECastCheckedType::NullChecked)
+FORCEINLINE typename TCopyQualifiersFromTo<From, To>::Type* CastChecked(const TObjectPtr<From>& Src, ECastCheckedType::Type CheckType = ECastCheckedType::NullChecked)
 {
 	static_assert(sizeof(From) > 0 && sizeof(To) > 0, "Attempting to cast between incomplete types");
 
 #if DO_CHECK
 	if (Src)
 	{
-		To* Result = Cast<To>(Src);
+		auto* Result = Cast<To>(Src);
 		if (!Result)
 		{
 			CastLogError(*GetFullNameForCastLogError(Src.Get()), *GetTypeName<To>());
diff --git a/Engine/Source/Runtime/CoreUObject/Public/UObject/UObjectBaseUtility.h b/Engine/Source/Runtime/CoreUObject/Public/UObject/UObjectBaseUtility.h
--- a/Engine/Source/Runtime/CoreUObject/Public/UObject/UObjectBaseUtility.h
+++ b/Engine/Source/Runtime/CoreUObject/Public/UObject/UObjectBaseUtility.h
@@ -522,10 +522,6 @@ public:
 	*/
 	COREUOBJECT_API void AddToCluster(UObjectBaseUtility* ClusterRootOrObjectFromCluster, bool bAddAsMutableObject = false);
 
-protected:
-	/** Helper function to create a cluster from UObject */
-	static COREUOBJECT_API void CreateClusterFromObject(UObjectBaseUtility* ClusterRootObject, UObjectBaseUtility* ReferencingObject);
-
 public:
 	/**
 	 * Walks up the chain of packages until it reaches the top level, which it ignores.
