#UE Add a bit more context to an ensure

diff --git a/Engine/Source/Runtime/CoreUObject/Private/UObject/SavePackage/PackageHarvester.cpp b/Engine/Source/Runtime/CoreUObject/Private/UObject/SavePackage/PackageHarvester.cpp
--- a/Engine/Source/Runtime/CoreUObject/Private/UObject/SavePackage/PackageHarvester.cpp
+++ b/Engine/Source/Runtime/CoreUObject/Private/UObject/SavePackage/PackageHarvester.cpp
@@ -888,7 +888,7 @@ void FPackageHarvester::ResolveOverrides()
 				}
 
 				// We currently only support object property
-				if (ensureAlwaysMsgf(InnerProp && InnerProp->IsA<FObjectProperty>(), TEXT("Save Overrides supports only object properties at the moment. Name: %s, Type: %s"), *Prop->GetName(), *Prop->GetClass()->GetName()))
+				if (ensureAlwaysMsgf(InnerProp && InnerProp->IsA<FObjectProperty>(), TEXT("Save Overrides supports only object properties at the moment. Name: %s, Type: %s. Package: %s"), *Prop->GetName(), *Prop->GetClass()->GetName(), *GetNameSafe(SaveContext.GetPackage())))
 				{
 					// Harvest the name of the property, since it is only made transient for the purpose of the package harvest, it will be needed for the LinkerSave
 					HarvestExportDataName(Prop->GetFName());
