SavePackage: Improve poor performance in PackageHarvester when there are imports of many inner objects in a deep outer chain by skipping the full harvest of imports in operator<< for objects that have already been included, in addition to the skip that is already done for objects that have already been excluded.
#rnx

diff --git a/Engine/Source/Runtime/CoreUObject/Private/UObject/SavePackage/PackageHarvester.cpp b/Engine/Source/Runtime/CoreUObject/Private/UObject/SavePackage/PackageHarvester.cpp
--- a/Engine/Source/Runtime/CoreUObject/Private/UObject/SavePackage/PackageHarvester.cpp
+++ b/Engine/Source/Runtime/CoreUObject/Private/UObject/SavePackage/PackageHarvester.cpp
@@ -236,6 +236,13 @@ FPackageHarvester::FHarvestScope FPackageHarvester::EnterRootReferencesScope()
 	return Scope;
 }
 
+FPackageHarvester::FHarvestScope FPackageHarvester::EnterRealmsArrayScope(FExportingRealmsArray& Array)
+{
+	FHarvestScope Scope(*this);
+	CurrentExportHarvestingRealms = Array;
+	return Scope;
+}
+
 FPackageHarvester::FHarvestScope FPackageHarvester::EnterConditionalEditorOnlyScope(bool bIsEditorOnly)
 {
 	FHarvestScope Scope(*this);
@@ -332,6 +339,24 @@ FPackageHarvester::FHarvestScope FPackageHarvester::EnterIncludedScope(TObjectPt
 	return Scope;
 }
 
+void FPackageHarvester::GetPreviouslyIncludedRealms(TObjectPtr<UObject> Object,
+	FExportingRealmsArray& OutAlreadyIncluded, FExportingRealmsArray& OutNotAlreadyIncluded)
+{
+	OutAlreadyIncluded.Reset();
+	OutNotAlreadyIncluded.Reset();
+	for (ESaveRealm Realm : CurrentExportHarvestingRealms)
+	{
+		if (SaveContext.GetHarvestedRealm(Realm).IsIncluded(Object))
+		{
+			OutAlreadyIncluded.Add(Realm);
+		}
+		else
+		{
+			OutNotAlreadyIncluded.Add(Realm);
+		}
+	}
+}
+
 bool FPackageHarvester::IsObjNative(TObjectPtr<UObject> InObj)
 {
 	bool bIsNative = InObj->IsNative();
@@ -621,7 +646,7 @@ void FPackageHarvester::TryHarvestImport(TObjectPtr<UObject> InObject)
 		}
 	}
 
-	// Filter out any realms in which the export is excluded
+	// Filter out any realms in which the import is excluded
 	FHarvestScope NotExcludedScope = EnterNotExcludedScope(InObject);
 	if (NotExcludedScope.IsEmpty())
 	{
@@ -634,6 +659,9 @@ void FPackageHarvester::TryHarvestImport(TObjectPtr<UObject> InObject)
 
 void FPackageHarvester::ProcessImport(TObjectPtr<UObject> InObject)
 {
+	++CurrentExportDependencies.ProcessImportDepth;
+	ON_SCOPE_EXIT{ --CurrentExportDependencies.ProcessImportDepth; };
+
 	bool bIsNative = IsObjNative(InObject);
 	TObjectPtr<UObject> ObjOuter = InObject.GetOuter();
 	UClass* ObjClass = InObject.GetClass();
@@ -716,9 +744,6 @@ void FPackageHarvester::ProcessImport(TObjectPtr<UObject> InObject)
 		HarvestPackageHeaderName(ObjClass->GetFName());
 		HarvestPackageHeaderName(ObjClass->GetOuter()->GetFName());
 	}
-
-	// If we have an illegal reference to an optional object, record it
-	EnterConditionalOptionalObjectScope(InObject);
 }
 
 FString FPackageHarvester::GetArchiveName() const
@@ -766,22 +791,39 @@ FArchive& FPackageHarvester::operator<<(UObject*& Obj)
 		return *this;
 	}
 
-	// if the object is in the save context package, try to tag it as export
-	if (Obj->IsInPackage(SaveContext.GetPackage()))
-	{
-		TryHarvestExportInternal(Obj);
-	}
-	// Otherwise visit the import
-	else
+	// For realms in which the object has already been included as an import or export,
+	// do reduced work - just the work that needs to be done per referencer that references it.
+	// This is important for performance, because a full harvest of the object can be expensive.
+	// For new realms that have not previously included the object, do a full harvest.
+	FExportingRealmsArray PreviouslyIncludedRealms;
+	FExportingRealmsArray NewRealms;
+	GetPreviouslyIncludedRealms(Obj, PreviouslyIncludedRealms, NewRealms);
+
+	if (!NewRealms.IsEmpty())
 	{
-		TryHarvestImport(Obj);
+		// if the object is in the save context package, try to tag it as export for the new realms.
+		// Otherwise visit the import for the new realms.
+		FHarvestScope NewRealmsScope = EnterRealmsArrayScope(NewRealms);
+		if (Obj->IsInPackage(SaveContext.GetPackage()))
+		{
+			TryHarvestExportInternal(Obj);
+		}
+		else
+		{
+			TryHarvestImport(Obj);
+		}
 	}
 
-	// Add a dependency from the current export to Obj if Obj was added as an import or export in any realm
+	// Work that needs to be done per referencing export. This work only needs to be done if Obj is an import or
+	// export in any realm.
 	FHarvestScope ObjIncludedScope = EnterIncludedScope(Obj);
 	if (!ObjIncludedScope.IsEmpty())
 	{
+		// Add a dependency from the current referencing export to Obj
 		HarvestDependency(Obj, IsObjNative(Obj));
+
+		// Validate the reference is allowed (this validation is a side-effect of EnterConditionalOptionalObjectScope)
+		EnterConditionalOptionalObjectScope(Obj);
 	}
 
 	return *this;
@@ -923,9 +965,10 @@ TMap<UObject*, TSet<FProperty*>> FPackageHarvester::ReleaseTransientPropertyOver
 void FPackageHarvester::HarvestDependency(TObjectPtr<UObject> InObj, bool bIsNative)
 {
 	// if we aren't currently processing an export or the referenced object is a package, do not harvest the dependency
-	if (CurrentExportDependencies.bIgnoreDependencies ||
-		CurrentExportDependencies.CurrentExport == nullptr ||
-		(InObj.GetOuter() == nullptr && InObj.GetClass()->GetFName() == NAME_Package))
+	if (CurrentExportDependencies.ProcessImportDepth > 0 || // Skip if we are processing a transitive import found inside ProcessImport
+		CurrentExportDependencies.bIgnoreDependencies || // Skip in stack-specific cases that put FIgnoreDependenciesScope on the stack
+		CurrentExportDependencies.CurrentExport == nullptr || // Skip if we are not currently processing an export
+		(InObj.GetOuter() == nullptr && InObj.GetClass()->GetFName() == NAME_Package)) // Skip if object is a package
 	{
 		return;
 	}
diff --git a/Engine/Source/Runtime/CoreUObject/Private/UObject/SavePackage/PackageHarvester.h b/Engine/Source/Runtime/CoreUObject/Private/UObject/SavePackage/PackageHarvester.h
--- a/Engine/Source/Runtime/CoreUObject/Private/UObject/SavePackage/PackageHarvester.h
+++ b/Engine/Source/Runtime/CoreUObject/Private/UObject/SavePackage/PackageHarvester.h
@@ -97,6 +97,7 @@ private:
 		UObject* CurrentExport = nullptr;
 		TSet<TObjectPtr<UObject>> ObjectReferences;
 		TSet<TObjectPtr<UObject>> NativeObjectReferences;
+		int32 ProcessImportDepth = 0;
 		bool bIgnoreDependencies = false;
 	};
 
@@ -147,6 +148,8 @@ private:
 	ESaveableStatus GetSaveableStatusForRealm(UObject* Obj, ESaveRealm RealmInWhichItIsUnsaveable,
 		TObjectPtr<UObject>& OutCulprit, FString& OutReason);
 
+	/** Set HarvestingRealms equal to the given Array. */
+	FHarvestScope EnterRealmsArrayScope(FExportingRealmsArray& Array);
 	/** If bIsEditorOnly is true, remove CurrentExportHarvestingRealms that do not follow editoronly references. */
 	FHarvestScope EnterConditionalEditorOnlyScope(bool bIsEditorOnly);
 	/** If the Object is optional, clear CurrentExportHarvestingRealms and populate it with the Optional Realm. */
@@ -159,6 +162,9 @@ private:
 	FHarvestScope EnterNotPreviouslyExcludedScope(TObjectPtr<UObject> Object);
 	/** Remove CurrentExportHarvestingRealms that return false for IsIncluded(Object). */
 	FHarvestScope EnterIncludedScope(TObjectPtr<UObject> Object);
+	/** Copy CurrentExportHarvestingRealms and split it into two arrays: AlreadyIncluded and NotAlreadyIncluded. */
+	void GetPreviouslyIncludedRealms(TObjectPtr<UObject> Object, FExportingRealmsArray& OutAlreadyIncluded,
+		FExportingRealmsArray& OutNotAlreadyIncluded);
 
 	FSaveContext& SaveContext;
 	TQueue<FExportWithContext> ExportsToProcess;
