Small optimization to GetDerivedClass/RecursivelyPopulateDerivedClasses: remove unnecessary allocation from TSet::Array and adding reserves.

diff --git a/Engine/Source/Runtime/CoreUObject/Private/UObject/UObjectHash.cpp b/Engine/Source/Runtime/CoreUObject/Private/UObject/UObjectHash.cpp
--- a/Engine/Source/Runtime/CoreUObject/Private/UObject/UObjectHash.cpp
+++ b/Engine/Source/Runtime/CoreUObject/Private/UObject/UObjectHash.cpp
@@ -1366,6 +1366,7 @@ static void RecursivelyPopulateDerivedClasses(FUObjectHashTables& ThreadHash, co
 		TSet<UClass*>* ChildSet = ThreadHash.ClassToChildListMap.Find(SearchClass);
 		if (ChildSet)
 		{
+			OutAllDerivedClass.Reserve(OutAllDerivedClass.Num() + ChildSet->Num());
 			for (UClass* ChildClass : *ChildSet)
 			{
 				OutAllDerivedClass.Add(ChildClass);
@@ -1474,9 +1475,13 @@ void GetDerivedClasses(const UClass* ClassToLookFor, TArray<UClass*>& Results, b
 	else
 	{
 		TSet<UClass*>* DerivedClasses = ThreadHash.ClassToChildListMap.Find(ClassToLookFor);
-		if ( DerivedClasses )
+		if (DerivedClasses)
 		{
-			Results.Append( DerivedClasses->Array() );
+			Results.Reserve(Results.Num() + DerivedClasses->Num());
+			for (UClass* DerivedClass : *DerivedClasses)
+			{
+				Results.Add(DerivedClass);
+			}
 		}
 	}
 }
