Improve logging
#rnx

diff --git a/Engine/Source/Runtime/CoreUObject/Private/UObject/Enum.cpp b/Engine/Source/Runtime/CoreUObject/Private/UObject/Enum.cpp
--- a/Engine/Source/Runtime/CoreUObject/Private/UObject/Enum.cpp
+++ b/Engine/Source/Runtime/CoreUObject/Private/UObject/Enum.cpp
@@ -723,7 +723,12 @@ int32 UEnum::GetIndexByNameString(const FString& InSearchString, EGetByNameFlags
 				SerializedObject = LoadContext->SerializedObject;
 			}
 		}
-		UE_LOG(LogEnum, Warning, TEXT("In asset '%s', there is an enum property of type '%s' with an invalid value of '%s'"), *GetPathNameSafe(SerializedObject ? SerializedObject : FUObjectThreadContext::Get().ConstructedObject), *GetName(), *InSearchString);
+		const bool bIsNativeOrLoaded = (!HasAnyFlags(RF_WasLoaded) || HasAnyFlags(RF_LoadCompleted));
+		UE_LOG(LogEnum, Warning, TEXT("UEnum: In asset '%s', there is an enum property of type '%s' with an invalid value of '%s' - %s"), 
+			*GetPathNameSafe(SerializedObject ? SerializedObject : FUObjectThreadContext::Get().ConstructedObject), 
+			*GetName(), 
+			*InSearchString,
+			bIsNativeOrLoaded ? TEXT("loaded") : TEXT("not loaded"));
 	}
 
 	return INDEX_NONE;
diff --git a/Engine/Source/Runtime/CoreUObject/Private/UObject/EnumProperty.cpp b/Engine/Source/Runtime/CoreUObject/Private/UObject/EnumProperty.cpp
--- a/Engine/Source/Runtime/CoreUObject/Private/UObject/EnumProperty.cpp
+++ b/Engine/Source/Runtime/CoreUObject/Private/UObject/EnumProperty.cpp
@@ -357,7 +357,12 @@ const TCHAR* FEnumProperty::ImportText_Internal(const TCHAR* InBuffer, void* Con
 					SerializedObject = LoadContext->SerializedObject;
 				}
 			}
-			ErrorText->Logf(ELogVerbosity::Warning, TEXT("In asset '%s', there is an enum property of type '%s' with an invalid value of '%s'"), *GetPathNameSafe(SerializedObject ? SerializedObject : FUObjectThreadContext::Get().ConstructedObject), *Enum->GetName(), *Temp);
+			const bool bIsNativeOrLoaded = (!Enum->HasAnyFlags(RF_WasLoaded) || Enum->HasAnyFlags(RF_LoadCompleted));
+			ErrorText->Logf(ELogVerbosity::Warning, TEXT("FEP: In asset '%s', there is an enum property of type '%s' with an invalid value of '%s' - %s"), 
+				*GetPathNameSafe(SerializedObject ? SerializedObject : FUObjectThreadContext::Get().ConstructedObject), 
+				*Enum->GetName(), 
+				*Temp,
+				bIsNativeOrLoaded ? TEXT("loaded") : TEXT("not loaded"));
 			return nullptr;
 		}
 	}
diff --git a/Engine/Source/Runtime/CoreUObject/Private/UObject/PropertyByte.cpp b/Engine/Source/Runtime/CoreUObject/Private/UObject/PropertyByte.cpp
--- a/Engine/Source/Runtime/CoreUObject/Private/UObject/PropertyByte.cpp
+++ b/Engine/Source/Runtime/CoreUObject/Private/UObject/PropertyByte.cpp
@@ -448,7 +448,12 @@ const TCHAR* FByteProperty::ImportText_Internal( const TCHAR* InBuffer, void* Co
 					SerializedObject = LoadContext->SerializedObject;
 				}
 			}
-			ErrorText->Logf(ELogVerbosity::Warning, TEXT("In asset '%s', there is an enum property of type '%s' with an invalid value of '%s'"), *GetPathNameSafe(SerializedObject ? SerializedObject : FUObjectThreadContext::Get().ConstructedObject), *Enum->GetName(), *Temp);
+			const bool bIsNativeOrLoaded = (!Enum->HasAnyFlags(RF_WasLoaded) || Enum->HasAnyFlags(RF_LoadCompleted));
+			ErrorText->Logf(ELogVerbosity::Warning, TEXT("FBP: In asset '%s', there is an enum property of type '%s' with an invalid value of '%s' - %s"), 
+				*GetPathNameSafe(SerializedObject ? SerializedObject : FUObjectThreadContext::Get().ConstructedObject), 
+				*Enum->GetName(), 
+				*Temp,
+				bIsNativeOrLoaded ? TEXT("loaded") : TEXT("not loaded"));
 			return nullptr;
 		}
 	}
