Modified auto optional asset check to consider class assets as well as support derived types

diff --git a/Engine/Source/Runtime/CoreUObject/Private/UObject/SavePackage/SaveContext.cpp b/Engine/Source/Runtime/CoreUObject/Private/UObject/SavePackage/SaveContext.cpp
--- a/Engine/Source/Runtime/CoreUObject/Private/UObject/SavePackage/SaveContext.cpp
+++ b/Engine/Source/Runtime/CoreUObject/Private/UObject/SavePackage/SaveContext.cpp
@@ -2,6 +2,7 @@
 
 #include "SaveContext.h"
 
+#include "Algo/Find.h"
 #include "Misc/ConfigCacheIni.h"
 #include "Serialization/PackageWriter.h"
 #include "UObject/UObjectGlobals.h"
@@ -187,7 +188,14 @@ void FSaveContext::SetupHarvestingRealms()
 
 	if (bAssetListGenerated && Asset)
 	{
-		bIsSaveAutoOptional = IsCooking() && IsSaveOptional() && AutomaticOptionalInclusionAssetTypeList.Contains(Asset->GetClass());
+		// if the asset type itself is a class (ie. BP) use that to check for auto optional
+		UClass* AssetType = Cast<UClass>(Asset);
+		AssetType = AssetType ? AssetType : Asset->GetClass();
+		bool bAllowedClass = Algo::FindByPredicate(AutomaticOptionalInclusionAssetTypeList, [AssetType](const UClass* InAssetClass)
+			{
+				return AssetType->IsChildOf(InAssetClass);
+			}) != nullptr;
+		bIsSaveAutoOptional = IsCooking() && IsSaveOptional() && bAllowedClass;
 	}
 }
 
