Improve logging
#rnx

diff --git a/Engine/Source/Runtime/CoreUObject/Private/UObject/Enum.cpp b/Engine/Source/Runtime/CoreUObject/Private/UObject/Enum.cpp
--- a/Engine/Source/Runtime/CoreUObject/Private/UObject/Enum.cpp
+++ b/Engine/Source/Runtime/CoreUObject/Private/UObject/Enum.cpp
@@ -724,11 +724,12 @@ int32 UEnum::GetIndexByNameString(const FString& InSearchString, EGetByNameFlags
 			}
 		}
 		const bool bIsNativeOrLoaded = (!HasAnyFlags(RF_WasLoaded) || HasAnyFlags(RF_LoadCompleted));
-		UE_LOG(LogEnum, Warning, TEXT("UEnum: In asset '%s', there is an enum property of type '%s' with an invalid value of '%s' - %s"), 
+		UE_LOG(LogEnum, Warning, TEXT("UEnum: In asset '%s', there is an enum property of type '%s' with an invalid value of '%s' - %s - %d"), 
 			*GetPathNameSafe(SerializedObject ? SerializedObject : FUObjectThreadContext::Get().ConstructedObject), 
 			*GetName(), 
 			*InSearchString,
-			bIsNativeOrLoaded ? TEXT("loaded") : TEXT("not loaded"));
+			bIsNativeOrLoaded ? TEXT("loaded") : TEXT("not loaded"),
+			Count);
 	}
 
 	return INDEX_NONE;
