Adding the ability to delay reachability iterations by the desired number of frames

diff --git a/Engine/Source/Runtime/CoreUObject/Private/UObject/GarbageCollection.cpp b/Engine/Source/Runtime/CoreUObject/Private/UObject/GarbageCollection.cpp
--- a/Engine/Source/Runtime/CoreUObject/Private/UObject/GarbageCollection.cpp
+++ b/Engine/Source/Runtime/CoreUObject/Private/UObject/GarbageCollection.cpp
@@ -225,6 +225,7 @@ FConsoleCommandWithOutputDeviceDelegate::CreateStatic(&UnlinkReachabilityStressD
 
 #if UE_BUILD_SHIPPING
 static constexpr int32 GGarbageReferenceTrackingEnabled = 0;
+static constexpr int32 GDelayReachabilityIterations = 0;
 #else
 int32 GGarbageReferenceTrackingEnabled = 0;
 static FAutoConsoleVariableRef CGarbageReferenceTrackingEnabled(
@@ -233,6 +234,14 @@ static FAutoConsoleVariableRef CGarbageReferenceTrackingEnabled(
 	TEXT("Causes the Garbage Collector to track and log unreleased garbage objects. If 1, will dump every reference. If 2, will dump a sample of the references to highlight problematic properties."),
 	ECVF_Default
 );
+
+int32 GDelayReachabilityIterations = 0;
+static FAutoConsoleVariableRef CDelayReachabilityIterations(
+	TEXT("gc.DelayReachabilityIterations"),
+	GDelayReachabilityIterations,
+	TEXT("Causes the Garbage Collector to delay incremental reachability iterations by the provided number of frames."),
+	ECVF_Default
+);
 #endif // UE_BUILD_SHIPPING
 
 static int32 GAllowIncrementalReachability = 0;
@@ -5169,16 +5178,21 @@ void FReachabilityAnalysisState::PerformReachabilityAnalysis()
 	if (!bIsSuspended)
 	{
 		Init();
+		NumRechabilityIterationsToSkip = FMath::Max(0, GDelayReachabilityIterations);
 	}
 
 	if (bPerformFullPurge)
 	{
 		UE::GC::CollectGarbageFull(ObjectKeepFlags);
 	}
-	else
+	else if (NumRechabilityIterationsToSkip == 0 || !bIsSuspended)
 	{
 		UE::GC::CollectGarbageIncremental(ObjectKeepFlags);
 	}
+	else
+	{
+		--NumRechabilityIterationsToSkip;
+	}
 
 	FinishIteration();
 }
diff --git a/Engine/Source/Runtime/CoreUObject/Private/UObject/ReachabilityAnalysisState.h b/Engine/Source/Runtime/CoreUObject/Private/UObject/ReachabilityAnalysisState.h
--- a/Engine/Source/Runtime/CoreUObject/Private/UObject/ReachabilityAnalysisState.h
+++ b/Engine/Source/Runtime/CoreUObject/Private/UObject/ReachabilityAnalysisState.h
@@ -49,6 +49,9 @@ private:
 	/** Number of reachability analysis iterations performed during reachability analysis */
 	int32 NumIterations = 0;
 
+	/** Number of reachability analysis iterations to skip when running with gc.DelayReachabilityIterations */
+	int32 NumRechabilityIterationsToSkip = 0;
+
 	alignas (PLATFORM_CACHE_LINE_SIZE) double IterationStartTime = 0.0;
 	alignas (PLATFORM_CACHE_LINE_SIZE) double IterationTimeLimit = 0.0;
 
