Fix crash with field metadata when using cooked data that includes editor only data (editor optional data)

diff --git a/Engine/Source/Runtime/CoreUObject/Private/UObject/Field.cpp b/Engine/Source/Runtime/CoreUObject/Private/UObject/Field.cpp
--- a/Engine/Source/Runtime/CoreUObject/Private/UObject/Field.cpp
+++ b/Engine/Source/Runtime/CoreUObject/Private/UObject/Field.cpp
@@ -418,7 +418,7 @@ void FField::Serialize(FArchive& Ar)
 	if (!Ar.IsCooking())
 	{
 		UPackage* Package = GetOutermost();
-		if (!Package || !Package->bIsCookedForEditor)
+		if (!Package || !Package->HasAnyPackageFlags(PKG_Cooked))
 		{
 			bool bHasMetaData = false;
 			if (Ar.IsLoading())
