Fix missing ReferencingObject in GC verification code by retrieving it from context object in UE::GC::TDefaultCollector

diff --git a/Engine/Source/Runtime/CoreUObject/Public/UObject/FastReferenceCollector.h b/Engine/Source/Runtime/CoreUObject/Public/UObject/FastReferenceCollector.h
--- a/Engine/Source/Runtime/CoreUObject/Public/UObject/FastReferenceCollector.h
+++ b/Engine/Source/Runtime/CoreUObject/Public/UObject/FastReferenceCollector.h
@@ -852,10 +852,18 @@ public:
 
 	virtual void HandleObjectReference(UObject*& Object, const UObject* ReferencingObject, const FProperty* ReferencingProperty) override
 	{
+		if (!ReferencingObject)
+		{
+			ReferencingObject = Context.GetReferencingObject();
+		}
 		Processor.HandleTokenStreamObjectReference(Context, const_cast<UObject*>(ReferencingObject), Object, EMemberlessId::Collector, EOrigin::Other, false);
 	}
 	virtual void HandleObjectReferences(UObject** InObjects, const int32 ObjectNum, const UObject* ReferencingObject, const FProperty* InReferencingProperty) override
 	{
+		if (!ReferencingObject)
+		{
+			ReferencingObject = Context.GetReferencingObject();
+		}
 		for (int32 ObjectIndex = 0; ObjectIndex < ObjectNum; ++ObjectIndex)
 		{
 			UObject*& Object = InObjects[ObjectIndex];
