Add an exception looking for dependencies when mapping object for the ClassGeneratedBy object on classes

diff --git a/Engine/Source/Runtime/CoreUObject/Private/UObject/LinkerSave.cpp b/Engine/Source/Runtime/CoreUObject/Private/UObject/LinkerSave.cpp
--- a/Engine/Source/Runtime/CoreUObject/Private/UObject/LinkerSave.cpp
+++ b/Engine/Source/Runtime/CoreUObject/Private/UObject/LinkerSave.cpp
@@ -265,9 +265,17 @@ FPackageIndex FLinkerSave::MapObject(TObjectPtr<const UObject> Object) const
 				}
 				if (!bFoundDep)
 				{
-					if (SavingExport.Object && SavingExport.Object->IsA(UClass::StaticClass()) && CastChecked<UClass>(SavingExport.Object)->GetDefaultObject() == Object)
+					if (SavingExport.Object && SavingExport.Object->IsA(UClass::StaticClass()))
 					{
-						bFoundDep = true; // the class is saving a ref to the CDO...which doesn't really work or do anything useful, but it isn't an error
+						UClass* Class = CastChecked<UClass>(SavingExport.Object);
+						if (Class->GetDefaultObject() == Object
+					#if WITH_EDITORONLY_DATA
+							|| Class->ClassGeneratedBy == Object
+					#endif
+							)
+						{
+							bFoundDep = true; // the class is saving a ref to the CDO...which doesn't really work or do anything useful, but it isn't an error or it is saving a reference to the class that generated it 
+						}
 					}
 				}
 				if (!bFoundDep)
