Reduce overhead in object handle access tracking by reducing high frequency thrash of TArrayView.  Synthetic test of 100 million handle accesses in a loop went from 4243ms total time to 1430ms total time.
Further improvements to be had by avoiding CheckCallable overhead of TFunction and making a lockless container to store read callbacks.  Those will be investigated in future changelists.

diff --git a/Engine/Source/Editor/UnrealEd/Classes/CookOnTheSide/CookOnTheFlyServer.h b/Engine/Source/Editor/UnrealEd/Classes/CookOnTheSide/CookOnTheFlyServer.h
--- a/Engine/Source/Editor/UnrealEd/Classes/CookOnTheSide/CookOnTheFlyServer.h
+++ b/Engine/Source/Editor/UnrealEd/Classes/CookOnTheSide/CookOnTheFlyServer.h
@@ -1409,7 +1409,7 @@ private:
 	/** Callback for analytics when a new UPackage is loaded. */
 	UNREALED_API void OnDiscoveredPackageDebug(FName PackageName, const UE::Cook::FInstigator& Instigator);
 	/** Callback for analytics when a TObjectPtr is read. */
-	UNREALED_API void OnObjectHandleReadDebug(TArrayView<const UObject*const> ReadObjects);
+	UNREALED_API void OnObjectHandleReadDebug(const TArrayView<const UObject*const>& ReadObjects);
 	/** Send warnings/telemetry when a discovered or read package is found to be a hidden dependency. */
 	UNREALED_API void ReportHiddenDependency(FName Referencer, FName Dependency);
 
diff --git a/Engine/Source/Editor/UnrealEd/Private/CookOnTheFlyServer.cpp b/Engine/Source/Editor/UnrealEd/Private/CookOnTheFlyServer.cpp
--- a/Engine/Source/Editor/UnrealEd/Private/CookOnTheFlyServer.cpp
+++ b/Engine/Source/Editor/UnrealEd/Private/CookOnTheFlyServer.cpp
@@ -11679,7 +11679,7 @@ void UCookOnTheFlyServer::RegisterCookByTheBookDelegates()
 	if (bHiddenDependenciesDebug)
 	{
 		ObjectHandleReadHandle = UE::CoreUObject::AddObjectHandleReadCallback(
-			[this](TArrayView<const UObject*const> ReadObjects) { UCookOnTheFlyServer::OnObjectHandleReadDebug(ReadObjects); });
+			[this](const TArrayView<const UObject*const>& ReadObjects) { UCookOnTheFlyServer::OnObjectHandleReadDebug(ReadObjects); });
 	}
 }
 
@@ -12527,7 +12527,7 @@ void UCookOnTheFlyServer::OnDiscoveredPackageDebug(FName PackageName, const UE::
 
 FName EngineTransientName(TEXT("/Engine/Transient"));
 
-void UCookOnTheFlyServer::OnObjectHandleReadDebug(TArrayView<const UObject*const> ReadObjects)
+void UCookOnTheFlyServer::OnObjectHandleReadDebug(const TArrayView<const UObject*const>& ReadObjects)
 {
 	using namespace UE::Cook;
 #if UE_WITH_PACKAGE_ACCESS_TRACKING
diff --git a/Engine/Source/Editor/UnrealEd/Private/Cooker/PackageBuildDependencyTracker.cpp b/Engine/Source/Editor/UnrealEd/Private/Cooker/PackageBuildDependencyTracker.cpp
--- a/Engine/Source/Editor/UnrealEd/Private/Cooker/PackageBuildDependencyTracker.cpp
+++ b/Engine/Source/Editor/UnrealEd/Private/Cooker/PackageBuildDependencyTracker.cpp
@@ -95,7 +95,7 @@ static bool ShouldSkipDependency(const UObject* Object)
 		(Object->GetClass() == UClass::StaticClass());
 }
 
-void FPackageBuildDependencyTracker::StaticOnObjectHandleRead(TArrayView<const UObject* const> Objects)
+void FPackageBuildDependencyTracker::StaticOnObjectHandleRead(const TArrayView<const UObject* const>& Objects)
 {
 	int Count = Objects.Num();
 	if(Count == 0 || (Count == 1 && ShouldSkipDependency(Objects[0])))
diff --git a/Engine/Source/Editor/UnrealEd/Private/Cooker/PackageBuildDependencyTracker.h b/Engine/Source/Editor/UnrealEd/Private/Cooker/PackageBuildDependencyTracker.h
--- a/Engine/Source/Editor/UnrealEd/Private/Cooker/PackageBuildDependencyTracker.h
+++ b/Engine/Source/Editor/UnrealEd/Private/Cooker/PackageBuildDependencyTracker.h
@@ -42,7 +42,7 @@ private:
 	virtual ~FPackageBuildDependencyTracker();
 
 	/** Track object reference reads */
-	static void StaticOnObjectHandleRead(TArrayView<const UObject* const> Objects);
+	static void StaticOnObjectHandleRead(const TArrayView<const UObject* const>& Objects);
 
 	// Use a mutex rather than a critical section for synchronization.  Calls into system libraries, such as windows critical section
 	// functions, are 50 times more expensive on build farm VMs, radically affecting cook times, which this avoids.  Saves 5% of total
diff --git a/Engine/Source/Runtime/CoreUObject/Private/UObject/ObjectHandleTracking.cpp b/Engine/Source/Runtime/CoreUObject/Private/UObject/ObjectHandleTracking.cpp
--- a/Engine/Source/Runtime/CoreUObject/Private/UObject/ObjectHandleTracking.cpp
+++ b/Engine/Source/Runtime/CoreUObject/Private/UObject/ObjectHandleTracking.cpp
@@ -84,11 +84,11 @@ namespace UE::CoreUObject
 	{
 		COREUOBJECT_API std::atomic<int32> HandleReadCallbackQuantity = 0;
 
-		struct ObjectHandleCallbacks
+		struct FObjectHandleCallbacks
 		{
-			static ObjectHandleCallbacks& Get()
+			static FObjectHandleCallbacks& Get()
 			{
-				static ObjectHandleCallbacks Callbacks;
+				static FObjectHandleCallbacks Callbacks;
 				return Callbacks;
 			}
 
@@ -104,6 +104,19 @@ namespace UE::CoreUObject
 				});
 			}
 
+			void OnHandleRead(const UObject* Object)
+			{
+				UE_AUTORTFM_OPEN(
+				{
+					FLightweightReadScopeLock _(HandleLock);
+					TArrayView<const UObject* const> Objects(&Object, 1);
+					for (auto&& Pair : ReadHandleCallbacks)
+					{
+						Pair.Value(Objects);
+					}
+				});
+			}
+
 			void OnClassReferenceResolved(const FObjectRef& ObjectRef, UPackage* Package, UClass* Class)
 			{
 				UE_AUTORTFM_OPEN(
@@ -271,22 +284,27 @@ namespace UE::CoreUObject
 
 		void OnHandleReadInternal(TArrayView<const UObject* const> Objects)
 		{
-			ObjectHandleCallbacks::Get().OnHandleRead(Objects);
+			FObjectHandleCallbacks::Get().OnHandleRead(Objects);
+		}
+
+		void OnHandleReadInternal(const UObject* Object)
+		{
+			FObjectHandleCallbacks::Get().OnHandleRead(Object);
 		}
 
 		void OnClassReferenceResolved(const FObjectRef& ObjectRef, UPackage* Package, UClass* Class)
 		{
-			ObjectHandleCallbacks::Get().OnClassReferenceResolved(ObjectRef, Package, Class);
+			FObjectHandleCallbacks::Get().OnClassReferenceResolved(ObjectRef, Package, Class);
 		}
 
 		void OnReferenceResolved(const FObjectRef& ObjectRef, UPackage* Package, UObject* Object)
 		{
-			ObjectHandleCallbacks::Get().OnReferenceResolved(ObjectRef, Package, Object);
+			FObjectHandleCallbacks::Get().OnReferenceResolved(ObjectRef, Package, Object);
 		}
 
 		void OnReferenceLoaded(const FObjectRef& ObjectRef, UPackage* Package, UObject* Object)
 		{
-			ObjectHandleCallbacks::Get().OnReferenceLoaded(ObjectRef, Package, Object);
+			FObjectHandleCallbacks::Get().OnReferenceLoaded(ObjectRef, Package, Object);
 		}
 
 	}
@@ -294,49 +312,49 @@ namespace UE::CoreUObject
 	FObjectHandleTrackingCallbackId AddObjectHandleReadCallback(FObjectHandleReadFunc Func)
 	{
 		using namespace UE::CoreUObject::Private;
-		return ObjectHandleCallbacks::Get().AddObjectHandleReadCallback(Func);
+		return FObjectHandleCallbacks::Get().AddObjectHandleReadCallback(Func);
 	}
 
 	void RemoveObjectHandleReadCallback(FObjectHandleTrackingCallbackId Handle)
 	{
 		using namespace UE::CoreUObject::Private;
-		ObjectHandleCallbacks::Get().RemoveObjectHandleReadCallback(Handle);
+		FObjectHandleCallbacks::Get().RemoveObjectHandleReadCallback(Handle);
 	}
 
 	FObjectHandleTrackingCallbackId AddObjectHandleClassResolvedCallback(FObjectHandleClassResolvedFunc Func)
 	{
 		using namespace UE::CoreUObject::Private;
-		return ObjectHandleCallbacks::Get().AddObjectHandleClassResolvedCallback(Func);
+		return FObjectHandleCallbacks::Get().AddObjectHandleClassResolvedCallback(Func);
 	}
 
 	void RemoveObjectHandleClassResolvedCallback(FObjectHandleTrackingCallbackId Handle)
 	{
 		using namespace UE::CoreUObject::Private;
-		ObjectHandleCallbacks::Get().RemoveObjectHandleClassResolvedCallback(Handle);
+		FObjectHandleCallbacks::Get().RemoveObjectHandleClassResolvedCallback(Handle);
 	}
 
 	FObjectHandleTrackingCallbackId AddObjectHandleReferenceResolvedCallback(FObjectHandleReferenceResolvedFunc Func)
 	{
 		using namespace UE::CoreUObject::Private;
-		return ObjectHandleCallbacks::Get().AddObjectHandleReferenceResolvedCallback(Func);
+		return FObjectHandleCallbacks::Get().AddObjectHandleReferenceResolvedCallback(Func);
 	}
 
 	void RemoveObjectHandleReferenceResolvedCallback(FObjectHandleTrackingCallbackId Handle)
 	{
 		using namespace UE::CoreUObject::Private;
-		ObjectHandleCallbacks::Get().RemoveObjectHandleReferenceResolvedCallback(Handle);
+		FObjectHandleCallbacks::Get().RemoveObjectHandleReferenceResolvedCallback(Handle);
 	}
 
 	FObjectHandleTrackingCallbackId AddObjectHandleReferenceLoadedCallback(FObjectHandleReferenceLoadedFunc Func)
 	{
 		using namespace UE::CoreUObject::Private;
-		return ObjectHandleCallbacks::Get().AddObjectHandleReferenceLoadedCallback(Func);
+		return FObjectHandleCallbacks::Get().AddObjectHandleReferenceLoadedCallback(Func);
 	}
 
 	void RemoveObjectHandleReferenceLoadedCallback(FObjectHandleTrackingCallbackId Handle)
 	{
 		using namespace UE::CoreUObject::Private;
-		ObjectHandleCallbacks::Get().RemoveObjectHandleReferenceLoadedCallback(Handle);
+		FObjectHandleCallbacks::Get().RemoveObjectHandleReferenceLoadedCallback(Handle);
 	}
 }
 #endif // UE_WITH_OBJECT_HANDLE_TRACKING
diff --git a/Engine/Source/Runtime/CoreUObject/Private/UObject/ObjectSerializeAccessScope.cpp b/Engine/Source/Runtime/CoreUObject/Private/UObject/ObjectSerializeAccessScope.cpp
--- a/Engine/Source/Runtime/CoreUObject/Private/UObject/ObjectSerializeAccessScope.cpp
+++ b/Engine/Source/Runtime/CoreUObject/Private/UObject/ObjectSerializeAccessScope.cpp
@@ -41,7 +41,7 @@ namespace UE::CoreUObject::Private
 		{
 			if (bEnabled)
 			{
-				FObjectHandleReadFunc Func = [this](TArrayView<const UObject* const> ReadObject)
+				FObjectHandleReadFunc Func = [this](const TArrayView<const UObject* const>& ReadObject)
 				{
 					FLocalData& Data = GetThreadLocalData();
 					if (!IsEnableSerializeAccessWarnings() || (Data.CurrentScope == nullptr || Data.CurrentScope->bSuspended))
diff --git a/Engine/Source/Runtime/CoreUObject/Public/UObject/ObjectHandleTracking.h b/Engine/Source/Runtime/CoreUObject/Public/UObject/ObjectHandleTracking.h
--- a/Engine/Source/Runtime/CoreUObject/Public/UObject/ObjectHandleTracking.h
+++ b/Engine/Source/Runtime/CoreUObject/Public/UObject/ObjectHandleTracking.h
@@ -26,7 +26,7 @@ namespace UE::CoreUObject
 	  *
 	  * @param ReadObject	The object that was read from a handle.
 	  */
-	using FObjectHandleReadFunc = TFunction<void(TArrayView<const UObject* const> Objects)>;
+	using FObjectHandleReadFunc = TFunction<void(const TArrayView<const UObject* const>& Objects)>;
 
 	/**
 	 * Callback notifying when a class is resolved from an object handle or object reference.
@@ -105,6 +105,7 @@ namespace UE::CoreUObject::Private
 {
 	extern COREUOBJECT_API std::atomic<int32> HandleReadCallbackQuantity;
 	COREUOBJECT_API void OnHandleReadInternal(TArrayView<const UObject* const> Objects);
+	COREUOBJECT_API void OnHandleReadInternal(const UObject* Object);
 	inline void OnHandleRead(TArrayView<const UObject* const> Objects)
 	{
 		if (HandleReadCallbackQuantity.load(std::memory_order_acquire) > 0)
@@ -116,8 +117,7 @@ namespace UE::CoreUObject::Private
 	{
 		if (HandleReadCallbackQuantity.load(std::memory_order_acquire) > 0)
 		{
-			TArrayView<const UObject* const> View(&Object, 1);
-			OnHandleReadInternal(View);
+			OnHandleReadInternal(Object);
 		}
 	}
 	COREUOBJECT_API void OnClassReferenceResolved(const FObjectRef& ObjectRef, UPackage* ClassPackage, UClass* Class);
diff --git a/Engine/Source/Runtime/CoreUObject/Tests/ObjectPtrTest.cpp b/Engine/Source/Runtime/CoreUObject/Tests/ObjectPtrTest.cpp
--- a/Engine/Source/Runtime/CoreUObject/Tests/ObjectPtrTest.cpp
+++ b/Engine/Source/Runtime/CoreUObject/Tests/ObjectPtrTest.cpp
@@ -796,7 +796,7 @@ void TestArrayConversion()
 #if UE_WITH_OBJECT_HANDLE_TRACKING
 	int ResolveCount = 0;
 	uint32 ObjCount = 0;
-	auto ResolveDelegate = [&](TArrayView<const UObject* const> Objects)
+	auto ResolveDelegate = [&](const TArrayView<const UObject* const>& Objects)
 	{
 		++ResolveCount;
 		ObjCount = Objects.Num();
@@ -988,7 +988,7 @@ TEST_CASE("CoreUObject::TObjectPtr::ArrayConversionReferenceForSet")
 #if UE_WITH_OBJECT_HANDLE_TRACKING
 	int ResolveCount = 0;
 	uint32 ObjCount = 0;
-	auto ResolveDelegate = [&](TArrayView<const UObject* const> Objects)
+	auto ResolveDelegate = [&](const TArrayView<const UObject* const>& Objects)
 	{
 		++ResolveCount;
 		ObjCount = Objects.Num();
@@ -1042,7 +1042,7 @@ TEST_CASE("CoreUObject::TObjectPtr::ConstArrayViewConversion")
 
 #if UE_WITH_OBJECT_HANDLE_TRACKING
 	int ResolveCount = 0;
-	auto ResolveDelegate = [&](TArrayView<const UObject* const> Objects)
+	auto ResolveDelegate = [&](const TArrayView<const UObject* const>& Objects)
 	{
 		++ResolveCount;
 		CHECK(Objects.Num() == 3);
diff --git a/Engine/Source/Runtime/CoreUObject/Tests/ObjectRefTrackingTestBase.cpp b/Engine/Source/Runtime/CoreUObject/Tests/ObjectRefTrackingTestBase.cpp
--- a/Engine/Source/Runtime/CoreUObject/Tests/ObjectRefTrackingTestBase.cpp
+++ b/Engine/Source/Runtime/CoreUObject/Tests/ObjectRefTrackingTestBase.cpp
@@ -31,11 +31,11 @@ TEST_CASE("UE::CoreUObject::ObjectHandleTracking::Callbacks")
 	int32 CallbackCount7 = 0;
 	int32 CallbackCount8 = 0;
 
-	FObjectHandleReadFunc ReadCallback1 = [&](TArrayView<const UObject* const> Objects)
+	FObjectHandleReadFunc ReadCallback1 = [&](const TArrayView<const UObject* const>& Objects)
 	{
 		++CallbackCount1;
 	};
-	FObjectHandleReadFunc ReadCallback2 = [&](TArrayView<const UObject* const> Objects)
+	FObjectHandleReadFunc ReadCallback2 = [&](const TArrayView<const UObject* const>& Objects)
 	{
 		++CallbackCount2;
 	};
@@ -206,7 +206,7 @@ TEST_CASE("UE::CoreUObject::ObjectHandleTracking::Callbacks::Verify")
 
 	TObjectPtr<UObject> ObjPtr = Obj1;
 	int32 ReadCount = 0;
-	auto ReadCallbackHandle1 = AddObjectHandleReadCallback([&](TArrayView<const UObject* const> ReadObject)
+	auto ReadCallbackHandle1 = AddObjectHandleReadCallback([&](const TArrayView<const UObject* const>& ReadObject)
 		{
 			++ReadCount;
 			CHECK(Obj1 == ReadObject[0]);
diff --git a/Engine/Source/Runtime/CoreUObject/Tests/ObjectRefTrackingTestBase.h b/Engine/Source/Runtime/CoreUObject/Tests/ObjectRefTrackingTestBase.h
--- a/Engine/Source/Runtime/CoreUObject/Tests/ObjectRefTrackingTestBase.h
+++ b/Engine/Source/Runtime/CoreUObject/Tests/ObjectRefTrackingTestBase.h
@@ -120,7 +120,7 @@ private:
 			NumFailedResolves++;
 		}
 	}
-	static void OnRefRead(TArrayView<const UObject* const> Objects)
+	static void OnRefRead(const TArrayView<const UObject* const>& Objects)
 	{
 		NumReads++;
 		for (auto Obj : Objects)
