Prevent bpvm execution up/down current callstack after Abort exception

diff --git a/Engine/Source/Runtime/CoreUObject/Public/UObject/Stack.h b/Engine/Source/Runtime/CoreUObject/Public/UObject/Stack.h
--- a/Engine/Source/Runtime/CoreUObject/Public/UObject/Stack.h
+++ b/Engine/Source/Runtime/CoreUObject/Public/UObject/Stack.h
@@ -202,6 +202,13 @@ public:
 		}
 #endif
 		PopThreadLocalTopStackFrame(PreviousTrackingFrame);
+		
+		if (PreviousTrackingFrame)
+		{
+			// we propagate bAbortingExecution to frames below to avoid losing abort state
+			// across heterogeneous frames (eg. bpvm -> c++ -> bpvm)
+			PreviousTrackingFrame->bAbortingExecution |= bAbortingExecution;
+		}
 	}
 
 	// Functions.
@@ -342,6 +349,19 @@ inline FFrame::FFrame( UObject* InObject, UFunction* InNode, void* InLocals, FFr
 	FBlueprintContextTracker::Get().ScriptStack.Push(this);
 #endif
 	PreviousTrackingFrame = PushThreadLocalTopStackFrame(this);
+	
+	{
+		// we propagate bAbortingExecution to *upper* frames to avoid invoking code
+		// on top of already-aborted frames
+		if (PreviousTrackingFrame)
+		{
+			bAbortingExecution |= PreviousTrackingFrame->bAbortingExecution;
+		}
+		if (InPreviousFrame)
+		{
+			bAbortingExecution |= InPreviousFrame->bAbortingExecution;
+		}
+	}
 
 #if PER_FUNCTION_SCRIPT_STATS
 	if (InPreviousFrame)
