LexToString implementation for object flags

diff --git a/Engine/Source/Runtime/CoreUObject/Private/UObject/LinkerDiff.cpp b/Engine/Source/Runtime/CoreUObject/Private/UObject/LinkerDiff.cpp
--- a/Engine/Source/Runtime/CoreUObject/Private/UObject/LinkerDiff.cpp
+++ b/Engine/Source/Runtime/CoreUObject/Private/UObject/LinkerDiff.cpp
@@ -449,11 +449,6 @@ void FLinkerDiff::GenerateSearchableNameMapDiff(FLinker* LHSLinker, FLinker* RHS
 	}
 }
 
-FString LexToString(EObjectFlags Flags)
-{
-	return LexToString((int32)Flags);
-}
-
 void FLinkerDiff::GenerateExportDiff(FLinker* LHSLinker, const FLinkerExportObject& LHSExport, const FLinkerExportObject& RHSExport)
 {
 	TArray<FString> Diffs;
diff --git a/Engine/Source/Runtime/CoreUObject/Private/UObject/UObjectGlobals.cpp b/Engine/Source/Runtime/CoreUObject/Private/UObject/UObjectGlobals.cpp
--- a/Engine/Source/Runtime/CoreUObject/Private/UObject/UObjectGlobals.cpp
+++ b/Engine/Source/Runtime/CoreUObject/Private/UObject/UObjectGlobals.cpp
@@ -9,6 +9,7 @@
 #include "HAL/FileManager.h"
 #include "Logging/StructuredLog.h"
 #include "Misc/AsciiSet.h"
+#include "Misc/EnumRange.h"
 #include "Misc/PackageAccessTrackingOps.h"
 #include "Misc/Paths.h"
 #include "Misc/ITransaction.h"
@@ -308,6 +309,61 @@ namespace
 	}
 }
 
+FString LexToString(EObjectFlags Flags)
+{
+	if (Flags == RF_NoFlags)
+	{
+		return TEXT("None");
+	}	
+	
+	static const TCHAR* Names[] = {
+		TEXT("Public"),
+		TEXT("Standalone"),
+		TEXT("MarkAsNative"),
+		TEXT("Transactional"),
+		TEXT("ClassDefaultObject"),
+		TEXT("ArchetypeObject"),
+		TEXT("Transient"),
+		TEXT("MarkAsRootSet"),
+		TEXT("TagGarbageTemp"),
+		TEXT("NeedInitialization"),
+		TEXT("NeedLoad"),
+		TEXT("KeepForCooker"),
+		TEXT("NeedPostLoad"),
+		TEXT("NeedPostLoadSubobjects"),
+		TEXT("NewerVersionExists"),
+		TEXT("BeginDestroyed"),
+		TEXT("FinishDestroyed"),
+		TEXT("BeingRegenerated"),
+		TEXT("DefaultSubObject"),
+		TEXT("WasLoaded"),
+		TEXT("TextExportTransient"),
+		TEXT("LoadCompleted"),
+		TEXT("InheritableComponentTemplate"),
+		TEXT("DuplicateTransient"),
+		TEXT("StrongRefOnFrame"),
+		TEXT("NonPIEDuplicateTransient"),
+		TEXT("Dynamic"),
+		TEXT("WillBeLoaded"),
+		TEXT("HasExternalPackage"),
+		TEXT("PendingKill"),
+		TEXT("Garbage"),
+		TEXT("AllocatedInSharedPage"),
+	};
+
+	TStringBuilder<1024> Builder;
+	for (EObjectFlags Flag : MakeFlagsRange(Flags))
+	{
+		int32 Index = FMath::FloorLog2((uint32)Flag);	
+		if (Builder.Len() > 0)
+		{
+			Builder << TEXT(" | ");
+		}	
+		Builder << Names[Index];
+	}
+	return Builder.ToString();
+}
+
 /** Object annotation used to keep track of the number suffixes  */
 struct FPerClassNumberSuffixAnnotation
 {
diff --git a/Engine/Source/Runtime/CoreUObject/Public/UObject/ObjectMacros.h b/Engine/Source/Runtime/CoreUObject/Public/UObject/ObjectMacros.h
--- a/Engine/Source/Runtime/CoreUObject/Public/UObject/ObjectMacros.h
+++ b/Engine/Source/Runtime/CoreUObject/Public/UObject/ObjectMacros.h
@@ -519,6 +519,7 @@ COREUOBJECT_API const TCHAR* LexToString(EPropertyObjectReferenceType Type);
 
 /**
  * Flags describing an object instance
+ * When modifying this enum, update the LexToString implementation! 
  */
 enum EObjectFlags
 {
@@ -591,6 +592,8 @@ PRAGMA_ENABLE_DEPRECATION_WARNINGS
 
 ENUM_CLASS_FLAGS(EObjectFlags);
 
+COREUOBJECT_API FString LexToString(EObjectFlags Flags);
+
 /** 
  * Objects flags for internal use (GC, low level UObject code)
  *
