Improved logging for when UObject::GetVersePath() fails.

diff --git a/Engine/Source/Runtime/CoreUObject/Private/UObject/UObjectBaseUtility.cpp b/Engine/Source/Runtime/CoreUObject/Private/UObject/UObjectBaseUtility.cpp
--- a/Engine/Source/Runtime/CoreUObject/Private/UObject/UObjectBaseUtility.cpp
+++ b/Engine/Source/Runtime/CoreUObject/Private/UObject/UObjectBaseUtility.cpp
@@ -279,7 +279,7 @@ UE::Core::FVersePath UObjectBaseUtility::GetVersePath() const
 		AlreadyLogged.Add(VerseModule, &bAlreadyInSet);
 		if (!bAlreadyInSet)
 		{
-			UE_LOG(LogCore, Display, TEXT("Unable to make a VersePath from %s"), *GetPathName());
+			UE_LOG(LogCore, Display, TEXT("Unable to make a VersePath for object '%s' with path '%s'"), *GetPathName(), *VerseModule);
 		}
 #endif
 	}
