Standalone builds always use raw transforms
#fix remove GIsEditor==true limitation for FCoreUObjectDelegates::OnEndLoadPackage broadcasting

diff --git a/Engine/Source/Runtime/CoreUObject/Private/Serialization/AsyncLoading.cpp b/Engine/Source/Runtime/CoreUObject/Private/Serialization/AsyncLoading.cpp
--- a/Engine/Source/Runtime/CoreUObject/Private/Serialization/AsyncLoading.cpp
+++ b/Engine/Source/Runtime/CoreUObject/Private/Serialization/AsyncLoading.cpp
@@ -4848,27 +4848,25 @@ EAsyncPackageState::Type FAsyncLoadingThread::ProcessLoadedPackages(bool bUseTim
 		Package->CallCompletionCallbacks(bInternalCallbacks, LoadingResult);
 		
 #if WITH_EDITOR
-		if (GIsEditor)
-		{
-			// In the editor we need to find any assets and packages and add them to list for later callback
-			Package->GetLoadedAssetsAndPackages(LoadedAssets, CompletedUPackages);
-		}
+		// In the editor we need to find any assets and packages and add them to list for later callback
+		Package->GetLoadedAssetsAndPackages(LoadedAssets, CompletedUPackages);
 #endif
 		// We don't need the package anymore
 		PackagesToDelete.AddUnique(Package);
 		Package->MarkRequestIDsAsComplete();
 	}
 #if WITH_EDITOR
-	if (GIsEditor)
+	// Call the global delegate for package endloads and set the bHasBeenLoaded flag that is used
+	// to check which packages have reached this state
+	for (UPackage* CompletedUPackage : CompletedUPackages)
+	{
+		CompletedUPackage->SetHasBeenEndLoaded(true);
+	}
+
+	if (CompletedUPackages.Num())
 	{
-		// Call the global delegate for package endloads and set the bHasBeenLoaded flag that is used
-		// to check which packages have reached this state
-		for (UPackage* CompletedUPackage : CompletedUPackages)
-		{
-			CompletedUPackage->SetHasBeenEndLoaded(true);
-		}
 		FCoreUObjectDelegates::OnEndLoadPackage.Broadcast(
-			FEndLoadPackageContext{ CompletedUPackages.Array(), 0, false /* bSynchronous */ });
+		FEndLoadPackageContext{ CompletedUPackages.Array(), 0, false /* bSynchronous */ });
 	}
 #endif
 
diff --git a/Engine/Source/Runtime/CoreUObject/Private/UObject/UObjectGlobals.cpp b/Engine/Source/Runtime/CoreUObject/Private/UObject/UObjectGlobals.cpp
--- a/Engine/Source/Runtime/CoreUObject/Private/UObject/UObjectGlobals.cpp
+++ b/Engine/Source/Runtime/CoreUObject/Private/UObject/UObjectGlobals.cpp
@@ -1592,7 +1592,7 @@ void BroadcastEndLoad(TArray<UPackage*>&& LoadedPackages)
 {
 #if WITH_EDITOR
 	// check(IsInGameThread()) was called by the caller, but we still need to test !IsInAsyncLoadingThread to exclude that callsite when the engine is single-threaded
-	if (GIsEditor && !IsInAsyncLoadingThread() && GGameThreadLoadCounter == 0)
+	if (!IsInAsyncLoadingThread() && GGameThreadLoadCounter == 0)
 	{
 		LoadedPackages.RemoveAllSwap([](UPackage* Package)
 			{
@@ -2241,23 +2241,20 @@ void EndLoad(FUObjectSerializeContext* LoadContext, TArray<UPackage*>* OutLoaded
 				SlowTask->TotalAmountOfWork += static_cast<float>(ObjLoaded.Num());
 				SlowTask->CurrentFrameScope = 0;
 			}
-#endif
 
-			if (GIsEditor)
+			for (int32 i = 0; i < ObjLoaded.Num(); i++)
 			{
-				for (int32 i = 0; i < ObjLoaded.Num(); i++)
+				UObject* Obj = ObjLoaded[i];
+				if (OutLoadedPackages)
 				{
-					UObject* Obj = ObjLoaded[i];
-					if (OutLoadedPackages)
-					{
-						LoadedPackages.Add(Obj->GetPackage());
-					}
-					if (Obj->GetLinker())
-					{
-						LoadedLinkers.Add(Obj->GetLinker());
-					}
+					LoadedPackages.Add(Obj->GetPackage());
+				}
+				if (GIsEditor && Obj->GetLinker())
+				{
+					LoadedLinkers.Add(Obj->GetLinker());
 				}
 			}
+#endif		
 
 			{
 				SCOPED_LOADTIMER(PostLoad);
diff --git a/Engine/Source/Runtime/CoreUObject/Public/UObject/UObjectGlobals.h b/Engine/Source/Runtime/CoreUObject/Public/UObject/UObjectGlobals.h
--- a/Engine/Source/Runtime/CoreUObject/Public/UObject/UObjectGlobals.h
+++ b/Engine/Source/Runtime/CoreUObject/Public/UObject/UObjectGlobals.h
@@ -3048,7 +3048,7 @@ struct FCoreUObjectDelegates
 	DECLARE_TS_MULTICAST_DELEGATE_OneParam(FOnObjectConstructed, UObject*);
 	static COREUOBJECT_API FOnObjectConstructed OnObjectConstructed;
 
-	/** Callback when packages end loading in LoadPackage or AsyncLoadPackage. All packages loaded recursively due to imports are included in the single call of the explicitly-loaded package. */
+	/** Callback when packages end loading in LoadPackage or AsyncLoadPackage. All packages loaded recursively due to imports are included in the single call of the explicitly-loaded package. It is called in all `WITH_EDITOR` cases, including -game. But it is not called in the runtime game compiled without WITH_EDITOR. This difference in behavior between the two game configurations is necessary because WITH_EDITOR uses a different loading mechanism than runtime game. */
 	DECLARE_MULTICAST_DELEGATE_OneParam(FOnEndLoadPackage, const FEndLoadPackageContext&);
 	static COREUOBJECT_API FOnEndLoadPackage OnEndLoadPackage;
 
