Fixed TScriptInterface's implicit conversions.

diff --git a/Engine/Source/Runtime/CoreUObject/Public/UObject/ScriptInterface.h b/Engine/Source/Runtime/CoreUObject/Public/UObject/ScriptInterface.h
--- a/Engine/Source/Runtime/CoreUObject/Public/UObject/ScriptInterface.h
+++ b/Engine/Source/Runtime/CoreUObject/Public/UObject/ScriptInterface.h
@@ -168,7 +168,7 @@ public:
 	 */
 	template <
 		typename OtherInterfaceType,
-		decltype(ImplicitConv<InInterfaceType*>(std::declval<OtherInterfaceType>()))* = nullptr
+		decltype(ImplicitConv<InInterfaceType*>(std::declval<OtherInterfaceType*>()))* = nullptr
 	>
 	FORCEINLINE TScriptInterface(const TScriptInterface<OtherInterfaceType>& Other)
 	{
@@ -190,12 +190,6 @@ public:
 		SetInterface(SourceInterface);
 	}
 
-	template <typename OtherInterfaceType>
-	TScriptInterface(const TScriptInterface<OtherInterfaceType>& Other)
-	{
-
-	}
-
 	/**
 	 * Copyable
 	 */
@@ -316,7 +310,7 @@ public:
 	 */
 	FORCEINLINE void SetInterface(InInterfaceType* InInterfacePointer)
 	{
-		FScriptInterface::SetInterface(InInterfacePointer);
+		FScriptInterface::SetInterface((void*)InInterfacePointer);
 	}
 
 	/**
