Remove sleep from zenloader flushrequest
  - Improves performance by waking up GT as soon as the async loading thread needs something processed on it
  - Improves performance by only entering the sleeping part when the GT is actually idle instead of all the time
  - Reduce power consumption and contention between GT/ALT by favoring longer sleep instead of spinning on platform not supporting sub milliseconds sleeps

diff --git a/Engine/Source/Runtime/CoreUObject/Private/Serialization/AsyncLoading2.cpp b/Engine/Source/Runtime/CoreUObject/Private/Serialization/AsyncLoading2.cpp
--- a/Engine/Source/Runtime/CoreUObject/Private/Serialization/AsyncLoading2.cpp
+++ b/Engine/Source/Runtime/CoreUObject/Private/Serialization/AsyncLoading2.cpp
@@ -66,6 +66,7 @@
 #include "ProfilingDebugging/AssetMetadataTrace.h"
 #include "Async/Async.h"
 #include "Async/ParallelFor.h"
+#include "Async/ManualResetEvent.h"
 #include "HAL/LowLevelMemStats.h"
 #include "HAL/IPlatformFileOpenLogWrapper.h"
 #include "Modules/ModuleManager.h"
@@ -2086,6 +2087,11 @@ public:
 		Zenaphore = InZenaphore;
 	}
 
+	void SetWakeEvent(UE::FManualResetEvent* InEvent)
+	{
+		WakeEvent = InEvent;
+	}
+
 	void SetOwnerThread(const FAsyncLoadingThreadState2* InOwnerThread)
 	{
 		OwnerThread = InOwnerThread;
@@ -2126,6 +2132,7 @@ private:
 
 	const FAsyncLoadingThreadState2* OwnerThread = nullptr;
 	FZenaphore* Zenaphore = nullptr;
+	UE::FManualResetEvent* WakeEvent = nullptr;
 	TIoPriorityQueue<FEventLoadNode2> LocalQueue;
 	FCriticalSection ExternalCritical;
 	TIoPriorityQueue<FEventLoadNode2> ExternalQueue;
@@ -3007,6 +3014,9 @@ private:
 	/** [ASYNC/GAME THREAD] Event used to signal that the async loading thread has resumed */
 	FEvent* ThreadResumedEvent;
 	TArray<FAsyncPackage2*> LoadedPackagesToProcess;
+	/** [ASYNC/GAME THREAD] Event used to signal the main thread that new processing is needed. Only used when ALT is active. */
+	UE::FManualResetEvent MainThreadWakeEvent;
+
 #if WITH_EDITOR
 	/** [GAME THREAD] */
 	TArray<UObject*> EditorLoadedAssets;
@@ -3105,8 +3115,33 @@ private:
 	TAtomic<int32> QueuedPackagesCounter { 0 };
 	/** [ASYNC/GAME THREAD] Number of packages being loaded on the async thread and post loaded on the game thread */
 	TAtomic<int32> LoadingPackagesCounter { 0 };
+	/** Encapsulate our counter to make sure all decrements are monitored. */
+	class FPackagesWithRemainingWorkCounter
+	{
+		UE::FManualResetEvent* WakeEvent = nullptr;
+		TAtomic<int32> PackagesWithRemainingWorkCounter {0};
+	public:
+		void SetWakeEvent(UE::FManualResetEvent* InWakeEvent) {	WakeEvent = InWakeEvent; }
+		int32 operator++() { return ++PackagesWithRemainingWorkCounter; }
+		int32 operator++(int) { return PackagesWithRemainingWorkCounter++; }
+		operator int32() const { return PackagesWithRemainingWorkCounter; }
+
+		// Only implement prefix
+		int32 operator--() 
+		{ 
+			int32 newValue = --PackagesWithRemainingWorkCounter;
+			if (newValue == 0)
+			{
+				if (WakeEvent)
+				{
+					WakeEvent->Notify();
+				}
+			}
+			return newValue;
+		}
+	};
 	/** [ASYNC/GAME THREAD] While this is non-zero there's work left to do */
-	TAtomic<int32> PackagesWithRemainingWorkCounter{ 0 };
+	FPackagesWithRemainingWorkCounter PackagesWithRemainingWorkCounter;
 
 	FThreadSafeCounter AsyncThreadReady;
 
@@ -3245,7 +3280,7 @@ public:
 	void UpdatePackagePriorityRecursive(FAsyncLoadingThreadState2& ThreadState, FAsyncPackage2* Package, int32 NewPriority);
 
 	FAsyncPackage2* FindOrInsertPackage(FAsyncLoadingThreadState2& ThreadState, FAsyncPackageDesc2& InDesc, bool& bInserted, FAsyncPackage2* ImportedByPackage, TUniquePtr<FLoadPackageAsyncDelegate>&& PackageLoadedDelegate = TUniquePtr<FLoadPackageAsyncDelegate>(), TUniquePtr<FLoadPackageAsyncProgressDelegate>&& PackageProgressDelegate = TUniquePtr<FLoadPackageAsyncProgressDelegate>());
-	void QueueMissingPackage(FAsyncPackageDesc2& PackageDesc, TUniquePtr<FLoadPackageAsyncDelegate>&& LoadPackageAsyncDelegate, TUniquePtr<FLoadPackageAsyncProgressDelegate>&& LoadPackageAsyncProgressDelegate);
+	void QueueMissingPackage(FAsyncLoadingThreadState2& ThreadState, FAsyncPackageDesc2& PackageDesc, TUniquePtr<FLoadPackageAsyncDelegate>&& LoadPackageAsyncDelegate, TUniquePtr<FLoadPackageAsyncProgressDelegate>&& LoadPackageAsyncProgressDelegate);
 
 	/**
 	* [ASYNC* THREAD] Loads all packages
@@ -3373,17 +3408,26 @@ public:
 	/**
 	 * [ASYNC/GAME THREAD] Removes a request ID from the list of pending requests
 	 */
-	void RemovePendingRequests(TConstArrayView<int32> RequestIDs)
+	void RemovePendingRequests(FAsyncLoadingThreadState2& ThreadState, TConstArrayView<int32> RequestIDs)
 	{
-		FScopeLock Lock(&PendingRequestsCritical);
-		for (int32 ID : RequestIDs)
+		int32 RemovedCount = 0;
 		{
-			PendingRequests.Remove(ID);
-			TRACE_LOADTIME_END_REQUEST(ID);
+			FScopeLock Lock(&PendingRequestsCritical);
+			for (int32 ID : RequestIDs)
+			{
+				RemovedCount += PendingRequests.Remove(ID);
+				TRACE_LOADTIME_END_REQUEST(ID);
+			}
+			if (PendingRequests.IsEmpty())
+			{
+				PendingRequests.Empty(DefaultAsyncPackagesReserveCount);
+			}
 		}
-		if (PendingRequests.IsEmpty())
+
+		// Any removed pending request is of interest to main thread as it might unblock a flush.
+		if (RemovedCount > 0 && ThreadState.bIsAsyncLoadingThread)
 		{
-			PendingRequests.Empty(DefaultAsyncPackagesReserveCount);
+			MainThreadWakeEvent.Notify();
 		}
 	}
 
@@ -3826,6 +3870,9 @@ void FAsyncLoadingThread2::UpdatePackagePriority(FAsyncLoadingThreadState2& Thre
 			if (Package->TryAddRef())
 			{
 				GameThreadState->PackagesToReprioritize.Enqueue(Package);
+
+				// Repriorization of packages is of interest to main thread as it could unblock a flush.
+				MainThreadWakeEvent.Notify();
 			}
 		}
 		else
@@ -4191,7 +4238,7 @@ bool FAsyncLoadingThread2::CreateAsyncPackagesFromQueue(FAsyncLoadingThreadState
 			FAsyncPackageDesc2 PackageDesc = FAsyncPackageDesc2::FromPackageRequest(Request, UPackageName, PackageIdToLoad);
 			if (PackageStatus == EPackageStoreEntryStatus::Missing)
 			{
-				QueueMissingPackage(PackageDesc, MoveTemp(Request.PackageLoadedDelegate), MoveTemp(Request.PackageProgressDelegate));
+				QueueMissingPackage(ThreadState, PackageDesc, MoveTemp(Request.PackageLoadedDelegate), MoveTemp(Request.PackageProgressDelegate));
 			}
 			else
 			{
@@ -4456,6 +4503,10 @@ void FAsyncLoadEventQueue2::PushExternal(FEventLoadNode2* Node)
 	{
 		Zenaphore->NotifyOne();
 	}
+	if (WakeEvent)
+	{
+		WakeEvent->Notify();
+	}
 }
 
 bool FAsyncLoadEventQueue2::PopAndExecute(FAsyncLoadingThreadState2& ThreadState)
@@ -6747,11 +6798,17 @@ void FAsyncPackage2::ConditionalFinishLoading(FAsyncLoadingThreadState2& ThreadS
 {
 	TRACE_CPUPROFILER_EVENT_SCOPE(ConditionalFinishLoading);
 	WaitForAllDependenciesToReachState(ThreadState, &FAsyncPackage2::AllDependenciesFullyLoadedState, EAsyncPackageLoadingState2::DeferredPostLoadDone, AsyncLoadingThread.ConditionalFinishLoadingTick,
-		[](FAsyncPackage2* Package)
+		[&ThreadState](FAsyncPackage2* Package)
 		{
 			check(Package->AsyncPackageLoadingState == EAsyncPackageLoadingState2::DeferredPostLoadDone);
 			Package->AsyncPackageLoadingState = EAsyncPackageLoadingState2::Finalize;
 			Package->AsyncLoadingThread.LoadedPackagesToProcess.Add(Package);
+
+			// Any update to LoadedPackagesToProcess is of interest to the main thread if we are on ALT.
+			if (ThreadState.bIsAsyncLoadingThread)
+			{
+				Package->AsyncLoadingThread.MainThreadWakeEvent.Notify();
+			}
 		});
 }
 
@@ -7097,7 +7154,7 @@ void FAsyncLoadingThread2::UpdateSyncLoadContext(FAsyncLoadingThreadState2& Thre
 				// Flushing a package while it's already being processed on the stack, if we're done preloading we let it pass and remove the request id
 				bool bPreloadIsDone = RequestedPackage->AsyncPackageLoadingState >= EAsyncPackageLoadingState2::DeferredPostLoad;
 				UE_CLOG(!bPreloadIsDone, LogStreaming, Fatal, TEXT("Flushing package %s while it's being preloaded in the same callstack is not permitted"), *RequestedPackage->Desc.UPackageName.ToString());
-				RemovePendingRequests({RequestID});
+				RemovePendingRequests(ThreadState, {RequestID});
 			}
 		}
 	}
@@ -7389,7 +7446,7 @@ EAsyncPackageState::Type FAsyncLoadingThread2::ProcessLoadedPackagesFromGameThre
 
 			if (Package->CompletionCallbacks.IsEmpty() && Package->ProgressCallbacks.IsEmpty())
 			{
-				RemovePendingRequests(Package->RequestIDs);
+				RemovePendingRequests(ThreadState, Package->RequestIDs);
 				Package->ReleaseRef();
 			}
 			else
@@ -7410,7 +7467,7 @@ EAsyncPackageState::Type FAsyncLoadingThread2::ProcessLoadedPackagesFromGameThre
 			if (FlushRequestIDs.Num() == 0 
 				|| Algo::AnyOf(FlushRequestIDs, [&CompletedPackageRequest](int32 FlushRequestID) { return CompletedPackageRequest.RequestIDs.Contains(FlushRequestID); }))
 			{
-				RemovePendingRequests(CompletedPackageRequest.RequestIDs);
+				RemovePendingRequests(ThreadState, CompletedPackageRequest.RequestIDs);
 				RequestsToProcess.Emplace(MoveTemp(CompletedPackageRequest));
 				CompletedPackageRequests.RemoveAt(CompletedPackageRequestIndex);
 				bLocalDidSomething = true;
@@ -7609,6 +7666,11 @@ void FAsyncLoadingThread2::StartThread()
 	{
 		AsyncLoadingThreadState = MakeUnique<FAsyncLoadingThreadState2>(GraphAllocator, IoDispatcher);
 		EventQueue.SetOwnerThread(AsyncLoadingThreadState.Get());
+
+		// When using ALT, we want to wake the main thread ASAP in case it's sleeping for lack of something to do during flush.
+		MainThreadEventQueue.SetWakeEvent(&MainThreadWakeEvent);
+		PackagesWithRemainingWorkCounter.SetWakeEvent(&MainThreadWakeEvent);
+
 		AsyncLoadingThreadState->bIsAsyncLoadingThread = true;
 		AsyncLoadingThreadState->bCanAccessAsyncLoadingThreadData = true;
 		GameThreadState->bCanAccessAsyncLoadingThreadData = false;
@@ -8617,7 +8679,7 @@ int32 FAsyncLoadingThread2::LoadPackage(const FPackagePath& InPackagePath, FName
 	return LoadPackageInternal(InPackagePath, InCustomName, MoveTemp(CompletionDelegate), TUniquePtr<FLoadPackageAsyncProgressDelegate>(), InPackageFlags, InPIEInstanceID, InPackagePriority, InInstancingContext, InLoadFlags);
 }
 
-void FAsyncLoadingThread2::QueueMissingPackage(FAsyncPackageDesc2& PackageDesc, TUniquePtr<FLoadPackageAsyncDelegate>&& PackageLoadedDelegate, TUniquePtr<FLoadPackageAsyncProgressDelegate>&& PackageProgressDelegate)
+void FAsyncLoadingThread2::QueueMissingPackage(FAsyncLoadingThreadState2& ThreadState, FAsyncPackageDesc2& PackageDesc, TUniquePtr<FLoadPackageAsyncDelegate>&& PackageLoadedDelegate, TUniquePtr<FLoadPackageAsyncProgressDelegate>&& PackageProgressDelegate)
 {
 	const FName FailedPackageName = PackageDesc.UPackageName;
 
@@ -8662,7 +8724,7 @@ void FAsyncLoadingThread2::QueueMissingPackage(FAsyncPackageDesc2& PackageDesc,
 	}
 	else
 	{
-		RemovePendingRequests(TArrayView<int32>(&PackageDesc.RequestID, 1));
+		RemovePendingRequests(ThreadState, TArrayView<int32>(&PackageDesc.RequestID, 1));
 		--PackagesWithRemainingWorkCounter;
 		TRACE_COUNTER_SET(AsyncLoadingPackagesWithRemainingWork, PackagesWithRemainingWorkCounter);
 	}
@@ -8778,11 +8840,22 @@ void FAsyncLoadingThread2::FlushLoading(TConstArrayView<int32> RequestIDs)
 
 				if (IsMultithreaded())
 				{
-					// Update the heartbeat and sleep. If we're not multithreading, the heartbeat is updated after each package has been processed
+					// Update the heartbeat and sleep a little if we had nothing to do unless ALT has made progress.
+					// If we're not multithreading, the heartbeat is updated after each package has been processed.
 					FThreadHeartBeat::Get().HeartBeat();
-					FPlatformProcess::SleepNoStats(0.0001f);
 
-					// Flush logging when runing cook-on-the-fly and waiting for packages
+					// Only going idle if nothing has been done
+					if (!bDidSomething)
+					{
+						TRACE_CPUPROFILER_EVENT_SCOPE(MainThreadWaitingOnAsyncLoadingThread);
+						// Still let the main thread tick at 60fps for processing message loop/etc.
+						MainThreadWakeEvent.WaitFor(UE::FMonotonicTimeSpan::FromMilliseconds(16));
+						// Reset the manual event right after we wake up so we don't miss any trigger.
+						// Worst case, we'll do an empty spin before going back to sleep.
+						MainThreadWakeEvent.Reset();
+					}
+
+					// Flush logging when running cook-on-the-fly and waiting for packages
 					if (IsRunningCookOnTheFly() && FPlatformTime::Seconds() - LogFlushTime > 1.0)
 					{
 						GLog->FlushThreadedLogs(EOutputDeviceRedirectorFlushOptions::Async);
