Add text to ensure.

diff --git a/Engine/Source/Runtime/CoreUObject/Private/UObject/UObjectGlobals.cpp b/Engine/Source/Runtime/CoreUObject/Private/UObject/UObjectGlobals.cpp
--- a/Engine/Source/Runtime/CoreUObject/Private/UObject/UObjectGlobals.cpp
+++ b/Engine/Source/Runtime/CoreUObject/Private/UObject/UObjectGlobals.cpp
@@ -3350,7 +3350,7 @@ UObject* StaticAllocateObject
 	if (bCreatingCDO)
 	{
 		check(InClass->GetClass());
-		ensure(!GIsDuplicatingClassForReinstancing || InClass->HasAnyClassFlags(CLASS_Native));
+		ensureMsgf(!GIsDuplicatingClassForReinstancing || InClass->HasAnyClassFlags(CLASS_Native), TEXT("GIsDuplicatingClassForReinstancing %d InClass %s"), (int)GIsDuplicatingClassForReinstancing, *InClass->GetPathName());
 		InName = InClass->GetDefaultObjectName();
 		// never call PostLoad on class default objects
 		InFlags &= ~(RF_NeedPostLoad|RF_NeedPostLoadSubobjects);
