Fixing GC clustering issues when incremental reachability is enabled.

- Fixing a bug where GC clusters are not marked as reachable when running a barrier on a clustered object
- Fixing a crash when constructing FWorkerContext for CollectReferences calls in between incremental reachability iterations

+ Adding automated tests for the above issues

diff --git a/Engine/Source/Runtime/CoreUObject/Private/UObject/GarbageCollection.cpp b/Engine/Source/Runtime/CoreUObject/Private/UObject/GarbageCollection.cpp
--- a/Engine/Source/Runtime/CoreUObject/Private/UObject/GarbageCollection.cpp
+++ b/Engine/Source/Runtime/CoreUObject/Private/UObject/GarbageCollection.cpp
@@ -492,7 +492,10 @@ namespace UE::GC
 
 namespace UE::GC::Private
 {
+	/** List of objects marked as reachable by GC barrier (see UObject::MarkAsReachable()) */
 	static TExpandingChunkedList<UObject*> GReachableObjects;
+	/** List of FUObjectItems representing cluster root objects marker as reachable by GC barrier (see UObject::MarkAsReachable()) */
+	static TExpandingChunkedList<FUObjectItem*> GReachableClusters;
 	bool GIsIncrementalReachabilityPending = false;
 }
 
@@ -1137,6 +1140,11 @@ int32 GetNumSlowAROs();
 // Currently doesn't free pages until shutdown
 class FPageAllocator
 {
+public:
+	// We need one extra worker cache to be able to perform non-async reference scans when incremental reachability is suspended
+	static constexpr int32 NumWorkerCaches = MaxWorkers + 1;
+
+private:
 	struct alignas(PLATFORM_CACHE_LINE_SIZE) FWorkerCache
 	{
 		UE_NONCOPYABLE(FWorkerCache);
@@ -1212,14 +1220,14 @@ class FPageAllocator
 	};	
 
 	FSharedCache SharedCache;
-	FWorkerCache WorkerCaches[MaxWorkers];
+	FWorkerCache WorkerCaches[NumWorkerCaches];
 
 public:
 	static constexpr uint64 PageSize = 4096;
 
 	void* AllocatePage(int32 WorkerIdx)
 	{
-		check(WorkerIdx >=0 && WorkerIdx < MaxWorkers);
+		check(WorkerIdx >=0 && WorkerIdx < NumWorkerCaches);
 		if (void* WorkerPage = WorkerCaches[WorkerIdx].Pop())
 		{
 			check(IsValidPage(WorkerPage));
@@ -1991,7 +1999,7 @@ FORCEINLINE_DEBUGGABLE void PadBlock(FWorkBlock& Block)
 class FWorkerIndexAllocator
 {
 	std::atomic<uint64> Used{0};
-	static_assert(MaxWorkers <= 64, "Currently supports single uint64 word");
+	static_assert(FPageAllocator::NumWorkerCaches <= 64, "Currently supports single uint64 word");
 public:
 	int32 AllocateChecked()
 	{
@@ -1999,7 +2007,7 @@ public:
 		{
 			uint64 UsedNow = Used.load(std::memory_order_relaxed);
 			uint64 FreeIndex = FPlatformMath::CountTrailingZeros64(~UsedNow);
-			checkf(FreeIndex < MaxWorkers, TEXT("Exceeded max active GC worker contexts"));
+			checkf(FreeIndex < FPageAllocator::NumWorkerCaches, TEXT("Exceeded max active GC worker contexts"));
 
 			uint64 Mask = uint64(1) << FreeIndex;
 			if ((Used.fetch_or(Mask) & Mask) == 0)
@@ -2011,7 +2019,7 @@ public:
 
 	void FreeChecked(int32 Index)
 	{
-		check(Index >= 0 && Index < MaxWorkers);
+		check(Index >= 0 && Index < FPageAllocator::NumWorkerCaches);
 		uint64 Mask = uint64(1) << Index;
 		uint64 Old = Used.fetch_and(~Mask);
 		checkf(Old & Mask, TEXT("Index already freed"));
@@ -3955,6 +3963,7 @@ private:
 
 		if (!Private::GReachableObjects.IsEmpty())
 		{
+			// Add objects marked with the GC barrier to the inital set of objects for the next iteration of incremental reachability
 			Private::GReachableObjects.PopAllAndEmpty(InitialObjects);
 			UE_LOG(LogGarbage, Verbose, TEXT("Adding %d object(s) marker by GC barrier to the list of objects to process"), InitialObjects.Num());
 			ConditionallyAddBarrierReferencesToHistory(*Context);
@@ -3964,6 +3973,18 @@ private:
 			Context->InitialNativeReferences = GetInitialReferences(Options);
 		}
 
+		if (!Private::GReachableClusters.IsEmpty())
+		{
+			// Process cluster roots that were marked as reachable by the GC barrier
+			TArray<FUObjectItem*> KeepClusterRefs;
+			Private::GReachableClusters.PopAllAndEmpty(KeepClusterRefs);
+			for (FUObjectItem* ObjectItem : KeepClusterRefs)
+			{
+				// Mark referenced clusters and mutable objects as reachable
+				MarkReferencedClustersAsReachable<EGCOptions::None>(ObjectItem->GetClusterIndex(), InitialObjects);
+			}
+		}
+
 		Context->SetInitialObjectsUnpadded(InitialObjects);
 
 		PerformReachabilityAnalysisOnObjects(Context, Options);
@@ -4004,7 +4025,7 @@ public:
 			do
 			{
 				PerformReachabilityAnalysisPass(Options);
-			} while (!Private::GReachableObjects.IsEmpty() && !GReachabilityState.IsSuspended());
+			} while ((!Private::GReachableObjects.IsEmpty() || !Private::GReachableClusters.IsEmpty()) && !GReachabilityState.IsSuspended());
 
 			UE_LOG(LogGarbage, Verbose, TEXT("%f ms for Reachability Analysis"), (FPlatformTime::Seconds() - StartTime) * 1000);
 		}
@@ -5410,19 +5431,39 @@ bool UObject::IsDestructionThreadSafe() const
 	return false;
 }
 
-void UObject::MarkAsReachable() const
-{	
+FORCEINLINE static void MarkObjectItemAsReachable(FUObjectItem* ObjectItem)
+{
 	using namespace UE::GC::Private;
 
-	FUObjectItem* ObjectItem = GUObjectArray.ObjectToObjectItem(const_cast<UObject*>(this));
-	if (ObjectItem->IsMaybeUnreachable())
+	checkf(GIsIncrementalReachabilityPending, TEXT("%s is marked as MaybeUnreachable but Incremental Reachability Analysis is not in progress"), *static_cast<UObject*>(ObjectItem->Object)->GetFullName());
+	if (ObjectItem->ThisThreadAtomicallyClearedMaybeUnreachable())
 	{
-		checkf(GIsIncrementalReachabilityPending, TEXT("%s is marked as MaybeUnreachable but Incremental Reachability Analysis is not in progress"), *GetFullName());
-		if (ObjectItem->ThisThreadAtomicallyClearedMaybeUnreachable())
+		if (ObjectItem->GetOwnerIndex() >= 0)
 		{
 			// This object became reachable so add it to a list of new objects to process in the next iteration of incremental GC because
 			// we need to mark objects it's referencing as reachable too
-			GReachableObjects.Push(const_cast<UObject*>(this));
+			GReachableObjects.Push(static_cast<UObject*>(ObjectItem->Object));
+		}
+		else
+		{
+			GReachableClusters.Push(ObjectItem);
+		}
+	}
+}
+
+void UObject::MarkAsReachable() const
+{	
+	FUObjectItem* ObjectItem = GUObjectArray.ObjectToObjectItem(const_cast<UObject*>(this));
+	if (ObjectItem->IsMaybeUnreachable())
+	{
+		MarkObjectItemAsReachable(ObjectItem);
+	}
+	else if (ObjectItem->GetOwnerIndex() > 0) // Clustered objects are never marked as MaybeUnreachable so we need to check if the cluster root is MaybeUnreachable
+	{
+		FUObjectItem* ClusterRootObjectItem = GUObjectArray.IndexToObject(ObjectItem->GetOwnerIndex());
+		if (ClusterRootObjectItem->IsMaybeUnreachable())
+		{
+			MarkObjectItemAsReachable(ClusterRootObjectItem);
 		}
 	}
 }
