Fix linker support that skip loading imports that are remapped to none from the instancing context, introduced in CL 25135064 and broken in CL 26540611. This feature is important in PIE for actors that have references to editor-only actors that aren't part of the streaming cells, but the actors still have references to them in their import tables.
Also added support for this in LinkerLoad's debug code that catches objects in external packages outered to objects in instanced packages, added in CL 27215992.
This fixes a couple UDNs reported with different setups of editor-only parent actors and some internal projects with editor-only parent actors generated through PCG.

diff --git a/Engine/Source/Runtime/CoreUObject/Private/UObject/LinkerInstancingContext.cpp b/Engine/Source/Runtime/CoreUObject/Private/UObject/LinkerInstancingContext.cpp
--- a/Engine/Source/Runtime/CoreUObject/Private/UObject/LinkerInstancingContext.cpp
+++ b/Engine/Source/Runtime/CoreUObject/Private/UObject/LinkerInstancingContext.cpp
@@ -45,11 +45,16 @@ public:
 		InstancedPackageMap.AddPackageMapping(Original, Instanced);
 	}
 
-	FName FindPackageMapping(FName Original) const
+	bool FindPackageMapping(FName Original, FName& Instanced) const
 	{
 		FReadScopeLock ScopeLock(Lock);
 		LLM_SCOPE_BYTAG(Loading_LinkerInstancingContext);
-		return InstancedPackageMap.InstancedPackageMapping.FindRef(Original);
+		if (const FName* InstancedPtr = InstancedPackageMap.InstancedPackageMapping.Find(Original))
+		{
+			Instanced = *InstancedPtr;
+			return true;
+		}
+		return false;
 	}
 
 	void FixupSoftObjectPath(FSoftObjectPath& InOutSoftObjectPath) const
@@ -190,9 +195,9 @@ void FLinkerInstancingContext::BuildPackageMapping(FName Original, FName Instanc
 	SharedData->BuildPackageMapping(Original, Instanced, GetSoftObjectPathRemappingEnabled());
 }
 
-FName FLinkerInstancingContext::FindPackageMapping(FName Original) const
+bool FLinkerInstancingContext::FindPackageMapping(FName Original, FName& Instanced) const
 {
-	return SharedData->FindPackageMapping(Original);
+	return SharedData->FindPackageMapping(Original, Instanced);
 }
 
 bool FLinkerInstancingContext::IsInstanced() const
@@ -283,14 +288,20 @@ void FLinkerInstancedPackageMap::AddPackageMapping(FName Original, FName Instanc
 {
 	if (InstanceMappingDirection == EInstanceMappingDirection::OriginalToInstanced)
 	{
-		InstancedPackageMapping.Add(Original, Instanced);
-		bIsInstanced |= !Instanced.IsNone();
+		if (!InstancedPackageMapping.Contains(Original))
+		{
+			InstancedPackageMapping.Add(Original, Instanced);
+			bIsInstanced |= !Instanced.IsNone();
+		}
 	}
 	else
 	{
 		check(InstanceMappingDirection == EInstanceMappingDirection::InstancedToOriginal);
-		InstancedPackageMapping.Add(Instanced, Original);
-		bIsInstanced |= !Original.IsNone();
+		if (!InstancedPackageMapping.Contains(Instanced))
+		{
+			InstancedPackageMapping.Add(Instanced, Original);
+			bIsInstanced |= !Original.IsNone();
+		}
 	}
 }
 
diff --git a/Engine/Source/Runtime/CoreUObject/Private/UObject/LinkerLoad.cpp b/Engine/Source/Runtime/CoreUObject/Private/UObject/LinkerLoad.cpp
--- a/Engine/Source/Runtime/CoreUObject/Private/UObject/LinkerLoad.cpp
+++ b/Engine/Source/Runtime/CoreUObject/Private/UObject/LinkerLoad.cpp
@@ -2064,9 +2064,9 @@ FLinkerLoad::ELinkerStatus FLinkerLoad::PopulateInstancingContext()
 		{
 			auto AddInstancedMapping = [this](const FString& OuterPackageName, FName InstancingPackageName) -> bool
 			{
-				FName InstancedName = InstancingContext.FindPackageMapping(InstancingPackageName);
+				FName InstancedName;
 				// if there's isn't already a remapping for that package, create one
-				if (InstancedName.IsNone())
+				if (!InstancingContext.FindPackageMapping(InstancingPackageName, InstancedName))
 				{
 					InstancedName = *FLinkerInstancingContext::GetInstancedPackageName(OuterPackageName, InstancingPackageName.ToString());
 					InstancingContext.AddPackageMapping(InstancingPackageName, InstancedName);
diff --git a/Engine/Source/Runtime/CoreUObject/Public/UObject/LinkerInstancingContext.h b/Engine/Source/Runtime/CoreUObject/Public/UObject/LinkerInstancingContext.h
--- a/Engine/Source/Runtime/CoreUObject/Public/UObject/LinkerInstancingContext.h
+++ b/Engine/Source/Runtime/CoreUObject/Public/UObject/LinkerInstancingContext.h
@@ -137,7 +137,7 @@ public:
 private:
 	void EnableAutomationTest();
 	void BuildPackageMapping(FName Original, FName Instanced);
-	FName FindPackageMapping(FName Original) const;
+	bool FindPackageMapping(FName Original, FName& Instanced) const;
 
 	friend class FLinkerLoad;
 	friend struct FAsyncPackage2;
