FAsyncPackage and FAsyncLoadingThread2 CallCompletionCallbacks now properly sets GPlayInEditorID with associated package request PIEInstanceID using FPlayInEditorLoadingScope.

diff --git a/Engine/Source/Runtime/CoreUObject/Private/Serialization/AsyncLoading.cpp b/Engine/Source/Runtime/CoreUObject/Private/Serialization/AsyncLoading.cpp
--- a/Engine/Source/Runtime/CoreUObject/Private/Serialization/AsyncLoading.cpp
+++ b/Engine/Source/Runtime/CoreUObject/Private/Serialization/AsyncLoading.cpp
@@ -7310,6 +7310,10 @@ void FAsyncPackage::CallCompletionCallbacks(bool bInternal, EAsyncLoadingResult:
 {
 	checkSlow(bInternal || !GetOwnerThread().IsInAsyncLoadThread());
 
+#if WITH_EDITOR
+	UE::Core::Private::FPlayInEditorLoadingScope PlayInEditorIDScope(Desc.PIEInstanceID);
+#endif
+
 	UPackage* LoadedPackage = (!bLoadHasFailed) ? LinkerRoot : nullptr;
 	for (FCompletionCallback& CompletionCallback : CompletionCallbacks)
 	{
diff --git a/Engine/Source/Runtime/CoreUObject/Private/Serialization/AsyncLoading2.cpp b/Engine/Source/Runtime/CoreUObject/Private/Serialization/AsyncLoading2.cpp
--- a/Engine/Source/Runtime/CoreUObject/Private/Serialization/AsyncLoading2.cpp
+++ b/Engine/Source/Runtime/CoreUObject/Private/Serialization/AsyncLoading2.cpp
@@ -3063,6 +3063,10 @@ private:
 		void CallCompletionCallbacks()
 		{
 			checkSlow(IsInGameThread());
+
+#if WITH_EDITOR
+			UE::Core::Private::FPlayInEditorLoadingScope PlayInEditorIDScope(AsyncPackage ? AsyncPackage->Desc.PIEInstanceID : INDEX_NONE);
+#endif
 			
 			if (CompletionCallbacks.Num() != 0)
 			{
