FDynamicallyTypedValue: add support for incremental GC

diff --git a/Engine/Source/Runtime/CoreUObject/Public/UObject/DynamicallyTypedValue.h b/Engine/Source/Runtime/CoreUObject/Public/UObject/DynamicallyTypedValue.h
--- a/Engine/Source/Runtime/CoreUObject/Public/UObject/DynamicallyTypedValue.h
+++ b/Engine/Source/Runtime/CoreUObject/Public/UObject/DynamicallyTypedValue.h
@@ -8,6 +8,7 @@
 #include "Misc/AssertionMacros.h"
 #include "Serialization/StructuredArchive.h"
 #include "Templates/UnrealTemplate.h"
+#include "UObject/UObjectGlobals.h"
 
 class FReferenceCollector;
 
@@ -107,6 +108,7 @@ namespace UE
 			Type = &NewType;
 			AllocateData();
 			Type->InitializeValue(GetDataPointer());
+			MarkTypeReachableIfIncrementalReachabilityPending();			
 		}
 
 		// Returns hash of the underlying FDynamicallyTypedValue's value. Added to allow for FDynamicallyTypedValue to be used as TMap keys.
@@ -116,7 +118,30 @@ namespace UE
 		}
 
 	private:
-
+		void MarkTypeReachableIfIncrementalReachabilityPending()
+		{
+			struct FTypeReferenceCollector final : public FReferenceCollector
+			{
+				bool IsIgnoringArchetypeRef() const override { return false; }
+				bool IsIgnoringTransient() const override { return false; }
+				void HandleObjectReference(UObject*& InObject, const UObject* InReferencingObject, const FProperty* InReferencingProperty) override
+				{
+					if (InObject)
+					{
+						InObject->MarkAsReachable();
+					}
+				}
+			};
+			
+			if (UE::GC::Private::GIsIncrementalReachabilityPending && Type)
+			{
+				// nb: this is done to simulate a write barrier for this type, which
+				//     enables it to behave properly with incremental gc.
+				FTypeReferenceCollector Collector;
+				Type->MarkReachable(Collector);
+			}
+		}
+		
 		FDynamicallyTypedValueType* Type;
 
 		// Store pointer-sized or smaller values inline, heap allocate all others.
@@ -131,6 +156,7 @@ namespace UE
 		{
 			Type = &NullType();
 			HeapData = nullptr;
+			MarkTypeReachableIfIncrementalReachabilityPending();			
 		}
 		// Deinitializes this value back to the primordial state.
 		void Deinit()
@@ -145,6 +171,7 @@ namespace UE
 			Type = Copyee.Type;
 			AllocateData();
 			Type->InitializeValueFromCopy(GetDataPointer(), Copyee.GetDataPointer());
+			MarkTypeReachableIfIncrementalReachabilityPending();			
 		}
 		// Moves the data from another value to this one, which is assumed to be in the primordial state.
 		// The source value is set to the null state.
@@ -154,6 +181,7 @@ namespace UE
 			// This assumes that the data is trivially relocatable.
 			Type = Movee.Type;
 			InlineData = Movee.InlineData;
+			MarkTypeReachableIfIncrementalReachabilityPending();			
 
 			// Reset the source value to null.
 			Movee.InitializeToNull();
