Fix non-unity build error

diff --git a/Engine/Source/Runtime/CoreUObject/Public/UObject/DynamicallyTypedValue.h b/Engine/Source/Runtime/CoreUObject/Public/UObject/DynamicallyTypedValue.h
--- a/Engine/Source/Runtime/CoreUObject/Public/UObject/DynamicallyTypedValue.h
+++ b/Engine/Source/Runtime/CoreUObject/Public/UObject/DynamicallyTypedValue.h
@@ -8,6 +8,7 @@
 #include "Misc/AssertionMacros.h"
 #include "Serialization/StructuredArchive.h"
 #include "Templates/UnrealTemplate.h"
+#include "UObject/Object.h"
 #include "UObject/UObjectGlobals.h"
 
 class FReferenceCollector;
