Adding missing TObjectPtr constructor barrier. Fixes TStrongObjectPtr referenced objects not having a GC barrier ran on them.

diff --git a/Engine/Source/Runtime/CoreUObject/Public/UObject/ObjectPtr.h b/Engine/Source/Runtime/CoreUObject/Public/UObject/ObjectPtr.h
--- a/Engine/Source/Runtime/CoreUObject/Public/UObject/ObjectPtr.h
+++ b/Engine/Source/Runtime/CoreUObject/Public/UObject/ObjectPtr.h
@@ -77,6 +77,9 @@ public:
 	explicit FORCEINLINE FObjectPtr(UObject* Object)
 		: Handle(UE::CoreUObject::Private::MakeObjectHandle(Object))
 	{
+#if UE_OBJECT_PTR_GC_BARRIER
+		ConditionallyMarkAsReachable(Object);
+#endif // UE_OBJECT_PTR_GC_BARRIER
 	}
 
 	UE_OBJPTR_DEPRECATED(5.0, "Construction with incomplete type pointer is deprecated.  Please update this code to use MakeObjectPtrUnsafe.")
@@ -89,6 +92,9 @@ public:
 	explicit FORCEINLINE FObjectPtr(FObjectHandle Handle)
 		: Handle(Handle)
 	{
+#if UE_OBJECT_PTR_GC_BARRIER
+		ConditionallyMarkAsReachable(*this);
+#endif // UE_OBJECT_PTR_GC_BARRIER
 	}
 #endif
 	
