VerifyNoUnreachable GC test will now also scan the disregard for GC set for stale references and will attempt to get the current FGCObject name when reporting them

diff --git a/Engine/Source/Runtime/CoreUObject/Private/UObject/GarbageCollectionVerification.cpp b/Engine/Source/Runtime/CoreUObject/Private/UObject/GarbageCollectionVerification.cpp
--- a/Engine/Source/Runtime/CoreUObject/Private/UObject/GarbageCollectionVerification.cpp
+++ b/Engine/Source/Runtime/CoreUObject/Private/UObject/GarbageCollectionVerification.cpp
@@ -447,9 +447,22 @@ public:
 					DebugInfo = TEXT("Native Reference");
 				}
 
+				FString ReferencingObjectName;
+				if (!ReferencingObject && FGCObject::GGCObjectReferencer)
+				{
+					if (FGCObject* CurrentlySerializingObject = FGCObject::GGCObjectReferencer->GetCurrentlySerializingObject())
+					{
+						ReferencingObjectName = CurrentlySerializingObject->GetReferencerName();
+					}
+				}
+				if (ReferencingObjectName.IsEmpty())
+				{
+					ReferencingObjectName = GetFullNameSafe(ReferencingObject);
+				}
+
 				UE_LOG(LogGarbage, Warning, TEXT("Unreachable object %s is being referenced by reachable object %s through %s"),
 					*Object->GetFullName(),
-					*GetFullNameSafe(ReferencingObject),
+					*ReferencingObjectName,
 					*DebugInfo);
 
 				NumErrors++;
@@ -461,7 +474,7 @@ public:
 void VerifyNoUnreachableObjects()
 {
 	const double StartTime = FPlatformTime::Seconds();
-	const int32 MaxNumberOfReachableObjects = GUObjectArray.GetObjectArrayNum() - GUObjectArray.GetFirstGCIndex();
+	const int32 MaxNumberOfReachableObjects = GUObjectArray.GetObjectArrayNum();
 	const int32 NumThreads = GetNumCollectReferenceWorkers();
 	const int32 NumberOfObjectsPerThread = (MaxNumberOfReachableObjects / NumThreads) + 1;
 	std::atomic<uint32> NumErrors(0);
@@ -473,9 +486,9 @@ void VerifyNoUnreachableObjects()
 		TArray<UObject*> ObjectsToSerialize;
 		ObjectsToSerialize.Reserve(NumberOfObjectsPerThread);		
 
-		for (int32 ObjectIndex = 0; ObjectIndex < NumObjects && (FirstObjectIndex + ObjectIndex + GUObjectArray.GetFirstGCIndex()) < GUObjectArray.GetObjectArrayNum(); ++ObjectIndex)
+		for (int32 ObjectIndex = 0; ObjectIndex < NumObjects && (FirstObjectIndex + ObjectIndex) < GUObjectArray.GetObjectArrayNum(); ++ObjectIndex)
 		{
-			FUObjectItem& ObjectItem = GUObjectArray.GetObjectItemArrayUnsafe()[FirstObjectIndex + ObjectIndex + GUObjectArray.GetFirstGCIndex()];
+			FUObjectItem& ObjectItem = GUObjectArray.GetObjectItemArrayUnsafe()[FirstObjectIndex + ObjectIndex];
 			if (ObjectItem.Object && !ObjectItem.HasAnyFlags(EInternalObjectFlags::MaybeUnreachable | EInternalObjectFlags::Unreachable))
 			{
 				ObjectsToSerialize.Add(static_cast<UObject*>(ObjectItem.Object));
