Added utility methods to notifiy changes on operation done via code on the entities/components, this in prevision of notifying the object of the property begin overriden.
Fix an entity reinstancing problems like fetching the owner entity when restarting an reinstantiated entity.
Fix the logic that is saving the world between the pre/post edit change callbacks.
Also fixed a replace reference not happening when the entity is only reference in a serialize method, need to be also be reference via the AddReferenceObject for the GO referencing table!
Fix the copy of the entity into the CDO when creating a prefab from a instance.

diff --git a/Engine/Source/Runtime/CoreUObject/Public/UObject/UnrealType.h b/Engine/Source/Runtime/CoreUObject/Public/UObject/UnrealType.h
--- a/Engine/Source/Runtime/CoreUObject/Public/UObject/UnrealType.h
+++ b/Engine/Source/Runtime/CoreUObject/Public/UObject/UnrealType.h
@@ -4892,7 +4892,7 @@ public:
 	}
 
 	/** Finds the associated pair from hash, rather than linearly searching */
-	uint8* FindMapPairPtrFromHash(const void* KeyPtr)
+	int32 FindMapPairIndexFromHash(const void* KeyPtr)
 	{
 		int32 Index = WithScriptMap([this, KeyPtr, LocalKeyPropForCapture = this->KeyProp](auto* Map)
 		{
@@ -4903,6 +4903,13 @@ public:
 				[LocalKeyPropForCapture](const void* A, const void* B) { return LocalKeyPropForCapture->Identical(A, B); }
 			);
 		});
+		return Index;
+	}
+
+	/** Finds the associated pair from hash, rather than linearly searching */
+	uint8* FindMapPairPtrFromHash(const void* KeyPtr)
+	{
+		int32 Index = FindMapPairIndexFromHash(KeyPtr);
 		uint8* Result = (Index >= 0) ? GetPairPtr(Index) : nullptr;
 		return Result;
 	}
