Clean up script interface comments to match the actual behavior:
For natively implemented interfaces, FScriptInterface stores the raw interface pointer implementation
For interfaces that are only implemented in blueprint, FScriptInterface is used as a passthrough and only stores the object

diff --git a/Engine/Source/Runtime/CoreUObject/Public/UObject/ScriptInterface.h b/Engine/Source/Runtime/CoreUObject/Public/UObject/ScriptInterface.h
--- a/Engine/Source/Runtime/CoreUObject/Public/UObject/ScriptInterface.h
+++ b/Engine/Source/Runtime/CoreUObject/Public/UObject/ScriptInterface.h
@@ -13,20 +13,21 @@
 #include <type_traits>
 
 /**
- * FScriptInterface
- *
- * This utility class stores the FProperty data for a native interface property.  ObjectPointer and InterfacePointer point to different locations in the same UObject.
+ * This utility class stores the FProperty data for an interface property used in both blueprint and native code.
+ * For objects natively implementing an interface, ObjectPointer and InterfacePointer point to different locations in the same UObject.
+ * For objects that only implement an interface in blueprint, only ObjectPointer will be set because there is no native representation.
+ * UClass::ImplementsInterface can be used along with Execute_ event wrappers to properly handle BP-implemented interfaces.
  */
 class FScriptInterface
 {
 private:
 	/**
-	 * A pointer to a UObject that implements a native interface.
+	 * A pointer to a UObject that implements an interface.
 	 */
 	TObjectPtr<UObject>	ObjectPointer = nullptr;
 
 	/**
-	 * Pointer to the location of the interface object within the UObject referenced by ObjectPointer.
+	 * For native interfaces, pointer to the location of the interface object within the UObject referenced by ObjectPointer.
 	 */
 	void*		InterfacePointer = nullptr;
 
@@ -71,12 +72,12 @@ public:
 	}
 
 	/**
-	 * Returns the pointer to the interface
+	 * Returns the pointer to the native interface if it is valid
 	 */
 	FORCEINLINE void* GetInterface() const
 	{
-		// only allow access to InterfacePointer if we have a valid ObjectPointer.  This is necessary because the garbage collector will set ObjectPointer to NULL
-		// without using the accessor methods
+		// Only access the InterfacePointer if we have a valid ObjectPointer. This is necessary because garbage collection may only clear the ObjectPointer.
+		// This will also return null for objects that only implement the interface in a blueprint class because there is no native representation.
 		return ObjectPointer ? InterfacePointer : nullptr;
 	}
 
@@ -129,7 +130,8 @@ template<> struct TIsPODType<class FScriptInterface> { enum { Value = true }; };
 template<> struct TIsZeroConstructType<class FScriptInterface> { enum { Value = true }; };
 
 /**
- * Templated version of FScriptInterface, which provides accessors and operators for referencing the interface portion of a UObject that implements a native interface.
+ * Templated version of FScriptInterface, which provides accessors and operators for referencing the interface portion of an object implementing an interface.
+ * This type is only useful with native interfaces, UClass::ImplementsInterface should be used to check for blueprint interfaces.
  */
 template <typename InInterfaceType>
 class TScriptInterface : public FScriptInterface
@@ -148,7 +150,7 @@ public:
 	TScriptInterface(TYPE_OF_NULLPTR) {}
 
 	/**
-	 * Construction from an object type that implements the InterfaceType native interface class
+	 * Construction from an object type that may natively implement InterfaceType
 	 */
 	template <
 		typename U,
@@ -156,9 +158,11 @@ public:
 	>
 	FORCEINLINE TScriptInterface(U&& Source)
 	{
+		// Always set the object
 		UObject* SourceObject = ImplicitConv<UObject*>(Source);
 		SetObject(SourceObject);
 
+		// Tries to set the native interface instance, this will set it to null for BP-implemented interfaces
 		InInterfaceType* SourceInterface = Cast<InInterfaceType>(SourceObject);
 		SetInterface(SourceInterface);
 	}
@@ -179,13 +183,15 @@ public:
 	}
 
 	/**
-	 * Assignment from an object type that implements the InterfaceType native interface class
+	 * Assignment from an object type that may natively implement InterfaceType
 	 */
 	template <typename ObjectType>
 	TScriptInterface(TObjectPtr<ObjectType> SourceObject)
 	{
+		// Always set the object
 		SetObject(SourceObject);
 
+		// Tries to set the native interface instance, this will set it to null for BP-implemented interfaces
 		InInterfaceType* SourceInterface = Cast<InInterfaceType>(ToRawPtr(SourceObject));
 		SetInterface(SourceInterface);
 	}
@@ -206,7 +212,7 @@ public:
 	}
 
 	/**
-	 * Assignment from an object type that implements the InterfaceType native interface class
+	 * Assignment from an object type that may natively implement InterfaceType
 	 */
 	template <
 		typename U,
@@ -232,7 +238,7 @@ public:
 	}
 
 	/**
-	 * Assignment from an object type that implements the InterfaceType native interface class
+	 * Assignment from an object type that may natively implement InterfaceType
 	 */
 	template <typename ObjectType>
 	TScriptInterface& operator=(TObjectPtr<ObjectType> SourceObject)
@@ -280,7 +286,7 @@ public:
 	}
 
 	/**
-	 * Member access operator.  Provides transparent access to the interface pointer contained by this TScriptInterface
+	 * Member access operator.  Provides transparent access to the native interface pointer contained by this TScriptInterface
 	 */
 	FORCEINLINE InInterfaceType* operator->() const
 	{
@@ -288,7 +294,7 @@ public:
 	}
 
 	/**
-	 * Dereference operator.  Provides transparent access to the interface pointer contained by this TScriptInterface
+	 * Dereference operator.  Provides transparent access to the native interface pointer contained by this TScriptInterface
 	 *
 	 * @return	a reference (of type InterfaceType) to the object pointed to by InterfacePointer
 	 */
@@ -314,9 +320,8 @@ public:
 	}
 
 	/**
-	 * Boolean operator.  Provides transparent access to the interface pointer contained by this TScriptInterface.
-	 *
-	 * @return	true if InterfacePointer is non-NULL.
+	 * Boolean operator, returns true if this object natively implements InterfaceType.
+	 * This will return false for objects that only implement the interface in blueprint classes.
 	 */
 	FORCEINLINE explicit operator bool() const
 	{
diff --git a/Engine/Source/Runtime/Engine/Classes/Kismet/KismetSystemLibrary.h b/Engine/Source/Runtime/Engine/Classes/Kismet/KismetSystemLibrary.h
--- a/Engine/Source/Runtime/Engine/Classes/Kismet/KismetSystemLibrary.h
+++ b/Engine/Source/Runtime/Engine/Classes/Kismet/KismetSystemLibrary.h
@@ -270,11 +270,14 @@ class UKismetSystemLibrary : public UBlueprintFunctionLibrary
 	UFUNCTION(BlueprintPure, meta=(DisplayName = "Cast To Class", DeterminesOutputType = "Class"), Category="Utilities")
 	static ENGINE_API UClass* Conv_ObjectToClass(UObject* Object, TSubclassOf<UObject> Class);
 
-	/** Converts an interfance into an object */
+	/** Converts an interface instance into an object */
 	UFUNCTION(BlueprintPure, meta=(DisplayName = "To Object (Interface)", CompactNodeTitle = "->", BlueprintAutocast), Category="Utilities")
 	static ENGINE_API UObject* Conv_InterfaceToObject(const FScriptInterface& Interface); 
 
-	/** Checks if the backing object is valid. */
+	/** 
+	 * Checks if the interface instance has a valid object for blueprint interface functions.
+	 * This will return true for both natively implemented and blueprint implemented interfaces.
+	 */
 	UFUNCTION(BlueprintPure, Category="Utilities")
 	static bool IsValidInterface(const FScriptInterface& Interface); 
 
