[ZenLoader] Fix recursive loading behavior in zenloader by following more closely how legacy loading works
  - Add a new partial loading mode where flushing the requestID will exit after serialization is done (after exportdone but before postload)
  - Automatically opt-in partial loading when an explicit load is called recursively from another load (treated as an import)
  - During recursive loading, attach postload of the recursive load to the caller's postload
  - Add loading unit tests to validate the new behavior for zenloader
  - Add a way to detect which loader is currently active to run unit tests on zenloader only

#rnx

diff --git a/Engine/Source/Runtime/Core/Private/Misc/CoreGlobals.cpp b/Engine/Source/Runtime/Core/Private/Misc/CoreGlobals.cpp
--- a/Engine/Source/Runtime/Core/Private/Misc/CoreGlobals.cpp
+++ b/Engine/Source/Runtime/Core/Private/Misc/CoreGlobals.cpp
@@ -356,6 +356,12 @@ bool (*IsAsyncLoadingMultithreaded)() = &IsAsyncLoadingCoreInternal;
 void (*SuspendTextureStreamingRenderTasks)() = &appNoop;
 void (*ResumeTextureStreamingRenderTasks)() = &appNoop;
 
+static FName appNotAvailable()
+{
+	return TEXT("N/A");
+}
+FName (*GetLoaderName)() = &appNotAvailable;
+
 /** Whether the editor is currently loading a package or not												*/
 bool					GIsEditorLoadingPackage				= false;
 /** Whether the cooker is currently loading a package or not												*/
diff --git a/Engine/Source/Runtime/Core/Public/CoreGlobals.h b/Engine/Source/Runtime/Core/Public/CoreGlobals.h
--- a/Engine/Source/Runtime/Core/Public/CoreGlobals.h
+++ b/Engine/Source/Runtime/Core/Public/CoreGlobals.h
@@ -443,6 +443,9 @@ extern CORE_API bool (*IsAsyncLoadingSuspended)();
 /** Returns true if async loading is using the async loading thread */
 extern CORE_API bool(*IsAsyncLoadingMultithreaded)();
 
+/** Returns the name of the currently active loader, if any. */
+extern CORE_API FName(*GetLoaderName)();
+
 /** Suspends texture updates caused by completed async IOs. */
 extern CORE_API void (*SuspendTextureStreamingRenderTasks)();
 
diff --git a/Engine/Source/Runtime/CoreUObject/Private/Serialization/AsyncLoading2.cpp b/Engine/Source/Runtime/CoreUObject/Private/Serialization/AsyncLoading2.cpp
--- a/Engine/Source/Runtime/CoreUObject/Private/Serialization/AsyncLoading2.cpp
+++ b/Engine/Source/Runtime/CoreUObject/Private/Serialization/AsyncLoading2.cpp
@@ -82,6 +82,19 @@
 
 #include <atomic>
 
+// For now, the partial request behavior is reserved for the editor only.
+#define WITH_PARTIAL_REQUEST_DURING_RECURSION WITH_EDITOR
+
+enum class ERequestType : uint8
+{
+#if WITH_PARTIAL_REQUEST_DURING_RECURSION
+	ExportsDone,	// ExportsDone on the requested package must be reached before exiting a flush of that request.
+#endif
+	Full,
+
+	Count
+};
+
 CSV_DECLARE_CATEGORY_MODULE_EXTERN(CORE_API, FileIO);
 CSV_DEFINE_STAT(FileIO, FrameCompletedExportBundleLoadsKB);
 
@@ -482,6 +495,7 @@ class FAsyncLoadingSyncLoadContext;
 struct FPackageRequest
 {
 	int32 RequestId = -1;
+	ERequestType RequestType = ERequestType::Full;
 	int32 Priority = -1;
 	EPackageFlags PackageFlags = PKG_None;
 #if WITH_EDITOR
@@ -506,11 +520,12 @@ struct FPackageRequest
 #endif
 	}
 
-	static FPackageRequest Create(int32 RequestId, EPackageFlags PackageFlags, uint32 LoadFlags, int32 PIEInstanceID, int32 Priority, const FLinkerInstancingContext* InstancingContext, const FPackagePath& PackagePath, FName CustomName, TUniquePtr<FLoadPackageAsyncDelegate> PackageLoadedDelegate, TUniquePtr<FLoadPackageAsyncProgressDelegate> PackageProgressDelegate)
+	static FPackageRequest Create(int32 RequestId, ERequestType RequestType, EPackageFlags PackageFlags, uint32 LoadFlags, int32 PIEInstanceID, int32 Priority, const FLinkerInstancingContext* InstancingContext, const FPackagePath& PackagePath, FName CustomName, TUniquePtr<FLoadPackageAsyncDelegate> PackageLoadedDelegate, TUniquePtr<FLoadPackageAsyncProgressDelegate> PackageProgressDelegate)
 	{
 		return FPackageRequest
 		{
 			RequestId,
+			RequestType,
 			Priority,
 			PackageFlags,
 #if WITH_EDITOR
@@ -533,6 +548,8 @@ struct FAsyncPackageDesc2
 {
 	// A unique request id for each external call to LoadPackage
 	int32 RequestID;
+	// The type of request
+	ERequestType RequestType;
 	// Package priority
 	int32 Priority;
 	/** The flags that should be applied to the package */
@@ -564,6 +581,7 @@ struct FAsyncPackageDesc2
 		return FAsyncPackageDesc2
 		{
 			Request.RequestId,
+			Request.RequestType,
 			Request.Priority,
 			Request.PackageFlags,
 #if WITH_EDITOR
@@ -592,6 +610,7 @@ struct FAsyncPackageDesc2
 		return FAsyncPackageDesc2
 		{
 			INDEX_NONE,
+			ERequestType::Full,
 			ImportingPackageDesc.Priority,
 			PKG_None,
 #if WITH_EDITOR
@@ -2278,7 +2297,7 @@ struct FAsyncLoadingThreadState2
 
 	FAsyncLoadEventGraphAllocator& GraphAllocator;
 	TArray<FEventLoadNode2*> NodesToFire;
-	FEventLoadNode2* CurrentlyExecutingEventNode = nullptr;
+	TArray<FEventLoadNode2*> CurrentlyExecutingEventNodeStack;
 	TArray<FAsyncLoadingSyncLoadContext*> SyncLoadContextStack;
 	TArray<FAsyncPackage2*> PackagesOnStack;
 	TSpscQueue<FAsyncLoadingSyncLoadContext*> SyncLoadContextsCreatedOnGameThread;
@@ -2571,7 +2590,9 @@ struct FAsyncPackage2
 	friend class FAsyncLoadingThread2;
 	friend class FAsyncLoadEventQueue2;
 
-	FAsyncPackage2(const FAsyncPackageDesc2& InDesc,
+	FAsyncPackage2(
+		FAsyncLoadingThreadState2& ThreadState,
+		const FAsyncPackageDesc2& InDesc,
 		FAsyncLoadingThread2& InAsyncLoadingThread,
 		FAsyncLoadEventGraphAllocator& InGraphAllocator,
 		const FAsyncLoadEventSpec* EventSpecs);
@@ -2623,7 +2644,7 @@ struct FAsyncPackage2
 	}
 
 	/** Adds new request ID to the existing package */
-	void AddRequestID(int32 Id);
+	void AddRequestID(FAsyncLoadingThreadState2& ThreadState, int32 Id, ERequestType RequestType);
 
 	uint64 GetSyncLoadContextId() const
 	{
@@ -2822,8 +2843,8 @@ private:
 	bool						bLoadHasFailed = false;
 	/** True if this package was created by this async package */
 	bool						bCreatedLinkerRoot = false;
-	/** List of all request handles */
-	TArray<int32, TInlineAllocator<2>> RequestIDs;
+	/** List of all request handles by types. */
+	TArray<int32, TInlineAllocator<2>> RequestIDs[(int32)ERequestType::Count];
 	/** List of ConstructedObjects = Exports + UPackage + ObjectsCreatedFromExports */
 	TArray<UObject*> ConstructedObjects;
 	TArray<FExternalReadCallback> ExternalReadDependencies;
@@ -2847,6 +2868,8 @@ public:
 		return GraphAllocator;
 	}
 
+	FAsyncPackage2* GetCallerPackage(FAsyncLoadingThreadState2& ThreadState) const;
+
 	/** [EDL] Begin Event driven loader specific stuff */
 
 	static EEventLoadNodeExecutionResult Event_ProcessExportBundle(FAsyncLoadingThreadState2& ThreadState, FAsyncPackage2* Package, int32 ExportBundleIndex);
@@ -2993,6 +3016,11 @@ public:
 	FAsyncLoadingThread2(FIoDispatcher& IoDispatcher, IAsyncPackageLoader* InUncookedPackageLoader);
 	virtual ~FAsyncLoadingThread2();
 
+	virtual FName GetLoaderName() const override
+	{
+		return TEXT("ZenLoader");
+	}
+
 private:
 	/** Thread to run the worker FRunnable on */
 	static constexpr int32 DefaultAsyncPackagesReserveCount = 512;
@@ -3048,16 +3076,21 @@ private:
 		static FCompletedPackageRequest FromLoadedPackage(
 			FAsyncPackage2* Package)
 		{
-			return FCompletedPackageRequest
+			FCompletedPackageRequest Request
 			{
 				Package->Desc.UPackageName,
 				Package->HasLoadFailed() ? EAsyncLoadingResult::Failed : EAsyncLoadingResult::Succeeded,
 				Package->LinkerRoot,
 				Package,
 				MoveTemp(Package->CompletionCallbacks),
-				MoveTemp(Package->ProgressCallbacks),
-				Package->RequestIDs
+				MoveTemp(Package->ProgressCallbacks)
 			};
+
+			for (int32 Index = 0; Index < (int32)ERequestType::Count; ++Index)
+			{
+				Request.RequestIDs.Append(Package->RequestIDs[Index]);
+			}
+			return Request;
 		}
 
 		void CallCompletionCallbacks()
@@ -3461,7 +3494,7 @@ private:
 #endif
 	void ConditionalBeginPostLoad(FAsyncLoadingThreadState2& ThreadState, FAsyncLoadingPostLoadGroup* PostLoadGroup);
 	void ConditionalBeginDeferredPostLoad(FAsyncLoadingThreadState2& ThreadState, FAsyncLoadingPostLoadGroup* DeferredPostLoadGroup);
-	void MergePostLoadGroups(FAsyncLoadingThreadState2& ThreadState, FAsyncLoadingPostLoadGroup* Target, FAsyncLoadingPostLoadGroup* Source);
+	void MergePostLoadGroups(FAsyncLoadingThreadState2& ThreadState, FAsyncLoadingPostLoadGroup* Target, FAsyncLoadingPostLoadGroup* Source, bool bSkipSyncLoadContext = false);
 
 	bool ProcessDeferredDeletePackagesQueue(int32 MaxCount = MAX_int32)
 	{
@@ -3572,11 +3605,11 @@ private:
 
 	bool CreateAsyncPackagesFromQueue(FAsyncLoadingThreadState2& ThreadState);
 
-	FAsyncPackage2* CreateAsyncPackage(const FAsyncPackageDesc2& Desc)
+	FAsyncPackage2* CreateAsyncPackage(FAsyncLoadingThreadState2& ThreadState, const FAsyncPackageDesc2& Desc)
 	{
 		UE_ASYNC_PACKAGE_DEBUG(Desc);
 
-		return new FAsyncPackage2(Desc, *this, GraphAllocator, EventSpecs.GetData());
+		return new FAsyncPackage2(ThreadState, Desc, *this, GraphAllocator, EventSpecs.GetData());
 	}
 
 	void InitializeAsyncPackageFromPackageStore(FAsyncLoadingThreadState2& ThreadState, FIoBatch* IoBatch, FAsyncPackage2* AsyncPackage, const FPackageStoreEntry& PackageStoreEntry)
@@ -3749,9 +3782,12 @@ private:
 	{
 		TRACE_CPUPROFILER_EVENT_SCOPE(DeleteAsyncPackage);
 		UE_ASYNC_PACKAGE_DEBUG(Package->Desc);
-		for (int32 RequestId : Package->RequestIDs)
+		for (int32 Index = 0; Index < (int32)ERequestType::Count; ++Index)
 		{
-			RequestIdToPackageMap.Remove(RequestId);
+			for (int32 RequestId : Package->RequestIDs[Index])
+			{
+				RequestIdToPackageMap.Remove(RequestId);
+			}
 		}
 		if (RequestIdToPackageMap.IsEmpty())
 		{
@@ -3957,7 +3993,7 @@ void FAsyncLoadingThread2::ConditionalBeginDeferredPostLoad(FAsyncLoadingThreadS
 	}
 }
 
-void FAsyncLoadingThread2::MergePostLoadGroups(FAsyncLoadingThreadState2& ThreadState, FAsyncLoadingPostLoadGroup* Target, FAsyncLoadingPostLoadGroup* Source)
+void FAsyncLoadingThread2::MergePostLoadGroups(FAsyncLoadingThreadState2& ThreadState, FAsyncLoadingPostLoadGroup* Target, FAsyncLoadingPostLoadGroup* Source, bool bSkipSyncLoadContext)
 {
 	if (Target == Source)
 	{
@@ -3974,17 +4010,23 @@ void FAsyncLoadingThread2::MergePostLoadGroups(FAsyncLoadingThreadState2& Thread
 	Target->PackagesWithExportsToSerializeCount += Source->PackagesWithExportsToSerializeCount;
 	check(Target->PackagesWithExportsToPostLoadCount == 0 && Source->PackagesWithExportsToPostLoadCount == 0);
 
-	const uint64 SyncLoadContextId = FMath::Max(Source->SyncLoadContextId, Target->SyncLoadContextId);
-	if (SyncLoadContextId)
+	// If the intention of the caller of this function is to merge postloads into the caller package so they are executed later after
+	// a partial flush, then we can't update the synccontext of the caller during the merge as it would make the whole hierarchy
+	// flush in a single swoop which is in direct contradiction with the partial flush feature.
+	if (!bSkipSyncLoadContext)
 	{
-		Target->SyncLoadContextId = SyncLoadContextId;
-		for (FAsyncPackage2* Package : Target->Packages)
+		const uint64 SyncLoadContextId = FMath::Max(Source->SyncLoadContextId, Target->SyncLoadContextId);
+		if (SyncLoadContextId)
 		{
-			Package->SyncLoadContextId = SyncLoadContextId;
-			if (Package->Desc.Priority < MAX_int32)
+			Target->SyncLoadContextId = SyncLoadContextId;
+			for (FAsyncPackage2* Package : Target->Packages)
 			{
-				Package->Desc.Priority = MAX_int32;
-				UpdatePackagePriority(ThreadState, Package);
+				Package->SyncLoadContextId = SyncLoadContextId;
+				if (Package->Desc.Priority < MAX_int32)
+				{
+					Package->Desc.Priority = MAX_int32;
+					UpdatePackagePriority(ThreadState, Package);
+				}
 			}
 		}
 	}
@@ -4002,7 +4044,7 @@ FAsyncPackage2* FAsyncLoadingThread2::FindOrInsertPackage(FAsyncLoadingThreadSta
 		Package = AsyncPackageLookup.FindRef(Desc.UPackageId);
 		if (!Package)
 		{
-			Package = CreateAsyncPackage(Desc);
+			Package = CreateAsyncPackage(ThreadState, Desc);
 			checkf(Package, TEXT("Failed to create async package %s"), *Desc.UPackageName.ToString());
 			Package->AddRef();
 			++LoadingPackagesCounter;
@@ -4014,7 +4056,7 @@ FAsyncPackage2* FAsyncLoadingThread2::FindOrInsertPackage(FAsyncLoadingThreadSta
 		{
 			if (Desc.RequestID > 0)
 			{
-				Package->AddRequestID(Desc.RequestID);
+				Package->AddRequestID(ThreadState, Desc.RequestID, Desc.RequestType);
 			}
 			if (Desc.Priority > Package->Desc.Priority)
 			{
@@ -4406,8 +4448,9 @@ EEventLoadNodeExecutionResult FEventLoadNode2::Execute(FAsyncLoadingThreadState2
 	check(BarrierCount.load(std::memory_order_relaxed) == 0);
 	EEventLoadNodeExecutionResult Result;
 	{
-		TGuardValue<FEventLoadNode2*> GuardCurrentlyExecutingEventNode(ThreadState.CurrentlyExecutingEventNode, this);
+		ThreadState.CurrentlyExecutingEventNodeStack.Push(this);
 		Result = Spec->Func(ThreadState, Package, ImportOrExportIndex);
+		ThreadState.CurrentlyExecutingEventNodeStack.Pop();
 	}
 	if (Result == EEventLoadNodeExecutionResult::Complete)
 	{
@@ -6517,6 +6560,20 @@ bool FAsyncPackage2::EventDrivenSerializeExport(const FAsyncPackageHeaderData& H
 	return true;
 }
 
+FAsyncPackage2* FAsyncPackage2::GetCallerPackage(FAsyncLoadingThreadState2& ThreadState) const
+{
+	FAsyncPackage2* CallerPackage = nullptr;
+	for (int32 Index = ThreadState.CurrentlyExecutingEventNodeStack.Num() - 1; Index >= 0; --Index)
+	{
+		if (ThreadState.CurrentlyExecutingEventNodeStack[Index]->GetPackage() != this)
+		{
+			CallerPackage = ThreadState.CurrentlyExecutingEventNodeStack[Index]->GetPackage();
+			break;
+		}
+	}
+	return CallerPackage;
+}
+
 EEventLoadNodeExecutionResult FAsyncPackage2::Event_ExportsDone(FAsyncLoadingThreadState2& ThreadState, FAsyncPackage2* Package, int32)
 {
 	TRACE_CPUPROFILER_EVENT_SCOPE(Event_ExportsDone);
@@ -6565,7 +6622,27 @@ EEventLoadNodeExecutionResult FAsyncPackage2::Event_ExportsDone(FAsyncLoadingThr
 	check(PostLoadGroup);
 	check(PostLoadGroup->PackagesWithExportsToSerializeCount > 0);
 	--PostLoadGroup->PackagesWithExportsToSerializeCount;
-	Package->AsyncLoadingThread.ConditionalBeginPostLoad(ThreadState, PostLoadGroup);
+
+#if WITH_PARTIAL_REQUEST_DURING_RECURSION
+	// Release any requests that are only waiting on exports, allowing a flush to exit.
+	Package->AsyncLoadingThread.RemovePendingRequests(ThreadState, Package->RequestIDs[(int32)ERequestType::ExportsDone]);
+
+	FAsyncPackage2* CallerPackage = Package->GetCallerPackage(ThreadState);
+
+	// If this is a recursive load and the caller has not started postloading yet, merge postload group.
+	if (CallerPackage && CallerPackage->PostLoadGroup != nullptr)
+	{
+		UE_LOG(LogStreaming, Display, TEXT("Merging postload groups of package %s with caller package %s"), *Package->Desc.UPackageName.ToString(), *CallerPackage->Desc.UPackageName.ToString());
+
+		// Do not adjust sync load context, we want to be able to exit the current one even if the caller has not finished yet.
+		const bool bSkipSyncLoadContext = false;
+		Package->AsyncLoadingThread.MergePostLoadGroups(ThreadState, Package->PostLoadGroup, CallerPackage->PostLoadGroup, bSkipSyncLoadContext);
+	}
+	else
+#endif // WITH_PARTIAL_REQUEST_DURING_RECURSION
+	{
+		Package->AsyncLoadingThread.ConditionalBeginPostLoad(ThreadState, PostLoadGroup);
+	}
 	return EEventLoadNodeExecutionResult::Complete;
 }
 
@@ -7136,6 +7213,18 @@ void FAsyncLoadingThread2::UpdateSyncLoadContext(FAsyncLoadingThreadState2& Thre
 				// Set RequestedPackage before setting bHasFoundRequestedPackage so that another thread looking at RequestedPackage
 				// after validating that bHasFoundRequestedPackage is true would see the proper value.
 				SyncLoadContext->RequestedPackages[i] = RequestedPackage;
+
+#if WITH_PARTIAL_REQUEST_DURING_RECURSION
+				FAsyncPackage2* CallerPackage = RequestedPackage->GetCallerPackage(ThreadState);
+				if (CallerPackage != nullptr && RequestedPackage->RequestIDs[(int32)ERequestType::Full].Contains(RequestID))
+				{
+					UE_LOG(LogStreaming, Error, TEXT("Flushing package %s recursively from another package %s can lead to undefined behavior and is strongly discouraged"),
+						*RequestedPackage->Desc.UPackageName.ToString(),
+						*CallerPackage->Desc.UPackageName.ToString()
+					);
+				}
+#endif
+
 				IncludePackageInSyncLoadContextRecursive(ThreadState, SyncLoadContext->ContextId, RequestedPackage);
 				++FoundPackages;	
 			}
@@ -7450,7 +7539,10 @@ EAsyncPackageState::Type FAsyncLoadingThread2::ProcessLoadedPackagesFromGameThre
 
 			if (Package->CompletionCallbacks.IsEmpty() && Package->ProgressCallbacks.IsEmpty())
 			{
-				RemovePendingRequests(ThreadState, Package->RequestIDs);
+				for (int32 Index = 0; Index < (int32)ERequestType::Count; ++Index)
+				{
+					RemovePendingRequests(ThreadState, Package->RequestIDs[Index]);
+				}
 				Package->ReleaseRef();
 			}
 			else
@@ -8239,6 +8331,7 @@ void FAsyncLoadingThread2::NotifyRegistrationComplete()
 * Constructor
 */
 FAsyncPackage2::FAsyncPackage2(
+	FAsyncLoadingThreadState2& ThreadState,
 	const FAsyncPackageDesc2& InDesc,
 	FAsyncLoadingThread2& InAsyncLoadingThread,
 	FAsyncLoadEventGraphAllocator& InGraphAllocator,
@@ -8249,7 +8342,7 @@ FAsyncPackage2::FAsyncPackage2(
 , ImportStore(AsyncLoadingThread.GlobalImportStore)
 {
 	TRACE_LOADTIME_NEW_ASYNC_PACKAGE(this);
-	AddRequestID(Desc.RequestID);
+	AddRequestID(ThreadState, Desc.RequestID, Desc.RequestType);
 
 	CreatePackageNodes(EventSpecs);
 
@@ -8370,7 +8463,7 @@ void FAsyncPackage2::ClearConstructedObjects()
 	}
 }
 
-void FAsyncPackage2::AddRequestID(int32 Id)
+void FAsyncPackage2::AddRequestID(FAsyncLoadingThreadState2& ThreadState, int32 Id, ERequestType RequestType)
 {
 	if (Id > 0)
 	{
@@ -8379,8 +8472,21 @@ void FAsyncPackage2::AddRequestID(int32 Id)
 			// For debug readability
 			Desc.RequestID = Id;
 		}
-		RequestIDs.Add(Id);
+
+		RequestIDs[(int32)RequestType].Add(Id);
+
+		// The Id is most likely already present because it's added as soon as the request is created.
 		AsyncLoadingThread.AddPendingRequest(Id);
+
+		// Remove any partial request if our state is already past that point.
+#if WITH_PARTIAL_REQUEST_DURING_RECURSION
+		static_assert((int32)ERequestType::Count == 2, "Please adjust this code when adding a new request type");
+		if (RequestType == ERequestType::ExportsDone && AsyncPackageLoadingState >= EAsyncPackageLoadingState2::ExportsDone)
+		{
+			AsyncLoadingThread.RemovePendingRequests(ThreadState, {Id});
+		}
+#endif
+
 		TRACE_LOADTIME_ASYNC_PACKAGE_REQUEST_ASSOCIATION(this, Id);
 	}
 }
@@ -8656,7 +8762,36 @@ PRAGMA_ENABLE_DEPRECATION_WARNINGS
 	TRACE_LOADTIME_BEGIN_REQUEST(RequestId);
 	AddPendingRequest(RequestId);
 
-	PackageRequestQueue.Enqueue(FPackageRequest::Create(RequestId, InPackageFlags, InLoadFlags, InPIEInstanceID, InPackagePriority, InInstancingContext, InPackagePath, InCustomName, MoveTemp(InCompletionDelegate), MoveTemp(InProgressDelegate)));
+	ERequestType RequestType = ERequestType::Full;
+
+#if WITH_PARTIAL_REQUEST_DURING_RECURSION
+	FAsyncLoadingThreadState2* ThreadState = FAsyncLoadingThreadState2::Get();
+
+	// We want to skip postloads inside recursive loads that happens during serialization and preload
+	// since it is mostly what BP expects and how the old LinkerLoad was implemented.
+	// In some cases, trying to postload right away during recursion could end-up in data loss (i.e. UE-190649).
+	// The expectation is that postload should be run as part of the package causing this load to occur.
+	// As a more general rule, we want any synchronous load during deserialization to act like
+	// an import. We want the deserialization steps done before returning from the load but
+	// we keep postload for after deserialization is done on the referencer too so that postload 
+	// sees everything initialized in the case where it also has a reference back to the referencer.
+	// Using the LOAD_Partial flag allows the FlushRequest to return as soon as exports are done which
+	// allows the caller to get a pointer to an object inside the package even if it's not fully loaded yet.
+	if (ThreadState != nullptr &&
+		ThreadState->CurrentlyExecutingEventNodeStack.Num() > 0 &&
+		ThreadState->CurrentlyExecutingEventNodeStack.Top()->GetPackage() != nullptr)
+	{
+		UE_LOG(LogStreaming, Display,
+			TEXT("LoadPackage on %s is being called from another package %s, partial loading is automatically being applied which means only exports are processed before exiting flush and postloads will be called along with the caller's"),
+			*PackageNameToLoad.ToString(),
+			*ThreadState->CurrentlyExecutingEventNodeStack.Top()->GetPackage()->Desc.UPackageName.ToString()
+		);
+
+		RequestType = ERequestType::ExportsDone;
+	}
+#endif
+
+	PackageRequestQueue.Enqueue(FPackageRequest::Create(RequestId, RequestType, InPackageFlags, InLoadFlags, InPIEInstanceID, InPackagePriority, InInstancingContext, InPackagePath, InCustomName, MoveTemp(InCompletionDelegate), MoveTemp(InProgressDelegate)));
 	++QueuedPackagesCounter;
 	++PackagesWithRemainingWorkCounter;
 
@@ -8774,12 +8909,12 @@ void FAsyncLoadingThread2::FlushLoading(TConstArrayView<int32> RequestIDs)
 		double LogFlushTime = StartTime;
 
 		FAsyncPackage2* CurrentlyExecutingPackage = nullptr;
-		if (GameThreadState->CurrentlyExecutingEventNode)
+		if (GameThreadState->CurrentlyExecutingEventNodeStack.Num() > 0)
 		{
 			TRACE_CPUPROFILER_EVENT_SCOPE(HandleCurrentlyExecutingEventNode);
 
 			UE_CLOG(RequestIDs.Num() == 0, LogStreaming, Fatal, TEXT("Flushing async loading while creating, serializing or postloading an object is not permitted"));
-			CurrentlyExecutingPackage = GameThreadState->CurrentlyExecutingEventNode->GetPackage();
+			CurrentlyExecutingPackage = GameThreadState->CurrentlyExecutingEventNodeStack.Top()->GetPackage();
 			GameThreadState->PackagesOnStack.Push(CurrentlyExecutingPackage);
 			// Update the state of any package that is waiting for the currently executing one
 			while (FAsyncPackage2* WaitingPackage = CurrentlyExecutingPackage->AllDependenciesFullyLoadedState.PackagesWaitingForThisHead)
@@ -8787,24 +8922,6 @@ void FAsyncLoadingThread2::FlushLoading(TConstArrayView<int32> RequestIDs)
 				FAsyncPackage2::FAllDependenciesState::RemoveFromWaitList(&FAsyncPackage2::AllDependenciesFullyLoadedState, CurrentlyExecutingPackage, WaitingPackage);
 				WaitingPackage->ConditionalFinishLoading(*GameThreadState);
 			}
-			if (GameThreadState->bCanAccessAsyncLoadingThreadData)
-			{
-				if (FAsyncLoadingPostLoadGroup* PostLoadGroup = CurrentlyExecutingPackage->PostLoadGroup)
-				{
-					check(PostLoadGroup->Packages.Contains(CurrentlyExecutingPackage));
-					check(PostLoadGroup->PackagesWithExportsToSerializeCount > 0);
-					if (PostLoadGroup->Packages.Num() > 1)
-					{
-						PostLoadGroup->Packages.Remove(CurrentlyExecutingPackage);
-						--PostLoadGroup->PackagesWithExportsToSerializeCount;
-						ConditionalBeginPostLoad(*GameThreadState, PostLoadGroup);
-						CurrentlyExecutingPackage->PostLoadGroup = new FAsyncLoadingPostLoadGroup();
-						CurrentlyExecutingPackage->PostLoadGroup->SyncLoadContextId = CurrentlyExecutingPackage->SyncLoadContextId;
-						CurrentlyExecutingPackage->PostLoadGroup->Packages.Add(CurrentlyExecutingPackage);
-						CurrentlyExecutingPackage->PostLoadGroup->PackagesWithExportsToSerializeCount = 1;
-					}
-				}
-			}
 		}
 
 		FAsyncLoadingSyncLoadContext* SyncLoadContext = nullptr;
diff --git a/Engine/Source/Runtime/CoreUObject/Private/Serialization/AsyncLoadingThread.h b/Engine/Source/Runtime/CoreUObject/Private/Serialization/AsyncLoadingThread.h
--- a/Engine/Source/Runtime/CoreUObject/Private/Serialization/AsyncLoadingThread.h
+++ b/Engine/Source/Runtime/CoreUObject/Private/Serialization/AsyncLoadingThread.h
@@ -348,6 +348,11 @@ public:
 	FAsyncLoadingThread(int32 InThreadIndex);
 	virtual ~FAsyncLoadingThread();
 
+	virtual FName GetLoaderName() const override
+	{
+		return TEXT("LegacyLoader");
+	}
+
 	IAsyncPackageLoader* GetIoStorePackageLoader() const
 	{
 		return IoStorePackageLoader;
diff --git a/Engine/Source/Runtime/CoreUObject/Private/Serialization/AsyncPackageLoader.cpp b/Engine/Source/Runtime/CoreUObject/Private/Serialization/AsyncPackageLoader.cpp
--- a/Engine/Source/Runtime/CoreUObject/Private/Serialization/AsyncPackageLoader.cpp
+++ b/Engine/Source/Runtime/CoreUObject/Private/Serialization/AsyncPackageLoader.cpp
@@ -326,6 +326,11 @@ bool IsAsyncLoadingMultithreadedCoreUObjectInternal()
 	return GetAsyncPackageLoader().IsMultithreaded();
 }
 
+FName GetLoaderNameInternal()
+{
+	return GetAsyncPackageLoader().GetLoaderName();
+}
+
 void SuspendAsyncLoadingInternal()
 {
 	LLM_SCOPE(ELLMTag::AsyncLoading);
diff --git a/Engine/Source/Runtime/CoreUObject/Private/Serialization/EditorPackageLoader.cpp b/Engine/Source/Runtime/CoreUObject/Private/Serialization/EditorPackageLoader.cpp
--- a/Engine/Source/Runtime/CoreUObject/Private/Serialization/EditorPackageLoader.cpp
+++ b/Engine/Source/Runtime/CoreUObject/Private/Serialization/EditorPackageLoader.cpp
@@ -25,6 +25,11 @@ public:
 
 	virtual ~FEditorPackageLoader () { }
 
+	virtual FName GetLoaderName() const override
+	{
+		return TEXT("EditorPackageLoader");
+	}
+
 	virtual void InitializeLoading() override
 	{
 		if (FIoDispatcher::Get().DoesChunkExist(CreateIoChunkId(0, 0, EIoChunkType::ScriptObjects)))
diff --git a/Engine/Source/Runtime/CoreUObject/Private/UObject/CoreNative.cpp b/Engine/Source/Runtime/CoreUObject/Private/UObject/CoreNative.cpp
--- a/Engine/Source/Runtime/CoreUObject/Private/UObject/CoreNative.cpp
+++ b/Engine/Source/Runtime/CoreUObject/Private/UObject/CoreNative.cpp
@@ -17,6 +17,7 @@ void SuspendAsyncLoadingInternal();
 void ResumeAsyncLoadingInternal();
 bool IsAsyncLoadingSuspendedInternal();
 bool IsAsyncLoadingMultithreadedCoreUObjectInternal();
+FName GetLoaderNameInternal();
 
 // CoreUObject module. Handles UObject system pre-init (registers init function with Core callbacks).
 class FCoreUObjectModule : public FDefaultModuleImpl
@@ -46,6 +47,7 @@ public:
 		ResumeAsyncLoading = &ResumeAsyncLoadingInternal;
 		IsAsyncLoadingSuspended = &IsAsyncLoadingSuspendedInternal;
 		IsAsyncLoadingMultithreaded = &IsAsyncLoadingMultithreadedCoreUObjectInternal;
+		GetLoaderName = &GetLoaderNameInternal;
 
 		// Register the script callstack callback to the runtime error logging
 #if UE_RAISE_RUNTIME_ERRORS
diff --git a/Engine/Source/Runtime/CoreUObject/Public/Serialization/AsyncPackageLoader.h b/Engine/Source/Runtime/CoreUObject/Public/Serialization/AsyncPackageLoader.h
--- a/Engine/Source/Runtime/CoreUObject/Public/Serialization/AsyncPackageLoader.h
+++ b/Engine/Source/Runtime/CoreUObject/Public/Serialization/AsyncPackageLoader.h
@@ -230,6 +230,8 @@ public:
 
 	virtual void NotifyRegistrationComplete() = 0;
 
+	virtual FName GetLoaderName() const = 0;
+
 protected:
 	static int32 GetNextRequestId();
 
diff --git a/Engine/Source/Runtime/Engine/Private/Tests/Loading/AsyncLoadingTests_RecursiveLoads.cpp b/Engine/Source/Runtime/Engine/Private/Tests/Loading/AsyncLoadingTests_RecursiveLoads.cpp
new file mode 100644
--- /dev/null
+++ b/Engine/Source/Runtime/Engine/Private/Tests/Loading/AsyncLoadingTests_RecursiveLoads.cpp
@@ -0,0 +1,248 @@
+// Copyright Epic Games, Inc. All Rights Reserved.
+
+#include "AsyncLoadingTests_RecursiveLoads.h"
+#include "Misc/AutomationTest.h"
+#include "Misc/PackageName.h"
+#include "UObject/CoreRedirects.h"
+#include "UObject/SavePackage.h"
+
+#include UE_INLINE_GENERATED_CPP_BY_NAME(AsyncLoadingTests_RecursiveLoads)
+
+UAsyncLoadingTests_RecursiveLoads::FOnPostLoadEvent UAsyncLoadingTests_RecursiveLoads::OnPostLoadEvent;
+UAsyncLoadingTests_RecursiveLoads::FOnSerializeEvent UAsyncLoadingTests_RecursiveLoads::OnSerializeEvent;
+
+#if WITH_DEV_AUTOMATION_TESTS
+
+class FLoadingTests_RecursiveLoads : public FAutomationTestBase
+{
+public:
+	using FAutomationTestBase::FAutomationTestBase;
+protected:
+	static constexpr const TCHAR* ObjectName = TEXT("TestObject");
+	static constexpr const TCHAR* PackagePath1 = TEXT("/Engine/TestRecursiveLoads_Package1");
+	static constexpr const TCHAR* ObjectPath1 = TEXT("/Engine/TestRecursiveLoads_Package1.TestObject");
+	static constexpr const TCHAR* PackagePath2 = TEXT("/Engine/TestRecursiveLoads_Package2");
+	static constexpr const TCHAR* ObjectPath2 = TEXT("/Engine/TestRecursiveLoads_Package2.TestObject");
+
+	virtual void CreateObjects()
+	{
+		UPackage* Package1 = CreatePackage(PackagePath1);
+		UAsyncLoadingTests_RecursiveLoads* Object1 = NewObject<UAsyncLoadingTests_RecursiveLoads>(Package1, ObjectName, RF_Public | RF_Standalone);
+
+		UPackage* Package2 = CreatePackage(PackagePath2);
+		UAsyncLoadingTests_RecursiveLoads* Object2 = NewObject<UAsyncLoadingTests_RecursiveLoads>(Package2, ObjectName, RF_Public | RF_Standalone);
+
+		// This is the soft reference that we want to test loading for
+		Object1->SoftReference = Object2;
+
+		// To avoid an error on save, we need to mark the package as fully loaded.
+		Package1->MarkAsFullyLoaded();
+		Package2->MarkAsFullyLoaded();
+
+		// Save packages to disk.
+		check(UPackage::SavePackage(Package1, nullptr, *FPackageName::LongPackageNameToFilename(PackagePath1, FPackageName::GetAssetPackageExtension()), FSavePackageArgs()));
+		check(UPackage::SavePackage(Package2, nullptr, *FPackageName::LongPackageNameToFilename(PackagePath2, FPackageName::GetAssetPackageExtension()), FSavePackageArgs()));
+
+		// Remove RF_Standalone from top level object.
+		Object1->ClearFlags(RF_Standalone);
+		Object2->ClearFlags(RF_Standalone);
+		
+		// Remove RF_Standalone from the UMetaData otherwise the package will not GC.
+		Package1->GetMetaData()->ClearFlags(RF_Standalone);
+		Package2->GetMetaData()->ClearFlags(RF_Standalone);
+
+		// GC and make sure everything gets cleaned up before loading
+		check(FindObject<UAsyncLoadingTests_RecursiveLoads>(nullptr, ObjectPath1) != nullptr);
+		check(FindObject<UPackage>(nullptr, PackagePath1) != nullptr);
+		check(FindObject<UAsyncLoadingTests_RecursiveLoads>(nullptr, ObjectPath2) != nullptr);
+		check(FindObject<UPackage>(nullptr, PackagePath2) != nullptr);
+
+		CollectGarbage(GARBAGE_COLLECTION_KEEPFLAGS);
+
+		check(FindObject<UAsyncLoadingTests_RecursiveLoads>(nullptr, ObjectPath1) == nullptr);
+		check(FindObject<UPackage>(nullptr, PackagePath1) == nullptr);
+		check(FindObject<UAsyncLoadingTests_RecursiveLoads>(nullptr, ObjectPath2) == nullptr);
+		check(FindObject<UPackage>(nullptr, PackagePath2) == nullptr);
+	}
+
+	virtual void LoadObjects()
+	{
+		LoadPackage(nullptr, PackagePath1, LOAD_None);
+
+		UAsyncLoadingTests_RecursiveLoads* Object2 = FindObject<UAsyncLoadingTests_RecursiveLoads>(nullptr, ObjectPath2);
+		TestFalse(TEXT("The object should have been properly loaded recursively"), Object2->HasAnyFlags(RF_NeedLoad | RF_NeedPostLoad));
+
+		UAsyncLoadingTests_RecursiveLoads* Object1 = FindObject<UAsyncLoadingTests_RecursiveLoads>(nullptr, ObjectPath1);
+		TestFalse(TEXT("The object should have been properly loaded recursively"), Object1->HasAnyFlags(RF_NeedLoad | RF_NeedPostLoad));
+	}
+
+	virtual void CleanupObjects()
+	{
+		// GC and make sure everything gets cleaned up before exiting
+		check(FindObject<UAsyncLoadingTests_RecursiveLoads>(nullptr, ObjectPath1) != nullptr);
+		check(FindObject<UPackage>(nullptr, PackagePath1) != nullptr);
+		check(FindObject<UAsyncLoadingTests_RecursiveLoads>(nullptr, ObjectPath2) != nullptr);
+		check(FindObject<UPackage>(nullptr, PackagePath2) != nullptr);
+
+		FindObject<UAsyncLoadingTests_RecursiveLoads>(nullptr, ObjectPath1)->ClearFlags(RF_Standalone);
+		FindObject<UAsyncLoadingTests_RecursiveLoads>(nullptr, ObjectPath2)->ClearFlags(RF_Standalone);
+
+		FindObject<UPackage>(nullptr, PackagePath1)->GetMetaData()->ClearFlags(RF_Standalone);
+		FindObject<UPackage>(nullptr, PackagePath2)->GetMetaData()->ClearFlags(RF_Standalone);
+
+		CollectGarbage(GARBAGE_COLLECTION_KEEPFLAGS);
+
+		check(FindObject<UAsyncLoadingTests_RecursiveLoads>(nullptr, ObjectPath1) == nullptr);
+		check(FindObject<UPackage>(nullptr, PackagePath1) == nullptr);
+		check(FindObject<UAsyncLoadingTests_RecursiveLoads>(nullptr, ObjectPath2) == nullptr);
+		check(FindObject<UPackage>(nullptr, PackagePath2) == nullptr);
+	}
+
+	bool CanRunInEnvironment(const FString& TestParams, FString* OutReason, bool* OutWarn) const override
+	{
+		if (GetLoaderName() != TEXT("ZenLoader"))
+		{
+			if (OutReason)
+			{
+				*OutReason = FString::Printf(TEXT("Test %s is for ZenLoader only. Cannot run on non-compliant loader currently active: %s"), *GetTestName(), *GetLoaderName().ToString());
+			}
+			return false;
+		}
+
+		return true;
+	}
+
+	bool DoTest()
+	{
+		// Just make sure the async loading queue is empty before beginning.
+		FlushAsyncLoading();
+
+		// Creation phase
+		CreateObjects();
+
+		// Loading phase
+		LoadObjects();
+
+		// Cleanup phase
+		CleanupObjects();
+
+		return true;
+	}
+};
+
+/**
+ * This test validates loading an object synchronously during serialize.
+ */
+IMPLEMENT_CUSTOM_SIMPLE_AUTOMATION_TEST(
+	FLoadingTests_RecursiveLoads_FromSerialize, 
+	FLoadingTests_RecursiveLoads, 
+	TEXT("System.Engine.Loading.RecursiveLoads.FromSerialize"),
+	EAutomationTestFlags::ApplicationContextMask | EAutomationTestFlags::EngineFilter
+)
+bool FLoadingTests_RecursiveLoads_FromSerialize::RunTest(const FString& Parameters)
+{
+	FDelegateHandle DelegateHandle = UAsyncLoadingTests_RecursiveLoads::OnSerializeEvent.AddLambda(
+		[this](FArchive& Ar, UAsyncLoadingTests_RecursiveLoads* Object)
+		{
+			if (Ar.IsLoading())
+			{
+				if (UObject* Obj = Object->SoftReference.LoadSynchronous())
+				{
+					TestTrue(TEXT("Recursive loads in serialize should be deserialized"), !Obj->HasAnyFlags(RF_NeedLoad));
+					TestTrue(TEXT("Recursive loads in serialize should have their postload delayed until the caller is also ready to postload"), Obj->HasAnyFlags(RF_NeedPostLoad));
+				}
+			}
+		}
+	);
+
+	ON_SCOPE_EXIT
+	{
+		UAsyncLoadingTests_RecursiveLoads::OnSerializeEvent.Remove(DelegateHandle);
+	};
+
+	return DoTest();
+}
+
+/**
+ * This test validates loading an object synchronously during postload.
+ */
+IMPLEMENT_CUSTOM_SIMPLE_AUTOMATION_TEST(
+	FLoadingTests_RecursiveLoads_FromPostLoad,
+	FLoadingTests_RecursiveLoads,
+	TEXT("System.Engine.Loading.RecursiveLoads.FromPostLoad"),
+	EAutomationTestFlags::ApplicationContextMask | EAutomationTestFlags::EngineFilter
+)
+
+bool FLoadingTests_RecursiveLoads_FromPostLoad::RunTest(const FString& Parameters)
+{
+	FDelegateHandle DelegateHandle = UAsyncLoadingTests_RecursiveLoads::OnPostLoadEvent.AddLambda(
+		[this](UAsyncLoadingTests_RecursiveLoads* Object)
+		{
+			if (UObject* Obj = Object->SoftReference.LoadSynchronous())
+			{
+				TestFalse(TEXT("Recursive loads in postload should be fully loaded"), Obj->HasAnyFlags(RF_NeedLoad | RF_NeedPostLoad));
+			}
+		}
+	);
+
+	ON_SCOPE_EXIT
+	{
+		UAsyncLoadingTests_RecursiveLoads::OnPostLoadEvent.Remove(DelegateHandle);
+	};
+
+	return DoTest();
+}
+
+/**
+ * This test validates an error is emitted when flushing a requestid that is not a partial load from inside a recursive function.
+ */
+class FLoadingTests_RecursiveLoads_FullFlushFromSerialize_Base : public FLoadingTests_RecursiveLoads
+{
+public:
+	using FLoadingTests_RecursiveLoads::FLoadingTests_RecursiveLoads;
+protected:
+	int32 RequestId;
+
+	virtual void LoadObjects() override
+	{
+		// Create a request before starting the loading test so we get a request that is not tagged as partial.
+		RequestId = LoadPackageAsync(PackagePath2);
+
+		FLoadingTests_RecursiveLoads::LoadObjects();
+	}
+};
+
+IMPLEMENT_CUSTOM_SIMPLE_AUTOMATION_TEST(
+	FLoadingTests_RecursiveLoads_FullFlushFromSerialize,
+	FLoadingTests_RecursiveLoads_FullFlushFromSerialize_Base,
+	TEXT("System.Engine.Loading.RecursiveLoads.FullFlushFromSerialize"),
+	EAutomationTestFlags::ApplicationContextMask | EAutomationTestFlags::EngineFilter
+)
+bool FLoadingTests_RecursiveLoads_FullFlushFromSerialize::RunTest(const FString& Parameters)
+{
+	AddExpectedError(TEXT("can lead to undefined behavior and is strongly discouraged"), EAutomationExpectedErrorFlags::Contains);
+
+	FDelegateHandle DelegateHandle = UAsyncLoadingTests_RecursiveLoads::OnSerializeEvent.AddLambda(
+		[this](FArchive& Ar, UAsyncLoadingTests_RecursiveLoads* Object)
+		{
+			// Do not try to flush ourself as this would lead to a fatal error :)
+			// Just flush Package2 when we're in Package1
+			if (Ar.IsLoading() && Object->GetPathName() == ObjectPath1)
+			{
+				// Flush the requestId that has been created outside of the recursive load. This request
+				// should be a full request and flushing it should result in an error being reported.
+				FlushAsyncLoading(RequestId);
+			}
+		}
+	);
+
+	ON_SCOPE_EXIT
+	{
+		UAsyncLoadingTests_RecursiveLoads::OnSerializeEvent.Remove(DelegateHandle);
+	};
+
+	return DoTest();
+}
+
+#undef TEST_NAME_ROOT
+#endif // WITH_DEV_AUTOMATION_TESTS
diff --git a/Engine/Source/Runtime/Engine/Private/Tests/Loading/AsyncLoadingTests_RecursiveLoads.h b/Engine/Source/Runtime/Engine/Private/Tests/Loading/AsyncLoadingTests_RecursiveLoads.h
new file mode 100644
--- /dev/null
+++ b/Engine/Source/Runtime/Engine/Private/Tests/Loading/AsyncLoadingTests_RecursiveLoads.h
@@ -0,0 +1,41 @@
+// Copyright Epic Games, Inc. All Rights Reserved.
+
+#pragma once
+
+#include "CoreMinimal.h"
+#include "Misc/AutomationTest.h"
+#include "UObject/Object.h"
+#include "UObject/ObjectPtr.h"
+#include "UObject/ObjectMacros.h"
+#include "AsyncLoadingTests_RecursiveLoads.generated.h"
+
+UCLASS()
+class UAsyncLoadingTests_RecursiveLoads : public UObject
+{
+	GENERATED_BODY()
+
+public:
+	UPROPERTY()
+	TSoftObjectPtr<UObject> SoftReference;
+
+	// Those delegates allow to easily change behavior between tests
+	DECLARE_TS_MULTICAST_DELEGATE_OneParam(FOnPostLoadEvent, UAsyncLoadingTests_RecursiveLoads* /*Object*/);
+	static FOnPostLoadEvent OnPostLoadEvent;
+
+	DECLARE_TS_MULTICAST_DELEGATE_TwoParams(FOnSerializeEvent, FArchive& Ar, UAsyncLoadingTests_RecursiveLoads* /*Object*/);
+	static FOnSerializeEvent OnSerializeEvent;
+
+	virtual void PostLoad() override
+	{
+		Super::PostLoad();
+
+		OnPostLoadEvent.Broadcast(this);
+	}
+
+	virtual void Serialize(FArchive& Ar) override
+	{
+		Super::Serialize(Ar);
+
+		OnSerializeEvent.Broadcast(Ar, this);
+	}
+};
