Add more details on what property is being imported.
Particularly when there are warnings importing text from a config file to avoid just getting "asset 'none'"
#rnx

diff --git a/Engine/Source/Runtime/CoreUObject/Private/UObject/Property.cpp b/Engine/Source/Runtime/CoreUObject/Private/UObject/Property.cpp
--- a/Engine/Source/Runtime/CoreUObject/Private/UObject/Property.cpp
+++ b/Engine/Source/Runtime/CoreUObject/Private/UObject/Property.cpp
@@ -2002,6 +2002,7 @@ const TCHAR* FProperty::ImportSingleProperty( const TCHAR* Str, void* DestData,
 						TArray<FString> ImportErrors;
 						ImportError.ParseIntoArray(ImportErrors, LINE_TERMINATOR, true);
 
+						Warn->Logf(ELogVerbosity::Warning, TEXT("While importing text for property '%s' in '%s':"), *Property->GetName(), *ObjectStruct->GetName());
 						for ( int32 ErrorIndex = 0; ErrorIndex < ImportErrors.Num(); ErrorIndex++ )
 						{
 							Warn->Logf(ELogVerbosity::Warning, TEXT("%s"), *ImportErrors[ErrorIndex]);
