Add overload of TryConvertFilenameToLongPackageName using FStringView and FStringBuilderBase

diff --git a/Engine/Source/Runtime/CoreUObject/Private/Misc/PackageName.cpp b/Engine/Source/Runtime/CoreUObject/Private/Misc/PackageName.cpp
--- a/Engine/Source/Runtime/CoreUObject/Private/Misc/PackageName.cpp
+++ b/Engine/Source/Runtime/CoreUObject/Private/Misc/PackageName.cpp
@@ -784,17 +784,41 @@ void FPackageName::InternalFilenameToLongPackageName(FStringView InFilename, FSt
 }
 
 bool FPackageName::TryConvertFilenameToLongPackageName(const FString& InFilename, FString& OutPackageName, FString* OutFailureReason)
+{
+	TStringBuilder<256> FailureReasonBuilder;
+	FStringBuilderBase* FailureReasonBuilderPtr = nullptr;
+	if (OutFailureReason != nullptr)
+	{
+		FailureReasonBuilderPtr = &FailureReasonBuilder;
+	}
+
+	TStringBuilder<256> PackageNameBuilder;
+	const bool bResult = TryConvertFilenameToLongPackageName(MakeStringView(InFilename), PackageNameBuilder, FailureReasonBuilderPtr);
+	if (bResult)
+	{
+		OutPackageName = PackageNameBuilder.ToView();
+	}
+	else if (OutFailureReason != nullptr)
+	{
+		*OutFailureReason = FailureReasonBuilder.ToView();
+	}
+	return bResult;
+}
+
+bool FPackageName::TryConvertFilenameToLongPackageName(FStringView InFilename, FStringBuilderBase& OutPackageName, FStringBuilderBase* OutFailureReason /*= nullptr*/)
 {
 	TStringBuilder<256> LongPackageNameBuilder;
 	InternalFilenameToLongPackageName(InFilename, LongPackageNameBuilder);
-	FStringView LongPackageName = LongPackageNameBuilder.ToString();
+	const FStringView LongPackageName = LongPackageNameBuilder.ToView();
 
 	if (LongPackageName.IsEmpty())
 	{
-		if (OutFailureReason)
+		if (OutFailureReason != nullptr)
 		{
 			FStringView FilenameWithoutExtension = FPathViews::GetBaseFilenameWithPath(InFilename);
-			*OutFailureReason = FString::Printf(TEXT("FilenameToLongPackageName failed to convert '%s'. The Result would be indistinguishable from using '%.*s' as the InFilename."), *InFilename, FilenameWithoutExtension.Len(), FilenameWithoutExtension.GetData());
+			OutFailureReason->Reset();
+			*OutFailureReason << TEXTVIEW("FilenameToLongPackageName failed to convert '") << InFilename << TEXTVIEW("'. ");
+			*OutFailureReason << TEXTVIEW("The Result would be indistinguishable from using '") << FilenameWithoutExtension << TEXTVIEW("' as the InFilename.");
 		}
 		return false;
 	}
@@ -813,10 +837,11 @@ bool FPackageName::TryConvertFilenameToLongPackageName(const FString& InFilename
 
 	// if the package name resolution failed and a relative path was provided, convert to an absolute path
 	// as content may be mounted in a different relative path to the one given
-	if (FPaths::IsRelative(InFilename))
+	if (FPathViews::IsRelativePath(InFilename))
 	{
-		FString AbsPath = FPaths::ConvertRelativePathToFull(InFilename);
-		if (!FPaths::IsRelative(AbsPath) && AbsPath.Len() > 1)
+		TStringBuilder<256> AbsPath;
+		FPathViews::ToAbsolutePath(InFilename, AbsPath);
+		if (!FPathViews::IsRelativePath(AbsPath) && AbsPath.Len() > 1)
 		{
 			if (TryConvertFilenameToLongPackageName(AbsPath, OutPackageName, nullptr))
 			{
@@ -827,20 +852,22 @@ bool FPackageName::TryConvertFilenameToLongPackageName(const FString& InFilename
 
 	if (OutFailureReason != nullptr)
 	{
-		FString InvalidChars;
+		TStringBuilder<16> InvalidChars;
 		if (bContainsDot)
 		{
-			InvalidChars += TEXT(".");
+			InvalidChars << TEXT('.');
 		}
 		if (bContainsBackslash)
 		{
-			InvalidChars += TEXT("\\");
+			InvalidChars << TEXT('\\');
 		}
 		if (bContainsColon)
 		{
-			InvalidChars += TEXT(":");
+			InvalidChars << TEXT(':');
 		}
-		*OutFailureReason = FString::Printf(TEXT("FilenameToLongPackageName failed to convert '%s'. Attempt result was '%.*s', but the path contains illegal characters '%s'"), *InFilename, LongPackageName.Len(), LongPackageName.GetData(), *InvalidChars);
+		OutFailureReason->Reset();
+		*OutFailureReason << TEXTVIEW("FilenameToLongPackageName failed to convert '") << InFilename << TEXTVIEW("'. ");
+		*OutFailureReason << TEXTVIEW("Attempt result was '") << LongPackageName << TEXTVIEW("', but the path contains illegal characters '") << InvalidChars << TEXTVIEW("'.");
 	}
 
 	return false;
diff --git a/Engine/Source/Runtime/CoreUObject/Public/Misc/PackageName.h b/Engine/Source/Runtime/CoreUObject/Public/Misc/PackageName.h
--- a/Engine/Source/Runtime/CoreUObject/Public/Misc/PackageName.h
+++ b/Engine/Source/Runtime/CoreUObject/Public/Misc/PackageName.h
@@ -103,6 +103,7 @@ public:
 	 * @return Returns true if the supplied filename properly maps to one of the long package roots.
 	 */
 	static COREUOBJECT_API bool TryConvertFilenameToLongPackageName(const FString& InFilename, FString& OutPackageName, FString* OutFailureReason = nullptr);
+	static COREUOBJECT_API bool TryConvertFilenameToLongPackageName(FStringView InFilename, FStringBuilderBase& OutPackageName, FStringBuilderBase* OutFailureReason = nullptr);
 
 	/** 
 	 * Converts the supplied filename to long package name.
