[ZenLoader] Fix export bundle being picked up recursively during a flush caused by a synchronous load invoked while running a previous export bundle of the same package

#rnx

diff --git a/Engine/Source/Runtime/CoreUObject/Private/Serialization/AsyncLoading2.cpp b/Engine/Source/Runtime/CoreUObject/Private/Serialization/AsyncLoading2.cpp
--- a/Engine/Source/Runtime/CoreUObject/Private/Serialization/AsyncLoading2.cpp
+++ b/Engine/Source/Runtime/CoreUObject/Private/Serialization/AsyncLoading2.cpp
@@ -5987,6 +5987,11 @@ EEventLoadNodeExecutionResult FAsyncPackage2::Event_ProcessExportBundle(FAsyncLo
 			Package->AsyncLoadingThread.ExternalReadQueue.Enqueue(Package);
 		}
 	}
+	else
+	{
+		// Release the next bundle now that we've finished.
+		Package->GetExportBundleNode(EEventLoadNode2::ExportBundle_Process, Package->ProcessedExportBundlesCount).ReleaseBarrier(&ThreadState);
+	}
 
 	return EEventLoadNodeExecutionResult::Complete;
 }
@@ -6867,9 +6872,11 @@ void FAsyncPackage2::ConditionalBeginProcessPackageExports(FAsyncLoadingThreadSt
 #endif
 			{
 				Package->AsyncPackageLoadingState = EAsyncPackageLoadingState2::ProcessExportBundles;
-				for (int32 ExportBundleIndex = 0; ExportBundleIndex < Package->Data.TotalExportBundleCount; ++ExportBundleIndex)
+				if (Package->Data.TotalExportBundleCount > 0)
 				{
-					Package->GetExportBundleNode(EEventLoadNode2::ExportBundle_Process, ExportBundleIndex).ReleaseBarrier(&ThreadState);
+					// Release a single export bundle node to avoid them being picked up recursively during a flush.
+					// When a node finishes, it will release another one.
+					Package->GetExportBundleNode(EEventLoadNode2::ExportBundle_Process, 0).ReleaseBarrier(&ThreadState);
 				}
 			}
 		});
