Revert JSONFieldName entry.

diff --git a/Engine/Source/Runtime/CoreUObject/Public/UObject/ObjectMacros.h b/Engine/Source/Runtime/CoreUObject/Public/UObject/ObjectMacros.h
--- a/Engine/Source/Runtime/CoreUObject/Public/UObject/ObjectMacros.h
+++ b/Engine/Source/Runtime/CoreUObject/Public/UObject/ObjectMacros.h
@@ -1479,9 +1479,6 @@ namespace UM
 
 	    /// [PropertyMetadata] Used for enum properties to define the subset of invalid values as a comma-separated string; values within this subset are not made available.
         InvalidEnumValues,
-
-		/// [PropertyMetadata] Used to override the JSON field name in serialization via FJsonObjectConverter.
-		JSONFieldName,
 	};
 
 	// Metadata usable in UPROPERTY for customizing the behavior of Persona and UMG
