[ZenLoader]
 - Further improve the partial load logic by avoiding some operations where possible
 - Also add some very verbose logging that can help identifying load ordering issues

#rnx

diff --git a/Engine/Source/Runtime/CoreUObject/Private/Serialization/AsyncLoading2.cpp b/Engine/Source/Runtime/CoreUObject/Private/Serialization/AsyncLoading2.cpp
--- a/Engine/Source/Runtime/CoreUObject/Private/Serialization/AsyncLoading2.cpp
+++ b/Engine/Source/Runtime/CoreUObject/Private/Serialization/AsyncLoading2.cpp
@@ -3494,7 +3494,7 @@ private:
 #endif
 	void ConditionalBeginPostLoad(FAsyncLoadingThreadState2& ThreadState, FAsyncLoadingPostLoadGroup* PostLoadGroup);
 	void ConditionalBeginDeferredPostLoad(FAsyncLoadingThreadState2& ThreadState, FAsyncLoadingPostLoadGroup* DeferredPostLoadGroup);
-	void MergePostLoadGroups(FAsyncLoadingThreadState2& ThreadState, FAsyncLoadingPostLoadGroup* Target, FAsyncLoadingPostLoadGroup* Source, bool bSkipSyncLoadContext = false);
+	void MergePostLoadGroups(FAsyncLoadingThreadState2& ThreadState, FAsyncLoadingPostLoadGroup* Target, FAsyncLoadingPostLoadGroup* Source, bool bUpdateSyncLoadContext = true);
 
 	bool ProcessDeferredDeletePackagesQueue(int32 MaxCount = MAX_int32)
 	{
@@ -3993,7 +3993,7 @@ void FAsyncLoadingThread2::ConditionalBeginDeferredPostLoad(FAsyncLoadingThreadS
 	}
 }
 
-void FAsyncLoadingThread2::MergePostLoadGroups(FAsyncLoadingThreadState2& ThreadState, FAsyncLoadingPostLoadGroup* Target, FAsyncLoadingPostLoadGroup* Source, bool bSkipSyncLoadContext)
+void FAsyncLoadingThread2::MergePostLoadGroups(FAsyncLoadingThreadState2& ThreadState, FAsyncLoadingPostLoadGroup* Target, FAsyncLoadingPostLoadGroup* Source, bool bUpdateSyncLoadContext)
 {
 	if (Target == Source)
 	{
@@ -4013,7 +4013,7 @@ void FAsyncLoadingThread2::MergePostLoadGroups(FAsyncLoadingThreadState2& Thread
 	// If the intention of the caller of this function is to merge postloads into the caller package so they are executed later after
 	// a partial flush, then we can't update the synccontext of the caller during the merge as it would make the whole hierarchy
 	// flush in a single swoop which is in direct contradiction with the partial flush feature.
-	if (!bSkipSyncLoadContext)
+	if (bUpdateSyncLoadContext)
 	{
 		const uint64 SyncLoadContextId = FMath::Max(Source->SyncLoadContextId, Target->SyncLoadContextId);
 		if (SyncLoadContextId)
@@ -4125,7 +4125,15 @@ void FAsyncLoadingThread2::IncludePackageInSyncLoadContextRecursive(FAsyncLoadin
 		return;
 	}
 
+	UE_ASYNC_PACKAGE_LOG(VeryVerbose, Package->Desc, TEXT("IncludePackageInSyncLoadContextRecursive"), TEXT("Setting SyncLoadContextId to %d"), ContextId);
+
 	Package->SyncLoadContextId = ContextId;
+
+	// When using the partial loading feature, don't try to upgrade postload groups into higher priority sync context as
+	// this is exactly what we're trying to prevent. Package that are merged together into a postload group are supposed
+	// to stay at the lower syncloadcontextid, that way, it allows to exit the flush when serialization is done and let
+	// postload run with the ones of the caller.
+#if !WITH_PARTIAL_REQUEST_DURING_RECURSION
 	FAsyncLoadingPostLoadGroup* PostLoadGroup = Package->PostLoadGroup ? Package->PostLoadGroup : Package->DeferredPostLoadGroup;
 	if (PostLoadGroup && PostLoadGroup->SyncLoadContextId < ContextId)
 	{
@@ -4138,6 +4146,7 @@ void FAsyncLoadingThread2::IncludePackageInSyncLoadContextRecursive(FAsyncLoadin
 			}
 		}
 	}
+#endif
 	for (FAsyncPackage2* ImportedPackage : Package->Data.ImportedAsyncPackages)
 	{
 		if (ImportedPackage && ImportedPackage->SyncLoadContextId < ContextId)
@@ -5891,6 +5900,8 @@ EEventLoadNodeExecutionResult FAsyncPackage2::Event_ProcessExportBundle(FAsyncLo
 	check(Package->AsyncPackageLoadingState >= EAsyncPackageLoadingState2::DependenciesReady);
 	Package->AsyncPackageLoadingState = EAsyncPackageLoadingState2::ProcessExportBundles;
 
+	UE_ASYNC_PACKAGE_LOG(VeryVerbose, Package->Desc, TEXT("ProcessExportBundle"), TEXT("Beginning Processing Export Bundle %d"), InExportBundleIndex);
+
 	FAsyncPackageScope2 Scope(Package);
 #if WITH_EDITOR
 	UE::Core::Private::FPlayInEditorLoadingScope PlayInEditorIDScope(Package->Desc.PIEInstanceID);
@@ -5993,6 +6004,8 @@ EEventLoadNodeExecutionResult FAsyncPackage2::Event_ProcessExportBundle(FAsyncLo
 		Package->GetExportBundleNode(EEventLoadNode2::ExportBundle_Process, Package->ProcessedExportBundlesCount).ReleaseBarrier(&ThreadState);
 	}
 
+	UE_ASYNC_PACKAGE_LOG(VeryVerbose, Package->Desc, TEXT("ProcessExportBundle"), TEXT("Finished Processing Export Bundle %d"), InExportBundleIndex);
+
 	return EEventLoadNodeExecutionResult::Complete;
 }
 
@@ -6640,8 +6653,8 @@ EEventLoadNodeExecutionResult FAsyncPackage2::Event_ExportsDone(FAsyncLoadingThr
 		UE_LOG(LogStreaming, Display, TEXT("Merging postload groups of package %s with caller package %s"), *Package->Desc.UPackageName.ToString(), *CallerPackage->Desc.UPackageName.ToString());
 
 		// Do not adjust sync load context, we want to be able to exit the current one even if the caller has not finished yet.
-		const bool bSkipSyncLoadContext = false;
-		Package->AsyncLoadingThread.MergePostLoadGroups(ThreadState, Package->PostLoadGroup, CallerPackage->PostLoadGroup, bSkipSyncLoadContext);
+		const bool bUpdateSyncLoadContext = false;
+		Package->AsyncLoadingThread.MergePostLoadGroups(ThreadState, Package->PostLoadGroup, CallerPackage->PostLoadGroup, bUpdateSyncLoadContext);
 	}
 	else
 #endif // WITH_PARTIAL_REQUEST_DURING_RECURSION
@@ -7178,6 +7191,8 @@ void FAsyncLoadingThread2::UpdateSyncLoadContext(FAsyncLoadingThreadState2& Thre
 		while (ThreadState.SyncLoadContextsCreatedOnGameThread.Dequeue(CreatedOnMainThread))
 		{
 			ThreadState.SyncLoadContextStack.Push(CreatedOnMainThread);
+
+			UE_LOG(LogStreaming, VeryVerbose, TEXT("Pushing ALT SyncLoadContext %d"), CreatedOnMainThread->ContextId);
 		}
 	}
 	if (ThreadState.SyncLoadContextStack.IsEmpty())
@@ -7190,6 +7205,8 @@ void FAsyncLoadingThread2::UpdateSyncLoadContext(FAsyncLoadingThreadState2& Thre
 		// Retire complete/invalid contexts for which we aren't loading any requests
 		while (!ContainsAnyRequestID(SyncLoadContext->RequestIDs))
 		{
+			UE_LOG(LogStreaming, VeryVerbose, TEXT("Popping ALT SyncLoadContext %d"), SyncLoadContext->ContextId);
+
 			SyncLoadContext->ReleaseRef();
 			ThreadState.SyncLoadContextStack.Pop();
 			if (ThreadState.SyncLoadContextStack.IsEmpty())
@@ -7213,7 +7230,7 @@ void FAsyncLoadingThread2::UpdateSyncLoadContext(FAsyncLoadingThreadState2& Thre
 			int32 RequestID = SyncLoadContext->RequestIDs[i];
 			if (SyncLoadContext->RequestedPackages[i] != nullptr)
 			{
-				++FoundPackages;	
+				++FoundPackages;
 			}
 			else if (FAsyncPackage2* RequestedPackage = RequestIdToPackageMap.FindRef(RequestID))
 			{
@@ -8956,6 +8973,9 @@ void FAsyncLoadingThread2::FlushLoading(TConstArrayView<int32> RequestIDs)
 		{
 			SyncLoadContext = new FAsyncLoadingSyncLoadContext(RequestIDs);
 			SyncLoadContext->RequestingPackageDebug = CurrentlyExecutingPackage;
+
+			UE_LOG(LogStreaming, VeryVerbose, TEXT("Pushing GT SyncLoadContext %d"), SyncLoadContext->ContextId);
+
 			GameThreadState->SyncLoadContextStack.Push(SyncLoadContext);
 			if (AsyncLoadingThreadState)
 			{
@@ -9019,6 +9039,9 @@ void FAsyncLoadingThread2::FlushLoading(TConstArrayView<int32> RequestIDs)
 		if (SyncLoadContext)
 		{
 			check(GameThreadState->SyncLoadContextStack.Top() == SyncLoadContext);
+
+			UE_LOG(LogStreaming, VeryVerbose, TEXT("Popping GT SyncLoadContext %d"), SyncLoadContext->ContextId);
+
 			SyncLoadContext->ReleaseRef();
 			GameThreadState->SyncLoadContextStack.Pop();
 			AltZenaphore.NotifyOne();
