fixed crash in AsyncFileHandle with unimplemented EnsureCompletion()

diff --git a/Engine/Source/Editor/UnrealEd/Private/EditorDomain/EditorDomainArchive.cpp b/Engine/Source/Editor/UnrealEd/Private/EditorDomain/EditorDomainArchive.cpp
--- a/Engine/Source/Editor/UnrealEd/Private/EditorDomain/EditorDomainArchive.cpp
+++ b/Engine/Source/Editor/UnrealEd/Private/EditorDomain/EditorDomainArchive.cpp
@@ -715,6 +715,11 @@ protected:
 	{
 	}
 
+	virtual void EnsureCompletion() override
+	{
+		WaitCompletionImpl(0.0f);
+	}
+
 	virtual void CancelImpl() override
 	{
 	}
@@ -757,6 +762,11 @@ protected:
 	{
 	}
 
+	virtual void EnsureCompletion() override
+	{
+		WaitCompletionImpl(0.0f);
+	}
+
 	virtual void CancelImpl() override
 	{
 	}
diff --git a/Engine/Source/Runtime/Core/Private/Microsoft/MicrosoftAsyncIO.h b/Engine/Source/Runtime/Core/Private/Microsoft/MicrosoftAsyncIO.h
--- a/Engine/Source/Runtime/Core/Private/Microsoft/MicrosoftAsyncIO.h
+++ b/Engine/Source/Runtime/Core/Private/Microsoft/MicrosoftAsyncIO.h
@@ -322,6 +322,12 @@ public:
 			}
 		}
 	}
+
+	virtual void EnsureCompletion() override
+	{
+		WaitCompletionImpl(0.0f);
+	}
+
 	virtual void CancelImpl() override
 	{
 		// no cancel support
@@ -355,6 +361,11 @@ public:
 		while (!*(volatile bool*)&bCompleteAndCallbackCalled);
 	}
 
+	virtual void EnsureCompletion() override
+	{
+		WaitCompletionImpl(0.0f);
+	}
+
 	virtual void CancelImpl() override
 	{
 	}
@@ -380,6 +391,11 @@ public:
 		while (!*(volatile bool*)&bCompleteAndCallbackCalled);
 	}
 
+	virtual void EnsureCompletion() override
+	{
+		WaitCompletionImpl(0.0f);
+	}
+
 	virtual void CancelImpl() override
 	{
 	}
diff --git a/Engine/Source/Runtime/Core/Public/Async/AsyncFileHandleNull.h b/Engine/Source/Runtime/Core/Public/Async/AsyncFileHandleNull.h
--- a/Engine/Source/Runtime/Core/Public/Async/AsyncFileHandleNull.h
+++ b/Engine/Source/Runtime/Core/Public/Async/AsyncFileHandleNull.h
@@ -24,6 +24,11 @@ protected:
 	{
 	}
 
+	virtual void EnsureCompletion() override
+	{
+		WaitCompletionImpl(0.0f);
+	}
+
 	virtual void CancelImpl() override
 	{
 	}
diff --git a/Engine/Source/Runtime/CoreUObject/Private/Serialization/BulkDataStreaming.cpp b/Engine/Source/Runtime/CoreUObject/Private/Serialization/BulkDataStreaming.cpp
--- a/Engine/Source/Runtime/CoreUObject/Private/Serialization/BulkDataStreaming.cpp
+++ b/Engine/Source/Runtime/CoreUObject/Private/Serialization/BulkDataStreaming.cpp
@@ -206,6 +206,13 @@ public:
 	virtual ~FChunkReadFileRequest();
 	
 	virtual void WaitCompletionImpl(float TimeLimitSeconds) override;
+
+	virtual void EnsureCompletion() override
+	{
+		WaitCompletionImpl(0.0f);
+	}
+
+
 	virtual void CancelImpl() override;
 	virtual void ReleaseMemoryOwnershipImpl() override;
 	virtual void HandleChunkResult(TIoStatusOr<FIoBuffer>&& Result) override;
@@ -280,6 +287,11 @@ private:
 		while (!*(volatile bool*)&bCompleteAndCallbackCalled);
 	}
 
+	virtual void EnsureCompletion() override
+	{
+		WaitCompletionImpl(0.0f);
+	}
+
 	virtual void CancelImpl() override
 	{
 	}
diff --git a/Engine/Source/Runtime/PakFile/Private/IPlatformFilePak.cpp b/Engine/Source/Runtime/PakFile/Private/IPlatformFilePak.cpp
--- a/Engine/Source/Runtime/PakFile/Private/IPlatformFilePak.cpp
+++ b/Engine/Source/Runtime/PakFile/Private/IPlatformFilePak.cpp
@@ -4166,6 +4166,12 @@ public:
 			WaitEvent = nullptr;
 		}
 	}
+
+	virtual void EnsureCompletion() override
+	{
+		WaitCompletionImpl(0.0f);
+	}
+
 	virtual void CancelImpl() override
 	{
 		check(!WaitEvent); // you canceled from a different thread that you waited from
