Mark the Outer and not the Object that is being renamed with GC barrier

diff --git a/Engine/Source/Runtime/CoreUObject/Private/UObject/UObjectBase.cpp b/Engine/Source/Runtime/CoreUObject/Private/UObject/UObjectBase.cpp
--- a/Engine/Source/Runtime/CoreUObject/Private/UObject/UObjectBase.cpp
+++ b/Engine/Source/Runtime/CoreUObject/Private/UObject/UObjectBase.cpp
@@ -241,12 +241,12 @@ void UObjectBase::LowLevelRename(FName NewName,UObject *NewOuter)
 	if (NewOuter)
 	{
 		OuterPrivate = NewOuter;
+		if (UE::GC::Private::GIsIncrementalReachabilityPending)
+		{
+			OuterPrivate->MarkAsReachable();
+		}
 	}
 	HashObject(this);
-	if (UE::GC::Private::GIsIncrementalReachabilityPending)
-	{
-		static_cast<UObject*>(this)->MarkAsReachable();
-	}
 }
 
 UPackage* UObjectBase::GetExternalPackage() const
