Making sure an object Outer is marked with GC barrier when constructing new objects in between incremental reachability iterations

#prefight 6516bac3c779a827b1a4e7dc

diff --git a/Engine/Source/Runtime/CoreUObject/Private/UObject/UObjectBase.cpp b/Engine/Source/Runtime/CoreUObject/Private/UObject/UObjectBase.cpp
--- a/Engine/Source/Runtime/CoreUObject/Private/UObject/UObjectBase.cpp
+++ b/Engine/Source/Runtime/CoreUObject/Private/UObject/UObjectBase.cpp
@@ -137,6 +137,15 @@ UObjectBase::UObjectBase(UClass* InClass,
 #if CSV_PROFILER && CSV_TRACK_UOBJECT_COUNT
 	UObjectStats::IncrementUObjectCount();
 #endif
+
+	if (UE::GC::Private::GIsIncrementalReachabilityPending)
+	{
+		ClassPrivate->MarkAsReachable();
+		if (OuterPrivate)
+		{
+			OuterPrivate->MarkAsReachable();
+		}	
+	}
 }
 
 
@@ -183,6 +192,15 @@ void UObjectBase::DeferredRegister(UClass *UClassStaticClass,const TCHAR* Packag
 	check(!ClassPrivate);
 	ClassPrivate = UClassStaticClass;
 
+	if (UE::GC::Private::GIsIncrementalReachabilityPending)
+	{
+		ClassPrivate->MarkAsReachable();
+		if (OuterPrivate)
+		{
+			OuterPrivate->MarkAsReachable();
+		}
+	}
+
 	// Add to the global object table.
 	AddObject(FName(InName), EInternalObjectFlags::None);
 	// At this point all compiled-in objects should have already been fully constructed so it's safe to remove the NotFullyConstructed flag
@@ -306,6 +324,11 @@ void UObjectBase::SetClass(UClass* NewClass)
 	ClassPrivate->DestroyPersistentUberGraphFrame((UObject*)this);
 #endif
 	ClassPrivate = NewClass;
+	if (UE::GC::Private::GIsIncrementalReachabilityPending)
+	{
+		checkf(ClassPrivate, TEXT("SetClass called on %s with a null class"), *GetFName().ToString());
+		ClassPrivate->MarkAsReachable();
+	}
 #if USE_UBER_GRAPH_PERSISTENT_FRAME
 	ClassPrivate->CreatePersistentUberGraphFrame((UObject*)this, /*bCreateOnlyIfEmpty =*/false, /*bSkipSuperClass =*/false, OldClass);
 #endif
