Cooker export verification: Change "was never saved as an export" message into clearer text "content is missing from cook", and add feedback for why it is missing from the cook.
#rnx

diff --git a/Engine/Source/Editor/UnrealEd/Private/Commandlets/CookCommandlet.cpp b/Engine/Source/Editor/UnrealEd/Private/Commandlets/CookCommandlet.cpp
--- a/Engine/Source/Editor/UnrealEd/Private/Commandlets/CookCommandlet.cpp
+++ b/Engine/Source/Editor/UnrealEd/Private/Commandlets/CookCommandlet.cpp
@@ -543,7 +543,13 @@ void UCookCommandlet::RunCookByTheBookCook(UCookOnTheFlyServer* CookOnTheFlyServ
 	if (bShouldVerifyEDLCookInfo)
 	{
 		bool bFullReferencesExpected = !(CookOptions & ECookByTheBookOptions::SkipHardReferences);
-		UE::SavePackageUtilities::VerifyEDLCookInfo(bFullReferencesExpected);
+		UE::SavePackageUtilities::VerifyEDLCookInfo([](ELogVerbosity::Type Verbosity, FStringView Message)
+			{
+#if !NO_LOGGING
+				FMsg::Logf(__FILE__, __LINE__, LogCook.GetCategoryName(), Verbosity, TEXT("%.*s"),
+				Message.Len(), Message.GetData());
+#endif
+			}, bFullReferencesExpected);
 	}
 }
 
diff --git a/Engine/Source/Runtime/CoreUObject/Private/UObject/SavePackage/SavePackageUtilities.cpp b/Engine/Source/Runtime/CoreUObject/Private/UObject/SavePackage/SavePackageUtilities.cpp
--- a/Engine/Source/Runtime/CoreUObject/Private/UObject/SavePackage/SavePackageUtilities.cpp
+++ b/Engine/Source/Runtime/CoreUObject/Private/UObject/SavePackage/SavePackageUtilities.cpp
@@ -1468,7 +1468,8 @@ FEDLCookChecker FEDLCookChecker::AccumulateAndClear()
 	return Accumulator;
 }
 
-void FEDLCookChecker::Verify(bool bFullReferencesExpected)
+void FEDLCookChecker::Verify(const UE::SavePackageUtilities::FEDLMessageCallback& MessageCallback,
+	bool bFullReferencesExpected)
 {
 	check(!GIsSavingPackage);
 	FEDLCookChecker Accumulator = AccumulateAndClear();
@@ -1496,9 +1497,39 @@ void FEDLCookChecker::Verify(bool bFullReferencesExpected)
 				}
 
 				// Any imports of this non-exported node are an error; log them all if they exist
+				if (NodeData.ImportingPackagesSorted.IsEmpty())
+				{
+					continue;
+				}
+
+				const FEDLNodeData* NodeDataOfExportPackage = &NodeData;
+				while (NodeDataOfExportPackage->ParentID != NodeIDInvalid)
+				{
+					int32 ParentNodeIndex = static_cast<int32>(NodeDataOfExportPackage->ParentID);
+					check(Accumulator.Nodes.IsValidIndex(ParentNodeIndex));
+					NodeDataOfExportPackage = &Accumulator.Nodes[ParentNodeIndex];
+				}
+
+				const TCHAR* ReasonExportIsMissing = TEXT("");
+				if (NodeDataOfExportPackage->bIsExport)
+				{
+					ReasonExportIsMissing = TEXT("the object was stripped out of the target package when saved");
+				}
+				else
+				{
+					ReasonExportIsMissing = TEXT("the target package was not cooked");
+				}
+
 				for (FName PackageName : NodeData.ImportingPackagesSorted)
 				{
-					UE_LOG(LogSavePackage, Warning, TEXT("%s imported %s, but it was never saved as an export."), *PackageName.ToString(), *NodeData.ToString(Accumulator));
+					TStringBuilder<512> Message;
+					Message << TEXTVIEW("Content is missing from cook. Source package referenced an object in target package but ");
+					Message << ReasonExportIsMissing << TEXT(".\n");
+					Message << TEXT("\tSource package: ") << PackageName << TEXT("\n");
+					Message << TEXT("\tTarget package: ") << NodeDataOfExportPackage->Name << TEXT("\n");
+					Message << TEXT("\tReferenced object: ");
+					NodeData.AppendPathName(Accumulator, Message);
+					MessageCallback(ELogVerbosity::Warning, Message);
 				}
 			}
 		}
@@ -1701,9 +1732,21 @@ void StartSavingEDLCookInfoForVerification()
 }
 
 void VerifyEDLCookInfo(bool bFullReferencesExpected)
+{
+	VerifyEDLCookInfo([](ELogVerbosity::Type Verbosity, FStringView Message)
+		{
+#if !NO_LOGGING
+			FMsg::Logf(__FILE__, __LINE__, LogSavePackage.GetCategoryName(), Verbosity, TEXT("%.*s"),
+				Message.Len(), Message.GetData());
+#endif
+		}, bFullReferencesExpected);
+}
+
+void VerifyEDLCookInfo(const UE::SavePackageUtilities::FEDLMessageCallback& MessageCallback,
+	bool bFullReferencesExpected)
 {
 	LLM_SCOPE_BYTAG(EDLCookChecker);
-	FEDLCookChecker::Verify(bFullReferencesExpected);
+	FEDLCookChecker::Verify(MessageCallback, bFullReferencesExpected);
 }
 
 void EDLCookInfoAddIterativelySkippedPackage(FName LongPackageName)
diff --git a/Engine/Source/Runtime/CoreUObject/Private/UObject/SavePackage/SavePackageUtilities.h b/Engine/Source/Runtime/CoreUObject/Private/UObject/SavePackage/SavePackageUtilities.h
--- a/Engine/Source/Runtime/CoreUObject/Private/UObject/SavePackage/SavePackageUtilities.h
+++ b/Engine/Source/Runtime/CoreUObject/Private/UObject/SavePackage/SavePackageUtilities.h
@@ -15,6 +15,7 @@
 #include "UObject/Package.h"
 #include "UObject/UObjectMarks.h"
 #include "UObject/ObjectPtr.h"
+#include "UObject/SavePackage.h"
 
 // This file contains private utilities shared by UPackage::Save and UPackage::Save2 
 
@@ -160,7 +161,8 @@ struct FEDLCookChecker
 	void AddPackageWithUnknownExports(FName LongPackageName);
 
 	static void StartSavingEDLCookInfoForVerification();
-	static void Verify(bool bFullReferencesExpected);
+	static void Verify(const UE::SavePackageUtilities::FEDLMessageCallback& MessageCallback,
+		bool bFullReferencesExpected);
 	static void MoveToCompactBinaryAndClear(FCbWriter& Writer, bool& bOutHasData);
 	static bool AppendFromCompactBinary(FCbFieldView Field);
 
diff --git a/Engine/Source/Runtime/CoreUObject/Public/UObject/SavePackage.h b/Engine/Source/Runtime/CoreUObject/Public/UObject/SavePackage.h
--- a/Engine/Source/Runtime/CoreUObject/Public/UObject/SavePackage.h
+++ b/Engine/Source/Runtime/CoreUObject/Public/UObject/SavePackage.h
@@ -18,6 +18,7 @@
 #include "Serialization/ArchiveCookData.h"
 #include "Serialization/FileRegions.h"
 #include "Serialization/PackageWriter.h"
+#include "Templates/Function.h"
 #include "Templates/UniquePtr.h"
 #include "Templates/UnrealTemplate.h"
 #include "UObject/NameTypes.h"
@@ -309,7 +310,9 @@ namespace UE::SavePackageUtilities
 	COREUOBJECT_API int32 GetNumPackagesSaved();
 
 	COREUOBJECT_API void StartSavingEDLCookInfoForVerification();
+	using FEDLMessageCallback = TFunction<void(ELogVerbosity::Type, FStringView)>;
 	COREUOBJECT_API void VerifyEDLCookInfo(bool bFullReferencesExpected = true);
+	COREUOBJECT_API void VerifyEDLCookInfo(const FEDLMessageCallback& MessageCallback, bool bFullReferencesExpected = true);
 	COREUOBJECT_API void EDLCookInfoAddIterativelySkippedPackage(FName LongPackageName);
 	COREUOBJECT_API void EDLCookInfoMoveToCompactBinaryAndClear(FCbWriter& Writer, bool& bOutHasData);
 	COREUOBJECT_API void EDLCookInfoMoveToCompactBinaryAndClear(FCbWriter& Writer, bool& bOutHasData, FName PackageName);
