[ZenLoader]
 - Add partial flush support from loading thread
 - Fix the ConvertFromType allowing to load the softobjectptr during its conversion to hardref even when called from ALT and zenloader is active
 - Add unittest to exercise softobjectptr to hardref conversion in ConvertFromType
 - Move asserts preventing flush from loading thread directly into legacy loader
 - Added a few veryverbose logs that were missing for linker load
 - Fix partial loading being applied to sync calls coming from a postload, preventing the object to be fully postloaded before finishing the sync load

diff --git a/Engine/Source/Runtime/CoreUObject/Private/Serialization/AsyncLoading.cpp b/Engine/Source/Runtime/CoreUObject/Private/Serialization/AsyncLoading.cpp
--- a/Engine/Source/Runtime/CoreUObject/Private/Serialization/AsyncLoading.cpp
+++ b/Engine/Source/Runtime/CoreUObject/Private/Serialization/AsyncLoading.cpp
@@ -7500,6 +7500,8 @@ void FAsyncLoadingThread::FlushLoading(TConstArrayView<int32> RequestIDs)
 {
  	if (IsAsyncLoadingPackages())
 	{
+		checkf(IsInGameThread(), TEXT("The current loader %s is unable to FlushAsyncLoading from any thread other than the game thread."), *GetLoaderName().ToString());
+
 		// Flushing async loading while loading is suspend will result in infinite stall
 		UE_CLOG(IsAsyncLoadingSuspendedInternal(), LogStreaming, Fatal, TEXT("Cannot Flush Async Loading while async loading is suspended (%d)"), GetAsyncLoadingSuspendedCount());
 
diff --git a/Engine/Source/Runtime/CoreUObject/Private/Serialization/AsyncLoading2.cpp b/Engine/Source/Runtime/CoreUObject/Private/Serialization/AsyncLoading2.cpp
--- a/Engine/Source/Runtime/CoreUObject/Private/Serialization/AsyncLoading2.cpp
+++ b/Engine/Source/Runtime/CoreUObject/Private/Serialization/AsyncLoading2.cpp
@@ -496,6 +496,7 @@ struct FPackageRequest
 {
 	int32 RequestId = -1;
 	ERequestType RequestType = ERequestType::Full;
+	FPackageId RequesterId;
 	int32 Priority = -1;
 	EPackageFlags PackageFlags = PKG_None;
 #if WITH_EDITOR
@@ -520,12 +521,13 @@ struct FPackageRequest
 #endif
 	}
 
-	static FPackageRequest Create(int32 RequestId, ERequestType RequestType, EPackageFlags PackageFlags, uint32 LoadFlags, int32 PIEInstanceID, int32 Priority, const FLinkerInstancingContext* InstancingContext, const FPackagePath& PackagePath, FName CustomName, TUniquePtr<FLoadPackageAsyncDelegate> PackageLoadedDelegate, TUniquePtr<FLoadPackageAsyncProgressDelegate> PackageProgressDelegate)
+	static FPackageRequest Create(int32 RequestId, ERequestType RequestType, FPackageId RequesterId, EPackageFlags PackageFlags, uint32 LoadFlags, int32 PIEInstanceID, int32 Priority, const FLinkerInstancingContext* InstancingContext, const FPackagePath& PackagePath, FName CustomName, TUniquePtr<FLoadPackageAsyncDelegate> PackageLoadedDelegate, TUniquePtr<FLoadPackageAsyncProgressDelegate> PackageProgressDelegate)
 	{
 		return FPackageRequest
 		{
 			RequestId,
 			RequestType,
+			RequesterId,
 			Priority,
 			PackageFlags,
 #if WITH_EDITOR
@@ -550,6 +552,8 @@ struct FAsyncPackageDesc2
 	int32 RequestID;
 	// The type of request
 	ERequestType RequestType;
+	// The Id of the top package on the stack (if any) when the request was issued.
+	FPackageId RequesterId;
 	// Package priority
 	int32 Priority;
 	/** The flags that should be applied to the package */
@@ -582,6 +586,7 @@ struct FAsyncPackageDesc2
 		{
 			Request.RequestId,
 			Request.RequestType,
+			Request.RequesterId,
 			Request.Priority,
 			Request.PackageFlags,
 #if WITH_EDITOR
@@ -611,6 +616,7 @@ struct FAsyncPackageDesc2
 		{
 			INDEX_NONE,
 			ERequestType::Full,
+			FPackageId(),
 			ImportingPackageDesc.Priority,
 			PKG_None,
 #if WITH_EDITOR
@@ -2204,16 +2210,16 @@ public:
 	uint64 ContextId;
 	TArray<int32, TInlineAllocator<4>> RequestIDs;
 	TArray<FAsyncPackage2*, TInlineAllocator<4>> RequestedPackages;
-	FAsyncPackage2* RequestingPackageDebug = nullptr;
+	FAsyncPackage2* RequestingPackage = nullptr;
 	std::atomic<bool> bHasFoundRequestedPackages { false };
 
 private:
 	std::atomic<int32> RefCount = 1;
 
-	static uint64 NextContextId;
+	static std::atomic<uint64> NextContextId;
 };
 
-uint64 FAsyncLoadingSyncLoadContext::NextContextId = 1;
+std::atomic<uint64> FAsyncLoadingSyncLoadContext::NextContextId { 1 };
 
 struct FAsyncLoadingThreadState2
 {
@@ -2868,7 +2874,7 @@ public:
 		return GraphAllocator;
 	}
 
-	FAsyncPackage2* GetCallerPackage(FAsyncLoadingThreadState2& ThreadState) const;
+	static FAsyncPackage2* GetCurrentlyExecutingPackage(FAsyncLoadingThreadState2& ThreadState, FAsyncPackage2* PackageToFilter = nullptr);
 
 	/** [EDL] Begin Event driven loader specific stuff */
 
@@ -3398,6 +3404,8 @@ public:
 
 	virtual void FlushLoading(TConstArrayView<int32> RequestIds) override;
 
+	void FlushLoadingFromLoadingThread(FAsyncLoadingThreadState2& ThreadState, TConstArrayView<int32> RequestIds);
+
 	virtual int32 GetNumQueuedPackages() override
 	{
 		return QueuedPackagesCounter;
@@ -3601,7 +3609,7 @@ private:
 	EAsyncPackageState::Type ProcessLoadedPackagesFromGameThread(FAsyncLoadingThreadState2& ThreadState, bool& bDidSomething, TConstArrayView<int32> FlushRequestIDs = {});
 
 	void IncludePackageInSyncLoadContextRecursive(FAsyncLoadingThreadState2& ThreadState, uint64 ContextId, FAsyncPackage2* Package);
-	void UpdateSyncLoadContext(FAsyncLoadingThreadState2& ThreadState);
+	void UpdateSyncLoadContext(FAsyncLoadingThreadState2& ThreadState, bool bAutoHandleSyncLoadContext = true);
 
 	bool CreateAsyncPackagesFromQueue(FAsyncLoadingThreadState2& ThreadState);
 
@@ -4054,6 +4062,10 @@ FAsyncPackage2* FAsyncLoadingThread2::FindOrInsertPackage(FAsyncLoadingThreadSta
 		}
 		else
 		{
+			if (Desc.RequesterId.IsValid())
+			{
+				Package->Desc.RequesterId = Desc.RequesterId;
+			}
 			if (Desc.RequestID > 0)
 			{
 				Package->AddRequestID(ThreadState, Desc.RequestID, Desc.RequestType);
@@ -5462,6 +5474,7 @@ bool FAsyncPackage2::PreloadLinkerLoadExports(FAsyncLoadingThreadState2& ThreadS
 		{
 			if (Object->HasAnyFlags(RF_NeedLoad))
 			{
+				UE_ASYNC_PACKAGE_LOG(VeryVerbose, Desc, TEXT("PreloadLinkerLoadExports"), TEXT("Preloading export %d: %s"), ExportIndex, *Object->GetPathName());
 				LinkerLoadState->Linker->Preload(Object);
 			}
 		}
@@ -5539,6 +5552,13 @@ bool FAsyncPackage2::ResolveLinkerLoadImports(FAsyncLoadingThreadState2& ThreadS
 				LinkerImport.XObject = FromImportStore;
 				LinkerImport.SourceIndex = FromImportStore->GetLinkerIndex();
 				LinkerImport.SourceLinker = FromImportStore->GetLinker();
+
+				UE_ASYNC_PACKAGE_LOG(VeryVerbose, Desc, TEXT("ResolveLinkerLoadImports"), TEXT("Resolved import %d: %s"), ImportIndex, *FromImportStore->GetPathName());
+			}
+
+			if (FromImportStore == nullptr)
+			{
+				UE_ASYNC_PACKAGE_LOG(Verbose, Desc, TEXT("ResolveLinkerLoadImports"), TEXT("Could not resolve import %d"), ImportIndex);
 			}
 		}
 	}
@@ -6578,18 +6598,19 @@ bool FAsyncPackage2::EventDrivenSerializeExport(const FAsyncPackageHeaderData& H
 	return true;
 }
 
-FAsyncPackage2* FAsyncPackage2::GetCallerPackage(FAsyncLoadingThreadState2& ThreadState) const
+FAsyncPackage2* FAsyncPackage2::GetCurrentlyExecutingPackage(FAsyncLoadingThreadState2& ThreadState, FAsyncPackage2* PackageToFilter)
 {
-	FAsyncPackage2* CallerPackage = nullptr;
+	FAsyncPackage2* CurrentlyExecutingPackage = nullptr;
 	for (int32 Index = ThreadState.CurrentlyExecutingEventNodeStack.Num() - 1; Index >= 0; --Index)
 	{
-		if (ThreadState.CurrentlyExecutingEventNodeStack[Index]->GetPackage() != this)
+		FAsyncPackage2* Package = ThreadState.CurrentlyExecutingEventNodeStack[Index]->GetPackage();
+		if (Package != nullptr && Package != PackageToFilter)
 		{
-			CallerPackage = ThreadState.CurrentlyExecutingEventNodeStack[Index]->GetPackage();
+			CurrentlyExecutingPackage = Package;
 			break;
 		}
 	}
-	return CallerPackage;
+	return CurrentlyExecutingPackage;
 }
 
 EEventLoadNodeExecutionResult FAsyncPackage2::Event_ExportsDone(FAsyncLoadingThreadState2& ThreadState, FAsyncPackage2* Package, int32)
@@ -6645,16 +6666,20 @@ EEventLoadNodeExecutionResult FAsyncPackage2::Event_ExportsDone(FAsyncLoadingThr
 	// Release any requests that are only waiting on exports, allowing a flush to exit.
 	Package->AsyncLoadingThread.RemovePendingRequests(ThreadState, Package->RequestIDs[(int32)ERequestType::ExportsDone]);
 
-	FAsyncPackage2* CallerPackage = Package->GetCallerPackage(ThreadState);
+	FAsyncPackage2* RequesterPackage = nullptr;
+	if (Package->Desc.RequesterId.IsValid())
+	{
+		RequesterPackage = Package->AsyncLoadingThread.AsyncPackageLookup.FindRef(Package->Desc.RequesterId);
+	}
 
 	// If this is a recursive load and the caller has not started postloading yet, merge postload group.
-	if (CallerPackage && CallerPackage->PostLoadGroup != nullptr)
+	if (RequesterPackage && RequesterPackage->PostLoadGroup != nullptr)
 	{
-		UE_LOG(LogStreaming, Display, TEXT("Merging postload groups of package %s with caller package %s"), *Package->Desc.UPackageName.ToString(), *CallerPackage->Desc.UPackageName.ToString());
+		UE_ASYNC_PACKAGE_LOG(Display, Package->Desc, TEXT("ExportsDone"), TEXT("Merging postload groups with requester package %s"), *RequesterPackage->Desc.UPackageName.ToString());
 
 		// Do not adjust sync load context, we want to be able to exit the current one even if the caller has not finished yet.
 		const bool bUpdateSyncLoadContext = false;
-		Package->AsyncLoadingThread.MergePostLoadGroups(ThreadState, Package->PostLoadGroup, CallerPackage->PostLoadGroup, bUpdateSyncLoadContext);
+		Package->AsyncLoadingThread.MergePostLoadGroups(ThreadState, Package->PostLoadGroup, RequesterPackage->PostLoadGroup, bUpdateSyncLoadContext);
 	}
 	else
 #endif // WITH_PARTIAL_REQUEST_DURING_RECURSION
@@ -7183,9 +7208,9 @@ FEventLoadNode2& FAsyncPackage2::GetExportBundleNode(EEventLoadNode2 Phase, uint
 	return Data.ExportBundleNodes[ExportBundleNodeIndex];
 }
 
-void FAsyncLoadingThread2::UpdateSyncLoadContext(FAsyncLoadingThreadState2& ThreadState)
+void FAsyncLoadingThread2::UpdateSyncLoadContext(FAsyncLoadingThreadState2& ThreadState, bool bAutoHandleSyncLoadContext)
 {
-	if (ThreadState.bIsAsyncLoadingThread)
+	if (ThreadState.bIsAsyncLoadingThread && bAutoHandleSyncLoadContext)
 	{
 		FAsyncLoadingSyncLoadContext* CreatedOnMainThread;
 		while (ThreadState.SyncLoadContextsCreatedOnGameThread.Dequeue(CreatedOnMainThread))
@@ -7200,7 +7225,7 @@ void FAsyncLoadingThread2::UpdateSyncLoadContext(FAsyncLoadingThreadState2& Thre
 		return;
 	}
 	FAsyncLoadingSyncLoadContext* SyncLoadContext = ThreadState.SyncLoadContextStack.Top();
-	if (ThreadState.bIsAsyncLoadingThread)
+	if (ThreadState.bIsAsyncLoadingThread && bAutoHandleSyncLoadContext)
 	{
 		// Retire complete/invalid contexts for which we aren't loading any requests
 		while (!ContainsAnyRequestID(SyncLoadContext->RequestIDs))
@@ -7239,18 +7264,41 @@ void FAsyncLoadingThread2::UpdateSyncLoadContext(FAsyncLoadingThreadState2& Thre
 				SyncLoadContext->RequestedPackages[i] = RequestedPackage;
 
 #if WITH_PARTIAL_REQUEST_DURING_RECURSION
-				FAsyncPackage2* CallerPackage = RequestedPackage->GetCallerPackage(ThreadState);
-				if (CallerPackage != nullptr && RequestedPackage->RequestIDs[(int32)ERequestType::Full].Contains(RequestID))
+				FAsyncPackage2* RequestingPackage = SyncLoadContext->RequestingPackage;
+				if (RequestingPackage != nullptr &&
+					RequestedPackage->RequestIDs[(int32)ERequestType::Full].Contains(RequestID))
 				{
-					UE_LOG(LogStreaming, Error, TEXT("Flushing package %s recursively from another package %s can lead to undefined behavior and is strongly discouraged"),
-						*RequestedPackage->Desc.UPackageName.ToString(),
-						*CallerPackage->Desc.UPackageName.ToString()
-					);
+					// When this happens, it's because a load was initiated from outside a package recursion and the requestid didn't have the partial loading flag applied to it.
+					// Then the requestid is being flushed from inside another package stack, in which case we might need to demote the request id to a partial one on certain condition.
+
+					// If the flush is coming from a step before the requesting package is back on GT, there is no way to fully flush
+					// the requested package unless its already done.
+					if (RequestingPackage->AsyncPackageLoadingState < EAsyncPackageLoadingState2::DeferredPostLoad &&
+						RequestedPackage->AsyncPackageLoadingState < EAsyncPackageLoadingState2::Complete)
+					{
+						//  
+						// Note: Update the FLoadingTests_RecursiveLoads_FullFlushFrom_Serialize test if you edit this error message
+						//
+						UE_LOG(LogStreaming, Error, TEXT("Fully flushing package %s recursively from another package %s would lead into a deadlock. Demoting requestID %d to a partial request instead."),
+							*RequestedPackage->Desc.UPackageName.ToString(),
+							*RequestingPackage->Desc.UPackageName.ToString(),
+							RequestID
+						);
+
+						RequestedPackage->RequestIDs[(int32)ERequestType::Full].Remove(RequestID);
+						RequestedPackage->RequestIDs[(int32)ERequestType::ExportsDone].Add(RequestID);
+						RequestedPackage->Desc.RequesterId = RequestingPackage->Desc.UPackageId;
+
+						if (RequestedPackage->AsyncPackageLoadingState >= EAsyncPackageLoadingState2::ExportsDone)
+						{
+							RemovePendingRequests(ThreadState, { RequestID });
+						}
+					}
 				}
 #endif
 
 				IncludePackageInSyncLoadContextRecursive(ThreadState, SyncLoadContext->ContextId, RequestedPackage);
-				++FoundPackages;	
+				++FoundPackages;
 			}
 		}
 		
@@ -8786,6 +8834,7 @@ PRAGMA_ENABLE_DEPRECATION_WARNINGS
 	TRACE_LOADTIME_BEGIN_REQUEST(RequestId);
 	AddPendingRequest(RequestId);
 
+	FPackageId RequesterId;
 	ERequestType RequestType = ERequestType::Full;
 
 #if WITH_PARTIAL_REQUEST_DURING_RECURSION
@@ -8799,23 +8848,30 @@ PRAGMA_ENABLE_DEPRECATION_WARNINGS
 	// an import. We want the deserialization steps done before returning from the load but
 	// we keep postload for after deserialization is done on the referencer too so that postload 
 	// sees everything initialized in the case where it also has a reference back to the referencer.
-	// Using the LOAD_Partial flag allows the FlushRequest to return as soon as exports are done which
-	// allows the caller to get a pointer to an object inside the package even if it's not fully loaded yet.
+	// Using a partial request flag allows the flush to return as soon as exports are done which
+	// gives the opportunity to the caller to get a pointer to an object inside the package even if it's not fully loaded yet.
 	if (ThreadState != nullptr &&
-		ThreadState->CurrentlyExecutingEventNodeStack.Num() > 0 &&
-		ThreadState->CurrentlyExecutingEventNodeStack.Top()->GetPackage() != nullptr)
+		ThreadState->CurrentlyExecutingEventNodeStack.Num() > 0)
 	{
-		UE_LOG(LogStreaming, Display,
-			TEXT("LoadPackage on %s is being called from another package %s, partial loading is automatically being applied which means only exports are processed before exiting flush and postloads will be called along with the caller's"),
-			*PackageNameToLoad.ToString(),
-			*ThreadState->CurrentlyExecutingEventNodeStack.Top()->GetPackage()->Desc.UPackageName.ToString()
-		);
+		FAsyncPackage2* Requester = ThreadState->CurrentlyExecutingEventNodeStack.Top()->GetPackage();
+
+		// If the requester doesn't have a postload group, it means it's already being postloaded and
+		// we don't need to load partially in that case.
+		if (Requester && Requester->PostLoadGroup)
+		{
+			UE_LOG(LogStreaming, Display,
+				TEXT("LoadPackage on %s is being called from another package %s, partial loading is automatically being applied which means only exports are processed before exiting flush and postloads will be called along with the caller's"),
+				*PackageNameToLoad.ToString(),
+				*Requester->Desc.UPackageName.ToString()
+			);
 
-		RequestType = ERequestType::ExportsDone;
+			RequesterId = Requester->Desc.UPackageId;
+			RequestType = ERequestType::ExportsDone;
+		}
 	}
 #endif
 
-	PackageRequestQueue.Enqueue(FPackageRequest::Create(RequestId, RequestType, InPackageFlags, InLoadFlags, InPIEInstanceID, InPackagePriority, InInstancingContext, InPackagePath, InCustomName, MoveTemp(InCompletionDelegate), MoveTemp(InProgressDelegate)));
+	PackageRequestQueue.Enqueue(FPackageRequest::Create(RequestId, RequestType, RequesterId, InPackageFlags, InLoadFlags, InPIEInstanceID, InPackagePriority, InInstancingContext, InPackagePath, InCustomName, MoveTemp(InCompletionDelegate), MoveTemp(InProgressDelegate)));
 	++QueuedPackagesCounter;
 	++PackagesWithRemainingWorkCounter;
 
@@ -8901,26 +8957,74 @@ EAsyncPackageState::Type FAsyncLoadingThread2::ProcessLoadingFromGameThread(FAsy
 	// CSV_CUSTOM_STAT(FileIO, EDLEventQueueDepth, (int32)GraphAllocator.TotalNodeCount, ECsvCustomStatOp::Set);
 	CSV_CUSTOM_STAT(FileIO, QueuedPackagesQueueDepth, GetNumQueuedPackages(), ECsvCustomStatOp::Set);
 	CSV_CUSTOM_STAT(FileIO, ExistingQueuedPackagesQueueDepth, GetNumAsyncPackages(), ECsvCustomStatOp::Set);
-	
+
 	bool bDidSomething = false;
 	TickAsyncLoadingFromGameThread(ThreadState, bUseTimeLimit, bUseFullTimeLimit, TimeLimit, {}, bDidSomething);
 	return IsAsyncLoadingPackages() ? EAsyncPackageState::TimeOut : EAsyncPackageState::Complete;
 }
 
+void FAsyncLoadingThread2::FlushLoadingFromLoadingThread(FAsyncLoadingThreadState2& ThreadState, TConstArrayView<int32> RequestIDs)
+{
+	TRACE_CPUPROFILER_EVENT_SCOPE(FlushLoadingFromLoadingThread);
+
+	if (RequestIDs.IsEmpty())
+	{
+		return;
+	}
+
+	FAsyncLoadingSyncLoadContext* SyncLoadContext = new FAsyncLoadingSyncLoadContext(RequestIDs);
+
+	FAsyncPackage2* GetCurrentlyExecutingPackage = FAsyncPackage2::GetCurrentlyExecutingPackage(ThreadState);
+	SyncLoadContext->RequestingPackage = GetCurrentlyExecutingPackage;
+
+	UE_LOG(LogStreaming, VeryVerbose, TEXT("Pushing ALT SyncLoadContext %d"), SyncLoadContext->ContextId);
+	AsyncLoadingThreadState->SyncLoadContextStack.Push(SyncLoadContext);
+
+	// We handle the context push/pop manually since we don't interact with the main thread during this time.
+	static constexpr bool bAutoHandleSyncLoadContext = false;
+	UpdateSyncLoadContext(ThreadState, bAutoHandleSyncLoadContext);
+	while (ContainsAnyRequestID(SyncLoadContext->RequestIDs))
+	{
+		EventQueue.ExecuteSyncLoadEvents(ThreadState);
+	}
+
+	check(AsyncLoadingThreadState->SyncLoadContextStack.Top() == SyncLoadContext);
+	UE_LOG(LogStreaming, VeryVerbose, TEXT("Popping ALT SyncLoadContext %d"), AsyncLoadingThreadState->SyncLoadContextStack.Top()->ContextId);
+	AsyncLoadingThreadState->SyncLoadContextStack.Pop();
+	delete SyncLoadContext;
+}
+
 void FAsyncLoadingThread2::FlushLoading(TConstArrayView<int32> RequestIDs)
 {
 	if (IsAsyncLoadingPackages())
 	{
+		// We can't possibly support flushing from async loading thread unless we have the partial request support active.
+		const bool bIsFlushSupportedOnCurrentThread = IsInGameThread() || (WITH_PARTIAL_REQUEST_DURING_RECURSION && IsInAsyncLoadingThread());
+		ensureMsgf(bIsFlushSupportedOnCurrentThread, TEXT("The current loader '%s' is unable to FlushAsyncLoading from the current thread."), *GetLoaderName().ToString());
+
+		if (!bIsFlushSupportedOnCurrentThread)
+		{
+			UE_LOG(LogStreaming, Error, TEXT("The current loader '%s' is unable to FlushAsyncLoading from the current thread. Flush will be ignored."), *GetLoaderName().ToString());
+			return;
+		}
+
 		// Flushing async loading while loading is suspend will result in infinite stall
 		UE_CLOG(SuspendRequestedCount.load(std::memory_order_relaxed) > 0, LogStreaming, Fatal, TEXT("Cannot Flush Async Loading while async loading is suspended"));
 
-		SCOPE_CYCLE_COUNTER(STAT_FAsyncPackage_FlushAsyncLoadingGameThread);
-
 		if (RequestIDs.Num() != 0 && !ContainsAnyRequestID(RequestIDs))
 		{
 			return;
 		}
 
+		FAsyncLoadingThreadState2* ThreadState = FAsyncLoadingThreadState2::Get();
+		if (ThreadState && ThreadState->bIsAsyncLoadingThread)
+		{
+			FlushLoadingFromLoadingThread(*ThreadState, RequestIDs);
+			return;
+		}
+
+		SCOPE_CYCLE_COUNTER(STAT_FAsyncPackage_FlushAsyncLoadingGameThread);
+
 		// if the sync count is 0, then this flush is not triggered from a sync load, broadcast the delegate in that case
 		FUObjectThreadContext& ThreadContext = FUObjectThreadContext::Get();
 		if (ThreadContext.SyncLoadUsingAsyncLoaderCount == 0)
@@ -8972,7 +9076,7 @@ void FAsyncLoadingThread2::FlushLoading(TConstArrayView<int32> RequestIDs)
 		if (RequestIDs.Num() != 0 && GOnlyProcessRequiredPackagesWhenSyncLoading)
 		{
 			SyncLoadContext = new FAsyncLoadingSyncLoadContext(RequestIDs);
-			SyncLoadContext->RequestingPackageDebug = CurrentlyExecutingPackage;
+			SyncLoadContext->RequestingPackage = CurrentlyExecutingPackage;
 
 			UE_LOG(LogStreaming, VeryVerbose, TEXT("Pushing GT SyncLoadContext %d"), SyncLoadContext->ContextId);
 
diff --git a/Engine/Source/Runtime/CoreUObject/Private/Serialization/AsyncPackageLoader.cpp b/Engine/Source/Runtime/CoreUObject/Private/Serialization/AsyncPackageLoader.cpp
--- a/Engine/Source/Runtime/CoreUObject/Private/Serialization/AsyncPackageLoader.cpp
+++ b/Engine/Source/Runtime/CoreUObject/Private/Serialization/AsyncPackageLoader.cpp
@@ -272,11 +272,11 @@ void FlushAsyncLoading(TConstArrayView<int32> RequestIds)
 	CheckImageIntegrityAtRuntime();
 #endif
 	LLM_SCOPE(ELLMTag::AsyncLoading);
-	checkf(IsInGameThread(), TEXT("Unable to FlushAsyncLoading from any thread other than the game thread."));
+
 	if (GPackageLoader)
 	{
 #if NO_LOGGING == 0
-		if (IsAsyncLoading())
+		if (IsAsyncLoading() && IsInGameThread())
 		{
 			// Log the flush, but only display once per frame to avoid log spam.
 			static uint64 LastFrameNumber = -1;
diff --git a/Engine/Source/Runtime/CoreUObject/Private/UObject/PropertyObject.cpp b/Engine/Source/Runtime/CoreUObject/Private/UObject/PropertyObject.cpp
--- a/Engine/Source/Runtime/CoreUObject/Private/UObject/PropertyObject.cpp
+++ b/Engine/Source/Runtime/CoreUObject/Private/UObject/PropertyObject.cpp
@@ -49,14 +49,15 @@ EConvertFromTypeResult FObjectProperty::ConvertFromType(const FPropertyTag& Tag,
 
 		UObject* PreviousValueObj = nullptr;
 
-		// If we're async loading it's not safe to do a sync load because it may crash or fail to set the variable, so throw an error if it's not already in memory
-		if (!IsInGameThread())
+		// If we're async loading, only zenloader supports synchronous loads from the async loading thread.
+		// Other loaders are not safe to do a sync load because they may crash or fail to set the variable, so throw an error if it's not already in memory
+		if (!IsInGameThread() && GetLoaderName() != TEXT("ZenLoader"))
 		{
 			PreviousValueObj = PreviousValue.Get();
 
 			if (!PreviousValueObj && !PreviousValue.IsNull())
 			{
-				UE_LOG(LogClass, Error, TEXT("Failed to convert soft path %s to unloaded object as this is not safe during async loading. Load and resave %s in the editor to fix!"), *PreviousValue.ToString(), *Slot.GetUnderlyingArchive().GetArchiveName());
+				UE_LOG(LogClass, Error, TEXT("Failed to convert soft path %s to unloaded object as this is not supported during async loading with the currently active loader '%s'. Load and resave %s in the editor to fix!"), *PreviousValue.ToString(), *GetLoaderName().ToString(), *Slot.GetUnderlyingArchive().GetArchiveName());
 			}
 		}
 		else
diff --git a/Engine/Source/Runtime/CoreUObject/Private/UObject/UObjectGlobals.cpp b/Engine/Source/Runtime/CoreUObject/Private/UObject/UObjectGlobals.cpp
--- a/Engine/Source/Runtime/CoreUObject/Private/UObject/UObjectGlobals.cpp
+++ b/Engine/Source/Runtime/CoreUObject/Private/UObject/UObjectGlobals.cpp
@@ -1638,7 +1638,6 @@ UPackage* LoadPackageInternal(UPackage* InOuter, const FPackagePath& PackagePath
 		UE_LOG(LogUObjectGlobals, Warning, TEXT("Attempted to LoadPackage from empty PackagePath."));
 		return nullptr;
 	}
-	checkf(IsInGameThread(), TEXT("Unable to load %s. Objects and Packages can only be loaded from the game thread."), *PackagePath.GetDebugName());
 
 	FUObjectThreadContext& ThreadContext = FUObjectThreadContext::Get();
 	if (ShouldAlwaysLoadPackageAsync(PackagePath))
@@ -1653,7 +1652,10 @@ UPackage* LoadPackageInternal(UPackage* InOuter, const FPackagePath& PackagePath
 
 		UE_TRACK_REFERENCING_PACKAGE_SCOPED(PackageName, PackageAccessTrackingOps::NAME_Load);
 
-		if (FCoreDelegates::OnSyncLoadPackage.IsBound())
+		// This delegate is not thread-safe and the subscribers are mostly interested by sync loads
+		// that might stall the game thread anyway. So for now, do not broadcast when sync loading
+		// from the loading thread.
+		if (IsInGameThread() && FCoreDelegates::OnSyncLoadPackage.IsBound())
 		{
 			FCoreDelegates::OnSyncLoadPackage.Broadcast(PackageName.ToString());
 		}
@@ -1702,6 +1704,7 @@ UPackage* LoadPackageInternal(UPackage* InOuter, const FPackagePath& PackagePath
 		}
 	}
 
+	checkf(IsInGameThread(), TEXT("Unable to load %s. Objects and Packages can only be loaded from the game thread with the currently active loader '%s'."), *PackagePath.GetDebugName(), *GetLoaderName().ToString());
 	UPackage* Result = nullptr;
 
 #if WITH_EDITOR
diff --git a/Engine/Source/Runtime/Engine/Private/Tests/Loading/AsyncLoadingTests_ConvertFromType.cpp b/Engine/Source/Runtime/Engine/Private/Tests/Loading/AsyncLoadingTests_ConvertFromType.cpp
new file mode 100644
--- /dev/null
+++ b/Engine/Source/Runtime/Engine/Private/Tests/Loading/AsyncLoadingTests_ConvertFromType.cpp
@@ -0,0 +1,130 @@
+// Copyright Epic Games, Inc. All Rights Reserved.
+
+#include "AsyncLoadingTests_ConvertFromType.h"
+#include "Misc/AutomationTest.h"
+#include "Misc/PackageName.h"
+#include "UObject/CoreRedirects.h"
+#include "UObject/SavePackage.h"
+
+#include UE_INLINE_GENERATED_CPP_BY_NAME(AsyncLoadingTests_ConvertFromType)
+
+#if WITH_DEV_AUTOMATION_TESTS
+
+class FClassRedirectScope
+{
+	TArray<FCoreRedirect> Redirects;
+public:
+	FClassRedirectScope(UClass* OldClass, UClass* NewClass)
+	{
+		Redirects.Emplace(ECoreRedirectFlags::Type_Class, OldClass->GetName(), NewClass->GetName());
+		check(FCoreRedirects::AddRedirectList(Redirects, TEXT("Tests")));
+	}
+
+	~FClassRedirectScope()
+	{
+		check(FCoreRedirects::RemoveRedirectList(Redirects, TEXT("Tests")));
+	}
+};
+
+/**
+ * This test validates ConvertFromType thread-safety.
+ */
+IMPLEMENT_SIMPLE_AUTOMATION_TEST(FConvertFromType_SoftToHard, TEXT("System.Engine.Loading.ConvertFromType_SoftToHard"), EAutomationTestFlags::ApplicationContextMask | EAutomationTestFlags::EngineFilter)
+bool FConvertFromType_SoftToHard::RunTest(const FString& Parameters)
+{
+	// Just make sure the async loading queue is empty before beginning.
+	FlushAsyncLoading();
+
+	constexpr const TCHAR* ObjectName   = TEXT("TestObject");
+	constexpr const TCHAR* PackagePath1 = TEXT("/Engine/TestConvertFromType_SoftToHard1");
+	constexpr const TCHAR* ObjectPath1  = TEXT("/Engine/TestConvertFromType_SoftToHard1.TestObject");
+	constexpr const TCHAR* PackagePath2 = TEXT("/Engine/TestConvertFromType_SoftToHard2");
+	constexpr const TCHAR* ObjectPath2  = TEXT("/Engine/TestConvertFromType_SoftToHard2.TestObject");
+
+	// Creation phase
+	{
+		UPackage* Package1 = CreatePackage(PackagePath1);
+		UAsyncLoadingTests_ConvertFromType_V1* Object1 = NewObject<UAsyncLoadingTests_ConvertFromType_V1>(Package1, ObjectName, RF_Public | RF_Standalone);
+
+		UPackage* Package2 = CreatePackage(PackagePath2);
+		UAsyncLoadingTests_ConvertFromType_V1* Object2 = NewObject<UAsyncLoadingTests_ConvertFromType_V1>(Package2, ObjectName, RF_Public | RF_Standalone);
+	
+		// We need inter package dependency to trigger a load during ConvertFromType.
+		Object1->Reference = Object2;
+
+		// To avoid an error on save, we need to mark the package as fully loaded.
+		Package1->MarkAsFullyLoaded();
+		Package2->MarkAsFullyLoaded();
+
+		// Save packages to disk.
+		check(UPackage::SavePackage(Package1, nullptr, *FPackageName::LongPackageNameToFilename(PackagePath1, FPackageName::GetAssetPackageExtension()), FSavePackageArgs()));
+		check(UPackage::SavePackage(Package2, nullptr, *FPackageName::LongPackageNameToFilename(PackagePath2, FPackageName::GetAssetPackageExtension()), FSavePackageArgs()));
+
+		// Remove RF_Standalone from top level object.
+		Object1->ClearFlags(RF_Standalone);
+		Object2->ClearFlags(RF_Standalone);
+
+		// Remove RF_Standalone from the UMetaData otherwise the package will not GC.
+		Package1->GetMetaData()->ClearFlags(RF_Standalone);
+		Package2->GetMetaData()->ClearFlags(RF_Standalone);
+
+		// GC and make sure everything gets cleaned up before loading
+		check(FindObject<UAsyncLoadingTests_ConvertFromType_V1>(nullptr, ObjectPath1) != nullptr);
+		check(FindObject<UPackage>(nullptr, PackagePath1) != nullptr);
+		check(FindObject<UAsyncLoadingTests_ConvertFromType_V1>(nullptr, ObjectPath2) != nullptr);
+		check(FindObject<UPackage>(nullptr, PackagePath2) != nullptr);
+
+		CollectGarbage(GARBAGE_COLLECTION_KEEPFLAGS);
+
+		check(FindObject<UAsyncLoadingTests_ConvertFromType_V1>(nullptr, ObjectPath1) == nullptr);
+		check(FindObject<UPackage>(nullptr, PackagePath1) == nullptr);
+		check(FindObject<UAsyncLoadingTests_ConvertFromType_V1>(nullptr, ObjectPath2) == nullptr);
+		check(FindObject<UPackage>(nullptr, PackagePath2) == nullptr);
+	}
+
+	// Loading phase
+	{
+		// We're going to deserialize with the V2 so we get a ConvertFromType call.
+		FClassRedirectScope RedirectScope(UAsyncLoadingTests_ConvertFromType_V1::StaticClass(), UAsyncLoadingTests_ConvertFromType_V2::StaticClass());
+
+		UPackage* LoadedPackage = nullptr;
+		UAsyncLoadingTests_ConvertFromType_V2* LoadedObject = nullptr;
+		int32 RequestID = LoadPackageAsync(PackagePath1,
+			FLoadPackageAsyncDelegate::CreateLambda(
+				[&LoadedPackage, &LoadedObject, &ObjectName](const FName& InPackageName, UPackage* InLoadedPackage, EAsyncLoadingResult::Type InResult)
+				{
+					LoadedPackage = InLoadedPackage;
+					LoadedObject  = FindObject<UAsyncLoadingTests_ConvertFromType_V2>(LoadedPackage, ObjectName);
+				}
+			));
+		FlushAsyncLoading(RequestID);
+
+		TestTrue(TEXT("The object should have been properly loaded"), LoadedObject != nullptr);
+		TestTrue(TEXT("The hard-ref should now point to the object inside the second package"), LoadedObject->Reference.Get() != nullptr);
+		TestTrue(TEXT("The hard-ref should now point to the object inside the second package"), LoadedObject->Reference->GetPathName() == ObjectPath2);
+	}
+
+	// Cleanup phase
+	{
+		// GC and make sure everything gets cleaned up before exiting
+		check(FindObject<UAsyncLoadingTests_ConvertFromType_V2>(nullptr, ObjectPath1) != nullptr);
+		check(FindObject<UPackage>(nullptr, PackagePath1) != nullptr);
+		check(FindObject<UAsyncLoadingTests_ConvertFromType_V2>(nullptr, ObjectPath2) != nullptr);
+		check(FindObject<UPackage>(nullptr, PackagePath2) != nullptr);
+
+		FindObject<UPackage>(nullptr, PackagePath1)->GetMetaData()->ClearFlags(RF_Standalone);
+		FindObject<UPackage>(nullptr, PackagePath2)->GetMetaData()->ClearFlags(RF_Standalone);
+
+		CollectGarbage(GARBAGE_COLLECTION_KEEPFLAGS);
+
+		check(FindObject<UAsyncLoadingTests_ConvertFromType_V2>(nullptr, ObjectPath1) == nullptr);
+		check(FindObject<UPackage>(nullptr, PackagePath1) == nullptr);
+		check(FindObject<UAsyncLoadingTests_ConvertFromType_V2>(nullptr, ObjectPath2) == nullptr);
+		check(FindObject<UPackage>(nullptr, PackagePath2) == nullptr);
+	}
+
+	return true;
+}
+
+#undef TEST_NAME_ROOT
+#endif // WITH_DEV_AUTOMATION_TESTS
diff --git a/Engine/Source/Runtime/Engine/Private/Tests/Loading/AsyncLoadingTests_ConvertFromType.h b/Engine/Source/Runtime/Engine/Private/Tests/Loading/AsyncLoadingTests_ConvertFromType.h
new file mode 100644
--- /dev/null
+++ b/Engine/Source/Runtime/Engine/Private/Tests/Loading/AsyncLoadingTests_ConvertFromType.h
@@ -0,0 +1,29 @@
+// Copyright Epic Games, Inc. All Rights Reserved.
+
+#pragma once
+
+#include "CoreMinimal.h"
+#include "UObject/Object.h"
+#include "UObject/ObjectPtr.h"
+#include "UObject/ObjectMacros.h"
+#include "AsyncLoadingTests_ConvertFromType.generated.h"
+
+UCLASS()
+class UAsyncLoadingTests_ConvertFromType_V1 : public UObject
+{
+	GENERATED_BODY()
+
+public:
+	UPROPERTY()
+	TSoftObjectPtr<UObject> Reference;
+};
+
+UCLASS()
+class UAsyncLoadingTests_ConvertFromType_V2 : public UObject
+{
+	GENERATED_BODY()
+
+public:
+	UPROPERTY()
+	TObjectPtr<UObject> Reference;
+};
diff --git a/Engine/Source/Runtime/Engine/Private/Tests/Loading/AsyncLoadingTests_RecursiveLoads.cpp b/Engine/Source/Runtime/Engine/Private/Tests/Loading/AsyncLoadingTests_RecursiveLoads.cpp
--- a/Engine/Source/Runtime/Engine/Private/Tests/Loading/AsyncLoadingTests_RecursiveLoads.cpp
+++ b/Engine/Source/Runtime/Engine/Private/Tests/Loading/AsyncLoadingTests_RecursiveLoads.cpp
@@ -96,6 +96,9 @@ protected:
 		check(FindObject<UPackage>(nullptr, PackagePath1) == nullptr);
 		check(FindObject<UAsyncLoadingTests_RecursiveLoads>(nullptr, ObjectPath2) == nullptr);
 		check(FindObject<UPackage>(nullptr, PackagePath2) == nullptr);
+
+		UAsyncLoadingTests_RecursiveLoads::OnPostLoadEvent.Clear();
+		UAsyncLoadingTests_RecursiveLoads::OnSerializeEvent.Clear();
 	}
 
 	bool CanRunInEnvironment(const FString& TestParams, FString* OutReason, bool* OutWarn) const override
@@ -141,7 +144,7 @@ IMPLEMENT_CUSTOM_SIMPLE_AUTOMATION_TEST(
 )
 bool FLoadingTests_RecursiveLoads_FromSerialize::RunTest(const FString& Parameters)
 {
-	FDelegateHandle DelegateHandle = UAsyncLoadingTests_RecursiveLoads::OnSerializeEvent.AddLambda(
+	UAsyncLoadingTests_RecursiveLoads::OnSerializeEvent.AddLambda(
 		[this](FArchive& Ar, UAsyncLoadingTests_RecursiveLoads* Object)
 		{
 			if (Ar.IsLoading())
@@ -155,11 +158,6 @@ bool FLoadingTests_RecursiveLoads_FromSerialize::RunTest(const FString& Paramete
 		}
 	);
 
-	ON_SCOPE_EXIT
-	{
-		UAsyncLoadingTests_RecursiveLoads::OnSerializeEvent.Remove(DelegateHandle);
-	};
-
 	return DoTest();
 }
 
@@ -175,7 +173,7 @@ IMPLEMENT_CUSTOM_SIMPLE_AUTOMATION_TEST(
 
 bool FLoadingTests_RecursiveLoads_FromPostLoad::RunTest(const FString& Parameters)
 {
-	FDelegateHandle DelegateHandle = UAsyncLoadingTests_RecursiveLoads::OnPostLoadEvent.AddLambda(
+	UAsyncLoadingTests_RecursiveLoads::OnPostLoadEvent.AddLambda(
 		[this](UAsyncLoadingTests_RecursiveLoads* Object)
 		{
 			if (UObject* Obj = Object->SoftReference.LoadSynchronous())
@@ -185,18 +183,13 @@ bool FLoadingTests_RecursiveLoads_FromPostLoad::RunTest(const FString& Parameter
 		}
 	);
 
-	ON_SCOPE_EXIT
-	{
-		UAsyncLoadingTests_RecursiveLoads::OnPostLoadEvent.Remove(DelegateHandle);
-	};
-
 	return DoTest();
 }
 
 /**
  * This test validates an error is emitted when flushing a requestid that is not a partial load from inside a recursive function.
  */
-class FLoadingTests_RecursiveLoads_FullFlushFromSerialize_Base : public FLoadingTests_RecursiveLoads
+class FLoadingTests_RecursiveLoads_FullFlushFrom_Base : public FLoadingTests_RecursiveLoads
 {
 public:
 	using FLoadingTests_RecursiveLoads::FLoadingTests_RecursiveLoads;
@@ -213,16 +206,16 @@ protected:
 };
 
 IMPLEMENT_CUSTOM_SIMPLE_AUTOMATION_TEST(
-	FLoadingTests_RecursiveLoads_FullFlushFromSerialize,
-	FLoadingTests_RecursiveLoads_FullFlushFromSerialize_Base,
-	TEXT("System.Engine.Loading.RecursiveLoads.FullFlushFromSerialize"),
+	FLoadingTests_RecursiveLoads_FullFlushFrom_Serialize,
+	FLoadingTests_RecursiveLoads_FullFlushFrom_Base,
+	TEXT("System.Engine.Loading.RecursiveLoads.FullFlushFrom.Serialize"),
 	EAutomationTestFlags::ApplicationContextMask | EAutomationTestFlags::EngineFilter
 )
-bool FLoadingTests_RecursiveLoads_FullFlushFromSerialize::RunTest(const FString& Parameters)
+bool FLoadingTests_RecursiveLoads_FullFlushFrom_Serialize::RunTest(const FString& Parameters)
 {
-	AddExpectedError(TEXT("can lead to undefined behavior and is strongly discouraged"), EAutomationExpectedErrorFlags::Contains);
+	AddExpectedError(TEXT("would lead into a deadlock. Demoting requestID"), EAutomationExpectedErrorFlags::Contains);
 
-	FDelegateHandle DelegateHandle = UAsyncLoadingTests_RecursiveLoads::OnSerializeEvent.AddLambda(
+	UAsyncLoadingTests_RecursiveLoads::OnSerializeEvent.AddLambda(
 		[this](FArchive& Ar, UAsyncLoadingTests_RecursiveLoads* Object)
 		{
 			// Do not try to flush ourself as this would lead to a fatal error :)
@@ -232,14 +225,41 @@ bool FLoadingTests_RecursiveLoads_FullFlushFromSerialize::RunTest(const FString&
 				// Flush the requestId that has been created outside of the recursive load. This request
 				// should be a full request and flushing it should result in an error being reported.
 				FlushAsyncLoading(RequestId);
+
+				UAsyncLoadingTests_RecursiveLoads* Object2 = FindObject<UAsyncLoadingTests_RecursiveLoads>(nullptr, ObjectPath2);
+				TestFalse(TEXT("The object should be serialized"), Object2->HasAnyFlags(RF_NeedLoad));
+				TestTrue(TEXT("The object should not have been postloaded"), Object2->HasAnyFlags(RF_NeedPostLoad));
 			}
 		}
 	);
 
-	ON_SCOPE_EXIT
-	{
-		UAsyncLoadingTests_RecursiveLoads::OnSerializeEvent.Remove(DelegateHandle);
-	};
+	return DoTest();
+}
+
+IMPLEMENT_CUSTOM_SIMPLE_AUTOMATION_TEST(
+	FLoadingTests_RecursiveLoads_FullFlushFrom_PostLoad,
+	FLoadingTests_RecursiveLoads_FullFlushFrom_Base,
+	TEXT("System.Engine.Loading.RecursiveLoads.FullFlushFrom.Postload"),
+	EAutomationTestFlags::ApplicationContextMask | EAutomationTestFlags::EngineFilter
+)
+bool FLoadingTests_RecursiveLoads_FullFlushFrom_PostLoad::RunTest(const FString& Parameters)
+{
+	UAsyncLoadingTests_RecursiveLoads::OnPostLoadEvent.AddLambda(
+		[this](UAsyncLoadingTests_RecursiveLoads* Object)
+		{
+			// Do not try to flush ourself as this would lead to a fatal error :)
+			// Just flush Package2 when we're in Package1
+			if (Object->GetPathName() == ObjectPath1)
+			{
+				// Flush the requestId that has been created outside of the recursive load. This request
+				// should be a full request and flushing it should result in an error being reported.
+				FlushAsyncLoading(RequestId);
+
+				UAsyncLoadingTests_RecursiveLoads* Object2 = FindObject<UAsyncLoadingTests_RecursiveLoads>(nullptr, ObjectPath2);
+				TestFalse(TEXT("The object should be serialized and postloaded"), Object2->HasAnyFlags(RF_NeedLoad | RF_NeedPostLoad));
+			}
+		}
+	);
 
 	return DoTest();
 }
