Move the IsAsyncLoading calls around open for the logging, not something that matters to roll back

diff --git a/Engine/Source/Runtime/CoreUObject/Private/UObject/Linker.cpp b/Engine/Source/Runtime/CoreUObject/Private/UObject/Linker.cpp
--- a/Engine/Source/Runtime/CoreUObject/Private/UObject/Linker.cpp
+++ b/Engine/Source/Runtime/CoreUObject/Private/UObject/Linker.cpp
@@ -479,10 +479,13 @@ void ResetLinkerExports(UPackage* InPackage)
 
 void ResetLoaders(UObject* InPkg)
 {
-	if (IsAsyncLoading())
-	{
-		UE_LOG(LogLinker, Log, TEXT("ResetLoaders(%s) is flushing async loading"), *GetPathNameSafe(InPkg));
-	}
+	// This call into IsAsyncLoading calls an atomic, which we want to avoid in a transact and this log is not something we want to roll back so ignore
+	UE_AUTORTFM_OPEN({
+		if (IsAsyncLoading())
+		{
+			UE_LOG(LogLinker, Log, TEXT("ResetLoaders(%s) is flushing async loading"), *GetPathNameSafe(InPkg));
+		}
+	});
 
 	// Make sure we're not in the middle of loading something in the background.
 	FlushAsyncLoading();
@@ -491,10 +494,13 @@ void ResetLoaders(UObject* InPkg)
 
 void ResetLoaders(TArrayView<UObject*> InOuters)
 {
-	if (IsAsyncLoading())
-	{
-		UE_LOG(LogLinker, Log, TEXT("ResetLoaders is flushing async loading"));
-	}
+	// This call into IsAsyncLoading calls an atomic, which we want to avoid in a transact and this log is not something we want to roll back so ignore
+	UE_AUTORTFM_OPEN({
+		if (IsAsyncLoading())
+		{
+			UE_LOG(LogLinker, Log, TEXT("ResetLoaders is flushing async loading"));
+		}
+	});
 
 	// Make sure we're not in the middle of loading something in the background.
 	FlushAsyncLoading();
