Add open around Client and Server checking exclusion list

diff --git a/Engine/Source/Runtime/CoreUObject/Private/UObject/Obj.cpp b/Engine/Source/Runtime/CoreUObject/Private/UObject/Obj.cpp
--- a/Engine/Source/Runtime/CoreUObject/Private/UObject/Obj.cpp
+++ b/Engine/Source/Runtime/CoreUObject/Private/UObject/Obj.cpp
@@ -853,21 +853,31 @@ FClassExclusionData GDedicatedClientExclusionList;
 
 bool UObject::NeedsLoadForServer() const
 {
-	return !GDedicatedServerExclusionList.IsExcluded(GetClass());
+	bool Ret = false;
+	UE_AUTORTFM_OPEN({
+		Ret = !GDedicatedServerExclusionList.IsExcluded(GetClass());
+	});
+	return Ret;
 }
 
 void UObject::UpdateClassesExcludedFromDedicatedServer(const TArray<FString>& InClassNames, const TArray<FString>& InModulesNames)
 {
+	// Avoid adding an UE_AUTORTFM_OPEN as this clears a list which wont be rolled back correctly
 	GDedicatedServerExclusionList.UpdateExclusionList(InClassNames, InModulesNames);
 }
 
 bool UObject::NeedsLoadForClient() const
 {
-	return !GDedicatedClientExclusionList.IsExcluded(GetClass());
+	bool Ret = false;
+	UE_AUTORTFM_OPEN({
+		Ret = !GDedicatedClientExclusionList.IsExcluded(GetClass());
+	});
+	return Ret;
 }
 
 void UObject::UpdateClassesExcludedFromDedicatedClient(const TArray<FString>& InClassNames, const TArray<FString>& InModulesNames)
 {
+	// Avoid adding an UE_AUTORTFM_OPEN as this clears a list which wont be rolled back correctly
 	GDedicatedClientExclusionList.UpdateExclusionList(InClassNames, InModulesNames);
 }
 
