PropertyBag: Initial version with basic support for adding, removing, and iterating properties

Missing features include serialization, versioning of unknown properties, nesting property bags, and integration with garbage collection.

diff --git a/Engine/Source/Runtime/CoreUObject/Private/UObject/PropertyBag.cpp b/Engine/Source/Runtime/CoreUObject/Private/UObject/PropertyBag.cpp
new file mode 100644
--- /dev/null
+++ b/Engine/Source/Runtime/CoreUObject/Private/UObject/PropertyBag.cpp
@@ -0,0 +1,261 @@
+// Copyright Epic Games, Inc. All Rights Reserved.
+
+#include "UObject/PropertyBag.h"
+
+#include "Serialization/NullArchive.h"
+#include "UObject/UnrealType.h"
+
+namespace UE
+{
+
+FPropertyBag::FPropertyBag() = default;
+FPropertyBag::~FPropertyBag() = default;
+
+FPropertyBag::FPropertyBag(FPropertyBag&&) = default;
+FPropertyBag& FPropertyBag::operator=(FPropertyBag&&) = default;
+
+void FPropertyBag::Empty()
+{
+	Properties.Empty();
+}
+
+void FPropertyBag::Copy(const UStruct* Struct, void* Data, EPropertyFlags SkipFlags)
+{
+	EFieldIterationFlags IterationFlags = EFieldIterationFlags::IncludeAll;
+	if (SkipFlags & CPF_Deprecated)
+	{
+		IterationFlags &= ~EFieldIterationFlags::IncludeDeprecated;
+	}
+
+	for (TFieldIterator<FProperty> It(Struct, IterationFlags); It; ++It)
+	{
+		if (FProperty* Property = *It; !(Property->GetPropertyFlags() & SkipFlags))
+		{
+			for (int32 ArrayIndex = 0, ArrayDim = Property->ArrayDim; ArrayIndex < ArrayDim; ++ArrayIndex)
+			{
+				Add(Property, Data, ArrayIndex);
+			}
+		}
+	}
+}
+
+void FPropertyBag::Add(FProperty* Property, void* Data, int32 ArrayIndex)
+{
+	// TODO: NullAr is a workaround to FPropertyTag requiring an archive to assert on versioned property serialization.
+	FNullArchive NullAr;
+
+	FPropertyBag::FValue& Value = Properties.FindOrAdd(Property->GetFName());
+
+	const bool bPropertyChanged = Value.Tag.Prop != Property;
+
+	if (bPropertyChanged)
+	{
+		Value.Destroy();
+		Value.Tag = FPropertyTag(NullAr, Property, INDEX_NONE, (uint8*)Data, nullptr);
+	}
+
+	Value.EnsurePropertyData(ArrayIndex);
+
+	// Subtract the offset because the property will re-add it.
+	void* TargetBase = (uint8*)Value.Data.Memory - Property->GetOffset_ForInternal();
+
+	void* Target = Property->ContainerPtrToValuePtr<void>(TargetBase, ArrayIndex);
+	void* Source = Property->ContainerPtrToValuePtr<void>(Data, ArrayIndex);
+	Property->CopySingleValue(Target, Source);
+}
+
+void FPropertyBag::Remove(FName Name)
+{
+	Properties.Remove(Name);
+}
+
+void FPropertyBag::LoadPropertyByTag(const FPropertyTag& Tag, FStructuredArchiveSlot& ValueSlot, const void* Defaults)
+{
+	FArchive& UnderlyingArchive = ValueSlot.GetUnderlyingArchive();
+
+	FPropertyBag::FValue& Value = Properties.FindOrAdd(Tag.Name);
+
+	const bool bPropertyChanged =
+		(Value.Tag.Prop != Tag.Prop && Value.Tag.Prop && Tag.Prop) ||
+		Value.Tag.Type != Tag.Type ||
+		Value.Tag.Name != Tag.Name ||
+		Value.Tag.StructName != Tag.StructName ||
+		Value.Tag.EnumName != Tag.EnumName ||
+		Value.Tag.InnerType != Tag.InnerType ||
+		Value.Tag.ValueType != Tag.ValueType ||
+		Value.Tag.StructGuid != Tag.StructGuid ||
+		Value.Tag.PropertyGuid != Tag.PropertyGuid;
+
+	if (bPropertyChanged)
+	{
+		Value.Destroy();
+		Value.Tag = Tag;
+		Value.Tag.Size = 0;
+		Value.Tag.ArrayIndex = INDEX_NONE;
+	}
+
+	if (FProperty* Property = Value.Tag.Prop)
+	{
+		Value.EnsurePropertyData(Tag.ArrayIndex);
+
+		// Subtract the offset because the property will re-add it.
+		void* TargetBase = (uint8*)Value.Data.Memory - Property->GetOffset_ForInternal();
+
+		void* Target = Property->ContainerPtrToValuePtr<void>(TargetBase, Tag.ArrayIndex);
+		Tag.SerializeTaggedProperty(ValueSlot, Property, (uint8*)Target, (const uint8*)Defaults);
+
+		return;
+	}
+
+	const FValue::FData Data = {FMemory::Malloc(Tag.Size), Tag.Size};
+	UnderlyingArchive.Serialize(Data.Memory, Data.Size);
+
+	if (LIKELY(Tag.ArrayIndex == 0))
+	{
+		Value.Data = Data;
+	}
+	else
+	{
+		Value.EnsureArrayData(Tag.ArrayIndex);
+
+		FBitReference AllocationFlag = (*Value.ArrayAllocationFlags)[Tag.ArrayIndex];
+		if (UNLIKELY(AllocationFlag))
+		{
+			FMemory::Free((*Value.ArrayData)[Tag.ArrayIndex].Memory);
+		}
+		else
+		{
+			AllocationFlag = true;
+		}
+		Value.ArrayData->Add(Tag.ArrayIndex, Data);
+	}
+}
+
+void FPropertyBag::AssignElement(FPropertyBagElement& Element, const FName& Name, const FValue& Value)
+{
+	FProperty* Property = Value.Tag.Prop;
+	Element.Name = Name;
+	Element.Property = Property;
+
+	if (Property)
+	{
+		Element.ArrayIndex = 0;
+		Element.ArrayDim = Property->ArrayDim;
+		Element.Value = (uint8*)Value.Data.Memory;
+	}
+	else if (Value.ArrayData)
+	{
+		Element.ArrayIndex = Value.ArrayAllocationFlags->Find(/*bValue*/ true);
+		Element.ArrayDim = Value.ArrayAllocationFlags->Num();
+		Element.Value = (*Value.ArrayData)[Element.ArrayIndex].Memory;
+	}
+	else
+	{
+		Element.ArrayIndex = 0;
+		Element.ArrayDim = 1;
+		Element.Value = Value.Data.Memory;
+	}
+}
+
+bool FPropertyBag::AdvanceElement(FPropertyBagElement& Element, const FValue& Value)
+{
+	if (const FProperty* Property = Element.Property)
+	{
+		Element.Value = (uint8*)Element.Value + Property->ElementSize;
+		return true;
+	}
+
+	Element.ArrayIndex = Value.ArrayAllocationFlags->FindFrom(/*bValue*/ true, Element.ArrayIndex);
+	if (Element.ArrayIndex >= 0)
+	{
+		Element.Value = (*Value.ArrayData)[Element.ArrayIndex].Memory;
+		return true;
+	}
+
+	return false;
+}
+
+///////////////////////////////////////////////////////////////////////////////////////////////////////////////////////
+
+FPropertyBag::FValue::~FValue()
+{
+	Destroy();
+}
+
+void FPropertyBag::FValue::Destroy()
+{
+	const auto DestroyValue = [Property = Tag.Prop](const FData& D)
+	{
+		if (Property)
+		{
+			Property->DestroyAndFreeValue(D.Memory);
+		}
+		else
+		{
+			FMemory::Free(D.Memory);
+		}
+	};
+
+	if (Data.Memory)
+	{
+		DestroyValue(Data);
+		Data = {};
+	}
+
+	if (ArrayData)
+	{
+		for (TTuple<int32, FData>& IndexAndData : *ArrayData)
+		{
+			DestroyValue(IndexAndData.Value);
+		}
+		delete ArrayData;
+		ArrayData = nullptr;
+	}
+
+	delete ArrayAllocationFlags;
+	ArrayAllocationFlags = nullptr;
+}
+
+void FPropertyBag::FValue::EnsurePropertyData(int32 ArrayIndex)
+{
+	const FProperty* const Property = Tag.Prop;
+	const int32 ArrayDim = Property->ArrayDim;
+
+	if (!Data.Memory)
+	{
+		Data.Memory = Property->AllocateAndInitializeValue();
+		Data.Size = Property->GetSize();
+
+		if (ArrayDim > 1)
+		{
+			ArrayAllocationFlags = new TBitArray<>(/*bValue*/ false, ArrayDim);
+		}
+	}
+
+	if (ArrayDim > 1)
+	{
+		(*ArrayAllocationFlags)[ArrayIndex] = true;
+	}
+}
+
+void FPropertyBag::FValue::EnsureArrayData(int32 ArrayIndex)
+{
+	if (ArrayAllocationFlags)
+	{
+		ArrayAllocationFlags->PadToNum(ArrayIndex + 1, /*bValue*/ false);
+		return;
+	}
+
+	ArrayData = new TMap<int32, FValue::FData>();
+	ArrayAllocationFlags = new TBitArray<>(/*bValue*/ false, ArrayIndex + 1);
+
+	// Move the scalar value into array storage as the first element.
+	if (Data.Memory)
+	{
+		(*ArrayAllocationFlags)[0] = true;
+		ArrayData->Add(0, Data);
+		Data = {};
+	}
+}
+
+} // UE
diff --git a/Engine/Source/Runtime/CoreUObject/Public/UObject/PropertyBag.h b/Engine/Source/Runtime/CoreUObject/Public/UObject/PropertyBag.h
new file mode 100644
--- /dev/null
+++ b/Engine/Source/Runtime/CoreUObject/Public/UObject/PropertyBag.h
@@ -0,0 +1,197 @@
+// Copyright Epic Games, Inc. All Rights Reserved.
+
+#pragma once
+
+#include "CoreTypes.h"
+#include "Containers/Map.h"
+#include "UObject/PropertyTag.h"
+
+#define UE_API COREUOBJECT_API
+
+enum EPropertyFlags : uint64;
+
+namespace UE { class FPropertyBagElement; }
+
+namespace UE
+{
+
+/**
+ * A property bag element is a single value from the bag.
+ *
+ * A fixed-size array consists of multiple elements. These are not guaranteed to be contiguous in
+ * memory, and iteration of the property bag skips elements that do not have a value.
+ */
+class FPropertyBagElement
+{
+public:
+	/** Name, which is always valid even when the property is null. */
+	FName Name;
+
+	/** Property, which can be null if the property type is not available. */
+	const FProperty* Property = nullptr;
+
+	/** Pointer to the start of the value. This is NOT a container pointer. */
+	void* Value = nullptr;
+
+	/** Array index for fixed-size arrays. Always 0 for scalar values. */
+	int32 ArrayIndex = 0;
+
+	/** Array dimension for fixed-size arrays. Always 1 for scalar values. */
+	int32 ArrayDim = 0;
+};
+
+/**
+ * A property bag is an unordered collection of properties and their values.
+ */
+class FPropertyBag
+{
+	struct FValue
+	{
+		~FValue();
+
+		void Destroy();
+
+		void EnsurePropertyData(int32 ArrayIndex);
+		void EnsureArrayData(int32 ArrayIndex);
+
+		struct FData
+		{
+			void* Memory = nullptr;
+			int32 Size = 0;
+		};
+
+		FPropertyTag Tag;
+
+		// A value is stored in one of several ways based on what we know of the property.
+		//
+		// A scalar is stored in Data.
+		// An array is stored in Data if it is of known dimension.
+		// An array is stored in Data if it is of unknown dimension and only index 0 is known.
+		// An array is stored in ArrayData if it is of unknown dimension.
+		//
+		// Arrays use ArrayAllocationFlags to track which values are known.
+		//
+		// A value is stored in its native form if it has a known property or a property can be created dynamically.
+		// A value is otherwise stored in its serialized form.
+
+		FData Data;
+		TMap<int32, FData>* ArrayData = nullptr;
+		TBitArray<>* ArrayAllocationFlags = nullptr;
+	};
+
+public:
+	UE_API FPropertyBag();
+	UE_API ~FPropertyBag();
+
+	UE_API FPropertyBag(FPropertyBag&&);
+	UE_API FPropertyBag& operator=(FPropertyBag&&);
+
+	FPropertyBag(const FPropertyBag&) = delete;
+	FPropertyBag& operator=(const FPropertyBag&) = delete;
+
+	/** True if the bag contains no properties. */
+	inline bool IsEmpty() const { return Properties.IsEmpty(); }
+
+	/**
+	 * Remove every property from the property bag.
+	 */
+	UE_API void Empty();
+
+	/**
+	 * Copy every non-skipped property of the struct into the property bag.
+	 *
+	 * Any property in the bag with the same name will be overwritten.
+	 */
+	UE_API void Copy(const UStruct* Struct, void* Data, EPropertyFlags SkipFlags);
+
+	/**
+	 * Add the property to the bag.
+	 *
+	 * Any property in the bag with the same name will be overwritten.
+	 */
+	UE_API void Add(FProperty* Property, void* Data, int32 ArrayIndex = 0);
+
+	/**
+	 * Remove a property from the bag by name.
+	 */
+	UE_API void Remove(FName Name);
+
+	/**
+	 * Load a property from the slot into the bag based on the tag.
+	 *
+	 * If the property has a known type then it will be deserialized and accessible.
+	 * If not, the serialized data will be retained to avoid losing the unrecognized property.
+	 */
+	UE_API void LoadPropertyByTag(const FPropertyTag& Tag, FStructuredArchiveSlot& ValueSlot, const void* Defaults = nullptr);
+
+	template <bool bConst, bool bRangedFor>
+	class TBaseIterator
+	{
+	public:
+		using FPropertyIterator = std::conditional_t<bConst,
+			std::conditional_t<bRangedFor, typename TMap<FName, FValue>::TRangedForConstIterator, typename TMap<FName, FValue>::TConstIterator>,
+			std::conditional_t<bRangedFor, typename TMap<FName, FValue>::TRangedForIterator, typename TMap<FName, FValue>::TIterator>>;
+
+		inline explicit TBaseIterator(const FPropertyIterator& InPropertyIt)
+			: PropertyIt(InPropertyIt)
+		{
+			if (PropertyIt)
+			{
+				AssignElement(Element, PropertyIt.Key(), PropertyIt.Value());
+			}
+		}
+
+		inline TBaseIterator& operator++()
+		{
+			if (UNLIKELY(++Element.ArrayIndex < Element.ArrayDim && AdvanceElement(Element, PropertyIt.Value())))
+			{
+			}
+			else if (++PropertyIt)
+			{
+				AssignElement(Element, PropertyIt.Key(), PropertyIt.Value());
+			}
+			else
+			{
+				Element = {};
+			}
+			return *this;
+		}
+
+		inline explicit operator bool() const { return !!Element.ArrayDim; }
+
+		inline bool operator==(const TBaseIterator& Rhs) const { return Element.Value == Rhs.Element.Value; }
+		inline bool operator!=(const TBaseIterator& Rhs) const { return Element.Value != Rhs.Element.Value; }
+
+		inline const FPropertyBagElement& operator*() const { return Element; }
+		inline const FPropertyBagElement* operator->() const { return &Element; }
+
+	private:
+		FPropertyIterator PropertyIt;
+		FPropertyBagElement Element;
+	};
+
+	using TIterator = TBaseIterator</*bConst*/ false, /*bRangedFor*/ false>;
+	using TConstIterator = TBaseIterator</*bConst*/ true, /*bRangedFor*/ false>;
+
+	inline TIterator CreateIterator() { return TIterator(Properties.CreateIterator()); }
+	inline TConstIterator CreateConstIterator() const { return TConstIterator(Properties.CreateConstIterator()); }
+
+private:
+	using TRangedForIterator = TBaseIterator</*bConst*/ false, /*bRangedFor*/ true>;
+	using TRangedForConstIterator = TBaseIterator</*bConst*/ true, /*bRangedFor*/ true>;
+
+	friend inline TRangedForIterator begin(FPropertyBag& Bag) { return TRangedForIterator(Bag.Properties.begin()); }
+	friend inline TRangedForIterator end(FPropertyBag& Bag) { return TRangedForIterator(Bag.Properties.end()); }
+
+	friend inline TRangedForConstIterator begin(const FPropertyBag& Bag) { return TRangedForConstIterator(Bag.Properties.begin()); }
+	friend inline TRangedForConstIterator end(const FPropertyBag& Bag) { return TRangedForConstIterator(Bag.Properties.end()); }
+
+	UE_API static void AssignElement(FPropertyBagElement& Element, const FName& Name, const FValue& Value);
+	UE_API static bool AdvanceElement(FPropertyBagElement& Element, const FValue& Value);
+
+	TMap<FName, FValue> Properties;
+};
+
+} // UE
+
+#undef UE_API
