Downgrading a warning to an error to unblock development while a proper fix is investigated

#rnx

diff --git a/Engine/Source/Runtime/CoreUObject/Private/Serialization/AsyncLoading2.cpp b/Engine/Source/Runtime/CoreUObject/Private/Serialization/AsyncLoading2.cpp
--- a/Engine/Source/Runtime/CoreUObject/Private/Serialization/AsyncLoading2.cpp
+++ b/Engine/Source/Runtime/CoreUObject/Private/Serialization/AsyncLoading2.cpp
@@ -7279,7 +7279,7 @@ void FAsyncLoadingThread2::UpdateSyncLoadContext(FAsyncLoadingThreadState2& Thre
 						//  
 						// Note: Update the FLoadingTests_RecursiveLoads_FullFlushFrom_Serialize test if you edit this error message
 						//
-						UE_LOG(LogStreaming, Error, TEXT("Fully flushing package %s recursively from another package %s would lead into a deadlock. Demoting requestID %d to a partial request instead."),
+						UE_LOG(LogStreaming, Warning, TEXT("Fully flushing package %s recursively from another package %s would lead into a deadlock. Demoting requestID %d to a partial request instead."),
 							*RequestedPackage->Desc.UPackageName.ToString(),
 							*RequestingPackage->Desc.UPackageName.ToString(),
 							RequestID
