Cooker: Upgrade severity of "Content is missing from cook" message to Error, but make it configurable per project.
#rnx

diff --git a/Engine/Config/BaseEditor.ini b/Engine/Config/BaseEditor.ini
--- a/Engine/Config/BaseEditor.ini
+++ b/Engine/Config/BaseEditor.ini
@@ -330,6 +330,13 @@ MaxUsedMemory=16384
 ; Packages that are only referenced by editoronly properties will be skipped.
 SkipOnlyEditorOnly=false
 BulkDataRegistryEnabled=true
+
+; If packages are referenced but are excluded from cooking due to NeverCook settings, or a referenced object
+; is culled from the cooked version of the package due to IsEditorOnly or other settings, we will issue an error
+; message at the end of cook. Some projects have safe or spurious occurrences of this condition, so the severity of
+; the message can be downgraded using this variable.
+CookContentMissingSeverity=Error
+
 ; HybridIterative uses the TargetDomain and BuildDependencies to decide when packages need to be recooked. If enabled,
 ; every cook is iterative for packages with allowed classes (IterativeClassAllowList) unless -noiterative is passed.
 HybridIterativeEnabled=false
diff --git a/Engine/Source/Runtime/CoreUObject/Private/UObject/SavePackage/SavePackageUtilities.cpp b/Engine/Source/Runtime/CoreUObject/Private/UObject/SavePackage/SavePackageUtilities.cpp
--- a/Engine/Source/Runtime/CoreUObject/Private/UObject/SavePackage/SavePackageUtilities.cpp
+++ b/Engine/Source/Runtime/CoreUObject/Private/UObject/SavePackage/SavePackageUtilities.cpp
@@ -1474,6 +1474,10 @@ void FEDLCookChecker::Verify(const UE::SavePackageUtilities::FEDLMessageCallback
 	check(!GIsSavingPackage);
 	FEDLCookChecker Accumulator = AccumulateAndClear();
 
+	FString SeverityStr;
+	GConfig->GetString(TEXT("CookSettings"), TEXT("CookContentMissingSeverity"), SeverityStr, GEditorIni);
+	ELogVerbosity::Type MissingContentSeverity = ParseLogVerbosityFromString(SeverityStr);
+
 	if (Accumulator.bIsActive)
 	{
 		double StartTime = FPlatformTime::Seconds();
@@ -1529,7 +1533,7 @@ void FEDLCookChecker::Verify(const UE::SavePackageUtilities::FEDLMessageCallback
 					Message << TEXT("\tTarget package: ") << NodeDataOfExportPackage->Name << TEXT("\n");
 					Message << TEXT("\tReferenced object: ");
 					NodeData.AppendPathName(Accumulator, Message);
-					MessageCallback(ELogVerbosity::Warning, Message);
+					MessageCallback(MissingContentSeverity, Message);
 				}
 			}
 		}
