Moved CookedMetaData to CoreUObject

This allows it to be used by things that create class/struct/enum types (which only requires CoreUObject) without them needing to add an additional dependency on the Engine

diff --git a/Engine/Source/Runtime/CoreUObject/Private/UObject/CookedMetaData.cpp b/Engine/Source/Runtime/CoreUObject/Private/UObject/CookedMetaData.cpp
new file mode 100644
--- /dev/null
+++ b/Engine/Source/Runtime/CoreUObject/Private/UObject/CookedMetaData.cpp
@@ -0,0 +1,226 @@
+// Copyright Epic Games, Inc. All Rights Reserved.
+
+#include "UObject/CookedMetaData.h"
+#include "UObject/Package.h"
+#include "UObject/MetaData.h"
+#include "UObject/UnrealType.h"
+
+#include UE_INLINE_GENERATED_CPP_BY_NAME(CookedMetaData)
+
+namespace CookedMetaDataUtil::Internal
+{
+
+void PrepareCookedMetaDataForPurge(UObject* CookedMetaDataPtr)
+{
+	// Skip the rename for cooked packages, as IO store cannot currently handle renames
+	if (!CookedMetaDataPtr->GetPackage()->HasAnyPackageFlags(PKG_Cooked))
+	{
+		FNameBuilder BaseMetaDataName(CookedMetaDataPtr->GetFName());
+		BaseMetaDataName << TEXT("_PURGED");
+		CookedMetaDataPtr->Rename(FNameBuilder(MakeUniqueObjectName(CookedMetaDataPtr->GetOuter(), CookedMetaDataPtr->GetClass(), FName(BaseMetaDataName))).ToString(), nullptr, REN_DoNotDirty | REN_DontCreateRedirectors | REN_ForceNoResetLoaders | REN_NonTransactional);
+	}
+
+	CookedMetaDataPtr->ClearFlags(RF_Standalone | RF_Public);
+}
+
+template <typename CookedMetaDataOuterType, typename CookedMetaDataType>
+void PostLoadCookedMetaData(CookedMetaDataType* CookedMetaDataPtr)
+{
+#if WITH_EDITORONLY_DATA
+	checkf(CookedMetaDataPtr->GetPackage()->HasAnyPackageFlags(PKG_Cooked), TEXT("Cooked meta-data should only be loaded for a cooked package!"));
+
+	if (CookedMetaDataOuterType* Owner = CastChecked<CookedMetaDataOuterType>(CookedMetaDataPtr->GetOuter()))
+	{
+		Owner->ConditionalPostLoad();
+		CookedMetaDataPtr->ApplyMetaData(Owner);
+		PrepareCookedMetaDataForPurge(CookedMetaDataPtr);
+	}
+#endif
+}
+
+} // namespace CookedMetaDataUtil::Internal
+
+
+bool FObjectCookedMetaDataStore::HasMetaData() const
+{
+	return ObjectMetaData.Num() > 0;
+}
+
+void FObjectCookedMetaDataStore::CacheMetaData(const UObject* SourceObject)
+{
+	ObjectMetaData.Reset();
+
+#if WITH_EDITORONLY_DATA
+	if (UPackage* SourcePackage = SourceObject->GetPackage())
+	{
+		if (const UMetaData* SourceMetaData = SourcePackage->GetMetaData())
+		{
+			if (const TMap<FName, FString>* SourceObjectMetaData = SourceMetaData->ObjectMetaDataMap.Find(SourceObject))
+			{
+				ObjectMetaData = *SourceObjectMetaData;
+			}
+		}
+	}
+#endif
+}
+
+void FObjectCookedMetaDataStore::ApplyMetaData(UObject* TargetObject) const
+{
+#if WITH_EDITORONLY_DATA
+	if (UPackage* TargetPackage = TargetObject->GetPackage())
+	{
+		if (UMetaData* TargetMetaData = TargetPackage->GetMetaData())
+		{
+			TargetMetaData->ObjectMetaDataMap.FindOrAdd(TargetObject).Append(ObjectMetaData);
+		}
+	}
+#endif
+}
+
+
+bool FFieldCookedMetaDataStore::HasMetaData() const
+{
+	return FieldMetaData.Num() > 0;
+}
+
+void FFieldCookedMetaDataStore::CacheMetaData(const FField* SourceField)
+{
+	FieldMetaData.Reset();
+
+#if WITH_EDITORONLY_DATA
+	if (const TMap<FName, FString>* SourceFieldMetaData = SourceField->GetMetaDataMap())
+	{
+		FieldMetaData = *SourceFieldMetaData;
+	}
+#endif
+}
+
+void FFieldCookedMetaDataStore::ApplyMetaData(FField* TargetField) const
+{
+#if WITH_EDITORONLY_DATA
+	TargetField->AppendMetaData(FieldMetaData);
+#endif
+}
+
+
+bool FStructCookedMetaDataStore::HasMetaData() const
+{
+	return ObjectMetaData.HasMetaData()
+		|| PropertiesMetaData.Num() > 0;
+}
+
+void FStructCookedMetaDataStore::CacheMetaData(const UStruct* SourceStruct)
+{
+	ObjectMetaData.CacheMetaData(SourceStruct);
+
+	for (TFieldIterator<const FProperty> SourcePropertyIt(SourceStruct, EFieldIterationFlags::None); SourcePropertyIt; ++SourcePropertyIt)
+	{
+		FFieldCookedMetaDataStore SourcePropertyMetaData;
+		SourcePropertyMetaData.CacheMetaData(*SourcePropertyIt);
+
+		if (SourcePropertyMetaData.HasMetaData())
+		{
+			PropertiesMetaData.Add(SourcePropertyIt->GetFName(), MoveTemp(SourcePropertyMetaData));
+		}
+	}
+}
+
+void FStructCookedMetaDataStore::ApplyMetaData(UStruct* TargetStruct) const
+{
+	ObjectMetaData.ApplyMetaData(TargetStruct);
+
+	for (TFieldIterator<FProperty> TargetPropertyIt(TargetStruct, EFieldIterationFlags::None); TargetPropertyIt; ++TargetPropertyIt)
+	{
+		if (const FFieldCookedMetaDataStore* TargetPropertyMetaData = PropertiesMetaData.Find(TargetPropertyIt->GetFName()))
+		{
+			TargetPropertyMetaData->ApplyMetaData(*TargetPropertyIt);
+		}
+	}
+}
+
+
+void UEnumCookedMetaData::PostLoad()
+{
+	Super::PostLoad();
+	CookedMetaDataUtil::Internal::PostLoadCookedMetaData<UEnum, UEnumCookedMetaData>(this);
+}
+
+bool UEnumCookedMetaData::HasMetaData() const
+{
+	return EnumMetaData.HasMetaData();
+}
+
+void UEnumCookedMetaData::CacheMetaData(const UEnum* SourceEnum)
+{
+	EnumMetaData.CacheMetaData(SourceEnum);
+}
+
+void UEnumCookedMetaData::ApplyMetaData(UEnum* TargetEnum) const
+{
+	EnumMetaData.ApplyMetaData(TargetEnum);
+}
+
+
+
+void UStructCookedMetaData::PostLoad()
+{
+	Super::PostLoad();
+	CookedMetaDataUtil::Internal::PostLoadCookedMetaData<UScriptStruct, UStructCookedMetaData>(this);
+}
+
+bool UStructCookedMetaData::HasMetaData() const
+{
+	return StructMetaData.HasMetaData();
+}
+
+void UStructCookedMetaData::CacheMetaData(const UScriptStruct* SourceStruct)
+{
+	StructMetaData.CacheMetaData(SourceStruct);
+}
+
+void UStructCookedMetaData::ApplyMetaData(UScriptStruct* TargetStruct) const
+{
+	StructMetaData.ApplyMetaData(TargetStruct);
+}
+
+
+void UClassCookedMetaData::PostLoad()
+{
+	Super::PostLoad();
+	CookedMetaDataUtil::Internal::PostLoadCookedMetaData<UClass, UClassCookedMetaData>(this);
+}
+
+bool UClassCookedMetaData::HasMetaData() const
+{
+	return ClassMetaData.HasMetaData()
+		|| FunctionsMetaData.Num() > 0;
+}
+
+void UClassCookedMetaData::CacheMetaData(const UClass* SourceClass)
+{
+	ClassMetaData.CacheMetaData(SourceClass);
+
+	for (TFieldIterator<const UFunction> SourceFunctionIt(SourceClass, EFieldIterationFlags::None); SourceFunctionIt; ++SourceFunctionIt)
+	{
+		FStructCookedMetaDataStore SourceFunctionMetaData;
+		SourceFunctionMetaData.CacheMetaData(*SourceFunctionIt);
+
+		if (SourceFunctionMetaData.HasMetaData())
+		{
+			FunctionsMetaData.Add(SourceFunctionIt->GetFName(), MoveTemp(SourceFunctionMetaData));
+		}
+	}
+}
+
+void UClassCookedMetaData::ApplyMetaData(UClass* TargetClass) const
+{
+	ClassMetaData.ApplyMetaData(TargetClass);
+
+	for (TFieldIterator<UFunction> TargetFunctionIt(TargetClass, EFieldIterationFlags::None); TargetFunctionIt; ++TargetFunctionIt)
+	{
+		if (const FStructCookedMetaDataStore* TargetFunctionMetaData = FunctionsMetaData.Find(TargetFunctionIt->GetFName()))
+		{
+			TargetFunctionMetaData->ApplyMetaData(*TargetFunctionIt);
+		}
+	}
+}
diff --git a/Engine/Source/Runtime/CoreUObject/Public/UObject/CookedMetaData.h b/Engine/Source/Runtime/CoreUObject/Public/UObject/CookedMetaData.h
new file mode 100644
--- /dev/null
+++ b/Engine/Source/Runtime/CoreUObject/Public/UObject/CookedMetaData.h
@@ -0,0 +1,157 @@
+// Copyright Epic Games, Inc. All Rights Reserved.
+
+#pragma once
+
+#include "UObject/Object.h"
+#include "UObject/UObjectGlobals.h"
+#include "Templates/SubclassOf.h"
+#include "CookedMetaData.generated.h"
+
+/**
+ * Cooked meta-data for a UObject.
+ */
+USTRUCT()
+struct FObjectCookedMetaDataStore
+{
+public:
+	GENERATED_BODY()
+
+	COREUOBJECT_API bool HasMetaData() const;
+	COREUOBJECT_API void CacheMetaData(const UObject* SourceObject);
+	COREUOBJECT_API void ApplyMetaData(UObject* TargetObject) const;
+
+private:
+	UPROPERTY()
+	TMap<FName, FString> ObjectMetaData;
+};
+
+/**
+ * Cooked meta-data for a FField.
+ */
+USTRUCT()
+struct FFieldCookedMetaDataStore
+{
+public:
+	GENERATED_BODY()
+
+	COREUOBJECT_API bool HasMetaData() const;
+	COREUOBJECT_API void CacheMetaData(const FField* SourceField);
+	COREUOBJECT_API void ApplyMetaData(FField* TargetField) const;
+
+private:
+	UPROPERTY()
+	TMap<FName, FString> FieldMetaData;
+};
+
+/**
+ * Cooked meta-data for a UStruct, including its nested FProperty data.
+ */
+USTRUCT()
+struct FStructCookedMetaDataStore
+{
+public:
+	GENERATED_BODY()
+
+	COREUOBJECT_API bool HasMetaData() const;
+	COREUOBJECT_API void CacheMetaData(const UStruct* SourceStruct);
+	COREUOBJECT_API void ApplyMetaData(UStruct* TargetStruct) const;
+
+private:
+	UPROPERTY()
+	FObjectCookedMetaDataStore ObjectMetaData;
+
+	UPROPERTY()
+	TMap<FName, FFieldCookedMetaDataStore> PropertiesMetaData;
+};
+
+/**
+ * Cooked meta-data for a UEnum.
+ */
+UCLASS(Optional, Within=Enum, MinimalAPI)
+class UEnumCookedMetaData : public UObject
+{
+public:
+	GENERATED_BODY()
+
+	COREUOBJECT_API virtual void PostLoad() override;
+	
+	COREUOBJECT_API virtual bool HasMetaData() const;
+	COREUOBJECT_API virtual void CacheMetaData(const UEnum* SourceEnum);
+	COREUOBJECT_API virtual void ApplyMetaData(UEnum* TargetEnum) const;
+
+protected:
+	UPROPERTY()
+	FObjectCookedMetaDataStore EnumMetaData;
+};
+
+/**
+ * Cooked meta-data for a UScriptStruct, including its nested FProperty data.
+ */
+UCLASS(Optional, Within=ScriptStruct, MinimalAPI)
+class UStructCookedMetaData : public UObject
+{
+public:
+	GENERATED_BODY()
+
+	COREUOBJECT_API virtual void PostLoad() override;
+	
+	COREUOBJECT_API virtual bool HasMetaData() const;
+	COREUOBJECT_API virtual void CacheMetaData(const UScriptStruct* SourceStruct);
+	COREUOBJECT_API virtual void ApplyMetaData(UScriptStruct* TargetStruct) const;
+
+protected:
+	UPROPERTY()
+	FStructCookedMetaDataStore StructMetaData;
+};
+
+/**
+ * Cooked meta-data for a UClass, including its nested FProperty and UFunction data.
+ */
+UCLASS(Optional, Within=Class, MinimalAPI)
+class UClassCookedMetaData : public UObject
+{
+public:
+	GENERATED_BODY()
+
+	COREUOBJECT_API virtual void PostLoad() override;
+	
+	COREUOBJECT_API virtual bool HasMetaData() const;
+	COREUOBJECT_API virtual void CacheMetaData(const UClass* SourceClass);
+	COREUOBJECT_API virtual void ApplyMetaData(UClass* TargetClass) const;
+
+protected:
+	UPROPERTY()
+	FStructCookedMetaDataStore ClassMetaData;
+
+	UPROPERTY()
+	TMap<FName, FStructCookedMetaDataStore> FunctionsMetaData;
+};
+
+namespace CookedMetaDataUtil
+{
+
+namespace Internal
+{
+COREUOBJECT_API void PrepareCookedMetaDataForPurge(UObject* CookedMetaDataPtr);
+}
+
+template <typename CookedMetaDataType>
+CookedMetaDataType* NewCookedMetaData(UObject* Outer, FName Name, TSubclassOf<CookedMetaDataType> Class = CookedMetaDataType::StaticClass())
+{
+	return NewObject<CookedMetaDataType>(Outer, Class, Name, RF_Standalone | RF_Public);
+}
+
+template <typename CookedMetaDataType>
+CookedMetaDataType* FindCookedMetaData(UObject* Outer, const TCHAR* Name)
+{
+	return FindObject<CookedMetaDataType>(Outer, Name);
+}
+
+template <typename CookedMetaDataType, typename CookedMetaDataPtrType>
+void PurgeCookedMetaData(CookedMetaDataPtrType& CookedMetaDataPtr)
+{
+	Internal::PrepareCookedMetaDataForPurge(CookedMetaDataPtr);
+	CookedMetaDataPtr = nullptr;
+}
+
+}
diff --git a/Engine/Source/Runtime/Engine/Private/BlueprintGeneratedClass.cpp b/Engine/Source/Runtime/Engine/Private/BlueprintGeneratedClass.cpp
--- a/Engine/Source/Runtime/Engine/Private/BlueprintGeneratedClass.cpp
+++ b/Engine/Source/Runtime/Engine/Private/BlueprintGeneratedClass.cpp
@@ -11,7 +11,6 @@
 #include "UObject/CoreRedirects.h"
 #include "UObject/ObjectSaveContext.h"
 #include "Serialization/ObjectWriter.h"
-#include "CookedMetaData.h"
 #include "Curves/CurveFloat.h"
 #include "Engine/DynamicBlueprintBinding.h"
 #include "Components/TimelineComponent.h"
@@ -38,6 +37,7 @@
 #include "BlueprintCompilationManager.h"
 #include "Engine/LevelScriptBlueprint.h"
 #include "WorldPartition/WorldPartitionActorDescUtils.h"
+#include "UObject/CookedMetaData.h"
 #include "UObject/UObjectGlobals.h"
 extern UNREALED_API class UEditorEngine* GEditor;
 #else
diff --git a/Engine/Source/Runtime/Engine/Private/CookedMetaData.cpp b/Engine/Source/Runtime/Engine/Private/CookedMetaData.cpp
deleted file mode 100644
diff --git a/Engine/Source/Runtime/Engine/Private/UserDefinedEnum.cpp b/Engine/Source/Runtime/Engine/Private/UserDefinedEnum.cpp
--- a/Engine/Source/Runtime/Engine/Private/UserDefinedEnum.cpp
+++ b/Engine/Source/Runtime/Engine/Private/UserDefinedEnum.cpp
@@ -4,13 +4,13 @@
 #include "Templates/SubclassOf.h"
 #include "UObject/EditorObjectVersion.h"
 #include "UObject/ObjectSaveContext.h"
-#include "CookedMetaData.h"
 #include "UObject/Package.h"
 
 #include UE_INLINE_GENERATED_CPP_BY_NAME(UserDefinedEnum)
 
 #if WITH_EDITOR
 #include "Kismet2/EnumEditorUtils.h"
+#include "UObject/CookedMetaData.h"
 #include "UObject/MetaData.h"
 #include "UObject/Package.h"
 #endif	// WITH_EDITOR
diff --git a/Engine/Source/Runtime/Engine/Private/UserDefinedStruct.cpp b/Engine/Source/Runtime/Engine/Private/UserDefinedStruct.cpp
--- a/Engine/Source/Runtime/Engine/Private/UserDefinedStruct.cpp
+++ b/Engine/Source/Runtime/Engine/Private/UserDefinedStruct.cpp
@@ -7,12 +7,12 @@
 #include "UObject/UnrealType.h"
 #include "UObject/LinkerLoad.h"
 #include "UObject/ObjectSaveContext.h"
-#include "CookedMetaData.h"
 #include "UObject/FrameworkObjectVersion.h"
 #include "Misc/PackageName.h"
 #include "Blueprint/BlueprintSupport.h"
 
 #if WITH_EDITOR
+#include "UObject/CookedMetaData.h"
 #include "UserDefinedStructure/UserDefinedStructEditorData.h"
 #endif //WITH_EDITOR
 
diff --git a/Engine/Source/Runtime/Engine/Public/CookedMetaData.h b/Engine/Source/Runtime/Engine/Public/CookedMetaData.h
deleted file mode 100644
