#UE Do not warn about duplicate lazy object pointer guids found in editor only level instance maps, since they do not affect the game. This is just another case of legitimately duplicating a level in memory, much like PIE

diff --git a/Engine/Source/Runtime/CoreUObject/Private/UObject/LazyObjectPtr.cpp b/Engine/Source/Runtime/CoreUObject/Private/UObject/LazyObjectPtr.cpp
--- a/Engine/Source/Runtime/CoreUObject/Private/UObject/LazyObjectPtr.cpp
+++ b/Engine/Source/Runtime/CoreUObject/Private/UObject/LazyObjectPtr.cpp
@@ -158,10 +158,11 @@ void FLazyObjectPtr::PossiblySerializeObjectGuid(UObject *Object, FStructuredArc
 						{
 							if (!bReassigning)
 							{
-								// Always warn for non-map packages, skip map packages in PIE or game
+								// Always warn for non-map packages, skip map packages in PIE, game, or when editor only. Editor-only maps are usually instances, and do not affect the game.
 								const bool bInGame = FApp::IsGame() || Package->HasAnyPackageFlags(PKG_PlayInEditor);
+								const bool bIsMapOrMapData = Package->ContainsMap() || Package->HasAnyPackageFlags(PKG_ContainsMapData);
 
-								UE_CLOG(!Package->ContainsMap() || !bInGame, LogUObjectGlobals, Warning,
+								UE_CLOG(!bIsMapOrMapData || (!bInGame && !Package->HasAnyPackageFlags(PKG_EditorOnly)), LogUObjectGlobals, Warning,
 									TEXT("Guid referenced by %s is already used by %s, which should never happen in the editor but could happen at runtime with duplicate level loading or PIE"),
 									*Object->GetFullName(), *OtherObject->GetFullName());
 							}
