Core: Make UScriptStruct::ImportText const-correct

diff --git a/Engine/Source/Runtime/CoreUObject/Private/UObject/Class.cpp b/Engine/Source/Runtime/CoreUObject/Private/UObject/Class.cpp
--- a/Engine/Source/Runtime/CoreUObject/Private/UObject/Class.cpp
+++ b/Engine/Source/Runtime/CoreUObject/Private/UObject/Class.cpp
@@ -2837,12 +2837,12 @@ void UScriptStruct::SerializeItem(FStructuredArchive::FSlot Slot, void* Value, v
 	}
 }
 
-const TCHAR* UScriptStruct::ImportText(const TCHAR* InBuffer, void* Value, UObject* OwnerObject, int32 PortFlags, FOutputDevice* ErrorText, const FString& StructName, bool bAllowNativeOverride)
+const TCHAR* UScriptStruct::ImportText(const TCHAR* InBuffer, void* Value, UObject* OwnerObject, int32 PortFlags, FOutputDevice* ErrorText, const FString& StructName, bool bAllowNativeOverride) const
 {
 	return ImportText(InBuffer, Value, OwnerObject, PortFlags, ErrorText, [&StructName](){return StructName;}, bAllowNativeOverride);
 }
 
-const TCHAR* UScriptStruct::ImportText(const TCHAR* InBuffer, void* Value, UObject* OwnerObject, int32 PortFlags, FOutputDevice* ErrorText, const TFunctionRef<FString()>& StructNameGetter, bool bAllowNativeOverride)
+const TCHAR* UScriptStruct::ImportText(const TCHAR* InBuffer, void* Value, UObject* OwnerObject, int32 PortFlags, FOutputDevice* ErrorText, const TFunctionRef<FString()>& StructNameGetter, bool bAllowNativeOverride) const
 {
 	FOutputDeviceNull NullErrorText;
 	if (!ErrorText)
diff --git a/Engine/Source/Runtime/CoreUObject/Private/UObject/Property.cpp b/Engine/Source/Runtime/CoreUObject/Private/UObject/Property.cpp
--- a/Engine/Source/Runtime/CoreUObject/Private/UObject/Property.cpp
+++ b/Engine/Source/Runtime/CoreUObject/Private/UObject/Property.cpp
@@ -1523,7 +1523,7 @@ void FProperty::PerformOperationWithGetter(void* OutContainer, const void* Direc
  * @param	Warn			the output device to send errors/warnings to
  * @return	the array index for this defaultproperties line.  INDEX_NONE if this line doesn't contains an array specifier, or 0 if there was an error parsing the specifier.
  */
-static const int32 ReadArrayIndex(UStruct* ObjectStruct, const TCHAR*& Str, FOutputDevice* Warn)
+static const int32 ReadArrayIndex(const UStruct* ObjectStruct, const TCHAR*& Str, FOutputDevice* Warn)
 {
 	const TCHAR* Start = Str;
 	int32 Index = INDEX_NONE;
@@ -1616,7 +1616,7 @@ static bool IsPropertyValueSpecified( const TCHAR* Buffer )
 	return Buffer && *Buffer && *Buffer != TCHAR(',') && *Buffer != TCHAR(')');
 }
 
-const TCHAR* FProperty::ImportSingleProperty( const TCHAR* Str, void* DestData, UStruct* ObjectStruct, UObject* SubobjectOuter, int32 PortFlags,
+const TCHAR* FProperty::ImportSingleProperty( const TCHAR* Str, void* DestData, const UStruct* ObjectStruct, UObject* SubobjectOuter, int32 PortFlags,
 											FOutputDevice* Warn, TArray<FDefinedProperty>& DefinedProperties )
 {
 	check(ObjectStruct);
@@ -2024,7 +2024,7 @@ const TCHAR* FProperty::ImportSingleProperty( const TCHAR* Str, void* DestData,
 	return Str;
 }
 
-FName FProperty::FindRedirectedPropertyName(UStruct* ObjectStruct, FName OldName)
+FName FProperty::FindRedirectedPropertyName(const UStruct* ObjectStruct, FName OldName)
 {
 	DECLARE_SCOPE_CYCLE_COUNTER(TEXT("FProperty::FindRedirectedPropertyName"), STAT_LinkerLoad_FindRedirectedPropertyName, STATGROUP_LoadTimeVerbose);
 
diff --git a/Engine/Source/Runtime/CoreUObject/Public/UObject/Class.h b/Engine/Source/Runtime/CoreUObject/Public/UObject/Class.h
--- a/Engine/Source/Runtime/CoreUObject/Public/UObject/Class.h
+++ b/Engine/Source/Runtime/CoreUObject/Public/UObject/Class.h
@@ -1689,7 +1689,7 @@ public:
 	 * @param	bAllowNativeOverride If true, will try to run native version of export text on the struct
 	 * @return Buffer after parsing has succeeded, or NULL on failure
 	 */
-	COREUOBJECT_API const TCHAR* ImportText(const TCHAR* Buffer, void* Value, UObject* OwnerObject, int32 PortFlags, FOutputDevice* ErrorText, const FString& StructName, bool bAllowNativeOverride = true);
+	COREUOBJECT_API const TCHAR* ImportText(const TCHAR* Buffer, void* Value, UObject* OwnerObject, int32 PortFlags, FOutputDevice* ErrorText, const FString& StructName, bool bAllowNativeOverride = true) const;
 
 	/**
 	 * Sets value of script struct based on imported string
@@ -1703,7 +1703,7 @@ public:
 	 * @param	bAllowNativeOverride If true, will try to run native version of export text on the struct
 	 * @return Buffer after parsing has succeeded, or NULL on failure
 	 */
-	COREUOBJECT_API const TCHAR* ImportText(const TCHAR* Buffer, void* Value, UObject* OwnerObject, int32 PortFlags, FOutputDevice* ErrorText, const TFunctionRef<FString()>& StructNameGetter, bool bAllowNativeOverride = true);
+	COREUOBJECT_API const TCHAR* ImportText(const TCHAR* Buffer, void* Value, UObject* OwnerObject, int32 PortFlags, FOutputDevice* ErrorText, const TFunctionRef<FString()>& StructNameGetter, bool bAllowNativeOverride = true) const;
 
 	/**
 	 * Compare two script structs
diff --git a/Engine/Source/Runtime/CoreUObject/Public/UObject/UnrealType.h b/Engine/Source/Runtime/CoreUObject/Public/UObject/UnrealType.h
--- a/Engine/Source/Runtime/CoreUObject/Public/UObject/UnrealType.h
+++ b/Engine/Source/Runtime/CoreUObject/Public/UObject/UnrealType.h
@@ -221,11 +221,11 @@ public:
 	 * @param DefinedProperties (out)	list of properties/indices that have been parsed by previous calls, so duplicate definitions cause an error
 	 * @return pointer to remaining text in the stream (even on failure, but on failure it may not be advanced past the entire key/value pair)
 	 */
-	static COREUOBJECT_API const TCHAR* ImportSingleProperty( const TCHAR* Str, void* DestData, class UStruct* ObjectStruct, UObject* SubobjectOuter, int32 PortFlags,
+	static COREUOBJECT_API const TCHAR* ImportSingleProperty( const TCHAR* Str, void* DestData, const UStruct* ObjectStruct, UObject* SubobjectOuter, int32 PortFlags,
 											FOutputDevice* Warn, TArray<struct FDefinedProperty>& DefinedProperties );
 
 	/** Gets a redirected property name, will return NAME_None if no redirection was found */
-	static COREUOBJECT_API FName FindRedirectedPropertyName(UStruct* ObjectStruct, FName OldName);
+	static COREUOBJECT_API FName FindRedirectedPropertyName(const UStruct* ObjectStruct, FName OldName);
 
 	/**
 	 * Returns the C++ name of the property, including the _DEPRECATED suffix if the
