Make MakeUniqueObjectNameReusingNumber transactionally safe.

diff --git a/Engine/Source/Runtime/CoreUObject/Private/UObject/UObjectGlobals.cpp b/Engine/Source/Runtime/CoreUObject/Private/UObject/UObjectGlobals.cpp
--- a/Engine/Source/Runtime/CoreUObject/Private/UObject/UObjectGlobals.cpp
+++ b/Engine/Source/Runtime/CoreUObject/Private/UObject/UObjectGlobals.cpp
@@ -2728,19 +2728,24 @@ namespace NameReuse
 
 		LLM_SCOPE(ELLMTag::UObject);
 
-		FNameEntryId BaseId = BaseName.GetDisplayIndex();
-		FName ReturnName = GRecentNameCache.Find(Parent, BaseId, BaseName);
+		FName ReturnName;
 
-		if (ReturnName.IsNone())
+		UE_AUTORTFM_OPEN(
 		{
-			ReturnName = GNameRangeCache.Find(Parent, BaseId, BaseName);
-		}
+			FNameEntryId BaseId = BaseName.GetDisplayIndex();
+			ReturnName = GRecentNameCache.Find(Parent, BaseId, BaseName);
 
-		if (ReturnName.IsNone())
-		{
-			// Store this name for reuse 
-			GRecentNameCache.Store(Parent, BaseId, ReturnName);
-		}
+			if (ReturnName.IsNone())
+			{
+				ReturnName = GNameRangeCache.Find(Parent, BaseId, BaseName);
+			}
+
+			if (ReturnName.IsNone())
+			{
+				// Store this name for reuse 
+				GRecentNameCache.Store(Parent, BaseId, ReturnName);
+			}
+		});
 
 		return ReturnName;
 	}
