Fix warning "Failed to find object 'Object None.None'" when deserializing FObjectAndNameAsStringProxyArchive.

Should also speed up loading a buffer with a lot of NULL objects considerably (see PR #6875).

#PR 6875

diff --git a/Engine/Source/Runtime/CoreUObject/Private/Serialization/ArchiveUObject.cpp b/Engine/Source/Runtime/CoreUObject/Private/Serialization/ArchiveUObject.cpp
--- a/Engine/Source/Runtime/CoreUObject/Private/Serialization/ArchiveUObject.cpp
+++ b/Engine/Source/Runtime/CoreUObject/Private/Serialization/ArchiveUObject.cpp
@@ -163,6 +163,14 @@ FArchive& FObjectAndNameAsStringProxyArchive::operator<<(UObject*& Obj)
 		// load the path name to the object
 		FString LoadedString;
 		InnerArchive << LoadedString;
+
+		// if it's empty, let's exit early
+		if (LoadedString.IsEmpty())
+		{
+			Obj = nullptr;
+			return *this;
+		}
+
 		// look up the object by fully qualified pathname
 		Obj = FindObject<UObject>(nullptr, *LoadedString, false);
 		// If we couldn't find it, and we want to load it, do that
@@ -171,12 +179,18 @@ FArchive& FObjectAndNameAsStringProxyArchive::operator<<(UObject*& Obj)
 			Obj = LoadObject<UObject>(nullptr, *LoadedString);
 		}
 	}
-	else
+	else if (Obj)
 	{
 		// save out the fully qualified object name
 		FString SavedString(Obj->GetPathName());
 		InnerArchive << SavedString;
 	}
+	else
+	{
+		// for null pointer, output empty string
+		FString EmptyString;
+		InnerArchive << EmptyString;
+	}
 	return *this;
 }
 
