Add comment about using package versions during saving.

diff --git a/Engine/Source/Runtime/CoreUObject/Private/UObject/ObjectResource.cpp b/Engine/Source/Runtime/CoreUObject/Private/UObject/ObjectResource.cpp
--- a/Engine/Source/Runtime/CoreUObject/Private/UObject/ObjectResource.cpp
+++ b/Engine/Source/Runtime/CoreUObject/Private/UObject/ObjectResource.cpp
@@ -207,6 +207,7 @@ void operator<<(FStructuredArchive::FSlot Slot, FObjectExport& E)
 	
 	if (!BaseArchive.UseUnversionedPropertySerialization() && BaseArchive.UEVer() >= EUnrealEngineObjectUE5Version::SCRIPT_SERIALIZATION_OFFSET)
 	{
+		// Note: this path may be taken when saving as well (for fast package duplication)
 		Record << SA_VALUE(TEXT("ScriptSerializationStartOffset"), E.ScriptSerializationStartOffset);
 		Record << SA_VALUE(TEXT("ScriptSerializationEndOffset"), E.ScriptSerializationEndOffset);
 	}
