TMaybeObjectPtr: fixing sometimes-incorrect handling of pointers returned from Cast<...>

diff --git a/Engine/Source/Runtime/CoreUObject/Public/UObject/ObjectPtr.h b/Engine/Source/Runtime/CoreUObject/Public/UObject/ObjectPtr.h
--- a/Engine/Source/Runtime/CoreUObject/Public/UObject/ObjectPtr.h
+++ b/Engine/Source/Runtime/CoreUObject/Public/UObject/ObjectPtr.h
@@ -1426,97 +1426,79 @@ namespace UE::Core::Private // private facilities; not for direct use
 		ViewType* View;
 	};
 
-	// nb: TMaybeObjectPtr class exists as a temporary compatibility shim with existing code.
-	// do not use in new code
+	// nb: TMaybeObjectPtr class exists as a temporary compatibility shim with existing code;
+	//     do not use in new code. it allows code that holds pointers to abstract classes that
+	//     might point to instances that also subclass UObject to properly interact with
+	//     garbage collection.
 	template <typename T>
-	class TMaybeObjectPtr
+	class TMaybeObjectPtr final
 	{
 		static_assert(!std::is_convertible_v<T, const UObjectBase*>, "TMaybeObjectPtr's type argument shouldn't be a subclass of UObjectBase");
 
 	public:
 		TMaybeObjectPtr() = default;
-		
-		explicit TMaybeObjectPtr(nullptr_t)
-			: Ptr{TObjectPtr<UObject>{}}, bIsObjectPtr{true}
-		{
-		}
-													
-		explicit TMaybeObjectPtr(UObject* X)
-			: Ptr{TObjectPtr<UObject>{X}}, bIsObjectPtr{true}
-		{
-		}
-		
-		explicit TMaybeObjectPtr(T* X)
+
+		explicit TMaybeObjectPtr(T* P)
+			: Ptr{P}
 		{
-			*this = X;
+			ConditionallyMarkAsReachable();
 		}
 
 		TMaybeObjectPtr(const TMaybeObjectPtr& Other)
+			: TMaybeObjectPtr{Other.Ptr}
 		{
-			*this = Other;
 		}
 
 		TMaybeObjectPtr(TMaybeObjectPtr&& Other)
+			: TMaybeObjectPtr{Other.Ptr}
 		{
-			*this = Other;
+		}
+
+		TMaybeObjectPtr& operator=(const TMaybeObjectPtr& Other)
+		{
+			return *this = Other.Ptr;
 		}
 
 		TMaybeObjectPtr& operator=(TMaybeObjectPtr&& Other)
 		{
 			return *this = Other;
 		}
-		
-		TMaybeObjectPtr& operator=(const TMaybeObjectPtr& Other)
-		{
-			if (Other.bIsObjectPtr)
-			{
-				Ptr.ObjectPtr = Other.Ptr.ObjectPtr;
-			}
-			else 
-			{
-				Ptr.NotObjectPtr = Other.Ptr.NotObjectPtr;
-			}
-			bIsObjectPtr = Other.bIsObjectPtr;
-			return *this;
-		}		
-		
-		TMaybeObjectPtr& operator=(T* MaybeObjectPtr)
+
+		TMaybeObjectPtr& operator=(T* P)
 		{
-			if (UObject* P = Cast<UObject>(const_cast<std::remove_cv_t<T>*>(MaybeObjectPtr)))
-			{
-				bIsObjectPtr = true;
-				Ptr.ObjectPtr = P;
-			}
-			else
-			{
-				bIsObjectPtr = false;
-				Ptr.NotObjectPtr = MaybeObjectPtr;
-			}
+			Ptr = P;
+			ConditionallyMarkAsReachable();
 			return *this;
 		}
-
+		
 		operator T*() const
 		{
-			return bIsObjectPtr ? nullptr : const_cast<std::remove_cv_t<T>*>(Ptr.NotObjectPtr);
+			return Ptr;
 		}
-
-		TObjectPtr<UObject>* AsNonNullObjectPtr() const
+		
+		void AddReferencedObject(FReferenceCollector& Collector, UObject* ReferencingObject)
 		{
-			if (!bIsObjectPtr || !Ptr.ObjectPtr)
+			if (const UObject* Obj = Cast<UObject>(Ptr))
 			{
-				return nullptr;
+				TObjectPtr<UObject> ObjectPtr{const_cast<UObject*>(Obj)};
+				Collector.AddReferencedObject(ObjectPtr, ReferencingObject);
+				if (!ObjectPtr)
+				{
+					Ptr = nullptr;
+				}
 			}
-			return const_cast<TObjectPtr<UObject>*>(&Ptr.ObjectPtr);
 		}
 
 	private:
-		union
+		FORCEINLINE void ConditionallyMarkAsReachable() const
 		{
-			TObjectPtr<UObject> ObjectPtr{};
-			T* NotObjectPtr;
-		} Ptr;
-		bool bIsObjectPtr{true};
-		static_assert(sizeof(Ptr) == sizeof(void*));
+			if (const UObject* Obj = Cast<UObject>(Ptr); Obj && UE::GC::Private::GIsIncrementalReachabilityPending)
+			{
+				Obj->MarkAsReachable();
+			}
+		}
+
+		T* Ptr{};
 	};	
 }
 
