[LinkerLoad] Fix missing attempt to reresolve a redirector when it is encountered multiple times recursively, potentially causing missing objects for complex case

#rnx

diff --git a/Engine/Source/Runtime/CoreUObject/Private/UObject/LinkerLoad.cpp b/Engine/Source/Runtime/CoreUObject/Private/UObject/LinkerLoad.cpp
--- a/Engine/Source/Runtime/CoreUObject/Private/UObject/LinkerLoad.cpp
+++ b/Engine/Source/Runtime/CoreUObject/Private/UObject/LinkerLoad.cpp
@@ -5576,6 +5576,16 @@ UObject* FLinkerLoad::CreateImport( int32 Index )
 					UObjectRedirector* Redirector = dynamic_cast<UObjectRedirector*>(Import.XObject);
 					if( Redirector )
 					{
+						// It may happen that the redirector is already being deserialized on the stack (i.e RF_LoadCompleted isn't set)
+						// but RF_NeedLoad has been removed already. We need to reresolve in that case right away
+						// otherwise the DestinationObject wouldn't be set until we unwind the stack and finish
+						// the deserialization, which may be too late.
+						if (!Redirector->HasAnyFlags(RF_NeedLoad|RF_LoadCompleted))
+						{
+							// Set the flag back if missing and preload hasn't completed yet so that 
+							// the preload we're going to run does something.
+							Redirector->SetFlags(RF_NeedLoad);
+						}
 						Preload(Redirector);
 						Import.XObject = Redirector->DestinationObject;
 					}
