PropertyBaseObject: When exporting a property to text, if the property is a default subobject make sure it is actually inside the parent's package instead of assuming so

diff --git a/Engine/Source/Runtime/CoreUObject/Private/UObject/PropertyBaseObject.cpp b/Engine/Source/Runtime/CoreUObject/Private/UObject/PropertyBaseObject.cpp
--- a/Engine/Source/Runtime/CoreUObject/Private/UObject/PropertyBaseObject.cpp
+++ b/Engine/Source/Runtime/CoreUObject/Private/UObject/PropertyBaseObject.cpp
@@ -266,7 +266,8 @@ void FObjectPropertyBase::ExportText_Internal( FString& ValueStr, const void* Pr
 				 // be that it's inside the package we are currently in, which means the Temp pointer should be resolved 
 				 // already.  So don't move forward with the check unless it's resolved, we don't want to force a deferred
 				 // loaded asset if we don't have to with this check.
-				 && Temp.IsResolved() && Temp && Temp->IsDefaultSubobject()
+		         // We also want to make sure the object is actually inside the package we are currently in
+				 && Temp.IsResolved() && Temp && Temp->IsDefaultSubobject() && Temp->IsIn(Parent->GetOutermostObject())
 				)
 		{
 			if (PortFlags & PPF_Delimited)
