Add an Open for AddObject and an OpenAbort to remove it if we fail during a transaction

diff --git a/Engine/Source/Runtime/CoreUObject/Private/Misc/GCObjectReferencer.cpp b/Engine/Source/Runtime/CoreUObject/Private/Misc/GCObjectReferencer.cpp
--- a/Engine/Source/Runtime/CoreUObject/Private/Misc/GCObjectReferencer.cpp
+++ b/Engine/Source/Runtime/CoreUObject/Private/Misc/GCObjectReferencer.cpp
@@ -263,7 +263,13 @@ void FGCObject::RegisterGCObject()
 		StaticInit();
 
 		// Add this instance to the referencer's list
-		GGCObjectReferencer->AddObject(this);
+		UE_AUTORTFM_OPEN({
+			GGCObjectReferencer->AddObject(this);
+		});
+		UE_AUTORTFM_OPENABORT({
+			GGCObjectReferencer->RemoveObject(this);
+		});
+
 		bReferenceAdded = true;
 	}
 }
