Fix race condition with FrankenGC that could prevent VCell->UObject references from being marked.

diff --git a/Engine/Source/Runtime/CoreUObject/Private/UObject/GarbageCollection.cpp b/Engine/Source/Runtime/CoreUObject/Private/UObject/GarbageCollection.cpp
--- a/Engine/Source/Runtime/CoreUObject/Private/UObject/GarbageCollection.cpp
+++ b/Engine/Source/Runtime/CoreUObject/Private/UObject/GarbageCollection.cpp
@@ -4193,7 +4193,11 @@ public:
 			do
 			{
 				PerformReachabilityAnalysisPass(Options);
-			} while ((!Private::GReachableObjects.IsEmpty() || !Private::GReachableClusters.IsEmpty() || VerseGCActive()) && !GReachabilityState.IsSuspended());
+			// NOTE: It is critical that VerseGCActive is called prior to checking GReachableObjects.  While VerseGCActive is true,
+			// items can still be added to GReachableObjects.  So if reversed, during the point where GReachableObjects is checked
+			// and VerseGCActive returns false, something might have been marked.  Reversing insures that Verse will not add anything 
+			// if Verse is no longer active.
+			} while ((VerseGCActive() || !Private::GReachableObjects.IsEmpty() || !Private::GReachableClusters.IsEmpty()) && !GReachabilityState.IsSuspended());
 
 			UE_LOG(LogGarbage, Verbose, TEXT("%f ms for Reachability Analysis"), (FPlatformTime::Seconds() - StartTime) * 1000);
 		}
