Making sure incremental reachability is complete before the engine performs the final UObject purge

diff --git a/Engine/Source/Runtime/CoreUObject/Private/UObject/GarbageCollection.cpp b/Engine/Source/Runtime/CoreUObject/Private/UObject/GarbageCollection.cpp
--- a/Engine/Source/Runtime/CoreUObject/Private/UObject/GarbageCollection.cpp
+++ b/Engine/Source/Runtime/CoreUObject/Private/UObject/GarbageCollection.cpp
@@ -5492,7 +5492,9 @@ void FReachabilityAnalysisState::PerformReachabilityAnalysis()
 	{
 		UE::GC::CollectGarbageFull(ObjectKeepFlags);
 	}
-	else if (NumRechabilityIterationsToSkip == 0 || !bIsSuspended)
+	else if (NumRechabilityIterationsToSkip == 0 || // Delay reachability analysis by NumRechabilityIterationsToSkip (if desired)
+		!bIsSuspended || // but only but only after the first iteration (which also does MarkObjectsAsUnreachable)
+		IterationTimeLimit <= 0.0f) // and only when using time limit (we're not using the limit when we're flushing reachability analysis when starting a new one or on exit)
 	{
 		UE::GC::CollectGarbageIncremental(ObjectKeepFlags);
 	}
@@ -5531,14 +5533,25 @@ bool IsIncrementalReachabilityAnalysisPending()
 	return UE::GC::Private::GIsIncrementalReachabilityPending;
 }
 
-void PerformIncrementalReachabilityAnalysis()
+void PerformIncrementalReachabilityAnalysis(double TimeLimit)
 {
 	checkf(UE::GC::Private::GIsIncrementalReachabilityPending, TEXT("Incremental reachability must be pending to perform its next iteration"));
 	// When performing Reachability Analysis iterations start the internal timer before acquiring GC lock 
 	// so that we don't spend more time than GIncrementalReachabilityTimeLimit on fully completing an iteration
-	UE::GC::GReachabilityState.StartTimer(GIncrementalReachabilityTimeLimit);
+	const bool bUsingTimeLimit = TimeLimit > 0.0;
+	UE::GC::GReachabilityState.StartTimer(bUsingTimeLimit ? TimeLimit : 0.0);
 	AcquireGCLock();	
-	UE::GC::GReachabilityState.PerformReachabilityAnalysisAndConditionallyPurgeGarbage(/*bReachabilityUsingTimeLimit = */ true);
+	UE::GC::GReachabilityState.PerformReachabilityAnalysisAndConditionallyPurgeGarbage(bUsingTimeLimit);
+
+	checkf(bUsingTimeLimit || !UE::GC::Private::GIsIncrementalReachabilityPending, TEXT("Incremental Reachability is still pending after completing the previous iteration without time limit."));
+}
+
+void FinalizeIncrementalReachabilityAnalysis()
+{
+	if (IsIncrementalReachabilityAnalysisPending())
+	{
+		PerformIncrementalReachabilityAnalysis(0.0);
+	}
 }
 
 FString FGarbageReferenceInfo::GetReferencingObjectInfo() const
diff --git a/Engine/Source/Runtime/CoreUObject/Private/UObject/Obj.cpp b/Engine/Source/Runtime/CoreUObject/Private/UObject/Obj.cpp
--- a/Engine/Source/Runtime/CoreUObject/Private/UObject/Obj.cpp
+++ b/Engine/Source/Runtime/CoreUObject/Private/UObject/Obj.cpp
@@ -5078,7 +5078,10 @@ void StaticExit()
 		GUObjectArray.CloseDisregardForGC();
 	}
 
-	// Complete any pending incremental GC
+	// Complete any pending incremental reachability analysis
+	FinalizeIncrementalReachabilityAnalysis();
+
+	// Complete any pending incremental purge
 	if (IsIncrementalPurgePending())
 	{
 		IncrementalPurgeGarbage(false);
diff --git a/Engine/Source/Runtime/CoreUObject/Public/UObject/UObjectGlobals.h b/Engine/Source/Runtime/CoreUObject/Public/UObject/UObjectGlobals.h
--- a/Engine/Source/Runtime/CoreUObject/Public/UObject/UObjectGlobals.h
+++ b/Engine/Source/Runtime/CoreUObject/Public/UObject/UObjectGlobals.h
@@ -877,14 +877,16 @@ COREUOBJECT_API void GatherUnreachableObjects(bool bForceSingleThreaded);
 COREUOBJECT_API bool IsIncrementalReachabilityAnalysisPending();
 
 /**
- * Incrementally purge garbage by deleting all unreferenced objects after routing Destroy.
- *
- * Calling code needs to be EXTREMELY careful when and how to call this function as
- * RF_Unreachable cannot change on any objects unless any pending purge has completed!
+ * Incrementally perform reachability analysis
  *
- * @param	bUseTimeLimit	whether the time limit parameter should be used
+ * @param	TimeLimit	Time limit (in seconds) for this function call. 0.0 results in no time limit being used.
+ */
+COREUOBJECT_API void PerformIncrementalReachabilityAnalysis(double TimeLimit);
+
+/**
+ * Finalizes incremental reachability analysis (if currently running) without any time limit
  */
-COREUOBJECT_API void PerformIncrementalReachabilityAnalysis();
+COREUOBJECT_API void FinalizeIncrementalReachabilityAnalysis();
 
 /**
  * Incrementally purge garbage by deleting all unreferenced objects after routing Destroy.
diff --git a/Engine/Source/Runtime/Engine/Private/UnrealEngine.cpp b/Engine/Source/Runtime/Engine/Private/UnrealEngine.cpp
--- a/Engine/Source/Runtime/Engine/Private/UnrealEngine.cpp
+++ b/Engine/Source/Runtime/Engine/Private/UnrealEngine.cpp
@@ -76,6 +76,7 @@ UnrealEngine.cpp: Implements the UEngine class and helpers.
 #include "EngineUtils.h"
 #include "Framework/Application/SlateApplication.h"
 #include "UObject/GarbageCollectionHistory.h"
+#include "UObject/ReachabilityAnalysis.h"
 #include "Widgets/Input/SButton.h"
 #include "Engine/EngineCustomTimeStep.h"
 #include "Engine/LevelStreamingPersistent.h"
@@ -1819,7 +1820,7 @@ void UEngine::ConditionalCollectGarbage()
 					else if (IsIncrementalReachabilityAnalysisPending())
 					{
 						SCOPE_CYCLE_COUNTER(STAT_GCMarkTime);
-						PerformIncrementalReachabilityAnalysis();
+						PerformIncrementalReachabilityAnalysis(GetReachabilityAnalysisTimeLimit());
 					}
 					// Perform incremental purge update if it's pending or in progress.
 					else if (!IsIncrementalPurgePending()
@@ -1866,7 +1867,7 @@ void UEngine::ConditionalCollectGarbage()
 	}
 	else if (IsIncrementalReachabilityAnalysisPending())
 	{
-		PerformIncrementalReachabilityAnalysis();
+		PerformIncrementalReachabilityAnalysis(GetReachabilityAnalysisTimeLimit());
 	}
 }
 
