Fixed Malloc/delete mismatch

diff --git a/Engine/Source/Runtime/CoreUObject/Private/UObject/ExpandingChunkedList.h b/Engine/Source/Runtime/CoreUObject/Private/UObject/ExpandingChunkedList.h
--- a/Engine/Source/Runtime/CoreUObject/Private/UObject/ExpandingChunkedList.h
+++ b/Engine/Source/Runtime/CoreUObject/Private/UObject/ExpandingChunkedList.h
@@ -33,7 +33,7 @@ class TExpandingChunkedList
 		for (FChunk* Chunk = Head; Chunk;)
 		{
 			FChunk* NextChunk = Chunk->Next;
-			delete Chunk;
+			FMemory::Free(Chunk);
 			Chunk = NextChunk;
 		}
 	}
@@ -71,7 +71,7 @@ public:
 		{
 			Head->Items[Head->NumItems++] = Item;
 			Lock.Unlock();
-			delete NewChunk;
+			FMemory::Free(NewChunk);
 			return;
 		}
 
