Add OPEN in AllocateSerialNumber and avoid undo'ing anything letting the Primary Serial Number keep increasing even if we abort during a transaction

diff --git a/Engine/Source/Runtime/CoreUObject/Private/UObject/UObjectArray.cpp b/Engine/Source/Runtime/CoreUObject/Private/UObject/UObjectArray.cpp
--- a/Engine/Source/Runtime/CoreUObject/Private/UObject/UObjectArray.cpp
+++ b/Engine/Source/Runtime/CoreUObject/Private/UObject/UObjectArray.cpp
@@ -405,14 +405,17 @@ int32 FUObjectArray::AllocateSerialNumber(int32 Index)
 	int32 SerialNumber = *SerialNumberPtr;
 	if (!SerialNumber)
 	{
-		SerialNumber = PrimarySerialNumber.Increment();
-		UE_CLOG(SerialNumber <= START_SERIAL_NUMBER, LogUObjectArray, Fatal, TEXT("UObject serial numbers overflowed (trying to allocate serial number %d)."), SerialNumber);
-		int32 ValueWas = FPlatformAtomics::InterlockedCompareExchange((int32*)SerialNumberPtr, SerialNumber, 0);
-		if (ValueWas != 0)
-		{
-			// someone else go it first, use their value
-			SerialNumber = ValueWas;
-		}
+		// Open around PrimarySerialNumber as if we fail/abort a transaction we dont need to undo this, simply allow it to grow for the next use
+		UE_AUTORTFM_OPEN({
+			SerialNumber = PrimarySerialNumber.Increment();
+			UE_CLOG(SerialNumber <= START_SERIAL_NUMBER, LogUObjectArray, Fatal, TEXT("UObject serial numbers overflowed (trying to allocate serial number %d)."), SerialNumber);
+			int32 ValueWas = FPlatformAtomics::InterlockedCompareExchange((int32*)SerialNumberPtr, SerialNumber, 0);
+			if (ValueWas != 0)
+			{
+				// someone else go it first, use their value
+				SerialNumber = ValueWas;
+			}
+		});
 	}
 	checkSlow(SerialNumber > START_SERIAL_NUMBER);
 	return SerialNumber;
