Fix multicast export code to export the path to the actor rather than the name so it can be parsed on import.
Fixes an issue where Alt+Click drag would ensure failing to find an actor class.

Marking the erroring primitive component BP assignable delegate TextExportTransient as we shouldn't copy and bind to the original class.

diff --git a/Engine/Source/Runtime/CoreUObject/Private/UObject/PropertyMulticastDelegate.cpp b/Engine/Source/Runtime/CoreUObject/Private/UObject/PropertyMulticastDelegate.cpp
--- a/Engine/Source/Runtime/CoreUObject/Private/UObject/PropertyMulticastDelegate.cpp
+++ b/Engine/Source/Runtime/CoreUObject/Private/UObject/PropertyMulticastDelegate.cpp
@@ -248,7 +248,7 @@ void FMulticastDelegateProperty::ExportText_Internal( FString& ValueStr, const v
 
 				bool bDelegateHasValue = CurInvocation->GetFunctionName() != NAME_None;
 				ValueStr += FString::Printf(TEXT("%s.%s"),
-					CurInvocation->GetUObject() != NULL ? *CurInvocation->GetUObject()->GetName() : TEXT("(null)"),
+					CurInvocation->GetUObject() != NULL ? *CurInvocation->GetUObject()->GetPathName() : TEXT("(null)"),
 					*CurInvocation->GetFunctionName().ToString());
 			}
 		}
diff --git a/Engine/Source/Runtime/Engine/Classes/Components/PrimitiveComponent.h b/Engine/Source/Runtime/Engine/Classes/Components/PrimitiveComponent.h
--- a/Engine/Source/Runtime/Engine/Classes/Components/PrimitiveComponent.h
+++ b/Engine/Source/Runtime/Engine/Classes/Components/PrimitiveComponent.h
@@ -1330,7 +1330,7 @@ public:
 	/**
 	 *	Event called when physics state is created or destroyed for this component
 	 */
-	UPROPERTY(BlueprintAssignable, Category = "Physics")
+	UPROPERTY(BlueprintAssignable, Category = "Physics", TextExportTransient)
 	FComponentPhysicsStateChanged OnComponentPhysicsStateChanged;
 
 	/** Event called when the mouse cursor is moved over this component and mouse over events are enabled in the player controller */
