[ALT] Fix all FMessageLog indirect calls to slate widget via a delegate that can be encountered when errors are reported during async loading

#rnx

diff --git a/Engine/Source/Runtime/CoreUObject/Private/UObject/LinkerLoad.cpp b/Engine/Source/Runtime/CoreUObject/Private/UObject/LinkerLoad.cpp
--- a/Engine/Source/Runtime/CoreUObject/Private/UObject/LinkerLoad.cpp
+++ b/Engine/Source/Runtime/CoreUObject/Private/UObject/LinkerLoad.cpp
@@ -8,6 +8,7 @@
 #include "Misc/ConfigCacheIni.h"
 #include "HAL/IConsoleManager.h"
 #include "Misc/SlowTask.h"
+#include "Async/Async.h"
 #include "Misc/ScopedSlowTask.h"
 #include "Misc/ObjectThumbnail.h"
 #include "Misc/App.h"
@@ -1341,8 +1342,13 @@ FLinkerLoad::ELinkerStatus FLinkerLoad::SerializePackageFileSummaryInternal()
 	// Check tag.
 	if (Summary.Tag != PACKAGE_FILE_TAG)
 	{
-		FMessageLog("LoadErrors").Error(FText::Format(NSLOCTEXT("Core", "LinkerLoad_PkgSumCorrupted", "The summary for the package '{0}' is invalid. Check that the file is of the expected type and not corrupted."),
-			FText::FromString(GetDebugName())));
+		Async(EAsyncExecution::TaskGraphMainThread,
+			[DebugName = GetDebugName()]()
+			{
+				FMessageLog("LoadErrors").Error(FText::Format(NSLOCTEXT("Core", "LinkerLoad_PkgSumCorrupted", "The summary for the package '{0}' is invalid. Check that the file is of the expected type and not corrupted."),
+					FText::FromString(DebugName)));
+			}
+		);
 
 		return LINKER_Failed;
 	}
@@ -1350,10 +1356,15 @@ FLinkerLoad::ELinkerStatus FLinkerLoad::SerializePackageFileSummaryInternal()
 	// Validate the summary.
 	if (Summary.IsFileVersionTooOld())
 	{
-		FMessageLog("LoadErrors").Warning(FText::Format(NSLOCTEXT("Core", "LinkerLoad_PkgVersionTooOld", "The package '{0}' was saved with an older version which is not backwards compatible with the current process. Min Required Version: {1}  Package Version: {2}"), 
-			FText::FromString(GetDebugName()), 
-			(int32)VER_UE4_OLDEST_LOADABLE_PACKAGE, 
-			Summary.GetFileVersionUE().FileVersionUE4));
+		Async(EAsyncExecution::TaskGraphMainThread,
+			[DebugName = GetDebugName(), FileVersion = Summary.GetFileVersionUE()]()
+			{
+				FMessageLog("LoadErrors").Warning(FText::Format(NSLOCTEXT("Core", "LinkerLoad_PkgVersionTooOld", "The package '{0}' was saved with an older version which is not backwards compatible with the current process. Min Required Version: {1}  Package Version: {2}"), 
+					FText::FromString(DebugName),
+					(int32)VER_UE4_OLDEST_LOADABLE_PACKAGE, 
+					FileVersion.FileVersionUE4));
+			}
+		);
 
 		return LINKER_Failed;
 	}
@@ -1380,12 +1391,18 @@ FLinkerLoad::ELinkerStatus FLinkerLoad::SerializePackageFileSummaryInternal()
 	// Don't load packages that are only compatible with an engine version newer than the current one.
 	if (bLoaderNeedsEngineVersionChecks && IsEnforcePackageCompatibleVersionCheck() && !FEngineVersion::Current().IsCompatibleWith(Summary.CompatibleWithEngineVersion))
 	{
-		FMessageLog("LoadErrors").Warning(FText::Format(NSLOCTEXT("Core", "LinkerLoad_EngineVersionIncompatible", "Package '{0}' has been saved with a newer engine version and can't be loaded. Current EngineVersion: {1} (Licensee={2}). Package EngineVersion: {3} (Licensee={4})"),
-			FText::FromString(GetDebugName()), 
-			FText::FromString(FEngineVersion::Current().ToString()),
-			FEngineVersion::Current().IsLicenseeVersion(),
-			FText::FromString(Summary.CompatibleWithEngineVersion.ToString()),
-			Summary.CompatibleWithEngineVersion.IsLicenseeVersion()));
+		// Send the warning to the game thread as slate is not thread-safe
+		Async(EAsyncExecution::TaskGraphMainThread,
+			[DebugName = GetDebugName(), CompatibleWith = Summary.CompatibleWithEngineVersion]()
+			{
+				FMessageLog("LoadErrors").Warning(FText::Format(NSLOCTEXT("Core", "LinkerLoad_EngineVersionIncompatible", "Package '{0}' has been saved with a newer engine version and can't be loaded. Current EngineVersion: {1} (Licensee={2}). Package EngineVersion: {3} (Licensee={4})"),
+					FText::FromString(DebugName),
+					FText::FromString(FEngineVersion::Current().ToString()),
+					FEngineVersion::Current().IsLicenseeVersion(),
+					FText::FromString(CompatibleWith.ToString()),
+					CompatibleWith.IsLicenseeVersion()));
+			}
+		);
 
 		return LINKER_Failed;
 	}
@@ -1424,12 +1441,18 @@ FLinkerLoad::ELinkerStatus FLinkerLoad::SerializePackageFileSummaryInternal()
 	// Don't load packages that were saved with package version newer than the current one.
 	if (bLoaderNeedsEngineVersionChecks && ((Summary.IsFileVersionTooNew()) || (Summary.GetFileVersionLicenseeUE() > GPackageFileLicenseeUEVersion)))
 	{
-		FMessageLog("LoadErrors").Warning(FText::Format(NSLOCTEXT("Core", "LinkerLoad_PkgVersionTooNew", "Package '{0}' contains a newer version than the current process supports. PackageVersion {1}, MaxExpected {2} : LicenseePackageVersion {3}, MaxExpected {4}."), 
-			FText::FromString(GetDebugName()), 
-			Summary.GetFileVersionUE().ToValue(), 
-			GPackageFileUEVersion.ToValue(), 
-			Summary.GetFileVersionLicenseeUE(), 
-			GPackageFileLicenseeUEVersion));
+		// Send the warning to the game thread as slate is not thread-safe
+		Async(EAsyncExecution::TaskGraphMainThread,
+			[DebugName = GetDebugName(), FileVersion = Summary.GetFileVersionUE(), FileVersionLicensee = Summary.GetFileVersionLicenseeUE(), PackageFileUEVersion = GPackageFileUEVersion, PackageFileLicenseeUEVersion = GPackageFileLicenseeUEVersion]()
+			{
+				FMessageLog("LoadErrors").Warning(FText::Format(NSLOCTEXT("Core", "LinkerLoad_PkgVersionTooNew", "Package '{0}' contains a newer version than the current process supports. PackageVersion {1}, MaxExpected {2} : LicenseePackageVersion {3}, MaxExpected {4}."), 
+					FText::FromString(DebugName),
+					FileVersion.ToValue(),
+					PackageFileUEVersion.ToValue(),
+					FileVersionLicensee,
+					PackageFileLicenseeUEVersion));
+			}
+		);
 
 		return LINKER_Failed;
 	}
@@ -1437,8 +1460,13 @@ FLinkerLoad::ELinkerStatus FLinkerLoad::SerializePackageFileSummaryInternal()
 	// don't load packages that contain editor only data in builds that don't support that and vise versa
 	if (!FPlatformProperties::HasEditorOnlyData() && !(Summary.GetPackageFlags() & PKG_FilterEditorOnly))
 	{
-		FMessageLog("LoadErrors").Warning(FText::Format(NSLOCTEXT("Core", "LinkerLoad_InvalidEditorOnlyData", "Unable to load package '{0}'. Package contains EditorOnly data which is not supported by the current build."), 
-			FText::FromString(GetDebugName())));
+		Async(EAsyncExecution::TaskGraphMainThread,
+			[DebugName = GetDebugName()]()
+			{
+				FMessageLog("LoadErrors").Warning(FText::Format(NSLOCTEXT("Core", "LinkerLoad_InvalidEditorOnlyData", "Unable to load package '{0}'. Package contains EditorOnly data which is not supported by the current build."), 
+					FText::FromString(DebugName)));
+			}
+		);
 
 		return LINKER_Failed;
 	}
@@ -1449,8 +1477,13 @@ FLinkerLoad::ELinkerStatus FLinkerLoad::SerializePackageFileSummaryInternal()
 		// This warning can be disabled in ini or project settings
 		if (!GAllowCookedDataInEditorBuilds)
 		{
-			FMessageLog("LoadErrors").Warning(FText::Format(NSLOCTEXT("Core", "LinkerLoad_InvalidCookedData", "Unable to load package '{0}'. Package contains cooked data which is not supported by the current build. Enable 'Allow Cooked Content In The Editor' in Project Settings under 'Engine - Cooker' section to load it."),
-				FText::FromString(GetDebugName())));
+			Async(EAsyncExecution::TaskGraphMainThread,
+				[DebugName = GetDebugName()]()
+				{
+					FMessageLog("LoadErrors").Warning(FText::Format(NSLOCTEXT("Core", "LinkerLoad_InvalidCookedData", "Unable to load package '{0}'. Package contains cooked data which is not supported by the current build. Enable 'Allow Cooked Content In The Editor' in Project Settings under 'Engine - Cooker' section to load it."),
+						FText::FromString(DebugName)));
+				}
+			);
 
 			return LINKER_Failed;
 		}
@@ -1475,8 +1508,13 @@ FLinkerLoad::ELinkerStatus FLinkerLoad::SerializePackageFileSummaryInternal()
 
 		if (Tag != PACKAGE_FILE_TAG)
 		{
-			FMessageLog("LoadErrors").Error(FText::Format(NSLOCTEXT("Core", "LinkerLoad_PkgTagCorrupted", "Unable to load package '{0}'. The end of package tag is not valid. Check that the file is of the expected type and not corrupted."), 
-				FText::FromString(GetDebugName())));
+			Async(EAsyncExecution::TaskGraphMainThread,
+				[DebugName = GetDebugName()]()
+				{
+					FMessageLog("LoadErrors").Error(FText::Format(NSLOCTEXT("Core", "LinkerLoad_PkgTagCorrupted", "Unable to load package '{0}'. The end of package tag is not valid. Check that the file is of the expected type and not corrupted."), 
+						FText::FromString(DebugName)));
+				}
+			);
 
 			return LINKER_Failed;
 		}
@@ -1561,10 +1599,15 @@ FLinkerLoad::ELinkerStatus FLinkerLoad::UpdateFromPackageFileSummary()
 			{
 				UE_ASSET_LOG(LogLinker, Error, PackagePath, TEXT("Package was saved with an invalid custom version. Tag %s  Version %d"), *Diff.Version->Key.ToString(), Diff.Version->Version);
 
-				FMessageLog("LoadErrors")
-					.SuppressLoggingToOutputLog(true)
-					.Error(FText::Format(NSLOCTEXT("Core", "LinkerLoad_InvalidCustomVersion", "Package {0} was saved with an invalid custom version and cannot be loaded, see output log for details"),
-						FText::FromString(GetDebugName())));
+				Async(EAsyncExecution::TaskGraphMainThread,
+					[DebugName = GetDebugName()]()
+					{
+						FMessageLog("LoadErrors")
+							.SuppressLoggingToOutputLog(true)
+							.Error(FText::Format(NSLOCTEXT("Core", "LinkerLoad_InvalidCustomVersion", "Package {0} was saved with an invalid custom version and cannot be loaded, see output log for details"),
+								FText::FromString(DebugName)));
+					}
+				);
 
 				return LINKER_Failed;
 			}
@@ -1576,10 +1619,15 @@ FLinkerLoad::ELinkerStatus FLinkerLoad::UpdateFromPackageFileSummary()
 				UE_ASSET_LOG(LogLinker, Error, PackagePath, TEXT("Package was saved with a newer custom version than the current. Tag %s Name '%s' PackageVersion %d  MaxExpected %d"),
 					*Diff.Version->Key.ToString(), *LatestVersion.GetFriendlyName().ToString(), Diff.Version->Version, LatestVersion.Version);
 
-				FMessageLog("LoadErrors")
-					.SuppressLoggingToOutputLog(true)
-					.Error(FText::Format(NSLOCTEXT("Core", "LinkerLoad_NewCustomVersion", "Package {0} was saved with a newer custom version than the current engine and cannot be loaded, see output log for details"),
-						FText::FromString(GetDebugName())));
+				Async(EAsyncExecution::TaskGraphMainThread,
+					[DebugName = GetDebugName()]()
+					{
+						FMessageLog("LoadErrors")
+							.SuppressLoggingToOutputLog(true)
+							.Error(FText::Format(NSLOCTEXT("Core", "LinkerLoad_NewCustomVersion", "Package {0} was saved with a newer custom version than the current engine and cannot be loaded, see output log for details"),
+								FText::FromString(DebugName)));
+					}
+				);
 
 				return LINKER_Failed;
 			}
@@ -1702,9 +1750,14 @@ FLinkerLoad::ELinkerStatus FLinkerLoad::SerializePackageTrailer()
 			{
 				UE_ASSET_LOG(LogLinker, Error, PackagePath, TEXT("Package has a corrupted package trailer"));
 
-				FMessageLog("LoadErrors").SuppressLoggingToOutputLog(true)
-					.Error(FText::Format(NSLOCTEXT("Core", "LinkerLoad_CorruptTrailer", "Package {0} has a corrupted package trailer"),
-					FText::FromString(GetDebugName())));
+				Async(EAsyncExecution::TaskGraphMainThread,
+					[DebugName = GetDebugName()]()
+					{
+						FMessageLog("LoadErrors").SuppressLoggingToOutputLog(true)
+							.Error(FText::Format(NSLOCTEXT("Core", "LinkerLoad_CorruptTrailer", "Package {0} has a corrupted package trailer"),
+							FText::FromString(DebugName)));
+					}
+				);
 
 				return LINKER_Failed;
 			}
