Fix confusing comment leaking into editor tooltips for Y axis of FVector/FVector4/FQuat properties

diff --git a/Engine/Source/Runtime/CoreUObject/Public/UObject/NoExportTypes.h b/Engine/Source/Runtime/CoreUObject/Public/UObject/NoExportTypes.h
--- a/Engine/Source/Runtime/CoreUObject/Public/UObject/NoExportTypes.h
+++ b/Engine/Source/Runtime/CoreUObject/Public/UObject/NoExportTypes.h
@@ -544,7 +544,7 @@ USTRUCT(immutable, noexport, BlueprintType, IsAlwaysAccessible, HasDefaults, Has
 struct FVector
 {
 	UPROPERTY(EditAnywhere, BlueprintReadWrite, Category = Vector, SaveGame)
-	FLargeWorldCoordinatesReal X;		// Alias for float/double depending on LWC status. Note: Will be refactored to double before UE5 ships.
+	FLargeWorldCoordinatesReal X;		//~ Alias for float/double depending on LWC status. Note: Will be refactored to double before UE5 ships.
 
 	UPROPERTY(EditAnywhere, BlueprintReadWrite, Category = Vector, SaveGame)
 	FLargeWorldCoordinatesReal Y;
@@ -603,7 +603,7 @@ USTRUCT(immutable, noexport, BlueprintType, IsAlwaysAccessible, HasDefaults, Has
 struct FVector4
 {
 	UPROPERTY(EditAnywhere, BlueprintReadWrite, Category=Vector4, SaveGame)
-	FLargeWorldCoordinatesReal X;		// Alias for float/double depending on LWC status. Note: Will be refactored to double before UE5 ships.
+	FLargeWorldCoordinatesReal X;		//~ Alias for float/double depending on LWC status. Note: Will be refactored to double before UE5 ships.
 
 	UPROPERTY(EditAnywhere, BlueprintReadWrite, Category=Vector4, SaveGame)
 	FLargeWorldCoordinatesReal Y;
@@ -908,7 +908,7 @@ USTRUCT(immutable, noexport, BlueprintType, IsAlwaysAccessible, HasDefaults, Has
 struct FQuat
 {
 	UPROPERTY(EditAnywhere, BlueprintReadWrite, Category = Quat, SaveGame)
-	FLargeWorldCoordinatesReal X;		// Alias for float/double depending on LWC status. Note: Will be refactored to double before UE5 ships.
+	FLargeWorldCoordinatesReal X;		//~ Alias for float/double depending on LWC status. Note: Will be refactored to double before UE5 ships.
 
 	UPROPERTY(EditAnywhere, BlueprintReadWrite, Category = Quat, SaveGame)
 	FLargeWorldCoordinatesReal Y;
