[ZenLoader]
  - Simplify partial loading logic which is now only applied during flush if required
  - During partial load, only remove the exact RequestID being flushed so that other requestid can be flushed independently
  - Add missing timeout handling for postloads
  - Add missing support to run linkerload thread-safe postloads on async loading thread
  - Add unittest to make sure timeout is respected for postloads
  - Add unittest to make sure behavior of IsReadyForAsyncPostLoad is respected until a flush occurs
  - Add unittest for syncloads coming from thread-safe postloads and validate partial loading behavior
  - Make unittest run in editor only as they rely on saving package before reloading them
  - Refactor unittest for better reuse

diff --git a/Engine/Source/Runtime/CoreUObject/Private/Serialization/AsyncLoading2.cpp b/Engine/Source/Runtime/CoreUObject/Private/Serialization/AsyncLoading2.cpp
--- a/Engine/Source/Runtime/CoreUObject/Private/Serialization/AsyncLoading2.cpp
+++ b/Engine/Source/Runtime/CoreUObject/Private/Serialization/AsyncLoading2.cpp
@@ -85,16 +85,6 @@
 // For now, the partial request behavior is reserved for the editor only.
 #define WITH_PARTIAL_REQUEST_DURING_RECURSION WITH_EDITOR
 
-enum class ERequestType : uint8
-{
-#if WITH_PARTIAL_REQUEST_DURING_RECURSION
-	ExportsDone,	// ExportsDone on the requested package must be reached before exiting a flush of that request.
-#endif
-	Full,
-
-	Count
-};
-
 CSV_DECLARE_CATEGORY_MODULE_EXTERN(CORE_API, FileIO);
 CSV_DEFINE_STAT(FileIO, FrameCompletedExportBundleLoadsKB);
 
@@ -495,8 +485,6 @@ class FAsyncLoadingSyncLoadContext;
 struct FPackageRequest
 {
 	int32 RequestId = -1;
-	ERequestType RequestType = ERequestType::Full;
-	FPackageId RequesterId;
 	int32 Priority = -1;
 	EPackageFlags PackageFlags = PKG_None;
 #if WITH_EDITOR
@@ -521,13 +509,11 @@ struct FPackageRequest
 #endif
 	}
 
-	static FPackageRequest Create(int32 RequestId, ERequestType RequestType, FPackageId RequesterId, EPackageFlags PackageFlags, uint32 LoadFlags, int32 PIEInstanceID, int32 Priority, const FLinkerInstancingContext* InstancingContext, const FPackagePath& PackagePath, FName CustomName, TUniquePtr<FLoadPackageAsyncDelegate> PackageLoadedDelegate, TUniquePtr<FLoadPackageAsyncProgressDelegate> PackageProgressDelegate)
+	static FPackageRequest Create(int32 RequestId, EPackageFlags PackageFlags, uint32 LoadFlags, int32 PIEInstanceID, int32 Priority, const FLinkerInstancingContext* InstancingContext, const FPackagePath& PackagePath, FName CustomName, TUniquePtr<FLoadPackageAsyncDelegate> PackageLoadedDelegate, TUniquePtr<FLoadPackageAsyncProgressDelegate> PackageProgressDelegate)
 	{
 		return FPackageRequest
 		{
 			RequestId,
-			RequestType,
-			RequesterId,
 			Priority,
 			PackageFlags,
 #if WITH_EDITOR
@@ -550,10 +536,6 @@ struct FAsyncPackageDesc2
 {
 	// A unique request id for each external call to LoadPackage
 	int32 RequestID;
-	// The type of request
-	ERequestType RequestType;
-	// The Id of the top package on the stack (if any) when the request was issued.
-	FPackageId RequesterId;
 	// Package priority
 	int32 Priority;
 	/** The flags that should be applied to the package */
@@ -585,8 +567,6 @@ struct FAsyncPackageDesc2
 		return FAsyncPackageDesc2
 		{
 			Request.RequestId,
-			Request.RequestType,
-			Request.RequesterId,
 			Request.Priority,
 			Request.PackageFlags,
 #if WITH_EDITOR
@@ -615,8 +595,6 @@ struct FAsyncPackageDesc2
 		return FAsyncPackageDesc2
 		{
 			INDEX_NONE,
-			ERequestType::Full,
-			FPackageId(),
 			ImportingPackageDesc.Priority,
 			PKG_None,
 #if WITH_EDITOR
@@ -2006,6 +1984,38 @@ enum class EAsyncPackageLoadingState2 : uint8
 	DeferredDelete,
 };
 
+const TCHAR* LexToString(EAsyncPackageLoadingState2 AsyncPackageLoadingState)
+{
+	switch(AsyncPackageLoadingState)
+	{
+	case EAsyncPackageLoadingState2::NewPackage: return TEXT("NewPackage");
+	case EAsyncPackageLoadingState2::WaitingForIo: return TEXT("WaitingForIo");
+	case EAsyncPackageLoadingState2::ProcessPackageSummary: return TEXT("ProcessPackageSummary");
+	case EAsyncPackageLoadingState2::WaitingForDependencies: return TEXT("WaitingForDependencies");
+	case EAsyncPackageLoadingState2::DependenciesReady: return TEXT("DependenciesReady");
+#if ALT2_ENABLE_LINKERLOAD_SUPPORT
+	case EAsyncPackageLoadingState2::CreateLinkerLoadExports: return TEXT("CreateLinkerLoadExports");
+	case EAsyncPackageLoadingState2::WaitingForLinkerLoadDependencies: return TEXT("WaitingForLinkerLoadDependencies");
+	case EAsyncPackageLoadingState2::ResolveLinkerLoadImports: return TEXT("ResolveLinkerLoadImports");
+	case EAsyncPackageLoadingState2::PreloadLinkerLoadExports: return TEXT("PreloadLinkerLoadExports");
+#endif
+	case EAsyncPackageLoadingState2::ProcessExportBundles: return TEXT("ProcessExportBundles");
+	case EAsyncPackageLoadingState2::WaitingForExternalReads: return TEXT("WaitingForExternalReads");
+	case EAsyncPackageLoadingState2::ExportsDone: return TEXT("ExportsDone");
+	case EAsyncPackageLoadingState2::PostLoad: return TEXT("PostLoad");
+	case EAsyncPackageLoadingState2::DeferredPostLoad: return TEXT("DeferredPostLoad");
+	case EAsyncPackageLoadingState2::DeferredPostLoadDone: return TEXT("DeferredPostLoadDone");
+	case EAsyncPackageLoadingState2::Finalize: return TEXT("Finalize");
+	case EAsyncPackageLoadingState2::PostLoadInstances: return TEXT("PostLoadInstances");
+	case EAsyncPackageLoadingState2::CreateClusters: return TEXT("CreateClusters");
+	case EAsyncPackageLoadingState2::Complete: return TEXT("Complete");
+	case EAsyncPackageLoadingState2::DeferredDelete: return TEXT("DeferredDelete");
+	default:
+		checkNoEntry();
+		return TEXT("Unknown");
+	}
+};
+
 class FEventLoadGraphAllocator;
 struct FAsyncLoadEventSpec;
 struct FAsyncLoadingThreadState2;
@@ -2650,7 +2660,7 @@ struct FAsyncPackage2
 	}
 
 	/** Adds new request ID to the existing package */
-	void AddRequestID(FAsyncLoadingThreadState2& ThreadState, int32 Id, ERequestType RequestType);
+	void AddRequestID(FAsyncLoadingThreadState2& ThreadState, int32 Id);
 
 	uint64 GetSyncLoadContextId() const
 	{
@@ -2849,8 +2859,8 @@ private:
 	bool						bLoadHasFailed = false;
 	/** True if this package was created by this async package */
 	bool						bCreatedLinkerRoot = false;
-	/** List of all request handles by types. */
-	TArray<int32, TInlineAllocator<2>> RequestIDs[(int32)ERequestType::Count];
+	/** List of all request handles. */
+	TArray<int32> RequestIDs;
 	/** List of ConstructedObjects = Exports + UPackage + ObjectsCreatedFromExports */
 	TArray<UObject*> ConstructedObjects;
 	TArray<FExternalReadCallback> ExternalReadDependencies;
@@ -2944,6 +2954,7 @@ private:
 	void WaitForAllDependenciesToReachState(FAsyncLoadingThreadState2& ThreadState, FAllDependenciesState FAsyncPackage2::* StateMemberPtr, EAsyncPackageLoadingState2 WaitForPackageState, uint32& CurrentTickVariable, TFunctionRef<void(FAsyncPackage2*)> OnStateReached);
 	void ConditionalBeginProcessPackageExports(FAsyncLoadingThreadState2& ThreadState);
 	void ConditionalFinishLoading(FAsyncLoadingThreadState2& ThreadState);
+	void ConditionalReleasePartialRequests(FAsyncLoadingThreadState2& ThreadState);
 
 #if ALT2_ENABLE_LINKERLOAD_SUPPORT
 	EEventLoadNodeExecutionResult ProcessLinkerLoadPackageSummary(FAsyncLoadingThreadState2& ThreadState);
@@ -2951,10 +2962,10 @@ private:
 	void ConditionalBeginResolveLinkerLoadImports(FAsyncLoadingThreadState2& ThreadState);
 	bool ResolveLinkerLoadImports(FAsyncLoadingThreadState2& ThreadState);
 	bool PreloadLinkerLoadExports(FAsyncLoadingThreadState2& ThreadState);
-	EEventLoadNodeExecutionResult PostLoadLinkerLoadPackageExports(FAsyncLoadingThreadState2& ThreadState);
+	EEventLoadNodeExecutionResult ExecutePostLoadLinkerLoadPackageExports(FAsyncLoadingThreadState2& ThreadState);
+	EEventLoadNodeExecutionResult ExecuteDeferredPostLoadLinkerLoadPackageExports(FAsyncLoadingThreadState2& ThreadState);
 #endif
 
-
 	void ProcessExportDependencies(const FAsyncPackageHeaderData& Header, int32 LocalExportIndex, FExportBundleEntry::EExportCommandType CommandType);
 	int32 GetPublicExportIndex(uint64 ExportHash, FAsyncPackageHeaderData*& OutHeader);
 	UObject* ConditionalCreateExport(const FAsyncPackageHeaderData& Header, int32 LocalExportIndex);
@@ -3092,10 +3103,7 @@ private:
 				MoveTemp(Package->ProgressCallbacks)
 			};
 
-			for (int32 Index = 0; Index < (int32)ERequestType::Count; ++Index)
-			{
-				Request.RequestIDs.Append(Package->RequestIDs[Index]);
-			}
+			Request.RequestIDs.Append(Package->RequestIDs);
 			return Request;
 		}
 
@@ -3405,6 +3413,7 @@ public:
 	virtual void FlushLoading(TConstArrayView<int32> RequestIds) override;
 
 	void FlushLoadingFromLoadingThread(FAsyncLoadingThreadState2& ThreadState, TConstArrayView<int32> RequestIds);
+	void WarnAboutPotentialSyncLoadStall(FAsyncLoadingSyncLoadContext* SyncLoadContext);
 
 	virtual int32 GetNumQueuedPackages() override
 	{
@@ -3790,12 +3799,9 @@ private:
 	{
 		TRACE_CPUPROFILER_EVENT_SCOPE(DeleteAsyncPackage);
 		UE_ASYNC_PACKAGE_DEBUG(Package->Desc);
-		for (int32 Index = 0; Index < (int32)ERequestType::Count; ++Index)
+		for (int32 RequestId : Package->RequestIDs)
 		{
-			for (int32 RequestId : Package->RequestIDs[Index])
-			{
-				RequestIdToPackageMap.Remove(RequestId);
-			}
+			RequestIdToPackageMap.Remove(RequestId);
 		}
 		if (RequestIdToPackageMap.IsEmpty())
 		{
@@ -3967,6 +3973,7 @@ void FAsyncLoadingThread2::ConditionalBeginPostLoad(FAsyncLoadingThreadState2& T
 			Package->DeferredPostLoadGroup = PostLoadGroup;
 			
 			Package->AsyncPackageLoadingState = EAsyncPackageLoadingState2::PostLoad;
+			Package->ConditionalReleasePartialRequests(ThreadState);
 			for (int32 ExportBundleIndex = 0; ExportBundleIndex < Package->Data.TotalExportBundleCount; ++ExportBundleIndex)
 			{
 				Package->GetExportBundleNode(EEventLoadNode2::ExportBundle_PostLoad, ExportBundleIndex).ReleaseBarrier(&ThreadState);
@@ -3992,6 +3999,7 @@ void FAsyncLoadingThread2::ConditionalBeginDeferredPostLoad(FAsyncLoadingThreadS
 			check(Package->AsyncPackageLoadingState == EAsyncPackageLoadingState2::PostLoad);
 			Package->DeferredPostLoadGroup = nullptr;
 			Package->AsyncPackageLoadingState = EAsyncPackageLoadingState2::DeferredPostLoad;
+			Package->ConditionalReleasePartialRequests(ThreadState);
 			for (int32 ExportBundleIndex = 0; ExportBundleIndex < Package->Data.TotalExportBundleCount; ++ExportBundleIndex)
 			{
 				Package->GetExportBundleNode(EEventLoadNode2::ExportBundle_DeferredPostLoad, ExportBundleIndex).ReleaseBarrier(&ThreadState);
@@ -4062,13 +4070,9 @@ FAsyncPackage2* FAsyncLoadingThread2::FindOrInsertPackage(FAsyncLoadingThreadSta
 		}
 		else
 		{
-			if (Desc.RequesterId.IsValid())
-			{
-				Package->Desc.RequesterId = Desc.RequesterId;
-			}
 			if (Desc.RequestID > 0)
 			{
-				Package->AddRequestID(ThreadState, Desc.RequestID, Desc.RequestType);
+				Package->AddRequestID(ThreadState, Desc.RequestID);
 			}
 			if (Desc.Priority > Package->Desc.Priority)
 			{
@@ -4090,7 +4094,15 @@ FAsyncPackage2* FAsyncLoadingThread2::FindOrInsertPackage(FAsyncLoadingThreadSta
 	// it is causing crashes and soft-locks in editor when FlushAsyncLoading is called on specific requests
 	// that are in the same postload group as their outer, preventing fine grained flush. So we disable this
 	// in editor for the time being since ALT is not yet activated by default and we'll sort out the race conditions instead.
-	constexpr bool bIsPostLoadGroupFeatureActive = !ALT2_ENABLE_LINKERLOAD_SUPPORT;
+#if ALT2_ENABLE_LINKERLOAD_SUPPORT
+	constexpr bool bIsPostLoadGroupFeatureActive = false;
+#else
+	// Prevents activating postload groups during boot because it was causing deadlock into the 
+	// highly recursive InitDefaultMaterials function on some platforms.
+	// Since postload groups are there to protect against race conditions between postloads and serialize,
+	// no such race can exists until the loading thread is started.
+	const bool bIsPostLoadGroupFeatureActive = bThreadStarted;
+#endif
 
 	if (bInserted)
 	{
@@ -5618,7 +5630,70 @@ bool FAsyncPackage2::CreateLinkerLoadExports(FAsyncLoadingThreadState2& ThreadSt
 	return true;
 }
 
-EEventLoadNodeExecutionResult FAsyncPackage2::PostLoadLinkerLoadPackageExports(FAsyncLoadingThreadState2& ThreadState)
+EEventLoadNodeExecutionResult FAsyncPackage2::ExecutePostLoadLinkerLoadPackageExports(FAsyncLoadingThreadState2& ThreadState)
+{
+	if (!bLoadHasFailed)
+	{
+		TRACE_LOADTIME_POSTLOAD_SCOPE;
+
+		FAsyncLoadingTickScope2 InAsyncLoadingTick(AsyncLoadingThread);
+
+#if WITH_EDITOR
+		UE::Core::Private::FPlayInEditorLoadingScope PlayInEditorIDScope(Desc.PIEInstanceID);
+#endif
+
+		// Begin async loading, simulates BeginLoad
+		BeginAsyncLoad();
+
+		FUObjectSerializeContext* LoadContext = GetSerializeContext();
+		TArray<UObject*>& ThreadObjLoaded = LoadContext->PRIVATE_GetObjectsLoadedInternalUseOnly();
+
+		// End async loading, simulates EndLoad
+		ON_SCOPE_EXIT { ThreadObjLoaded.Reset(); EndAsyncLoad(); };
+
+		SCOPED_LOADTIMER(PostLoadObjectsTime);
+
+		FUObjectThreadContext& ThreadContext = FUObjectThreadContext::Get();
+		TGuardValue<bool> GuardIsRoutingPostLoad(ThreadContext.IsRoutingPostLoad, true);
+
+		// Only postload exports instead of constructed objects to avoid cross-package interference.
+		// Also this step is only about running the thread-safe postload that are ready to run, the rest will be deferred.
+		// This code applies to both non threaded loader and threaded loader to exercise the same code path for both
+		// Any non thread-safe or not ready postload is automatically being moved to the deferred phase.
+		while (LinkerLoadState->PostLoadExportIndex < Data.Exports.Num())
+		{
+			const int32 ExportIndex = LinkerLoadState->PostLoadExportIndex++;
+			const FExportObject& Export = Data.Exports[ExportIndex];
+		
+			if (UObject* Object = Export.Object)
+			{
+				if (Object->HasAnyFlags(RF_NeedPostLoad) && CanPostLoadOnAsyncLoadingThread(Object) && Object->IsReadyForAsyncPostLoad())
+				{
+					ThreadContext.CurrentlyPostLoadedObjectByALT = Object;
+					Object->ConditionalPostLoad();
+					ThreadContext.CurrentlyPostLoadedObjectByALT = nullptr;
+				}
+
+				if (ThreadState.IsTimeLimitExceeded(TEXT("ExecutePostLoadLinkerLoadPackageExports")))
+				{
+					return EEventLoadNodeExecutionResult::Timeout;
+				}
+			}
+		}
+	}
+
+	// Reset this to be reused for the deferred postload phase
+	LinkerLoadState->PostLoadExportIndex = 0;
+
+	check(DeferredPostLoadGroup);
+	check(DeferredPostLoadGroup->PackagesWithExportsToPostLoadCount > 0);
+	--DeferredPostLoadGroup->PackagesWithExportsToPostLoadCount;
+
+	AsyncLoadingThread.ConditionalBeginDeferredPostLoad(ThreadState, DeferredPostLoadGroup);
+	return EEventLoadNodeExecutionResult::Complete;
+}
+
+EEventLoadNodeExecutionResult FAsyncPackage2::ExecuteDeferredPostLoadLinkerLoadPackageExports(FAsyncLoadingThreadState2& ThreadState)
 {
 	TRACE_LOADTIME_POSTLOAD_SCOPE;
 
@@ -5628,19 +5703,39 @@ EEventLoadNodeExecutionResult FAsyncPackage2::PostLoadLinkerLoadPackageExports(F
 	UE::Core::Private::FPlayInEditorLoadingScope PlayInEditorIDScope(Desc.PIEInstanceID);
 #endif
 
+	FUObjectSerializeContext* LoadContext = GetSerializeContext();
+	TArray<UObject*>& ThreadObjLoaded = LoadContext->PRIVATE_GetObjectsLoadedInternalUseOnly();
+
+	ON_SCOPE_EXIT{ ThreadObjLoaded.Reset(); };
+
+	// We can't return timeout during a flush as we're expected to be able to finish
+	const bool bIsReadyForAsyncPostLoadAllowed = ThreadState.SyncLoadContextStack.IsEmpty();
+
 	const int32 ObjectCount = ConstructedObjects.Num();
 	while (LinkerLoadState->PostLoadExportIndex < ObjectCount)
 	{
 		const int32 ObjectIndex = LinkerLoadState->PostLoadExportIndex++;
-		ConstructedObjects[ObjectIndex]->ConditionalPostLoad();
-	}
+		if (UObject* Object = ConstructedObjects[ObjectIndex])
+		{
+			// Only allow to wait when there is no flush waiting on us
+			if (bIsReadyForAsyncPostLoadAllowed && !Object->IsReadyForAsyncPostLoad())
+			{
+				--LinkerLoadState->PostLoadExportIndex;
+				return EEventLoadNodeExecutionResult::Timeout;
+			}
 
-	FUObjectSerializeContext* LoadContext = GetSerializeContext();
-	TArray<UObject*>& ThreadObjLoaded = LoadContext->PRIVATE_GetObjectsLoadedInternalUseOnly();
-	ThreadObjLoaded.Empty();
+			Object->ConditionalPostLoad();
+		}
+
+		if (ThreadState.IsTimeLimitExceeded(TEXT("ExecuteDeferredPostLoadLinkerLoadPackageExports")))
+		{
+			return EEventLoadNodeExecutionResult::Timeout;
+		}
+	}
 
 	AsyncPackageLoadingState = EAsyncPackageLoadingState2::DeferredPostLoadDone;
 	ConditionalFinishLoading(ThreadState);
+
 	return EEventLoadNodeExecutionResult::Complete;
 }
 
@@ -6613,6 +6708,55 @@ FAsyncPackage2* FAsyncPackage2::GetCurrentlyExecutingPackage(FAsyncLoadingThread
 	return CurrentlyExecutingPackage;
 }
 
+void FAsyncPackage2::ConditionalReleasePartialRequests(FAsyncLoadingThreadState2& ThreadState)
+{
+#if WITH_PARTIAL_REQUEST_DURING_RECURSION
+	FAsyncLoadingSyncLoadContext* CurrentSyncLoadContext = ThreadState.SyncLoadContextStack.Num() ? ThreadState.SyncLoadContextStack.Top() : nullptr;
+	if (CurrentSyncLoadContext && CurrentSyncLoadContext->RequestingPackage)
+	{
+		EAsyncPackageLoadingState2 RequesterState = CurrentSyncLoadContext->RequestingPackage->AsyncPackageLoadingState;
+
+		if (AsyncPackageLoadingState > RequesterState)
+		{
+			int32 Index = CurrentSyncLoadContext->RequestedPackages.Find(this);
+			if (Index != INDEX_NONE)
+			{
+				int32 RequestId = CurrentSyncLoadContext->RequestIDs[Index];
+				// Release the sync context request tied to our package, allowing the current flush to exit.
+				UE_LOG(LogStreaming, Display, TEXT("Package %s has reached state %s > %s, releasing request %d to allow recursive sync load to finish"),
+					*Desc.UPackageName.ToString(),
+					LexToString(AsyncPackageLoadingState),
+					LexToString(RequesterState),
+					RequestId
+				);
+				AsyncLoadingThread.RemovePendingRequests(ThreadState, { RequestId });
+			}
+
+			// We want to skip postloads inside recursive loads that happens during serialization and preload
+			// since it is mostly what BP expects and how the old LinkerLoad was implemented.
+			// In some cases, trying to postload right away during recursion could end-up in data loss (i.e. UE-190649).
+			// The expectation is that postload should be run as part of the package causing this load to occur.
+			// As a more general rule, we want any synchronous load during deserialization to act like
+			// an import. We want the deserialization steps done before returning from the load but
+			// we keep postload for after deserialization is done on the referencer too so that postload 
+			// sees everything initialized in the case where it also has a reference back to the referencer.
+			// Using a partial request flag allows the flush to return as soon as exports are done which
+			// gives the opportunity to the caller to get a pointer to an object inside the package even if it's not fully loaded yet.
+			if (AsyncPackageLoadingState == EAsyncPackageLoadingState2::ExportsDone && CurrentSyncLoadContext->RequestingPackage->PostLoadGroup)
+			{
+				UE_LOG(LogStreaming, Display, TEXT("Merging postload groups of package %s with requester package %s"),
+					*Desc.UPackageName.ToString(),
+					*CurrentSyncLoadContext->RequestingPackage->Desc.UPackageName.ToString());
+
+				// Do not adjust sync load context, we want to be able to exit the current one even if the caller has not finished yet.
+				const bool bUpdateSyncLoadContext = false;
+				AsyncLoadingThread.MergePostLoadGroups(ThreadState, PostLoadGroup, CurrentSyncLoadContext->RequestingPackage->PostLoadGroup, bUpdateSyncLoadContext);
+			}
+		}
+	}
+#endif // WITH_PARTIAL_REQUEST_DURING_RECURSION
+}
+
 EEventLoadNodeExecutionResult FAsyncPackage2::Event_ExportsDone(FAsyncLoadingThreadState2& ThreadState, FAsyncPackage2* Package, int32)
 {
 	TRACE_CPUPROFILER_EVENT_SCOPE(Event_ExportsDone);
@@ -6662,30 +6806,8 @@ EEventLoadNodeExecutionResult FAsyncPackage2::Event_ExportsDone(FAsyncLoadingThr
 	check(PostLoadGroup->PackagesWithExportsToSerializeCount > 0);
 	--PostLoadGroup->PackagesWithExportsToSerializeCount;
 
-#if WITH_PARTIAL_REQUEST_DURING_RECURSION
-	// Release any requests that are only waiting on exports, allowing a flush to exit.
-	Package->AsyncLoadingThread.RemovePendingRequests(ThreadState, Package->RequestIDs[(int32)ERequestType::ExportsDone]);
-
-	FAsyncPackage2* RequesterPackage = nullptr;
-	if (Package->Desc.RequesterId.IsValid())
-	{
-		RequesterPackage = Package->AsyncLoadingThread.AsyncPackageLookup.FindRef(Package->Desc.RequesterId);
-	}
-
-	// If this is a recursive load and the caller has not started postloading yet, merge postload group.
-	if (RequesterPackage && RequesterPackage->PostLoadGroup != nullptr)
-	{
-		UE_ASYNC_PACKAGE_LOG(Display, Package->Desc, TEXT("ExportsDone"), TEXT("Merging postload groups with requester package %s"), *RequesterPackage->Desc.UPackageName.ToString());
-
-		// Do not adjust sync load context, we want to be able to exit the current one even if the caller has not finished yet.
-		const bool bUpdateSyncLoadContext = false;
-		Package->AsyncLoadingThread.MergePostLoadGroups(ThreadState, Package->PostLoadGroup, RequesterPackage->PostLoadGroup, bUpdateSyncLoadContext);
-	}
-	else
-#endif // WITH_PARTIAL_REQUEST_DURING_RECURSION
-	{
-		Package->AsyncLoadingThread.ConditionalBeginPostLoad(ThreadState, PostLoadGroup);
-	}
+	Package->ConditionalReleasePartialRequests(ThreadState);
+	Package->AsyncLoadingThread.ConditionalBeginPostLoad(ThreadState, PostLoadGroup);
 	return EEventLoadNodeExecutionResult::Complete;
 }
 
@@ -6950,13 +7072,7 @@ EEventLoadNodeExecutionResult FAsyncPackage2::Event_PostLoadExportBundle(FAsyncL
 #if ALT2_ENABLE_LINKERLOAD_SUPPORT
 	if (Package->LinkerLoadState.IsSet())
 	{
-		// No async postload for now
-		FAsyncLoadingPostLoadGroup* DeferredPostLoadGroup = Package->DeferredPostLoadGroup;
-		check(DeferredPostLoadGroup);
-		check(DeferredPostLoadGroup->PackagesWithExportsToPostLoadCount > 0);
-		--DeferredPostLoadGroup->PackagesWithExportsToPostLoadCount;
-		Package->AsyncLoadingThread.ConditionalBeginDeferredPostLoad(ThreadState, DeferredPostLoadGroup);
-		return EEventLoadNodeExecutionResult::Complete;
+		return Package->ExecutePostLoadLinkerLoadPackageExports(ThreadState);
 	}
 #endif
 
@@ -7105,7 +7221,7 @@ EEventLoadNodeExecutionResult FAsyncPackage2::Event_DeferredPostLoadExportBundle
 #if ALT2_ENABLE_LINKERLOAD_SUPPORT
 	if (Package->LinkerLoadState.IsSet())
 	{
-		return Package->PostLoadLinkerLoadPackageExports(ThreadState);
+		return Package->ExecuteDeferredPostLoadLinkerLoadPackageExports(ThreadState);
 	}
 #endif
 
@@ -7265,34 +7381,26 @@ void FAsyncLoadingThread2::UpdateSyncLoadContext(FAsyncLoadingThreadState2& Thre
 
 #if WITH_PARTIAL_REQUEST_DURING_RECURSION
 				FAsyncPackage2* RequestingPackage = SyncLoadContext->RequestingPackage;
-				if (RequestingPackage != nullptr &&
-					RequestedPackage->RequestIDs[(int32)ERequestType::Full].Contains(RequestID))
+				if (RequestingPackage != nullptr)
 				{
-					// When this happens, it's because a load was initiated from outside a package recursion and the requestid didn't have the partial loading flag applied to it.
-					// Then the requestid is being flushed from inside another package stack, in which case we might need to demote the request id to a partial one on certain condition.
-
 					// If the flush is coming from a step before the requesting package is back on GT, there is no way to fully flush
-					// the requested package unless its already done.
+					// the requested package unless its already done. We have no choice but to trigger partial loading in that case.
 					if (RequestingPackage->AsyncPackageLoadingState < EAsyncPackageLoadingState2::DeferredPostLoad &&
 						RequestedPackage->AsyncPackageLoadingState < EAsyncPackageLoadingState2::Complete)
 					{
-						//  
+						//
 						// Note: Update the FLoadingTests_RecursiveLoads_FullFlushFrom_Serialize test if you edit this error message
 						//
-						UE_LOG(LogStreaming, Warning, TEXT("Fully flushing package %s recursively from another package %s would lead into a deadlock. Demoting requestID %d to a partial request instead."),
+						UE_LOG(LogStreaming, Display, TEXT("Flushing package %s (state: %s) recursively from another package %s (state: %s) will result in a partially loaded package to avoid a deadlock."),
 							*RequestedPackage->Desc.UPackageName.ToString(),
+							LexToString(RequestedPackage->AsyncPackageLoadingState),
 							*RequestingPackage->Desc.UPackageName.ToString(),
+							LexToString(RequestingPackage->AsyncPackageLoadingState),
 							RequestID
 						);
 
-						RequestedPackage->RequestIDs[(int32)ERequestType::Full].Remove(RequestID);
-						RequestedPackage->RequestIDs[(int32)ERequestType::ExportsDone].Add(RequestID);
-						RequestedPackage->Desc.RequesterId = RequestingPackage->Desc.UPackageId;
-
-						if (RequestedPackage->AsyncPackageLoadingState >= EAsyncPackageLoadingState2::ExportsDone)
-						{
-							RemovePendingRequests(ThreadState, { RequestID });
-						}
+						// Check if partial loading rules allow to release the package right now.
+						RequestedPackage->ConditionalReleasePartialRequests(ThreadState);
 					}
 				}
 #endif
@@ -7611,10 +7719,7 @@ EAsyncPackageState::Type FAsyncLoadingThread2::ProcessLoadedPackagesFromGameThre
 
 			if (Package->CompletionCallbacks.IsEmpty() && Package->ProgressCallbacks.IsEmpty())
 			{
-				for (int32 Index = 0; Index < (int32)ERequestType::Count; ++Index)
-				{
-					RemovePendingRequests(ThreadState, Package->RequestIDs[Index]);
-				}
+				RemovePendingRequests(ThreadState, Package->RequestIDs);
 				Package->ReleaseRef();
 			}
 			else
@@ -8414,7 +8519,7 @@ FAsyncPackage2::FAsyncPackage2(
 , ImportStore(AsyncLoadingThread.GlobalImportStore)
 {
 	TRACE_LOADTIME_NEW_ASYNC_PACKAGE(this);
-	AddRequestID(ThreadState, Desc.RequestID, Desc.RequestType);
+	AddRequestID(ThreadState, Desc.RequestID);
 
 	CreatePackageNodes(EventSpecs);
 
@@ -8535,7 +8640,7 @@ void FAsyncPackage2::ClearConstructedObjects()
 	}
 }
 
-void FAsyncPackage2::AddRequestID(FAsyncLoadingThreadState2& ThreadState, int32 Id, ERequestType RequestType)
+void FAsyncPackage2::AddRequestID(FAsyncLoadingThreadState2& ThreadState, int32 Id)
 {
 	if (Id > 0)
 	{
@@ -8545,20 +8650,11 @@ void FAsyncPackage2::AddRequestID(FAsyncLoadingThreadState2& ThreadState, int32
 			Desc.RequestID = Id;
 		}
 
-		RequestIDs[(int32)RequestType].Add(Id);
+		RequestIDs.Add(Id);
 
 		// The Id is most likely already present because it's added as soon as the request is created.
 		AsyncLoadingThread.AddPendingRequest(Id);
 
-		// Remove any partial request if our state is already past that point.
-#if WITH_PARTIAL_REQUEST_DURING_RECURSION
-		static_assert((int32)ERequestType::Count == 2, "Please adjust this code when adding a new request type");
-		if (RequestType == ERequestType::ExportsDone && AsyncPackageLoadingState >= EAsyncPackageLoadingState2::ExportsDone)
-		{
-			AsyncLoadingThread.RemovePendingRequests(ThreadState, {Id});
-		}
-#endif
-
 		TRACE_LOADTIME_ASYNC_PACKAGE_REQUEST_ASSOCIATION(this, Id);
 	}
 }
@@ -8834,44 +8930,7 @@ PRAGMA_ENABLE_DEPRECATION_WARNINGS
 	TRACE_LOADTIME_BEGIN_REQUEST(RequestId);
 	AddPendingRequest(RequestId);
 
-	FPackageId RequesterId;
-	ERequestType RequestType = ERequestType::Full;
-
-#if WITH_PARTIAL_REQUEST_DURING_RECURSION
-	FAsyncLoadingThreadState2* ThreadState = FAsyncLoadingThreadState2::Get();
-
-	// We want to skip postloads inside recursive loads that happens during serialization and preload
-	// since it is mostly what BP expects and how the old LinkerLoad was implemented.
-	// In some cases, trying to postload right away during recursion could end-up in data loss (i.e. UE-190649).
-	// The expectation is that postload should be run as part of the package causing this load to occur.
-	// As a more general rule, we want any synchronous load during deserialization to act like
-	// an import. We want the deserialization steps done before returning from the load but
-	// we keep postload for after deserialization is done on the referencer too so that postload 
-	// sees everything initialized in the case where it also has a reference back to the referencer.
-	// Using a partial request flag allows the flush to return as soon as exports are done which
-	// gives the opportunity to the caller to get a pointer to an object inside the package even if it's not fully loaded yet.
-	if (ThreadState != nullptr &&
-		ThreadState->CurrentlyExecutingEventNodeStack.Num() > 0)
-	{
-		FAsyncPackage2* Requester = ThreadState->CurrentlyExecutingEventNodeStack.Top()->GetPackage();
-
-		// If the requester doesn't have a postload group, it means it's already being postloaded and
-		// we don't need to load partially in that case.
-		if (Requester && Requester->PostLoadGroup)
-		{
-			UE_LOG(LogStreaming, Display,
-				TEXT("LoadPackage on %s is being called from another package %s, partial loading is automatically being applied which means only exports are processed before exiting flush and postloads will be called along with the caller's"),
-				*PackageNameToLoad.ToString(),
-				*Requester->Desc.UPackageName.ToString()
-			);
-
-			RequesterId = Requester->Desc.UPackageId;
-			RequestType = ERequestType::ExportsDone;
-		}
-	}
-#endif
-
-	PackageRequestQueue.Enqueue(FPackageRequest::Create(RequestId, RequestType, RequesterId, InPackageFlags, InLoadFlags, InPIEInstanceID, InPackagePriority, InInstancingContext, InPackagePath, InCustomName, MoveTemp(InCompletionDelegate), MoveTemp(InProgressDelegate)));
+	PackageRequestQueue.Enqueue(FPackageRequest::Create(RequestId, InPackageFlags, InLoadFlags, InPIEInstanceID, InPackagePriority, InInstancingContext, InPackagePath, InCustomName, MoveTemp(InCompletionDelegate), MoveTemp(InProgressDelegate)));
 	++QueuedPackagesCounter;
 	++PackagesWithRemainingWorkCounter;
 
@@ -8963,6 +9022,23 @@ EAsyncPackageState::Type FAsyncLoadingThread2::ProcessLoadingFromGameThread(FAsy
 	return IsAsyncLoadingPackages() ? EAsyncPackageState::TimeOut : EAsyncPackageState::Complete;
 }
 
+void FAsyncLoadingThread2::WarnAboutPotentialSyncLoadStall(FAsyncLoadingSyncLoadContext* SyncLoadContext)
+{
+	for (int32 Index = 0; Index < SyncLoadContext->RequestIDs.Num(); ++Index)
+	{
+		int32 RequestId = SyncLoadContext->RequestIDs[Index];
+		if (ContainsRequestID(RequestId))
+		{
+			FAsyncPackage2* Package = SyncLoadContext->RequestedPackages[Index];
+			UE_LOG(LogStreaming, Warning, TEXT("A flush request appear to be stuck waiting on package %s at state %s to reach state > %s"),
+				*Package->Desc.UPackageName.ToString(),
+				LexToString(Package->AsyncPackageLoadingState),
+				LexToString(SyncLoadContext->RequestingPackage->AsyncPackageLoadingState)
+			);
+		}
+	}
+}
+
 void FAsyncLoadingThread2::FlushLoadingFromLoadingThread(FAsyncLoadingThreadState2& ThreadState, TConstArrayView<int32> RequestIDs)
 {
 	TRACE_CPUPROFILER_EVENT_SCOPE(FlushLoadingFromLoadingThread);
@@ -8983,9 +9059,19 @@ void FAsyncLoadingThread2::FlushLoadingFromLoadingThread(FAsyncLoadingThreadStat
 	// We handle the context push/pop manually since we don't interact with the main thread during this time.
 	static constexpr bool bAutoHandleSyncLoadContext = false;
 	UpdateSyncLoadContext(ThreadState, bAutoHandleSyncLoadContext);
+
+	int64 DidNothingCount = 0;
 	while (ContainsAnyRequestID(SyncLoadContext->RequestIDs))
 	{
-		EventQueue.ExecuteSyncLoadEvents(ThreadState);
+		const bool bDidSomething = EventQueue.ExecuteSyncLoadEvents(ThreadState);
+		if (bDidSomething)
+		{
+			DidNothingCount = 0;
+		}
+		else if (DidNothingCount++ == 100)
+		{
+			WarnAboutPotentialSyncLoadStall(SyncLoadContext);
+		}
 	}
 
 	check(AsyncLoadingThreadState->SyncLoadContextStack.Top() == SyncLoadContext);
diff --git a/Engine/Source/Runtime/Engine/Private/Tests/Loading/AsyncLoadingTests_ConvertFromType.cpp b/Engine/Source/Runtime/Engine/Private/Tests/Loading/AsyncLoadingTests_ConvertFromType.cpp
--- a/Engine/Source/Runtime/Engine/Private/Tests/Loading/AsyncLoadingTests_ConvertFromType.cpp
+++ b/Engine/Source/Runtime/Engine/Private/Tests/Loading/AsyncLoadingTests_ConvertFromType.cpp
@@ -29,7 +29,7 @@ public:
 /**
  * This test validates ConvertFromType thread-safety.
  */
-IMPLEMENT_SIMPLE_AUTOMATION_TEST(FConvertFromType_SoftToHard, TEXT("System.Engine.Loading.ConvertFromType_SoftToHard"), EAutomationTestFlags::ApplicationContextMask | EAutomationTestFlags::EngineFilter)
+IMPLEMENT_SIMPLE_AUTOMATION_TEST(FConvertFromType_SoftToHard, TEXT("System.Engine.Loading.ConvertFromType_SoftToHard"), EAutomationTestFlags::EditorContext | EAutomationTestFlags::EngineFilter)
 bool FConvertFromType_SoftToHard::RunTest(const FString& Parameters)
 {
 	// Just make sure the async loading queue is empty before beginning.
@@ -115,6 +115,10 @@ bool FConvertFromType_SoftToHard::RunTest(const FString& Parameters)
 		check(FindObject<UAsyncLoadingTests_ConvertFromType_V2>(nullptr, ObjectPath2) != nullptr);
 		check(FindObject<UPackage>(nullptr, PackagePath2) != nullptr);
 
+		// Remove RF_Standalone from top level object.
+		FindObject<UAsyncLoadingTests_ConvertFromType_V2>(nullptr, ObjectPath1)->ClearFlags(RF_Standalone);
+		FindObject<UAsyncLoadingTests_ConvertFromType_V2>(nullptr, ObjectPath2)->ClearFlags(RF_Standalone);
+
 		FindObject<UPackage>(nullptr, PackagePath1)->GetMetaData()->ClearFlags(RF_Standalone);
 		FindObject<UPackage>(nullptr, PackagePath2)->GetMetaData()->ClearFlags(RF_Standalone);
 
@@ -129,5 +133,4 @@ bool FConvertFromType_SoftToHard::RunTest(const FString& Parameters)
 	return true;
 }
 
-#undef TEST_NAME_ROOT
 #endif // WITH_DEV_AUTOMATION_TESTS
diff --git a/Engine/Source/Runtime/Engine/Private/Tests/Loading/AsyncLoadingTests_IsReadyForAsyncPostLoad.cpp b/Engine/Source/Runtime/Engine/Private/Tests/Loading/AsyncLoadingTests_IsReadyForAsyncPostLoad.cpp
new file mode 100644
--- /dev/null
+++ b/Engine/Source/Runtime/Engine/Private/Tests/Loading/AsyncLoadingTests_IsReadyForAsyncPostLoad.cpp
@@ -0,0 +1,95 @@
+// Copyright Epic Games, Inc. All Rights Reserved.
+
+#include "AsyncLoadingTests_Shared.h"
+#include "Misc/AutomationTest.h"
+#include "Misc/PackageName.h"
+
+#if WITH_DEV_AUTOMATION_TESTS
+
+IMPLEMENT_SIMPLE_AUTOMATION_TEST(
+	FLoadingTests_IsReadyForAsyncPostLoad,
+	TEXT("System.Engine.Loading.IsReadyForAsyncPostLoad.Tick"),
+	EAutomationTestFlags::EditorContext | EAutomationTestFlags::EngineFilter
+)
+bool FLoadingTests_IsReadyForAsyncPostLoad::RunTest(const FString& Parameters)
+{
+	FLoadingTestsScope LoadingTestScope(this);
+
+	// Return IsReadyForAsyncPostLoad until directed otherwise.
+	bool bIsReadyForAsyncPostLoad = false;
+	UAsyncLoadingTests_Shared::OnIsReadyForAsyncPostLoad.BindLambda(
+		[&bIsReadyForAsyncPostLoad](const UAsyncLoadingTests_Shared* Object) 
+		{
+			return bIsReadyForAsyncPostLoad;
+		}
+	);
+
+	// Sentinel to validate at which point postload is called exactly
+	bool bPostLoadCalled = false;
+	UAsyncLoadingTests_Shared::OnPostLoad.BindLambda(
+		[&bPostLoadCalled](UAsyncLoadingTests_Shared* Object) 
+		{ 
+			bPostLoadCalled = true; 
+		} 
+	);
+
+	// Trigger the async loading
+	int32 RequestId = LoadPackageAsync(FLoadingTestsScope::PackagePath1);
+
+	// Wait for a good amount of time to make sure we're really timing out because bIsReadyForAsyncPostLoad is false.
+	EAsyncPackageState::Type State = ProcessAsyncLoadingUntilComplete([&bPostLoadCalled]() { return bPostLoadCalled; }, 5.0);
+	TestNotEqual("Async loading should not have been able to complete due to bIsReadyForAsyncPostLoad being false", State, EAsyncPackageState::Complete);
+	TestFalse("No postload should have been called", bPostLoadCalled);
+
+	// Start returning that we are ready for postload
+	bIsReadyForAsyncPostLoad = true;
+	State = ProcessAsyncLoadingUntilComplete([&bPostLoadCalled]() { return bPostLoadCalled; }, 1.0);
+
+	// Now we expect this to have completed in a timely manner
+	TestEqual("AsyncLoading should now be complete", State, EAsyncPackageState::Complete);
+	TestTrue("Postload should have been called", bPostLoadCalled);
+
+	return true;
+}
+
+IMPLEMENT_CUSTOM_SIMPLE_AUTOMATION_TEST(
+	FLoadingTests_IsReadyForAsyncPostLoad_Flush,
+	FLoadingTests_ZenLoaderOnly_Base,
+	TEXT("System.Engine.Loading.IsReadyForAsyncPostLoad.Flush"),
+	EAutomationTestFlags::EditorContext | EAutomationTestFlags::EngineFilter
+)
+bool FLoadingTests_IsReadyForAsyncPostLoad_Flush::RunTest(const FString& Parameters)
+{
+	FLoadingTestsScope LoadingTestScope(this);
+
+	// Return IsReadyForAsyncPostLoad until directed otherwise.
+	bool bIsReadyForAsyncPostLoad = false;
+	UAsyncLoadingTests_Shared::OnIsReadyForAsyncPostLoad.BindLambda(
+		[&bIsReadyForAsyncPostLoad](const UAsyncLoadingTests_Shared* Object)
+		{
+			return bIsReadyForAsyncPostLoad;
+		}
+	);
+
+	// Sentinel to validate at which point postload is called exactly
+	bool bPostLoadCalled = false;
+	UAsyncLoadingTests_Shared::OnPostLoad.BindLambda(
+		[&bPostLoadCalled](UAsyncLoadingTests_Shared* Object)
+		{
+			bPostLoadCalled = true;
+		}
+	);
+
+	// Trigger the async loading
+	int32 RequestId = LoadPackageAsync(FLoadingTestsScope::PackagePath1);
+
+	// When a flush occurs, we have no choice but to ignore bIsReadyForAsyncPostLoad
+	FlushAsyncLoading(RequestId);
+
+	// Now we expect this to have completed in a timely manner
+	TestTrue("Postload should have been called", bPostLoadCalled);
+
+	return true;
+}
+
+#endif // WITH_DEV_AUTOMATION_TESTS
diff --git a/Engine/Source/Runtime/Engine/Private/Tests/Loading/AsyncLoadingTests_RecursiveLoads.cpp b/Engine/Source/Runtime/Engine/Private/Tests/Loading/AsyncLoadingTests_RecursiveLoads.cpp
--- a/Engine/Source/Runtime/Engine/Private/Tests/Loading/AsyncLoadingTests_RecursiveLoads.cpp
+++ b/Engine/Source/Runtime/Engine/Private/Tests/Loading/AsyncLoadingTests_RecursiveLoads.cpp
@@ -1,180 +1,154 @@
 // Copyright Epic Games, Inc. All Rights Reserved.
 
-#include "AsyncLoadingTests_RecursiveLoads.h"
+#include "AsyncLoadingTests_Shared.h"
 #include "Misc/AutomationTest.h"
 #include "Misc/PackageName.h"
 #include "UObject/CoreRedirects.h"
 #include "UObject/SavePackage.h"
 
-#include UE_INLINE_GENERATED_CPP_BY_NAME(AsyncLoadingTests_RecursiveLoads)
-
-UAsyncLoadingTests_RecursiveLoads::FOnPostLoadEvent UAsyncLoadingTests_RecursiveLoads::OnPostLoadEvent;
-UAsyncLoadingTests_RecursiveLoads::FOnSerializeEvent UAsyncLoadingTests_RecursiveLoads::OnSerializeEvent;
-
 #if WITH_DEV_AUTOMATION_TESTS
 
-class FLoadingTests_RecursiveLoads : public FAutomationTestBase
+// All RecursiveLoads tests should run on zenloader only as the other loaders are not compliant.
+typedef FLoadingTests_ZenLoaderOnly_Base FLoadingTests_RecursiveLoads_Base;
+
+/**
+ * This test validates loading an object synchronously during serialize.
+ */
+IMPLEMENT_CUSTOM_SIMPLE_AUTOMATION_TEST(
+	FLoadingTests_RecursiveLoads_FromSerialize, 
+	FLoadingTests_RecursiveLoads_Base,
+	TEXT("System.Engine.Loading.RecursiveLoads.FromSerialize"),
+	EAutomationTestFlags::EditorContext | EAutomationTestFlags::EngineFilter
+)
+bool FLoadingTests_RecursiveLoads_FromSerialize::RunTest(const FString& Parameters)
 {
-public:
-	using FAutomationTestBase::FAutomationTestBase;
-protected:
-	static constexpr const TCHAR* ObjectName = TEXT("TestObject");
-	static constexpr const TCHAR* PackagePath1 = TEXT("/Engine/TestRecursiveLoads_Package1");
-	static constexpr const TCHAR* ObjectPath1 = TEXT("/Engine/TestRecursiveLoads_Package1.TestObject");
-	static constexpr const TCHAR* PackagePath2 = TEXT("/Engine/TestRecursiveLoads_Package2");
-	static constexpr const TCHAR* ObjectPath2 = TEXT("/Engine/TestRecursiveLoads_Package2.TestObject");
-
-	virtual void CreateObjects()
-	{
-		UPackage* Package1 = CreatePackage(PackagePath1);
-		UAsyncLoadingTests_RecursiveLoads* Object1 = NewObject<UAsyncLoadingTests_RecursiveLoads>(Package1, ObjectName, RF_Public | RF_Standalone);
-
-		UPackage* Package2 = CreatePackage(PackagePath2);
-		UAsyncLoadingTests_RecursiveLoads* Object2 = NewObject<UAsyncLoadingTests_RecursiveLoads>(Package2, ObjectName, RF_Public | RF_Standalone);
-
-		// This is the soft reference that we want to test loading for
-		Object1->SoftReference = Object2;
-
-		// To avoid an error on save, we need to mark the package as fully loaded.
-		Package1->MarkAsFullyLoaded();
-		Package2->MarkAsFullyLoaded();
-
-		// Save packages to disk.
-		check(UPackage::SavePackage(Package1, nullptr, *FPackageName::LongPackageNameToFilename(PackagePath1, FPackageName::GetAssetPackageExtension()), FSavePackageArgs()));
-		check(UPackage::SavePackage(Package2, nullptr, *FPackageName::LongPackageNameToFilename(PackagePath2, FPackageName::GetAssetPackageExtension()), FSavePackageArgs()));
-
-		// Remove RF_Standalone from top level object.
-		Object1->ClearFlags(RF_Standalone);
-		Object2->ClearFlags(RF_Standalone);
-		
-		// Remove RF_Standalone from the UMetaData otherwise the package will not GC.
-		Package1->GetMetaData()->ClearFlags(RF_Standalone);
-		Package2->GetMetaData()->ClearFlags(RF_Standalone);
-
-		// GC and make sure everything gets cleaned up before loading
-		check(FindObject<UAsyncLoadingTests_RecursiveLoads>(nullptr, ObjectPath1) != nullptr);
-		check(FindObject<UPackage>(nullptr, PackagePath1) != nullptr);
-		check(FindObject<UAsyncLoadingTests_RecursiveLoads>(nullptr, ObjectPath2) != nullptr);
-		check(FindObject<UPackage>(nullptr, PackagePath2) != nullptr);
-
-		CollectGarbage(GARBAGE_COLLECTION_KEEPFLAGS);
-
-		check(FindObject<UAsyncLoadingTests_RecursiveLoads>(nullptr, ObjectPath1) == nullptr);
-		check(FindObject<UPackage>(nullptr, PackagePath1) == nullptr);
-		check(FindObject<UAsyncLoadingTests_RecursiveLoads>(nullptr, ObjectPath2) == nullptr);
-		check(FindObject<UPackage>(nullptr, PackagePath2) == nullptr);
-	}
-
-	virtual void LoadObjects()
-	{
-		LoadPackage(nullptr, PackagePath1, LOAD_None);
-
-		UAsyncLoadingTests_RecursiveLoads* Object2 = FindObject<UAsyncLoadingTests_RecursiveLoads>(nullptr, ObjectPath2);
-		TestFalse(TEXT("The object should have been properly loaded recursively"), Object2->HasAnyFlags(RF_NeedLoad | RF_NeedPostLoad));
-
-		UAsyncLoadingTests_RecursiveLoads* Object1 = FindObject<UAsyncLoadingTests_RecursiveLoads>(nullptr, ObjectPath1);
-		TestFalse(TEXT("The object should have been properly loaded recursively"), Object1->HasAnyFlags(RF_NeedLoad | RF_NeedPostLoad));
-	}
-
-	virtual void CleanupObjects()
-	{
-		// GC and make sure everything gets cleaned up before exiting
-		check(FindObject<UAsyncLoadingTests_RecursiveLoads>(nullptr, ObjectPath1) != nullptr);
-		check(FindObject<UPackage>(nullptr, PackagePath1) != nullptr);
-		check(FindObject<UAsyncLoadingTests_RecursiveLoads>(nullptr, ObjectPath2) != nullptr);
-		check(FindObject<UPackage>(nullptr, PackagePath2) != nullptr);
-
-		FindObject<UAsyncLoadingTests_RecursiveLoads>(nullptr, ObjectPath1)->ClearFlags(RF_Standalone);
-		FindObject<UAsyncLoadingTests_RecursiveLoads>(nullptr, ObjectPath2)->ClearFlags(RF_Standalone);
-
-		FindObject<UPackage>(nullptr, PackagePath1)->GetMetaData()->ClearFlags(RF_Standalone);
-		FindObject<UPackage>(nullptr, PackagePath2)->GetMetaData()->ClearFlags(RF_Standalone);
-
-		CollectGarbage(GARBAGE_COLLECTION_KEEPFLAGS);
-
-		check(FindObject<UAsyncLoadingTests_RecursiveLoads>(nullptr, ObjectPath1) == nullptr);
-		check(FindObject<UPackage>(nullptr, PackagePath1) == nullptr);
-		check(FindObject<UAsyncLoadingTests_RecursiveLoads>(nullptr, ObjectPath2) == nullptr);
-		check(FindObject<UPackage>(nullptr, PackagePath2) == nullptr);
-
-		UAsyncLoadingTests_RecursiveLoads::OnPostLoadEvent.Clear();
-		UAsyncLoadingTests_RecursiveLoads::OnSerializeEvent.Clear();
-	}
-
-	bool CanRunInEnvironment(const FString& TestParams, FString* OutReason, bool* OutWarn) const override
-	{
-		if (GetLoaderName() != TEXT("ZenLoader"))
+	FLoadingTestsScope LoadingTestScope(this);
+
+	UAsyncLoadingTests_Shared::OnSerialize.BindLambda(
+		[this](FArchive& Ar, UAsyncLoadingTests_Shared* Object)
 		{
-			if (OutReason)
+			if (Ar.IsLoading())
 			{
-				*OutReason = FString::Printf(TEXT("Test %s is for ZenLoader only. Cannot run on non-compliant loader currently active: %s"), *GetTestName(), *GetLoaderName().ToString());
+				if (UObject* Obj = Object->SoftReference.LoadSynchronous())
+				{
+					TestTrue(TEXT("Recursive loads in serialize should be deserialized"), !Obj->HasAnyFlags(RF_NeedLoad));
+					TestTrue(TEXT("Recursive loads in serialize should have their postload delayed until the caller is also ready to postload"), Obj->HasAnyFlags(RF_NeedPostLoad));
+				}
 			}
-			return false;
 		}
+	);
 
-		return true;
-	}
+	LoadingTestScope.LoadObjects();
 
-	bool DoTest()
-	{
-		// Just make sure the async loading queue is empty before beginning.
-		FlushAsyncLoading();
+	return true;
+}
 
-		// Creation phase
-		CreateObjects();
+/**
+ * This test validates loading an object with a thread-safe postload synchronously inside a thread-safe postload.
+ */
+IMPLEMENT_CUSTOM_SIMPLE_AUTOMATION_TEST(
+	FLoadingTests_RecursiveLoads_FromPostLoad_ThreadSafe,
+	FLoadingTests_RecursiveLoads_Base,
+	TEXT("System.Engine.Loading.RecursiveLoads.FromPostLoad_ThreadSafe"),
+	EAutomationTestFlags::EditorContext | EAutomationTestFlags::EngineFilter
+)
+bool FLoadingTests_RecursiveLoads_FromPostLoad_ThreadSafe::RunTest(const FString& Parameters)
+{
+	FLoadingTestsScope LoadingTestScope(this);
 
-		// Loading phase
-		LoadObjects();
+	// Set both object for thread-safe postloads so we get called earlier.
+	UAsyncLoadingTests_Shared::OnIsPostLoadThreadSafe.BindLambda(
+		[](const UAsyncLoadingTests_Shared* Object)
+		{
+			return true;
+		}
+	);
 
-		// Cleanup phase
-		CleanupObjects();
+	// When running with async loading thread, this should be called on ALT.
+	UAsyncLoadingTests_Shared::OnPostLoad.BindLambda(
+		[this](UAsyncLoadingTests_Shared* Object)
+		{
+			// We expect objects that are thread-safe to postload, to have been postloaded before returning from their sync load.
+			if (UObject* Obj = Object->SoftReference.LoadSynchronous())
+			{
+				TestFalse(TEXT("Sync loads inside thread-safe postload of objects that are thread-safe to postload should be fully loaded"), Obj->HasAnyFlags(RF_NeedLoad | RF_NeedPostLoad));
+			}
+		}
+	);
 
-		return true;
-	}
-};
+	LoadingTestScope.LoadObjects();
+
+	return true;
+}
 
 /**
- * This test validates loading an object synchronously during serialize.
+ * This test validates loading an object synchronously inside a thread-safe (non-deferred) postload.
  */
 IMPLEMENT_CUSTOM_SIMPLE_AUTOMATION_TEST(
-	FLoadingTests_RecursiveLoads_FromSerialize, 
-	FLoadingTests_RecursiveLoads, 
-	TEXT("System.Engine.Loading.RecursiveLoads.FromSerialize"),
-	EAutomationTestFlags::ApplicationContextMask | EAutomationTestFlags::EngineFilter
+	FLoadingTests_RecursiveLoads_FromPostLoad_ThreadUnsafe,
+	FLoadingTests_RecursiveLoads_Base,
+	TEXT("System.Engine.Loading.RecursiveLoads.FromPostLoad_ThreadUnsafe"),
+	EAutomationTestFlags::EditorContext | EAutomationTestFlags::EngineFilter
 )
-bool FLoadingTests_RecursiveLoads_FromSerialize::RunTest(const FString& Parameters)
+bool FLoadingTests_RecursiveLoads_FromPostLoad_ThreadUnsafe::RunTest(const FString& Parameters)
 {
-	UAsyncLoadingTests_RecursiveLoads::OnSerializeEvent.AddLambda(
-		[this](FArchive& Ar, UAsyncLoadingTests_RecursiveLoads* Object)
+	FLoadingTestsScope LoadingTestScope(this);
+
+	// Make the first postloads as thread-safe
+	UAsyncLoadingTests_Shared::OnIsPostLoadThreadSafe.BindLambda(
+		[](const UAsyncLoadingTests_Shared* Object)
 		{
-			if (Ar.IsLoading())
+			return Object->GetPathName() == FLoadingTestsScope::ObjectPath1;
+		}
+	);
+
+	UAsyncLoadingTests_Shared::OnPostLoad.BindLambda(
+		[this](UAsyncLoadingTests_Shared* Object)
+		{
+			if (Object->GetPathName() == FLoadingTestsScope::ObjectPath1)
 			{
+				if (IsAsyncLoadingMultithreaded())
+				{
+					TestFalse(TEXT("Thread-safe postloads should get called from the async loading thread when it is active"), IsInGameThread());
+				}
+			
 				if (UObject* Obj = Object->SoftReference.LoadSynchronous())
 				{
-					TestTrue(TEXT("Recursive loads in serialize should be deserialized"), !Obj->HasAnyFlags(RF_NeedLoad));
-					TestTrue(TEXT("Recursive loads in serialize should have their postload delayed until the caller is also ready to postload"), Obj->HasAnyFlags(RF_NeedPostLoad));
+					TestTrue(TEXT("Sync loads of non thread-safe objects from thread-safe postload should be deserialized"), !Obj->HasAnyFlags(RF_NeedLoad));
+
+					// Since the Object returned will continue to postload on the game-thread, we can't verify the RF_NeedPostLoad flag as we could race
+					// trying to look at the value depending on its state on the game-thread.
+					
+					// What we can do is verify below that postload is called from the game-thread on the object
 				}
 			}
+			else
+			{
+				TestTrue(TEXT("Sync loads of non thread-safe objects from thread-safe postload should have their postload deferred on the game-thread"), IsInGameThread());
+			}
 		}
 	);
 
-	return DoTest();
+	LoadingTestScope.LoadObjects();
+
+	return true;
 }
 
 /**
  * This test validates loading an object synchronously during postload.
  */
 IMPLEMENT_CUSTOM_SIMPLE_AUTOMATION_TEST(
-	FLoadingTests_RecursiveLoads_FromPostLoad,
-	FLoadingTests_RecursiveLoads,
-	TEXT("System.Engine.Loading.RecursiveLoads.FromPostLoad"),
-	EAutomationTestFlags::ApplicationContextMask | EAutomationTestFlags::EngineFilter
+	FLoadingTests_RecursiveLoads_FromDeferredPostLoad,
+	FLoadingTests_RecursiveLoads_Base,
+	TEXT("System.Engine.Loading.RecursiveLoads.FromDeferredPostLoad"),
+	EAutomationTestFlags::EditorContext | EAutomationTestFlags::EngineFilter
 )
-
-bool FLoadingTests_RecursiveLoads_FromPostLoad::RunTest(const FString& Parameters)
+bool FLoadingTests_RecursiveLoads_FromDeferredPostLoad::RunTest(const FString& Parameters)
 {
-	UAsyncLoadingTests_RecursiveLoads::OnPostLoadEvent.AddLambda(
-		[this](UAsyncLoadingTests_RecursiveLoads* Object)
+	FLoadingTestsScope LoadingTestScope(this);
+
+	UAsyncLoadingTests_Shared::OnPostLoad.BindLambda(
+		[this](UAsyncLoadingTests_Shared* Object)
 		{
 			if (UObject* Obj = Object->SoftReference.LoadSynchronous())
 			{
@@ -183,86 +157,87 @@ bool FLoadingTests_RecursiveLoads_FromPostLoad::RunTest(const FString& Parameter
 		}
 	);
 
-	return DoTest();
+	LoadingTestScope.LoadObjects();
+
+	return true;
 }
 
 /**
  * This test validates an error is emitted when flushing a requestid that is not a partial load from inside a recursive function.
  */
-class FLoadingTests_RecursiveLoads_FullFlushFrom_Base : public FLoadingTests_RecursiveLoads
-{
-public:
-	using FLoadingTests_RecursiveLoads::FLoadingTests_RecursiveLoads;
-protected:
-	int32 RequestId;
-
-	virtual void LoadObjects() override
-	{
-		// Create a request before starting the loading test so we get a request that is not tagged as partial.
-		RequestId = LoadPackageAsync(PackagePath2);
-
-		FLoadingTests_RecursiveLoads::LoadObjects();
-	}
-};
-
 IMPLEMENT_CUSTOM_SIMPLE_AUTOMATION_TEST(
 	FLoadingTests_RecursiveLoads_FullFlushFrom_Serialize,
-	FLoadingTests_RecursiveLoads_FullFlushFrom_Base,
+	FLoadingTests_RecursiveLoads_Base,
 	TEXT("System.Engine.Loading.RecursiveLoads.FullFlushFrom.Serialize"),
-	EAutomationTestFlags::ApplicationContextMask | EAutomationTestFlags::EngineFilter
+	EAutomationTestFlags::EditorContext | EAutomationTestFlags::EngineFilter
 )
 bool FLoadingTests_RecursiveLoads_FullFlushFrom_Serialize::RunTest(const FString& Parameters)
 {
-	AddExpectedError(TEXT("would lead into a deadlock. Demoting requestID"), EAutomationExpectedErrorFlags::Contains);
+	FLoadingTestsScope LoadingTestScope(this);
 
-	UAsyncLoadingTests_RecursiveLoads::OnSerializeEvent.AddLambda(
-		[this](FArchive& Ar, UAsyncLoadingTests_RecursiveLoads* Object)
+	AddExpectedMessage(TEXT("will result in a partially loaded package to avoid a deadlock."), EAutomationExpectedErrorFlags::Contains);
+
+	int32 RequestId = 0;
+	UAsyncLoadingTests_Shared::OnSerialize.BindLambda(
+		[this, &RequestId](FArchive& Ar, UAsyncLoadingTests_Shared* Object)
 		{
 			// Do not try to flush ourself as this would lead to a fatal error :)
 			// Just flush Package2 when we're in Package1
-			if (Ar.IsLoading() && Object->GetPathName() == ObjectPath1)
+			if (Ar.IsLoading() && Object->GetPathName() == FLoadingTestsScope::ObjectPath1)
 			{
 				// Flush the requestId that has been created outside of the recursive load. This request
 				// should be a full request and flushing it should result in an error being reported.
 				FlushAsyncLoading(RequestId);
 
-				UAsyncLoadingTests_RecursiveLoads* Object2 = FindObject<UAsyncLoadingTests_RecursiveLoads>(nullptr, ObjectPath2);
+				UAsyncLoadingTests_Shared* Object2 = FindObject<UAsyncLoadingTests_Shared>(nullptr, FLoadingTestsScope::ObjectPath2);
 				TestFalse(TEXT("The object should be serialized"), Object2->HasAnyFlags(RF_NeedLoad));
 				TestTrue(TEXT("The object should not have been postloaded"), Object2->HasAnyFlags(RF_NeedPostLoad));
 			}
 		}
 	);
 
-	return DoTest();
+	// Create a request before starting the loading test so we get a request that is not tagged as partial.
+	RequestId = LoadPackageAsync(FLoadingTestsScope::PackagePath2);
+
+	LoadingTestScope.LoadObjects();
+
+	return true;
 }
 
 IMPLEMENT_CUSTOM_SIMPLE_AUTOMATION_TEST(
 	FLoadingTests_RecursiveLoads_FullFlushFrom_PostLoad,
-	FLoadingTests_RecursiveLoads_FullFlushFrom_Base,
+	FLoadingTests_RecursiveLoads_Base,
 	TEXT("System.Engine.Loading.RecursiveLoads.FullFlushFrom.Postload"),
-	EAutomationTestFlags::ApplicationContextMask | EAutomationTestFlags::EngineFilter
+	EAutomationTestFlags::EditorContext | EAutomationTestFlags::EngineFilter
 )
 bool FLoadingTests_RecursiveLoads_FullFlushFrom_PostLoad::RunTest(const FString& Parameters)
 {
-	UAsyncLoadingTests_RecursiveLoads::OnPostLoadEvent.AddLambda(
-		[this](UAsyncLoadingTests_RecursiveLoads* Object)
+	FLoadingTestsScope LoadingTestScope(this);
+
+	int32 RequestId = 0;
+	UAsyncLoadingTests_Shared::OnPostLoad.BindLambda(
+		[this, &RequestId](UAsyncLoadingTests_Shared* Object)
 		{
 			// Do not try to flush ourself as this would lead to a fatal error :)
 			// Just flush Package2 when we're in Package1
-			if (Object->GetPathName() == ObjectPath1)
+			if (Object->GetPathName() == FLoadingTestsScope::ObjectPath1)
 			{
 				// Flush the requestId that has been created outside of the recursive load. This request
 				// should be a full request and flushing it should result in an error being reported.
 				FlushAsyncLoading(RequestId);
 
-				UAsyncLoadingTests_RecursiveLoads* Object2 = FindObject<UAsyncLoadingTests_RecursiveLoads>(nullptr, ObjectPath2);
+				UAsyncLoadingTests_Shared* Object2 = FindObject<UAsyncLoadingTests_Shared>(nullptr, FLoadingTestsScope::ObjectPath2);
 				TestFalse(TEXT("The object should be serialized and postloaded"), Object2->HasAnyFlags(RF_NeedLoad | RF_NeedPostLoad));
 			}
 		}
 	);
 
-	return DoTest();
+	// Create a request before starting the loading test so we get a request that is not tagged as partial.
+	RequestId = LoadPackageAsync(FLoadingTestsScope::PackagePath2);
+
+	LoadingTestScope.LoadObjects();
+
+	return true;
 }
 
-#undef TEST_NAME_ROOT
 #endif // WITH_DEV_AUTOMATION_TESTS
diff --git a/Engine/Source/Runtime/Engine/Private/Tests/Loading/AsyncLoadingTests_RecursiveLoads.h b/Engine/Source/Runtime/Engine/Private/Tests/Loading/AsyncLoadingTests_RecursiveLoads.h
deleted file mode 100644
diff --git a/Engine/Source/Runtime/Engine/Private/Tests/Loading/AsyncLoadingTests_Shared.cpp b/Engine/Source/Runtime/Engine/Private/Tests/Loading/AsyncLoadingTests_Shared.cpp
new file mode 100644
--- /dev/null
+++ b/Engine/Source/Runtime/Engine/Private/Tests/Loading/AsyncLoadingTests_Shared.cpp
@@ -0,0 +1,115 @@
+// Copyright Epic Games, Inc. All Rights Reserved.
+
+#include "AsyncLoadingTests_Shared.h"
+#include "Misc/AutomationTest.h"
+#include "Misc/PackageName.h"
+#include "UObject/SavePackage.h"
+
+UAsyncLoadingTests_Shared::FOnPostLoadDelegate UAsyncLoadingTests_Shared::OnPostLoad;
+UAsyncLoadingTests_Shared::FOnSerializeDelegate UAsyncLoadingTests_Shared::OnSerialize;
+UAsyncLoadingTests_Shared::FOnIsReadyForAsyncPostLoadDelegate UAsyncLoadingTests_Shared::OnIsReadyForAsyncPostLoad;
+UAsyncLoadingTests_Shared::FOnIsPostLoadThreadSafeDelegate UAsyncLoadingTests_Shared::OnIsPostLoadThreadSafe;
+
+#if WITH_DEV_AUTOMATION_TESTS
+
+void FLoadingTestsScope::CreateObjects()
+{
+	Package1 = CreatePackage(PackagePath1);
+	Object1 = NewObject<UAsyncLoadingTests_Shared>(Package1, ObjectName, RF_Public | RF_Standalone);
+
+	Package2 = CreatePackage(PackagePath2);
+	Object2 = NewObject<UAsyncLoadingTests_Shared>(Package2, ObjectName, RF_Public | RF_Standalone);
+}
+
+void FLoadingTestsScope::MutateObjects()
+{
+	// This is the soft reference that we want to test loading for
+	Object1->SoftReference = Object2;
+}
+
+void FLoadingTestsScope::SaveObjects()
+{
+	// To avoid an error on save, we need to mark the package as fully loaded.
+	Package1->MarkAsFullyLoaded();
+	Package2->MarkAsFullyLoaded();
+
+	// Save packages to disk.
+	check(UPackage::SavePackage(Package1, nullptr, *FPackageName::LongPackageNameToFilename(PackagePath1, FPackageName::GetAssetPackageExtension()), FSavePackageArgs()));
+	check(UPackage::SavePackage(Package2, nullptr, *FPackageName::LongPackageNameToFilename(PackagePath2, FPackageName::GetAssetPackageExtension()), FSavePackageArgs()));
+
+	// Remove RF_Standalone from top level object.
+	Object1->ClearFlags(RF_Standalone);
+	Object2->ClearFlags(RF_Standalone);
+
+	// Remove RF_Standalone from the UMetaData otherwise the package will not GC.
+	Package1->GetMetaData()->ClearFlags(RF_Standalone);
+	Package2->GetMetaData()->ClearFlags(RF_Standalone);
+
+	// GC and make sure everything gets cleaned up before loading
+	check(FindObject<UAsyncLoadingTests_Shared>(nullptr, ObjectPath1) != nullptr);
+	check(FindObject<UPackage>(nullptr, PackagePath1) != nullptr);
+	check(FindObject<UAsyncLoadingTests_Shared>(nullptr, ObjectPath2) != nullptr);
+	check(FindObject<UPackage>(nullptr, PackagePath2) != nullptr);
+
+	CollectGarbage(GARBAGE_COLLECTION_KEEPFLAGS);
+
+	check(FindObject<UAsyncLoadingTests_Shared>(nullptr, ObjectPath1) == nullptr);
+	check(FindObject<UPackage>(nullptr, PackagePath1) == nullptr);
+	check(FindObject<UAsyncLoadingTests_Shared>(nullptr, ObjectPath2) == nullptr);
+	check(FindObject<UPackage>(nullptr, PackagePath2) == nullptr);
+}
+
+void FLoadingTestsScope::LoadObjects()
+{
+	LoadPackage(nullptr, PackagePath1, LOAD_None);
+
+	UAsyncLoadingTests_Shared* LocalObject1 = FindObject<UAsyncLoadingTests_Shared>(nullptr, ObjectPath1);
+	AutomationTest.TestTrue(TEXT("The object should have been properly loaded recursively"), LocalObject1 != nullptr);
+	if (LocalObject1)
+	{
+		AutomationTest.TestFalse(TEXT("The object should have been properly loaded recursively"), LocalObject1->HasAnyFlags(RF_NeedLoad | RF_NeedPostLoad));
+	}
+}
+
+void FLoadingTestsScope::CleanupObjects()
+{
+	// GC and make sure everything gets cleaned up before exiting
+	Object1 = FindObject<UAsyncLoadingTests_Shared>(nullptr, ObjectPath1);
+	Object2 = FindObject<UAsyncLoadingTests_Shared>(nullptr, ObjectPath2);
+	
+	if (Object1)
+	{
+		Object1->ClearFlags(RF_Standalone);
+	}
+	if (Object2)
+	{
+		Object2->ClearFlags(RF_Standalone);
+	}
+
+	Package1 = FindObject<UPackage>(nullptr, PackagePath1);
+	Package2 = FindObject<UPackage>(nullptr, PackagePath2);
+
+	if (Package1)
+	{
+		Package1->GetMetaData()->ClearFlags(RF_Standalone);
+	}
+	
+	if (Package2)
+	{
+		Package2->GetMetaData()->ClearFlags(RF_Standalone);
+	}
+
+	CollectGarbage(GARBAGE_COLLECTION_KEEPFLAGS);
+
+	check(FindObject<UAsyncLoadingTests_Shared>(nullptr, ObjectPath1) == nullptr);
+	check(FindObject<UPackage>(nullptr, PackagePath1) == nullptr);
+	check(FindObject<UAsyncLoadingTests_Shared>(nullptr, ObjectPath2) == nullptr);
+	check(FindObject<UPackage>(nullptr, PackagePath2) == nullptr);
+
+	UAsyncLoadingTests_Shared::OnPostLoad.Unbind();
+	UAsyncLoadingTests_Shared::OnSerialize.Unbind();
+	UAsyncLoadingTests_Shared::OnIsPostLoadThreadSafe.Unbind();
+	UAsyncLoadingTests_Shared::OnIsReadyForAsyncPostLoad.Unbind();
+}
+
+#endif // WITH_DEV_AUTOMATION_TESTS
diff --git a/Engine/Source/Runtime/Engine/Private/Tests/Loading/AsyncLoadingTests_Shared.h b/Engine/Source/Runtime/Engine/Private/Tests/Loading/AsyncLoadingTests_Shared.h
new file mode 100644
--- /dev/null
+++ b/Engine/Source/Runtime/Engine/Private/Tests/Loading/AsyncLoadingTests_Shared.h
@@ -0,0 +1,136 @@
+// Copyright Epic Games, Inc. All Rights Reserved.
+
+#pragma once
+
+#include "CoreMinimal.h"
+#include "Misc/AutomationTest.h"
+#include "UObject/Object.h"
+#include "UObject/ObjectPtr.h"
+#include "UObject/ObjectMacros.h"
+#include "AsyncLoadingTests_Shared.generated.h"
+
+UCLASS()
+class UAsyncLoadingTests_Shared : public UObject
+{
+	GENERATED_BODY()
+
+public:
+	UPROPERTY()
+	TSoftObjectPtr<UObject> SoftReference;
+
+	// Those delegates allow to easily change behavior between tests
+	using FOnPostLoadDelegate = TDelegate<void(UAsyncLoadingTests_Shared*), FDefaultTSDelegateUserPolicy>;
+	static FOnPostLoadDelegate OnPostLoad;
+
+	using FOnSerializeDelegate = TDelegate<void(FArchive& Ar, UAsyncLoadingTests_Shared*), FDefaultTSDelegateUserPolicy>;
+	static FOnSerializeDelegate OnSerialize;
+
+	using FOnIsReadyForAsyncPostLoadDelegate = TDelegate<bool(const UAsyncLoadingTests_Shared*), FDefaultTSDelegateUserPolicy>;
+	static FOnIsReadyForAsyncPostLoadDelegate OnIsReadyForAsyncPostLoad;
+
+	using FOnIsPostLoadThreadSafeDelegate = TDelegate<bool(const UAsyncLoadingTests_Shared*), FDefaultTSDelegateUserPolicy>;
+	static FOnIsReadyForAsyncPostLoadDelegate OnIsPostLoadThreadSafe;
+
+	virtual void PostLoad() override
+	{
+		Super::PostLoad();
+
+		OnPostLoad.ExecuteIfBound(this);
+	}
+
+	virtual void Serialize(FArchive& Ar) override
+	{
+		Super::Serialize(Ar);
+
+		OnSerialize.ExecuteIfBound(Ar, this);
+	}
+
+	virtual bool IsReadyForAsyncPostLoad() const override
+	{
+		if (OnIsReadyForAsyncPostLoad.IsBound())
+		{
+			return OnIsReadyForAsyncPostLoad.Execute(this);
+		}
+
+		return Super::IsReadyForAsyncPostLoad();
+	}
+
+	virtual bool IsPostLoadThreadSafe() const override
+	{
+		if (OnIsPostLoadThreadSafe.IsBound())
+		{
+			return OnIsPostLoadThreadSafe.Execute(this);
+		}
+
+		return Super::IsPostLoadThreadSafe();
+	}
+};
+
+#if WITH_DEV_AUTOMATION_TESTS
+
+class FLoadingTestsScope
+{
+private:
+	FAutomationTestBase& AutomationTest;
+public:
+	static constexpr const TCHAR* ObjectName = TEXT("TestObject");
+	static constexpr const TCHAR* PackagePath1 = TEXT("/Engine/TestRecursiveLoads_Package1");
+	static constexpr const TCHAR* ObjectPath1 = TEXT("/Engine/TestRecursiveLoads_Package1.TestObject");
+	static constexpr const TCHAR* PackagePath2 = TEXT("/Engine/TestRecursiveLoads_Package2");
+	static constexpr const TCHAR* ObjectPath2 = TEXT("/Engine/TestRecursiveLoads_Package2.TestObject");
+
+	UPackage* Package1 = nullptr;
+	UPackage* Package2 = nullptr;
+
+	UAsyncLoadingTests_Shared* Object1 = nullptr;
+	UAsyncLoadingTests_Shared* Object2 = nullptr;
+
+	virtual void CreateObjects();
+	virtual void MutateObjects();
+	virtual void SaveObjects();
+	virtual void LoadObjects();
+	virtual void CleanupObjects();
+
+	FLoadingTestsScope(FAutomationTestBase* InAutomationTest)
+		: AutomationTest(*InAutomationTest)
+	{
+		// Just make sure the async loading queue is empty before beginning.
+		FlushAsyncLoading();
+
+		CreateObjects();
+
+		MutateObjects();
+
+		SaveObjects();
+	}
+
+	virtual ~FLoadingTestsScope()
+	{
+		// Just to be sure in case the test fails, don't leave some pending loads behind
+		FlushAsyncLoading();
+
+		CleanupObjects();
+	}
+};
+
+class FLoadingTests_ZenLoaderOnly_Base : public FAutomationTestBase
+{
+public:
+	using FAutomationTestBase::FAutomationTestBase;
+protected:
+	bool CanRunInEnvironment(const FString& TestParams, FString* OutReason, bool* OutWarn) const override
+	{
+		if (GetLoaderName() != TEXT("ZenLoader"))
+		{
+			if (OutReason)
+			{
+				*OutReason = FString::Printf(TEXT("Test %s is for ZenLoader only. Cannot run on non-compliant loader currently active: %s"), *GetTestName(), *GetLoaderName().ToString());
+			}
+			return false;
+		}
+
+		return true;
+	}
+};
+
+#endif // WITH_DEV_AUTOMATION_TESTS
diff --git a/Engine/Source/Runtime/Engine/Private/Tests/Loading/AsyncLoadingTests_Timeouts.cpp b/Engine/Source/Runtime/Engine/Private/Tests/Loading/AsyncLoadingTests_Timeouts.cpp
new file mode 100644
--- /dev/null
+++ b/Engine/Source/Runtime/Engine/Private/Tests/Loading/AsyncLoadingTests_Timeouts.cpp
@@ -0,0 +1,51 @@
+// Copyright Epic Games, Inc. All Rights Reserved.
+
+#include "AsyncLoadingTests_Shared.h"
+#include "Misc/AutomationTest.h"
+#include "Misc/PackageName.h"
+#include "Async/ManualResetEvent.h"
+
+#if WITH_DEV_AUTOMATION_TESTS
+
+IMPLEMENT_SIMPLE_AUTOMATION_TEST(
+	FLoadingTests_Timeouts_Postload,
+	TEXT("System.Engine.Loading.Timeouts.Postload"),
+	EAutomationTestFlags::EditorContext | EAutomationTestFlags::EngineFilter
+)
+bool FLoadingTests_Timeouts_Postload::RunTest(const FString& Parameters)
+{
+	FLoadingTestsScope LoadingTestScope(this);
+
+	UE::FManualResetEvent PostLoadWaitEvent;
+
+	int32 PostLoadCount = 0;
+	UAsyncLoadingTests_Shared::OnPostLoad.BindLambda(
+		[this, &PostLoadCount](UAsyncLoadingTests_Shared* Object)
+		{
+			PostLoadCount++;
+			FPlatformProcess::Sleep(1.5);
+		}
+	);
+
+	// Trigger the async loading
+	LoadPackageAsync(FLoadingTestsScope::PackagePath1);
+	LoadPackageAsync(FLoadingTestsScope::PackagePath2);
+
+	// Try to wait until all postload are finished but we're timing out after 1 seconds, which doesn't let enough time to process both postloads.
+	EAsyncPackageState::Type State = ProcessAsyncLoadingUntilComplete([&PostLoadCount]() { return PostLoadCount == 2; }, 1.0);
+	
+	// The first call should have postloaded a single object and returned due to the elapse time.
+	TestEqual("Async loading should have timed out", State, EAsyncPackageState::TimeOut);
+	TestEqual("Only the first package should have been postloaded due to the timeout", PostLoadCount, 1);
+
+	// On the second iteration, we should be able to get throught the second postload
+	State = ProcessAsyncLoadingUntilComplete([&PostLoadCount]() { return PostLoadCount == 2; }, 5.0);
+
+	// The old loader could return PendingImports... do we want to make a special case for zenloader? 
+	TestTrue("Async loading should not have timeout on the second wait", State == EAsyncPackageState::Complete || State == EAsyncPackageState::PendingImports); 
+	TestEqual("All packages should have been postloaded by now", PostLoadCount, 2);
+
+	return true;
+}
+
+#endif // WITH_DEV_AUTOMATION_TESTS
