UE-188423 - Fix loader assumptions for BP objects

Added forced flushing of exports for Blueprint classes. Toggle with the "linker.EnableFullBlueprintPreloading" cvar in case any issues arise.

When Blueprints and their corresponding classes are compiled, they often need to reference other exports in the package. The expectation is that these exports are fully preloaded prior to compilation, though that's often not the case.
Both RegenerateBlueprintClass and RegenerateClass have incomplete logic for preloading the exports, and often miss several, such as portions of the widget tree in a widget Blueprint.

Instead of attempting to fix FPreloadMembersHelper to find all dependencies correctly, we can just exploit the fact that nearly all of the exports in a given Blueprint package will be referenced by either the Blueprint or its class. We can just force all of the exports to preload, and avoid having to traipse through object property references to figure out exactly what they reference.

While the previous behavior was incorrect, it didn't expose any obvious errors. However, we're trying to ensure that objects are always preloaded prior to postloading, which was often violated during Blueprint compilation.

diff --git a/Engine/Source/Runtime/CoreUObject/Private/Blueprint/BlueprintSupport.cpp b/Engine/Source/Runtime/CoreUObject/Private/Blueprint/BlueprintSupport.cpp
--- a/Engine/Source/Runtime/CoreUObject/Private/Blueprint/BlueprintSupport.cpp
+++ b/Engine/Source/Runtime/CoreUObject/Private/Blueprint/BlueprintSupport.cpp
@@ -49,6 +49,12 @@ const FName FBlueprintTags::NumNativeComponents(TEXT("NativeComponents"));
 const FName FBlueprintTags::NumBlueprintComponents(TEXT("BlueprintComponents"));
 const FName FBlueprintTags::BlueprintPathWithinPackage(TEXT("BlueprintPath"));
 
+static TAutoConsoleVariable<bool> CVarEnableFullBlueprintPreloading(
+	TEXT("linker.EnableFullBlueprintPreloading"),
+	true,
+	TEXT("If true, Blueprint class regeneration will perform a complete preload of all dependencies.")
+);
+
 /**
  * Defined in BlueprintSupport.cpp
  * Duplicates all fields of a class in depth-first order. It makes sure that everything contained
@@ -652,15 +658,11 @@ private:
 	TMap< UClass*, TArray<ULinkerPlaceholderExportObject*> > ClassToPlaceholderMap;
 };
 
-/**
- * Regenerates/Refreshes a blueprint class
- *
- * @param	LoadClass		Instance of the class currently being loaded and which is the parent for the blueprint
- * @param	ExportObject	Current object being exported
- * @return	Returns true if regeneration was successful, otherwise false
- */
-bool FLinkerLoad::RegenerateBlueprintClass(UClass* LoadClass, UObject* ExportObject)
+bool FLinkerLoad::RegenerateBlueprintClass(UClass* LoadClass, UObject* ClassDefaultObject)
 {
+	check(LoadClass);
+	ensure(ClassDefaultObject && ClassDefaultObject->HasAnyFlags(RF_ClassDefaultObject));
+
 	auto GetClassSourceObjectLambda = [](UClass* ForClass) -> UObject*
 	{
 #if WITH_EDITORONLY_DATA
@@ -671,10 +673,84 @@ bool FLinkerLoad::RegenerateBlueprintClass(UClass* LoadClass, UObject* ExportObj
 	};
 
 	UObject* ClassSourceObject = GetClassSourceObjectLambda(LoadClass);
+	check(ClassSourceObject);
+
+	if (CVarEnableFullBlueprintPreloading.GetValueOnAnyThread())
+	{
+		// "Re-preload" cyclic dependencies.
+		// 
+		// Some known objects, specifically UMetadata and UBlueprint, have a cyclic relationship with the current class.
+		// When these objects preload a UClass, they may not have finished preloading their remaining fields.
+		// In these cases, we need to effectively "re-preload" them to ensure that they actually complete the preload step.
+		// 
+		// For example, serializing UBlueprintCore::GeneratedClass could lead us to this point, but the memory
+		// for GeneratedClass hasn't been resolved yet. Performing a preload again will correctly fix GeneratedClass,
+		// since the class object is in memory during the second preload.
+		auto ForcePreloadObject = [this](UObject& Object)
+		{
+			// The absence of the RF_LoadCompleted flag implies that they're currently in the preload stage.
+			if (!Object.HasAnyFlags(RF_LoadCompleted))
+			{
+				Object.SetFlags(RF_NeedLoad);
+				Preload(&Object);
+			}
+		};
+
+		// Preload order is important here.
+		// Metdata exports can serialize all sorts of various exports, and has an implicit requirement that the
+		// Blueprint has been completely preloaded.
+		ForcePreloadObject(*ClassSourceObject);
+
+#if WITH_EDITORONLY_DATA
+		int32 MetadataIndex = LoadMetaDataFromExportMap(true);
+
+		// Older content may not have a metadata object in its package.
+		if (MetadataIndex != INDEX_NONE)
+		{
+			const FObjectExport& MetadataExport = Exp(FPackageIndex::FromExport(MetadataIndex));
+			check(MetadataExport.Object);
+			ForcePreloadObject(*MetadataExport.Object);
+		}
+#endif
+
+		// Flush (ie: create and preload) all remaining exports in the package.
+		//
+		// A Blueprint and its generated class often reference other exports in the same package.
+		// These need to be preloaded prior Blueprint compilation.
+		// This technique is a bit heavy-handed since we might be preloading exports that aren't
+		// used by either the Blueprint or its generated class.
+		// However, this is still preferable because
+		// 1) These exports are going to be preloaded anyway. We're just doing it now.
+		// 2) In most cases, the majority of exports are referenced by the Blueprint and its generated class.
+		// 3) This is more efficient (and less error prone) than using FReferenceFinder to find specific dependencies.
+		for (int32 ExportIndex = 0; ExportIndex < ExportMap.Num(); ++ExportIndex)
+		{
+			// If there was an earlier load error, we might not have a valid class.
+			if (UClass* Class = GetExportLoadClass(ExportIndex))
+			{
+				const bool bForcePreload = true;
+				CreateExportAndPreload(ExportIndex, bForcePreload);
+			}
+		}
+
+		{
+			// RegenerateClass largely performs redundant work since we already preloaded the remaining exports.
+			// However, in some circumstances, a re-preload might occur on an export that doesn't have RF_LoadCompleted set due to a cycle.
+			// This is largely dependent on how the derived Blueprint (eg: Widget BPs, Animation BPs, etc.) decides to handle this.
+			ClassSourceObject->SetFlags(RF_BeingRegenerated);
+			ClassSourceObject->RegenerateClass(LoadClass, ClassDefaultObject);
+			ClassSourceObject->ClearFlags(RF_BeingRegenerated);
+
+			// A regenerated class won't be post-loaded, so we can clear these flags.
+			LoadClass->ClearFlags(RF_NeedPostLoad | RF_NeedPostLoadSubobjects);
+		}
+
+		return true;
+	}
 
 	// determine if somewhere further down the callstack, we're already in this
 	// function for this class
-	bool const bAlreadyRegenerating = ClassSourceObject->HasAnyFlags(RF_BeingRegenerated);
+	const bool bAlreadyRegenerating = ClassSourceObject->HasAnyFlags(RF_BeingRegenerated);
 	// Flag the class source object, so we know we're already in the process of compiling this class
 	ClassSourceObject->SetFlags(RF_BeingRegenerated);
 
@@ -682,9 +758,8 @@ bool FLinkerLoad::RegenerateBlueprintClass(UClass* LoadClass, UObject* ExportObj
 	// manually... do this before we Preload() any children members so that if 
 	// one of those preloads subsequently ends up back here for this class, 
 	// then the ExportObject is carried along and used in the eventual RegenerateClass() call
-	UObject* CurrentCDO = ExportObject;
-	check(!bAlreadyRegenerating || (LoadClass->ClassDefaultObject == ExportObject));
-	LoadClass->ClassDefaultObject = CurrentCDO;
+	check(!bAlreadyRegenerating || (LoadClass->ClassDefaultObject == ClassDefaultObject));
+	LoadClass->ClassDefaultObject = ClassDefaultObject;
 
 	// Finish loading the class here, so we have all the appropriate data to copy over to the new CDO
 	TArray<UObject*> AllChildMembers;
@@ -697,7 +772,7 @@ bool FLinkerLoad::RegenerateBlueprintClass(UClass* LoadClass, UObject* ExportObj
 
 	// if this was subsequently regenerated from one of the above preloads, then 
 	// we don't have to finish this off, it was already done
-	bool const bWasSubsequentlyRegenerated = !ClassSourceObject->HasAnyFlags(RF_BeingRegenerated);
+	const bool bWasSubsequentlyRegenerated = !ClassSourceObject->HasAnyFlags(RF_BeingRegenerated);
 	// @TODO: find some other condition to block this if we've already  
 	//        regenerated the class (not just if we've regenerated the class 
 	//        from an above Preload(Member))... UBlueprint::RegenerateClass() 
@@ -708,7 +783,7 @@ bool FLinkerLoad::RegenerateBlueprintClass(UClass* LoadClass, UObject* ExportObj
 		Preload(LoadClass);
 
 		LoadClass->StaticLink(true);
-		Preload(CurrentCDO);
+		Preload(ClassDefaultObject);
 
 		// CDO preloaded - we can now resolve placeholders:
 		FResolvingExportTracker::Get().ResolvePlaceholders(LoadClass);
@@ -749,7 +824,7 @@ bool FLinkerLoad::RegenerateBlueprintClass(UClass* LoadClass, UObject* ExportObj
 			// Preload the blueprint to make sure it has all the data the class needs for regeneration
 			FPreloadMembersHelper::PreloadObject(ClassSourceObject);
 
-			UClass* RegeneratedClass = ClassSourceObject->RegenerateClass(LoadClass, CurrentCDO);
+			UClass* RegeneratedClass = ClassSourceObject->RegenerateClass(LoadClass, ClassDefaultObject);
 			if (RegeneratedClass)
 			{
 				ClassSourceObject->ClearFlags(RF_BeingRegenerated);
@@ -768,7 +843,7 @@ bool FLinkerLoad::RegenerateBlueprintClass(UClass* LoadClass, UObject* ExportObj
 		}
 	}
 
-	bool const bSuccessfulRegeneration = !ClassSourceObject->HasAnyFlags(RF_BeingRegenerated);
+	const bool bSuccessfulRegeneration = !ClassSourceObject->HasAnyFlags(RF_BeingRegenerated);
 	// if this wasn't already flagged as regenerating when we first entered this 
 	// function, the clear it ourselves.
 	if (!bAlreadyRegenerating)
