autortfm: move the UE_AUTORTFM_OPEN calls into NewObject itself and remove it from all of the callsites
- we keep repeating this same pattern, so time to clean it up

diff --git a/Engine/Source/Runtime/CoreUObject/Public/UObject/UObjectGlobals.h b/Engine/Source/Runtime/CoreUObject/Public/UObject/UObjectGlobals.h
--- a/Engine/Source/Runtime/CoreUObject/Public/UObject/UObjectGlobals.h
+++ b/Engine/Source/Runtime/CoreUObject/Public/UObject/UObjectGlobals.h
@@ -1730,6 +1730,9 @@ FUNCTION_NON_NULL_RETURN_END
 
 	T* Result = nullptr;
 
+	// AutoRTFM: the idea here is for us to run the entire UObject creation as uninstrumented, including
+	// the object allocation. If our transaction gets aborted, we leave it up to the GC to realize that this
+	// object is no longer reachable and should be destroyed.
 	UE_AUTORTFM_OPEN(
 	{
 		Result = static_cast<T*>(StaticConstructObject_Internal(Params));
@@ -1748,7 +1751,15 @@ FUNCTION_NON_NULL_RETURN_END
 
 	FStaticConstructObjectParameters Params(T::StaticClass());
 	Params.Outer = Outer;
-	return static_cast<T*>(StaticConstructObject_Internal(Params));
+
+	T* Result = nullptr;
+
+	UE_AUTORTFM_OPEN(
+	{
+		Result = static_cast<T*>(StaticConstructObject_Internal(Params));
+	});
+
+	return Result;
 }
 
 template< class T >
@@ -1768,7 +1779,15 @@ FUNCTION_NON_NULL_RETURN_END
 	Params.Template = Template;
 	Params.bCopyTransientsFromClassDefaults = bCopyTransientsFromClassDefaults;
 	Params.InstanceGraph = InInstanceGraph;
-	return static_cast<T*>(StaticConstructObject_Internal(Params));
+
+	T* Result = nullptr;
+
+	UE_AUTORTFM_OPEN(
+	{
+		Result = static_cast<T*>(StaticConstructObject_Internal(Params));
+	});
+
+	return Result;
 }
 
 /**
