Support editor helper paths for actors in a level instance.  This is so -game instances can find corresponding actor objects when transactions are applied.

diff --git a/Engine/Source/Runtime/CoreUObject/Private/Misc/TransactionObjectEvent.cpp b/Engine/Source/Runtime/CoreUObject/Private/Misc/TransactionObjectEvent.cpp
--- a/Engine/Source/Runtime/CoreUObject/Private/Misc/TransactionObjectEvent.cpp
+++ b/Engine/Source/Runtime/CoreUObject/Private/Misc/TransactionObjectEvent.cpp
@@ -4,6 +4,7 @@
 #include "UObject/Object.h"
 #include "UObject/Package.h"
 #include "UObject/Class.h"
+#include "Misc/EditorPathHelper.h"
 
 void FTransactionObjectId::SetObject(const UObject* Object)
 {
@@ -14,10 +15,25 @@ void FTransactionObjectId::SetObject(const UObject* Object)
 		return FName(ObjectPathNameBuilder);
 	};
 
-	ObjectPackageName = Object->GetPackage()->GetFName();
+	auto ObjectPathHelper = [GetPathNameHelper](const UObject* ObjectToGet) -> FName
+	{
+#if WITH_EDITOR
+		FSoftObjectPath Path = FEditorPathHelper::GetEditorPath(ObjectToGet);
+		return FName(Path.ToString());
+#else
+		return GetPathNameHelper(ObjectToGet);
+#endif
+	};
+#if WITH_EDITOR
+	const bool bEditorPathHelper = FEditorPathHelper::IsEnabled();
+#else
+	const bool bEditorPathHelper = false;
+#endif
+	const UPackage* ObjectPackage = Object->GetPackage();
+	ObjectPackageName = bEditorPathHelper ? ObjectPathHelper(ObjectPackage) : ObjectPackage->GetFName();
 	ObjectName = Object->GetFName();
-	ObjectPathName = GetPathNameHelper(Object);
-	ObjectOuterPathName = Object->GetOuter() ? GetPathNameHelper(Object->GetOuter()) : FName();
+	ObjectPathName = ObjectPathHelper(Object);
+	ObjectOuterPathName = Object->GetOuter() ? ObjectPathHelper(Object->GetOuter()) : FName();
 	ObjectExternalPackageName = Object->GetExternalPackage() ? Object->GetExternalPackage()->GetFName() : FName();
 	ObjectClassPathName = GetPathNameHelper(Object->GetClass());
 }
