This change adds the requested priority to the package summary so we can assess the order in which packages are loaded and if AsyncLoading is respecting package priorities as expected.
Discussions suggest we don't have enough data points to debug/ensure priorities are being respected, and this would be good to have in the engine in any case.

diff --git a/Engine/Source/Developer/TraceInsights/Private/Insights/ViewModels/LoadingTimingTrack.cpp b/Engine/Source/Developer/TraceInsights/Private/Insights/ViewModels/LoadingTimingTrack.cpp
--- a/Engine/Source/Developer/TraceInsights/Private/Insights/ViewModels/LoadingTimingTrack.cpp
+++ b/Engine/Source/Developer/TraceInsights/Private/Insights/ViewModels/LoadingTimingTrack.cpp
@@ -365,6 +365,7 @@ void FLoadingTimingTrack::InitTooltip(FTooltipDrawState& InOutTooltip, const ITi
 				InOutTooltip.AddNameValueTextLine(TEXT("Package Name:"), Package->Name);
 				InOutTooltip.AddNameValueTextLine(TEXT("Header Size:"), FString::Printf(TEXT("%s bytes"), *FText::AsNumber(Package->Summary.TotalHeaderSize).ToString()));
 				InOutTooltip.AddNameValueTextLine(TEXT("Package Summary:"), FString::Printf(TEXT("%d imports, %d exports"), Package->Summary.ImportCount, Package->Summary.ExportCount));
+				InOutTooltip.AddNameValueTextLine(TEXT("Request Priority:"), FString::Printf(TEXT("%d"), Package->Summary.Priority));
 				if (!Export)
 				{
 					InOutTooltip.AddNameValueTextLine(TEXT("Event:"), TEXT("ProcessPackageSummary"));
diff --git a/Engine/Source/Developer/TraceServices/Private/Analyzers/LoadTimeTraceAnalysis.cpp b/Engine/Source/Developer/TraceServices/Private/Analyzers/LoadTimeTraceAnalysis.cpp
--- a/Engine/Source/Developer/TraceServices/Private/Analyzers/LoadTimeTraceAnalysis.cpp
+++ b/Engine/Source/Developer/TraceServices/Private/Analyzers/LoadTimeTraceAnalysis.cpp
@@ -245,7 +245,8 @@ bool FAsyncLoadingTraceAnalyzer::OnEvent(uint16 RouteId, EStyle Style, const FOn
 				Summary.TotalHeaderSize = EventData.GetValue<uint32>("TotalHeaderSize");
 				Summary.ImportCount = EventData.GetValue<uint32>("ImportCount");
 				Summary.ExportCount = EventData.GetValue<uint32>("ExportCount");
-				UE_LOAD_TIME_TRACE_ANALYSIS_LOG1(TEXT("TotalHeaderSize=%u ImportCount=%u ExportCount=%u"), Summary.TotalHeaderSize, Summary.ImportCount, Summary.ExportCount);
+				Summary.Priority = EventData.GetValue<int32>("Priority", 0);  // added in UE 5.4
+				UE_LOAD_TIME_TRACE_ANALYSIS_LOG1(TEXT("TotalHeaderSize=%u ImportCount=%u ExportCount=%u Priority=%d"), Summary.TotalHeaderSize, Summary.ImportCount, Summary.ExportCount, Summary.Priority);
 			}
 			else
 			{
diff --git a/Engine/Source/Developer/TraceServices/Public/TraceServices/Model/LoadTimeProfiler.h b/Engine/Source/Developer/TraceServices/Public/TraceServices/Model/LoadTimeProfiler.h
--- a/Engine/Source/Developer/TraceServices/Public/TraceServices/Model/LoadTimeProfiler.h
+++ b/Engine/Source/Developer/TraceServices/Public/TraceServices/Model/LoadTimeProfiler.h
@@ -67,6 +67,7 @@ struct FPackageSummaryInfo
 	uint32 TotalHeaderSize = 0;
 	uint32 ImportCount = 0;
 	uint32 ExportCount = 0;
+	int32 Priority = 0;
 };
 
 struct FClassInfo
diff --git a/Engine/Source/Runtime/CoreUObject/Private/Serialization/AsyncLoading2.cpp b/Engine/Source/Runtime/CoreUObject/Private/Serialization/AsyncLoading2.cpp
--- a/Engine/Source/Runtime/CoreUObject/Private/Serialization/AsyncLoading2.cpp
+++ b/Engine/Source/Runtime/CoreUObject/Private/Serialization/AsyncLoading2.cpp
@@ -5947,7 +5947,7 @@ EEventLoadNodeExecutionResult FAsyncPackage2::Event_ProcessPackageSummary(FAsync
 			Package->InstanceContext->BuildPackageMapping(PackageNameToLoad, Package->Desc.UPackageName);
 		}
 
-		TRACE_LOADTIME_PACKAGE_SUMMARY(Package, Package->HeaderData.PackageName, Package->HeaderData.PackageSummary->HeaderSize, Package->HeaderData.ImportMap.Num(), Package->HeaderData.ExportMap.Num());
+		TRACE_LOADTIME_PACKAGE_SUMMARY(Package, Package->HeaderData.PackageName, Package->HeaderData.PackageSummary->HeaderSize, Package->HeaderData.ImportMap.Num(), Package->HeaderData.ExportMap.Num(), Package->Desc.Priority);
 	}
 
 	Package->AsyncLoadingThread.FinishInitializeAsyncPackage(ThreadState, Package);
diff --git a/Engine/Source/Runtime/CoreUObject/Private/Serialization/LoadTimeTrace.cpp b/Engine/Source/Runtime/CoreUObject/Private/Serialization/LoadTimeTrace.cpp
--- a/Engine/Source/Runtime/CoreUObject/Private/Serialization/LoadTimeTrace.cpp
+++ b/Engine/Source/Runtime/CoreUObject/Private/Serialization/LoadTimeTrace.cpp
@@ -32,6 +32,7 @@ UE_TRACE_EVENT_BEGIN(LoadTime, PackageSummary)
 	UE_TRACE_EVENT_FIELD(uint32, ImportCount)
 	UE_TRACE_EVENT_FIELD(uint32, ExportCount)
 	UE_TRACE_EVENT_FIELD(UE::Trace::WideString, Name)
+	UE_TRACE_EVENT_FIELD(int32, Priority) // added in UE 5.4
 UE_TRACE_EVENT_END()
 
 UE_TRACE_EVENT_BEGIN(LoadTime, BeginProcessSummary)
@@ -183,7 +184,7 @@ void FLoadTimeProfilerTracePrivate::OutputDestroyLinker(const void* Linker)
 		<< DestroyLinker.Linker(Linker);
 }
 
-void FLoadTimeProfilerTracePrivate::OutputPackageSummary(const void* AsyncPackage, const FName& PackageName, uint32 TotalHeaderSize, uint32 ImportCount, uint32 ExportCount)
+void FLoadTimeProfilerTracePrivate::OutputPackageSummary(const void* AsyncPackage, const FName& PackageName, uint32 TotalHeaderSize, uint32 ImportCount, uint32 ExportCount, int Priority)
 {
 	TCHAR Buffer[FName::StringBufferSize];
 	uint32 NameLength = PackageName.ToString(Buffer);
@@ -193,7 +194,8 @@ void FLoadTimeProfilerTracePrivate::OutputPackageSummary(const void* AsyncPackag
 		<< PackageSummary.TotalHeaderSize(TotalHeaderSize)
 		<< PackageSummary.ImportCount(ImportCount)
 		<< PackageSummary.ExportCount(ExportCount)
-		<< PackageSummary.Name(Buffer, NameLength);
+		<< PackageSummary.Name(Buffer, NameLength)
+		<< PackageSummary.Priority(Priority);
 }
 
 void FLoadTimeProfilerTracePrivate::OutputAsyncPackageRequestAssociation(const void* AsyncPackage, uint64 RequestId)
diff --git a/Engine/Source/Runtime/CoreUObject/Private/Serialization/LoadTimeTracePrivate.h b/Engine/Source/Runtime/CoreUObject/Private/Serialization/LoadTimeTracePrivate.h
--- a/Engine/Source/Runtime/CoreUObject/Private/Serialization/LoadTimeTracePrivate.h
+++ b/Engine/Source/Runtime/CoreUObject/Private/Serialization/LoadTimeTracePrivate.h
@@ -20,7 +20,7 @@ struct FLoadTimeProfilerTracePrivate
 	static void OutputDestroyLinker(const void* Linker);
 	static void OutputBeginRequest(uint64 RequestId);
 	static void OutputEndRequest(uint64 RequestId);
-	static void OutputPackageSummary(const void* AsyncPackage, const FName& PackageName, uint32 TotalHeaderSize, uint32 ImportCount, uint32 ExportCount);
+	static void OutputPackageSummary(const void* AsyncPackage, const FName& PackageName, uint32 TotalHeaderSize, uint32 ImportCount, uint32 ExportCount, int32 Priority);
 	static void OutputAsyncPackageImportDependency(const void* Package, const void* ImportedPackage);
 	static void OutputAsyncPackageRequestAssociation(const void* AsyncPackage, uint64 RequestId);
 	static void OutputAsyncPackageLinkerAssociation(const void* AsyncPackage, const void* Linker);
@@ -90,8 +90,8 @@ struct FLoadTimeProfilerTracePrivate
 #define TRACE_LOADTIME_DESTROY_LINKER(Linker) \
 	FLoadTimeProfilerTracePrivate::OutputDestroyLinker(Linker);
 
-#define TRACE_LOADTIME_PACKAGE_SUMMARY(AsyncPackage, PackageName, TotalHeaderSize, ImportCount, ExportCount) \
-	FLoadTimeProfilerTracePrivate::OutputPackageSummary(AsyncPackage, PackageName, TotalHeaderSize, ImportCount, ExportCount);
+#define TRACE_LOADTIME_PACKAGE_SUMMARY(AsyncPackage, PackageName, TotalHeaderSize, ImportCount, ExportCount, Priority) \
+	FLoadTimeProfilerTracePrivate::OutputPackageSummary(AsyncPackage, PackageName, TotalHeaderSize, ImportCount, ExportCount, Priority);
 
 #define TRACE_LOADTIME_ASYNC_PACKAGE_REQUEST_ASSOCIATION(AsyncPackage, RequestId) \
 	FLoadTimeProfilerTracePrivate::OutputAsyncPackageRequestAssociation(AsyncPackage, RequestId);
diff --git a/Engine/Source/Runtime/CoreUObject/Private/UObject/LinkerLoad.cpp b/Engine/Source/Runtime/CoreUObject/Private/UObject/LinkerLoad.cpp
--- a/Engine/Source/Runtime/CoreUObject/Private/UObject/LinkerLoad.cpp
+++ b/Engine/Source/Runtime/CoreUObject/Private/UObject/LinkerLoad.cpp
@@ -2703,7 +2703,7 @@ FLinkerLoad::ELinkerStatus FLinkerLoad::FinalizeCreation(TMap<TPair<FName, FPack
 
 		if (LinkerRoot)
 		{
-			TRACE_LOADTIME_PACKAGE_SUMMARY(this, LinkerRoot->GetFName(), Summary.TotalHeaderSize, Summary.ImportCount, Summary.ExportCount);
+			TRACE_LOADTIME_PACKAGE_SUMMARY(this, LinkerRoot->GetFName(), Summary.TotalHeaderSize, Summary.ImportCount, Summary.ExportCount, 0);
 		}
 
 		// Avoid duplicate work in the case of async linker creation.
