Updating object clusters to include referenced mutable cells.

diff --git a/Engine/Source/Runtime/CoreUObject/Private/UObject/GarbageCollection.cpp b/Engine/Source/Runtime/CoreUObject/Private/UObject/GarbageCollection.cpp
--- a/Engine/Source/Runtime/CoreUObject/Private/UObject/GarbageCollection.cpp
+++ b/Engine/Source/Runtime/CoreUObject/Private/UObject/GarbageCollection.cpp
@@ -49,6 +49,7 @@
 #include "VerseVM/VVMCollectionCycleRequest.h"
 #include "VerseVM/VVMContext.h"
 #include "VerseVM/VVMHeap.h"
+#include "VerseVM/VVMMarkStack.h"
 #include "VerseVM/VVMValue.h"
 #include "VerseVM/VVMWriteBarrier.h"
 #endif
@@ -1037,6 +1038,25 @@ static bool MarkClusterMutableObjectsAsReachable(FUObjectCluster& Cluster, Conta
 	return bAddClusterObjectsToSerialize;
 }
 
+#if WITH_VERSE_VM
+static void MarkClusterMutableCellsAsReachable(FUObjectCluster& Cluster)
+{
+	if (GIsFrankenGCCollecting && Cluster.MutableCells.Num() > 0)
+	{
+		Verse::FMarkStack MarkStack;
+		for (Verse::VCell* Cell : Cluster.MutableCells)
+		{
+			MarkStack.MarkNonNull(Cell);
+		}
+		Verse::FHeap::AddExternalMarkStack(MoveTemp(MarkStack));
+	}
+}
+#else
+FORCEINLINE static void MarkClusterMutableCellsAsReachable(FUObjectCluster& Cluster)
+{
+}
+#endif
+
 /** Marks all clusters referenced by another cluster as reachable */
 template<EGCOptions Options, class ContainerType>
 static FORCENOINLINE void MarkReferencedClustersAsReachable(int32 ClusterIndex, ContainerType& ObjectsToSerialize)
@@ -1084,6 +1104,7 @@ static FORCENOINLINE void MarkReferencedClustersAsReachable(int32 ClusterIndex,
 	{
 		bAddClusterObjectsToSerialize = true;
 	}
+	MarkClusterMutableCellsAsReachable(Cluster);
 	if (bAddClusterObjectsToSerialize)
 	{
 		// We need to process all cluster objects to handle PendingKill objects we nulled out (-1) from the cluster.
diff --git a/Engine/Source/Runtime/CoreUObject/Private/UObject/GarbageCollectionVerification.cpp b/Engine/Source/Runtime/CoreUObject/Private/UObject/GarbageCollectionVerification.cpp
--- a/Engine/Source/Runtime/CoreUObject/Private/UObject/GarbageCollectionVerification.cpp
+++ b/Engine/Source/Runtime/CoreUObject/Private/UObject/GarbageCollectionVerification.cpp
@@ -303,6 +303,39 @@ public:
 			}
 		}
 	}
+
+#if WITH_VERSE_VM
+	/**
+	* Handles VCell reference from the token stream. Performance is critical here so we're FORCEINLINING this function.
+	*
+	* @param Context Context of the reference collection
+	* @param ReferencingObject Object referencing the object to process.
+	* @param Cell Cell being processed
+	* @param MemberId Index to the token stream where the reference was found.
+	* @param Origin Declares if a schema represents a blueprint generated type
+	*/
+	FORCEINLINE void HandleTokenStreamVerseCellReference(FWorkerContext& Context, UObject* ReferencingObject, Verse::VCell* Cell, FMemberId MemberId, EOrigin Origin)
+	{
+		if (Cell)
+		{
+			if (Context.GetReferencingObject() != CurrentObject)
+			{
+				SetCurrentObjectAndCluster(Context.GetReferencingObject());
+			}
+			check(CurrentObject);
+
+			const FUObjectItem* ClusterRootObjectItem = GUObjectArray.ObjectToObjectItem(ClusterRootObject);
+			UE_CLOG(
+				!Cluster->MutableCells.Contains(Cell),
+				LogGarbage, Warning,
+				TEXT("Object %s from source cluster %s (%d) is referencing cell (0x%016llx) which is not part of cluster."),
+				*GetFullNameSafe(ReferencingObject),
+				*ClusterRootObject->GetFullName(),
+				ClusterRootObjectItem->GetClusterIndex(),
+				(int64)(PTRINT)Cell);
+		}
+	}
+#endif
 };
 
 void VerifyClustersAssumptions()
diff --git a/Engine/Source/Runtime/CoreUObject/Private/UObject/UObjectClusters.cpp b/Engine/Source/Runtime/CoreUObject/Private/UObject/UObjectClusters.cpp
--- a/Engine/Source/Runtime/CoreUObject/Private/UObject/UObjectClusters.cpp
+++ b/Engine/Source/Runtime/CoreUObject/Private/UObject/UObjectClusters.cpp
@@ -98,6 +98,9 @@ void FUObjectClusterContainer::FreeCluster(int32 InClusterIndex)
 	Cluster.MutableObjects.Reset();
 	Cluster.ReferencedClusters.Reset();
 	Cluster.ReferencedByClusters.Reset();
+#if WITH_VERSE_VM
+	Cluster.MutableCells.Reset();
+#endif
 	Cluster.bNeedsDissolving = false;
 	FreeClusterIndices.Add(InClusterIndex);
 	NumAllocatedClusters--;
@@ -304,8 +307,13 @@ void DumpClusterToLog(const FUObjectCluster& Cluster, bool bHierarchy, bool bInd
 
 	FUObjectItem* RootItem = GUObjectArray.IndexToObjectUnsafeForGC(Cluster.RootIndex);
 	UObject* RootObject = static_cast<UObject*>(RootItem->Object);
-	UE_LOG(LogObj, Display, TEXT("%s (Index: %d), Size: %d, ReferencedClusters: %d, MutableObjects: %d, ReferencedByClusters: %d"), 
-		*RootObject->GetFullName(), Cluster.RootIndex, Cluster.Objects.Num(), Cluster.ReferencedClusters.Num(), Cluster.MutableObjects.Num(), Cluster.ReferencedByClusters.Num());
+#if WITH_VERSE_VM
+	FString ExtraDetail = FString::Printf(TEXT(", MutableCells: %d"), Cluster.MutableCells.Num());
+#else
+	FString ExtraDetail;
+#endif
+	UE_LOG(LogObj, Display, TEXT("%s (Index: %d), Size: %d, ReferencedClusters: %d, MutableObjects: %d%s, ReferencedByClusters: %d"),
+		*RootObject->GetFullName(), Cluster.RootIndex, Cluster.Objects.Num(), Cluster.ReferencedClusters.Num(), Cluster.MutableObjects.Num(), *ExtraDetail, Cluster.ReferencedByClusters.Num());
 	if (bHierarchy)
 	{
 		int32 Index = 0;
@@ -364,6 +372,13 @@ void DumpClusterToLog(const FUObjectCluster& Cluster, bool bHierarchy, bool bInd
 				UE_LOG(LogObj, Display, TEXT("    => nullptr"));
 			}
 		}
+#if WITH_VERSE_VM
+		UE_LOG(LogObj, Display, TEXT("  External (mutable) cells: %d"), Cluster.MutableCells.Num());
+		for (Verse::VCell* Cell : Cluster.MutableCells)
+		{
+			UE_LOG(LogObj, Display, TEXT("    => %p"), Cell);
+		}
+#endif
 	}
 }
 
@@ -409,6 +424,15 @@ void ListClusters(const TArray<FString>& Args)
 				return A->MutableObjects.Num();
 			});
 		}
+#if WITH_VERSE_VM
+		else if (Arg == TEXT("SortByMutableCellCount"))
+		{
+			Algo::SortBy(AllClusters, [](FUObjectCluster* A)
+				{
+					return A->MutableCells.Num();
+				});
+		}
+#endif
 		else if (Arg == TEXT("SortByReferencedClustersCount"))
 		{
 			Algo::SortBy(AllClusters, [](FUObjectCluster* A)
@@ -627,10 +651,10 @@ public:
 	 * @param ObjectIndex UObject index in GUObjectArray
 	 * @param ObjectItem UObject's entry in GUObjectArray
 	 * @param Obj The object to add to cluster
-	 * @param ObjectsToSerialize An array of remaining objects to serialize (Obj must be added to it if Obj can be added to cluster)
+	* @param Context Context of the reference collection
 	 * @param bOuterAndClass If true, the Obj's Outer and Class will also be added to the cluster
 	 */
-	void AddObjectToCluster(int32 ObjectIndex, FUObjectItem* ObjectItem, UObject* Obj, FWorkerContext& ObjectsToSerializeStruct, bool bOuterAndClass)
+	void AddObjectToCluster(int32 ObjectIndex, FUObjectItem* ObjectItem, UObject* Obj, FWorkerContext& Context, bool bOuterAndClass)
 	{
 		// If we haven't finished loading, we can't be sure we know all the references
 		checkf(!Obj->HasAnyFlags(RF_NeedLoad), TEXT("%s hasn't been loaded (%s) but is being added to cluster %s"), 
@@ -642,7 +666,7 @@ public:
 		check(Obj->CanBeInCluster());
 		if (ObjectIndex != ClusterRootIndex && ObjectItem->GetOwnerIndex() == 0 && !GUObjectArray.IsDisregardForGC(Obj) && !Obj->IsRooted())
 		{
-			ObjectsToSerializeStruct.ObjectsToSerialize.Add<Options>(Obj);
+			Context.ObjectsToSerialize.Add<Options>(Obj);
 			check(!ObjectItem->HasAnyFlags(EInternalObjectFlags::ClusterRoot));
 			ObjectItem->SetOwnerIndex(ClusterRootIndex);
 			Cluster.Objects.Add(ObjectIndex);
@@ -652,14 +676,14 @@ public:
 				UObject* ObjOuter = Obj->GetOuter();
 				if (ObjOuter)
 				{
-					HandleTokenStreamObjectReference(ObjectsToSerializeStruct, Obj, ObjOuter, UE::GC::EMemberlessId::Outer, EOrigin::Other, true);
+					HandleTokenStreamObjectReference(Context, Obj, ObjOuter, UE::GC::EMemberlessId::Outer, EOrigin::Other, true);
 				}
 				if (!Obj->GetClass()->HasAllClassFlags(CLASS_Native))
 				{
 					UObject* ObjectClass = Obj->GetClass();
-					HandleTokenStreamObjectReference(ObjectsToSerializeStruct, Obj, ObjectClass, UE::GC::EMemberlessId::Class, EOrigin::Other, true);
+					HandleTokenStreamObjectReference(Context, Obj, ObjectClass, UE::GC::EMemberlessId::Class, EOrigin::Other, true);
 					UObject* ObjectClassOuter = Obj->GetClass()->GetOuter();
-					HandleTokenStreamObjectReference(ObjectsToSerializeStruct, Obj, ObjectClassOuter, UE::GC::EMemberlessId::ClassOuter, EOrigin::Other, true);
+					HandleTokenStreamObjectReference(Context, Obj, ObjectClassOuter, UE::GC::EMemberlessId::ClassOuter, EOrigin::Other, true);
 				}
 			}
 		}
@@ -668,12 +692,14 @@ public:
 	/**
 	* Handles UObject reference from the token stream. Performance is critical here so we're FORCEINLINING this function.
 	*
-	* @param ObjectsToSerialize An array of remaining objects to serialize (Obj must be added to it if Obj can be added to cluster)
+	* @param Context Context of the reference collection
 	* @param ReferencingObject Object referencing the object to process.
-	* @param TokenIndex Index to the token stream where the reference was found.
+	* @param Object Object being processed
+	* @param MemberId Index to the token stream where the reference was found.
+	* @param Origin Declares if a schema represents a blueprint generated type
 	* @param bAllowReferenceElimination True if reference elimination is allowed (ignored when constructing clusters).
 	*/
-	FORCEINLINE void HandleTokenStreamObjectReference(FGCArrayStruct& ObjectsToSerializeStruct, UObject* ReferencingObject, UObject*& Object, FMemberId MemberId, EOrigin Origin, bool bAllowReferenceElimination)
+	FORCEINLINE void HandleTokenStreamObjectReference(FWorkerContext& Context, UObject* ReferencingObject, UObject*& Object, FMemberId MemberId, EOrigin Origin, bool bAllowReferenceElimination)
 	{
 		if (!Object)
 		{
@@ -735,9 +761,30 @@ public:
 			check(ObjectItem->GetOwnerIndex() == 0);
 
 			// New object in the same package, add it to the cluster.
-			AddObjectToCluster(GUObjectArray.ObjectToIndex(Object), ObjectItem, Object, ObjectsToSerializeStruct, true);
+			AddObjectToCluster(GUObjectArray.ObjectToIndex(Object), ObjectItem, Object, Context, true);
+		}
+	}
+
+#if WITH_VERSE_VM
+	/**
+	* Handles VCell reference from the token stream. Performance is critical here so we're FORCEINLINING this function.
+	*
+	* @param Context Context of the reference collection
+	* @param ReferencingObject Object referencing the object to process.
+	* @param Cell Cell being processed
+	* @param MemberId Index to the token stream where the reference was found.
+	* @param Origin Declares if a schema represents a blueprint generated type
+	*/
+	FORCEINLINE void HandleTokenStreamVerseCellReference(FWorkerContext& Context, UObject* ReferencingObject, Verse::VCell* Cell, FMemberId MemberId, EOrigin Origin)
+	{
+		// As with mutable objects, we assume that the cell is mutable but the pointer to the cell is not.  If the cell pointer needs to be modified,
+		// then the underlying object must be mutable.
+		if (Cell)
+		{
+			Cluster.MutableCells.AddUnique(Cell);
 		}
 	}
+#endif
 };
 
 static void SuggestClusters(const TArray<FString>& Args)
@@ -919,11 +966,16 @@ void UObjectBaseUtility::CreateCluster()
 		// Add new cluster to the global cluster map.
 		Cluster.Objects.Sort();
 		Cluster.ReferencedClusters.Sort();		
-		Cluster.MutableObjects.Sort();		
+		Cluster.MutableObjects.Sort();
 
 #if UE_GCCLUSTER_VERBOSE_LOGGING
-		UE_LOG(LogObj, Log, TEXT("Created cluster (%d) with %d objects, %d referenced clusters and %d mutable objects."),
-			ClusterIndex, Cluster.Objects.Num(), Cluster.ReferencedClusters.Num(), Cluster.MutableObjects.Num());
+#if WITH_VERSE_VM
+		FString ExtraDetail = FString::Printf(TEXT(", %d verse cells"), Cluster.MutableCells.Num());
+#else
+		FString ExtraDetail;
+#endif
+		UE_LOG(LogObj, Log, TEXT("Created cluster (%d) with %d objects, %d referenced clusters%s and %d mutable objects."),
+			ClusterIndex, Cluster.Objects.Num(), Cluster.ReferencedClusters.Num(), *ExtraDetail, Cluster.MutableObjects.Num());
 
 		DumpClusterToLog(Cluster, true, false);
 #endif
@@ -931,8 +983,13 @@ void UObjectBaseUtility::CreateCluster()
 	else
 	{
 #if UE_GCCLUSTER_VERBOSE_LOGGING
-		UE_LOG(LogObj, Log, TEXT("Not creating cluster (%d) with %d objects, %d referenced clusters and %d mutable objects."),
-			ClusterIndex, Cluster.Objects.Num(), Cluster.ReferencedClusters.Num(), Cluster.MutableObjects.Num());
+#if WITH_VERSE_VM
+		FString ExtraDetail = FString::Printf(TEXT(", %d verse cells"), Cluster.MutableCells.Num());
+#else
+		FString ExtraDetail;
+#endif
+		UE_LOG(LogObj, Log, TEXT("Not creating cluster (%d) with %d objects, %d referenced clusters%s and %d mutable objects."),
+			ClusterIndex, Cluster.Objects.Num(), Cluster.ReferencedClusters.Num(), *ExtraDetail, Cluster.MutableObjects.Num());
 
 		DumpClusterToLog(Cluster, true, false);
 #endif
diff --git a/Engine/Source/Runtime/CoreUObject/Public/UObject/UObjectArray.h b/Engine/Source/Runtime/CoreUObject/Public/UObject/UObjectArray.h
--- a/Engine/Source/Runtime/CoreUObject/Public/UObject/UObjectArray.h
+++ b/Engine/Source/Runtime/CoreUObject/Public/UObject/UObjectArray.h
@@ -11,6 +11,13 @@
 #include "UObject/ObjectMacros.h"
 #include "UObject/UObjectBase.h"
 
+#if WITH_VERSE_VM
+namespace Verse
+{
+	struct VCell;
+}
+#endif
+
 /**
 * Controls whether the number of available elements is being tracked in the ObjObjects array.
 * By default it is only tracked in WITH_EDITOR builds as it adds a small amount of tracking overhead
@@ -1195,6 +1202,10 @@ struct FUObjectCluster
 	TArray<int32> MutableObjects;
 	/** List of clusters that direcly reference this cluster. Used when dissolving a cluster. */
 	TArray<int32> ReferencedByClusters;
+#if WITH_VERSE_VM
+	/** All verse cells are considered mutable.  They will just be added directly to verse gc when the cluster is marked */
+	TArray<Verse::VCell*> MutableCells;
+#endif
 
 	/** Cluster needs dissolving, probably due to PendingKill reference */
 	bool bNeedsDissolving;
diff --git a/Engine/Source/Runtime/Engine/Private/LevelActorContainer.cpp b/Engine/Source/Runtime/Engine/Private/LevelActorContainer.cpp
--- a/Engine/Source/Runtime/Engine/Private/LevelActorContainer.cpp
+++ b/Engine/Source/Runtime/Engine/Private/LevelActorContainer.cpp
@@ -101,9 +101,11 @@ public:
 	/**
 	* Handles UObject reference from the token stream. Performance is critical here so we're FORCEINLINING this function.
 	*
-	* @param ObjectsToSerialize An array of remaining objects to serialize (Obj must be added to it if Obj can be added to cluster)
+	* @param Context Context of the reference collection
 	* @param ReferencingObject Object referencing the object to process.
+	* @param Object Object being processed
 	* @param MemberId Index to the token stream where the reference was found.
+	* @param Origin Declares if a schema represents a blueprint generated type
 	* @param bAllowReferenceElimination True if reference elimination is allowed (ignored when constructing clusters).
 	*/
 	FORCEINLINE void HandleTokenStreamObjectReference(FWorkerContext& Context, UObject* ReferencingObject, UObject*& Object, FMemberId MemberId, EOrigin Origin, bool bAllowReferenceElimination)
@@ -158,6 +160,27 @@ public:
 			}
 		}
 	}
+
+#if WITH_VERSE_VM
+	/**
+	* Handles VCell reference from the token stream. Performance is critical here so we're FORCEINLINING this function.
+	*
+	* @param Context Context of the reference collection
+	* @param ReferencingObject Object referencing the object to process.
+	* @param Cell Cell being processed
+	* @param MemberId Index to the token stream where the reference was found.
+	* @param Origin Declares if a schema represents a blueprint generated type
+	*/
+	FORCEINLINE void HandleTokenStreamVerseCellReference(FWorkerContext& Context, UObject* ReferencingObject, Verse::VCell* Cell, FMemberId MemberId, EOrigin Origin)
+	{
+		// As with mutable objects, we assume that the cell is mutable but the pointer to the cell is not.  If the cell pointer needs to be modified,
+		// then the underlying object must be mutable.
+		if (Cell)
+		{
+			Cluster.MutableCells.AddUnique(Cell);
+		}
+	}
+#endif
 };
 
 void ULevelActorContainer::CreateCluster()
@@ -196,8 +219,16 @@ void ULevelActorContainer::CreateCluster()
 		Cluster.ReferencedClusters.Sort();
 		Cluster.MutableObjects.Sort();
 
-		UE_LOG(LogLevelActorContainer, Verbose, TEXT("Created LevelActorCluster (%d) for %s with %d objects, %d referenced clusters and %d mutable objects."),
-			ClusterIndex, *GetOuter()->GetPathName(), Cluster.Objects.Num(), Cluster.ReferencedClusters.Num(), Cluster.MutableObjects.Num());
+		if (UE_LOG_ACTIVE(LogLevelActorContainer, Verbose))
+		{
+#if WITH_VERSE_VM
+			FString ExtraDetail = FString::Printf(TEXT(", %d verse cells"), Cluster.MutableCells.Num());
+#else
+			FString ExtraDetail;
+#endif
+			UE_LOG(LogLevelActorContainer, Verbose, TEXT("Created LevelActorCluster (%d) for %s with %d objects, %d referenced clusters%s and %d mutable objects."),
+				ClusterIndex, *GetOuter()->GetPathName(), Cluster.Objects.Num(), Cluster.ReferencedClusters.Num(), *ExtraDetail, Cluster.MutableObjects.Num());
+		}
 
 #if UE_GCCLUSTER_VERBOSE_LOGGING
 		DumpClusterToLog(Cluster, true, false);
