Disable an error message checking the input of FChunkReadFileHandle::ReadRequest

#rnx

- In CL 27337756  we added a log error if a FChunkReadFileHandle::ReadRequest attempted to read outside of the bulkdata payload, however this error is being triggered by FFileCache which is still using the handle to represent the entire .ubulk file.
- Although this is something we want to move away from in the future the read is currently valid so for now we are disabling the log message via a new define UE_ENABLE_BULKDATA_RANGE_TEST

diff --git a/Engine/Source/Runtime/CoreUObject/Private/Serialization/BulkDataStreaming.cpp b/Engine/Source/Runtime/CoreUObject/Private/Serialization/BulkDataStreaming.cpp
--- a/Engine/Source/Runtime/CoreUObject/Private/Serialization/BulkDataStreaming.cpp
+++ b/Engine/Source/Runtime/CoreUObject/Private/Serialization/BulkDataStreaming.cpp
@@ -20,6 +20,13 @@
 TRACE_DECLARE_INT_COUNTER(BulkDataBatchRequest_Count, TEXT("BulkData/BatchRequest/Count"));
 TRACE_DECLARE_INT_COUNTER(BulkDataBatchRequest_PendingCount, TEXT("BulkData/BatchRequest/Pending"));
 
+/**
+ * When enabled calls to FChunkReadFileHandle::ReadRequest will validate that the request
+ * is within the bulkdata payload bounds. Currently disabled as FFileCache still uses the
+ * handle to represent the entire .ubulk file rather than the specific bulkdata payload.
+ */
+#define UE_ENABLE_BULKDATA_RANGE_TEST 0
+
 FBulkDataIORequest::FBulkDataIORequest(IAsyncReadFileHandle* InFileHandle)
 	: FileHandle(InFileHandle)
 	, ReadRequest(nullptr)
@@ -344,13 +351,16 @@ IAsyncReadRequest* FChunkReadFileHandle::ReadRequest(
 	FAsyncFileCallBack* CompleteCallback,
 	uint8* UserSuppliedMemory)
 {
+#if UE_ENABLE_BULKDATA_RANGE_TEST
 	const bool bIsOutsideBulkDataRange =
 		(Offset < static_cast<int64>(ChunkRange.GetOffset())) ||
 		((Offset + BytesToRead) > static_cast<int64>(ChunkRange.GetOffset() + AvailableChunkSize));
 
+
 	UE_CLOG(bIsOutsideBulkDataRange, LogSerialization, Warning,
 		TEXT("Reading outside of bulk data range, RequestRange='%lld, %lld', BulkDataRange='%llu, %llu', ChunkId='%s'"),
 		Offset, BytesToRead, ChunkRange.GetOffset(), ChunkRange.GetLength(), *LexToString(ChunkId));
+#endif //UE_ENABLE_BULKDATA_RANGE_TEST
 
 	FIoBuffer Buffer = UserSuppliedMemory ? FIoBuffer(FIoBuffer::Wrap, UserSuppliedMemory, BytesToRead) : FIoBuffer(BytesToRead);
 	FChunkReadFileRequest* Request = new FChunkReadFileRequest(CompleteCallback, MoveTemp(Buffer));
@@ -1208,3 +1218,5 @@ FBulkDataRequest::EStatus FBulkDataBatchRequest::FScatterGatherBuilder::Issue(FI
 }
 
 //////////////////////////////////////////////////////////////////////////////
+
+#undef UE_ENABLE_BULKDATA_RANGE_TEST
