Compilation fix: remove unnecessary macro-continuation character / that was causing a bug in some compilers because it caused #else to be moved onto the preceeding line.
#rnx

diff --git a/Engine/Source/Runtime/CoreUObject/Public/UObject/ObjectMacros.h b/Engine/Source/Runtime/CoreUObject/Public/UObject/ObjectMacros.h
--- a/Engine/Source/Runtime/CoreUObject/Public/UObject/ObjectMacros.h
+++ b/Engine/Source/Runtime/CoreUObject/Public/UObject/ObjectMacros.h
@@ -1893,16 +1893,17 @@ public: \
 
 #define UOBJECT_CPPCLASS_STATICFUNCTIONS_ALLCONFIGS(TClass) \
 	FUObjectCppClassStaticFunctions::AddReferencedObjectsType(&TClass::AddReferencedObjects)
-	/* UObjectCppClassStaticFunctions: Extend this macro with the address of your new static function, if it applies to all configs. */ \
-	/* Order must match the order in the FUObjectCppClassStaticFunctions constructor. */ \
+	/* UObjectCppClassStaticFunctions: Extend this macro with the address of your new static function, if it applies to all configs. */
+	/* Order must match the order in the FUObjectCppClassStaticFunctions constructor. */
 
 #if WITH_EDITORONLY_DATA
 	#define UOBJECT_CPPCLASS_STATICFUNCTIONS_WITHEDITORONLYDATA(TClass) \
 		, FUObjectCppClassStaticFunctions::DeclareCustomVersionsType(&TClass::DeclareCustomVersions) \
 		, FUObjectCppClassStaticFunctions::AppendToClassSchemaType(&TClass::AppendToClassSchema) \
 		, FUObjectCppClassStaticFunctions::DeclareConstructClassesType(&TClass::DeclareConstructClasses)
-		/* UObjectCppClassStaticFunctions: Extend this macro with the address of your new static function, if it is editor-only. */ \
-		/* Order must match the order in the FUObjectCppClassStaticFunctions constructor. */ \
+		/* UObjectCppClassStaticFunctions: Extend this macro with the address of your new static function, if it is editor-only. */
+		/* Order must match the order in the FUObjectCppClassStaticFunctions constructor. */
+
 #else
 	#define UOBJECT_CPPCLASS_STATICFUNCTIONS_WITHEDITORONLYDATA(TClass)
 #endif
