[AsyncLoading] - Fix to avoid potential deadlock during loading

#rnx
submitted on behalf of Danny.Couture

diff --git a/Engine/Source/Runtime/CoreUObject/Private/Serialization/AsyncLoading.cpp b/Engine/Source/Runtime/CoreUObject/Private/Serialization/AsyncLoading.cpp
--- a/Engine/Source/Runtime/CoreUObject/Private/Serialization/AsyncLoading.cpp
+++ b/Engine/Source/Runtime/CoreUObject/Private/Serialization/AsyncLoading.cpp
@@ -571,6 +571,14 @@ static FAutoConsoleVariableRef CVarAsyncLoadingThreadEnabled(
 	ECVF_Default
 	);
 
+static int32 GAsyncLoadingAlwaysProcessPackages = 1;
+static FAutoConsoleVariableRef CVarAsyncLoadingAlwaysProcessPackages(
+	TEXT("s.AsyncLoadingAlwaysProcessPackages"),
+	GAsyncLoadingAlwaysProcessPackages,
+	TEXT("When flushing, will process all packages instead of only what's needed. (Used to avoid a hard to repro potential deadlock)"),
+	ECVF_Default
+	);
+
 static int32 GFlushStreamingOnExit = 1;
 static FAutoConsoleVariableRef CFlushStreamingOnExit(
 	TEXT("s.FlushStreamingOnExit"),
@@ -4711,7 +4719,7 @@ void FAsyncLoadingThread::AdjustFlushRequest(FFlushRequest& FlushRequest)
 
 bool FAsyncLoadingThread::ShouldProcessPackage(FAsyncPackage* InAsyncPackage, const FFlushRequest& FlushRequest)
 {
-	if (!FlushRequest)
+	if (!FlushRequest || GAsyncLoadingAlwaysProcessPackages)
 	{
 		return true;
 	}
