Make sure objects preloaded during the async deferred post load step of the legacy loader adds them to the list of objects it tracked

diff --git a/Engine/Source/Runtime/CoreUObject/Private/Serialization/AsyncLoading.cpp b/Engine/Source/Runtime/CoreUObject/Private/Serialization/AsyncLoading.cpp
--- a/Engine/Source/Runtime/CoreUObject/Private/Serialization/AsyncLoading.cpp
+++ b/Engine/Source/Runtime/CoreUObject/Private/Serialization/AsyncLoading.cpp
@@ -7062,6 +7062,7 @@ EAsyncPackageState::Type FAsyncPackage::PostLoadDeferredObjects(double InTickSta
 						if (PreLoadObject && PreLoadObject->GetLinker())
 						{
 							PreLoadObject->GetLinker()->Preload(PreLoadObject);
+							PackageObjLoaded.Add(PreLoadObject);
 						}
 					}
 
