Cache the 'NotExcluded' state alongside 'Excluded' and 'Included' to prevent excessive query ConditionallyExcludeObjectForRealm on deeply nested object graph with a lot of public leaves

diff --git a/Engine/Source/Runtime/CoreUObject/Private/UObject/SavePackage/PackageHarvester.cpp b/Engine/Source/Runtime/CoreUObject/Private/UObject/SavePackage/PackageHarvester.cpp
--- a/Engine/Source/Runtime/CoreUObject/Private/UObject/SavePackage/PackageHarvester.cpp
+++ b/Engine/Source/Runtime/CoreUObject/Private/UObject/SavePackage/PackageHarvester.cpp
@@ -83,6 +83,10 @@ bool ConditionallyExcludeObjectForRealm(FSaveContext& SaveContext, TObjectPtr<UO
 	{
 		return false;
 	}
+	if (RealmBeingChecked.IsNotExcluded(Obj))
+	{
+		return false;
+	}
 
 	const EObjectMark ExcludedObjectMarks = SaveContext.GetExcludedObjectMarks(HarvestingContext);
 	const ITargetPlatform* TargetPlatform = SaveContext.GetTargetPlatform();
@@ -132,6 +136,7 @@ bool ConditionallyExcludeObjectForRealm(FSaveContext& SaveContext, TObjectPtr<UO
 		}
 	}
 
+	RealmBeingChecked.AddNotExcluded(Obj);
 	return false;
 }
 
diff --git a/Engine/Source/Runtime/CoreUObject/Private/UObject/SavePackage/SaveContext.h b/Engine/Source/Runtime/CoreUObject/Private/UObject/SavePackage/SaveContext.h
--- a/Engine/Source/Runtime/CoreUObject/Private/UObject/SavePackage/SaveContext.h
+++ b/Engine/Source/Runtime/CoreUObject/Private/UObject/SavePackage/SaveContext.h
@@ -150,6 +150,11 @@ struct FHarvestedRealm
 		Excluded.Add(InObject);
 	}
 
+	void AddNotExcluded(TObjectPtr<UObject> InObject)
+	{
+		NotExcluded.Add(InObject);
+	}
+
 	bool IsImport(TObjectPtr<UObject> InObject) const
 	{
 		return Imports.Contains(InObject);
@@ -174,6 +179,11 @@ struct FHarvestedRealm
 		return Excluded.Contains(InObject);
 	}
 
+	bool IsNotExcluded(TObjectPtr<UObject> InObject) const
+	{
+		return NotExcluded.Contains(InObject);
+	}
+
 	TSet<FTaggedExport>& GetExports()
 	{
 		return Exports;
@@ -368,6 +378,8 @@ private:
 
 	// Set of objects excluded (import or exports) through marks or otherwise (i.e. transient flags, etc)
 	TSet<TObjectPtr<UObject>> Excluded;
+	// Set of objects not excluded through marks or otherwise (i.e. transient flags, etc) while not being marked specifically included yet
+	TSet<TObjectPtr<UObject>> NotExcluded;
 	// Set of objects marked as export
 	TSet<FTaggedExport> Exports;
 	// Set of objects marked as import
