Changed how we load the UStruct::Children chain. A function can be generated as transient. That function will not be saved and would create a nullptr in the array.

diff --git a/Engine/Source/Runtime/CoreUObject/Private/UObject/Class.cpp b/Engine/Source/Runtime/CoreUObject/Private/UObject/Class.cpp
--- a/Engine/Source/Runtime/CoreUObject/Private/UObject/Class.cpp
+++ b/Engine/Source/Runtime/CoreUObject/Private/UObject/Class.cpp
@@ -1854,18 +1854,35 @@ void UStruct::Serialize(FArchive& Ar)
 		if (Ar.IsLoading())
 		{
 			Ar << ChildArray;
-			if (ChildArray.Num())
+
+			// skip null fields
+			Children = nullptr;
+			if (ChildArray.Num() > 0)
 			{
-				for (int32 Index = 0; Index + 1 < ChildArray.Num(); Index++)
+				UField* CurrentChild = nullptr;
+ 				int32 Index = 0;
+				for (; Index < ChildArray.Num(); ++Index)
 				{
-					ChildArray[Index]->Next = ChildArray[Index + 1];
+					if (ChildArray[Index])
+					{
+						Children = ChildArray[Index];
+						CurrentChild = Children;
+						break;
+					}
+				}
+				if (CurrentChild)
+				{
+					for (Index+=1; Index < ChildArray.Num(); ++Index)
+					{
+						if (ChildArray[Index])
+						{
+							CurrentChild->Next = ChildArray[Index];
+							CurrentChild = ChildArray[Index];
+						}
+					}
+
+					CurrentChild->Next = nullptr;
 				}
-				Children = ChildArray[0];
-				ChildArray[ChildArray.Num() - 1]->Next = nullptr;
-			}
-			else
-			{
-				Children = nullptr;
 			}
 		}
 		else
