Unify NumPackagesLoaded cook stat for all loaders and add NumRequestedLoads as an additional metric

diff --git a/Engine/Source/Editor/UnrealEd/Private/Cooker/CookProfiling.cpp b/Engine/Source/Editor/UnrealEd/Private/Cooker/CookProfiling.cpp
--- a/Engine/Source/Editor/UnrealEd/Private/Cooker/CookProfiling.cpp
+++ b/Engine/Source/Editor/UnrealEd/Private/Cooker/CookProfiling.cpp
@@ -251,6 +251,8 @@ namespace DetailedCookStats
 	FCookStatsManager::FAutoRegisterCallback RegisterCookOnTheFlyServerStats([](FCookStatsManager::AddStatFuncRef AddStat)
 		{
 			AddStat(TEXT("Package.Load"), FCookStatsManager::CreateKeyValueArray(
+				TEXT("NumRequestedLoads"), NumRequestedLoads,
+				TEXT("NumPackagesLoaded"), NumDetectedLoads.load(),
 				TEXT("NumInlineLoads"), NumDetectedLoads - NumRequestedLoads));
 			AddStat(TEXT("Package.Save"), FCookStatsManager::CreateKeyValueArray(TEXT("NumPackagesIterativelySkipped"), NumPackagesIterativelySkipped));
 			AddStat(TEXT("CookOnTheFlyServer"), FCookStatsManager::CreateKeyValueArray(
diff --git a/Engine/Source/Runtime/CoreUObject/Private/UObject/UObjectGlobals.cpp b/Engine/Source/Runtime/CoreUObject/Private/UObject/UObjectGlobals.cpp
--- a/Engine/Source/Runtime/CoreUObject/Private/UObject/UObjectGlobals.cpp
+++ b/Engine/Source/Runtime/CoreUObject/Private/UObject/UObjectGlobals.cpp
@@ -137,11 +137,9 @@ CSV_DEFINE_CATEGORY(UObject, false);
 namespace LoadPackageStats
 {
 	static double LoadPackageTimeSec = 0.0;
-	int NumPackagesLoaded = 0;
 	static FCookStatsManager::FAutoRegisterCallback RegisterCookStats([](FCookStatsManager::AddStatFuncRef AddStat)
 	{
 		AddStat(TEXT("Package.Load"), FCookStatsManager::CreateKeyValueArray(
-			TEXT("NumPackagesLoaded"), NumPackagesLoaded,
 			TEXT("LoadPackageTimeSec"), LoadPackageTimeSec));
 	});
 }
@@ -1621,7 +1619,6 @@ bool ShouldAlwaysLoadPackageAsync(const FPackagePath& InPackagePath);
 UPackage* LoadPackageInternal(UPackage* InOuter, const FPackagePath& PackagePath, uint32 LoadFlags, FLinkerLoad* ImportLinker, FArchive* InReaderOverride,
 	const FLinkerInstancingContext* InstancingContext, const FPackagePath* DiffPackagePath)
 {
-	COOK_STAT(LoadPackageStats::NumPackagesLoaded++);
 	DECLARE_SCOPE_CYCLE_COUNTER(TEXT("LoadPackageInternal"), STAT_LoadPackageInternal, STATGROUP_ObjectVerbose);
 
 	FString TracePackageName;
