[ZenLoader] Do not propagate an instancing context to our imports when those aren't also instanced

diff --git a/Engine/Source/Runtime/CoreUObject/Private/Serialization/AsyncLoading2.cpp b/Engine/Source/Runtime/CoreUObject/Private/Serialization/AsyncLoading2.cpp
--- a/Engine/Source/Runtime/CoreUObject/Private/Serialization/AsyncLoading2.cpp
+++ b/Engine/Source/Runtime/CoreUObject/Private/Serialization/AsyncLoading2.cpp
@@ -4855,6 +4855,7 @@ void FAsyncPackage2::ImportPackagesRecursiveInner(FAsyncLoadingThreadState2& Thr
 		}
 #endif
 
+		bool bIsInstanced = false;
 #if WITH_EDITORONLY_DATA && ALT2_ENABLE_LINKERLOAD_SUPPORT
 		if (bHasImportedPackageNames && LinkerLoadState.IsSet())
 		{
@@ -4863,6 +4864,7 @@ void FAsyncPackage2::ImportPackagesRecursiveInner(FAsyncLoadingThreadState2& Thr
 			ImportedPackageUPackageName = InstancingContext.RemapPackage(ImportedPackageNameToLoad);
 			if (ImportedPackageUPackageName != ImportedPackageNameToLoad)
 			{
+				bIsInstanced = true;
 				if (ImportedPackageUPackageName.IsNone())
 				{
 					ImportedPackageIdToLoad = FPackageId::FromName(NAME_None);
@@ -5045,7 +5047,8 @@ void FAsyncPackage2::ImportPackagesRecursiveInner(FAsyncLoadingThreadState2& Thr
 #if ALT2_ENABLE_LINKERLOAD_SUPPORT
 				if (!bIsZenPackageImport)
 				{
-					ImportedPackage->InitializeLinkerLoadState(bIsZenPackage ? nullptr : &LinkerLoadState->Linker->GetInstancingContext());
+					// Only propagate the instancing context if the imported package is also instanced and it isn't a zen package
+					ImportedPackage->InitializeLinkerLoadState(bIsZenPackage || !bIsInstanced ? nullptr : &LinkerLoadState->Linker->GetInstancingContext());
 				}
 				else
 #endif
