[ALT] - Fix TSAN warning on non-atomic package loading state

#rnx

diff --git a/Engine/Source/Runtime/CoreUObject/Private/Serialization/AsyncLoading2.cpp b/Engine/Source/Runtime/CoreUObject/Private/Serialization/AsyncLoading2.cpp
--- a/Engine/Source/Runtime/CoreUObject/Private/Serialization/AsyncLoading2.cpp
+++ b/Engine/Source/Runtime/CoreUObject/Private/Serialization/AsyncLoading2.cpp
@@ -2756,7 +2756,8 @@ private:
 	int32						DeferredClusterIndex = 0;
 	/** Current index into Export objects array used to spread routing PostLoadInstances over several frames			*/
 	int32						PostLoadInstanceIndex = 0;
-	EAsyncPackageLoadingState2	AsyncPackageLoadingState = EAsyncPackageLoadingState2::NewPackage;
+	/** Current loading state of a package. */
+	std::atomic<EAsyncPackageLoadingState2> AsyncPackageLoadingState { EAsyncPackageLoadingState2::NewPackage };
 
 	struct FAllDependenciesState
 	{
