Cooker: MPCook: Make UE::Cook::IMPCollector a public WITH_EDITOR type in CoreUObject so it can be subclassed next to runtime Engine or licensee types that need it.
#rnx

diff --git a/Engine/Source/Editor/UnrealEd/Private/Commandlets/AssetRegistryGenerator.cpp b/Engine/Source/Editor/UnrealEd/Private/Commandlets/AssetRegistryGenerator.cpp
--- a/Engine/Source/Editor/UnrealEd/Private/Commandlets/AssetRegistryGenerator.cpp
+++ b/Engine/Source/Editor/UnrealEd/Private/Commandlets/AssetRegistryGenerator.cpp
@@ -9,12 +9,12 @@
 #include "AssetRegistry/CookTagList.h"
 #include "CollectionManagerModule.h"
 #include "CollectionManagerTypes.h"
-#include "Cooker/CookMPCollector.h"
 #include "Cooker/CookPackageData.h"
 #include "Cooker/CookPlatformManager.h"
 #include "Cooker/CookProfiling.h"
 #include "Cooker/CookSandbox.h"
 #include "Cooker/CookWorkerClient.h"
+#include "Cooker/MPCollector.h"
 #include "CookMetadata.h"
 #include "Commandlets/ChunkDependencyInfo.h"
 #include "Commandlets/IChunkDataGenerator.h"
diff --git a/Engine/Source/Editor/UnrealEd/Private/CookOnTheFlyServer.cpp b/Engine/Source/Editor/UnrealEd/Private/CookOnTheFlyServer.cpp
--- a/Engine/Source/Editor/UnrealEd/Private/CookOnTheFlyServer.cpp
+++ b/Engine/Source/Editor/UnrealEd/Private/CookOnTheFlyServer.cpp
@@ -22,7 +22,6 @@
 #include "Cooker/AsyncIODelete.h"
 #include "Cooker/CookDiagnostics.h"
 #include "Cooker/CookDirector.h"
-#include "Cooker/CookMPCollector.h"
 #include "Cooker/CookOnTheFlyServerInterface.h"
 #include "Cooker/CookPackageData.h"
 #include "Cooker/CookPlatformManager.h"
@@ -36,6 +35,7 @@
 #include "Cooker/IoStoreCookOnTheFlyRequestManager.h"
 #include "Cooker/IterativeValidatePackageWriter.h"
 #include "Cooker/LooseCookedPackageWriter.h"
+#include "Cooker/MPCollector.h"
 #include "Cooker/NetworkFileCookOnTheFlyRequestManager.h"
 #include "Cooker/PackageTracker.h"
 #include "Cooker/WorkerRequestsLocal.h"
diff --git a/Engine/Source/Editor/UnrealEd/Private/Cooker/CookDirector.cpp b/Engine/Source/Editor/UnrealEd/Private/Cooker/CookDirector.cpp
--- a/Engine/Source/Editor/UnrealEd/Private/Cooker/CookDirector.cpp
+++ b/Engine/Source/Editor/UnrealEd/Private/Cooker/CookDirector.cpp
@@ -4,10 +4,9 @@
 
 #include "Async/Fundamental/Scheduler.h"
 #include "CompactBinaryTCP.h"
-#include "CookMPCollector.h"
-#include "CookPackageData.h"
-#include "CookPlatformManager.h"
-#include "CookWorkerServer.h"
+#include "Cooker/CookPackageData.h"
+#include "Cooker/CookPlatformManager.h"
+#include "Cooker/CookWorkerServer.h"
 #include "CookOnTheSide/CookOnTheFlyServer.h"
 #include "CoreGlobals.h"
 #include "GenericPlatform/GenericPlatformOutputDevices.h"
diff --git a/Engine/Source/Editor/UnrealEd/Private/Cooker/CookDirector.h b/Engine/Source/Editor/UnrealEd/Private/Cooker/CookDirector.h
--- a/Engine/Source/Editor/UnrealEd/Private/Cooker/CookDirector.h
+++ b/Engine/Source/Editor/UnrealEd/Private/Cooker/CookDirector.h
@@ -7,9 +7,9 @@
 #include "Containers/ArrayView.h"
 #include "Containers/Map.h"
 #include "Containers/UnrealString.h"
-#include "CookMPCollector.h"
-#include "CookSockets.h"
-#include "CookTypes.h"
+#include "Cooker/CookSockets.h"
+#include "Cooker/CookTypes.h"
+#include "Cooker/MPCollector.h"
 #include "HAL/CriticalSection.h"
 #include "HAL/Event.h"
 #include "HAL/LowLevelMemTracker.h"
diff --git a/Engine/Source/Editor/UnrealEd/Private/Cooker/CookMPCollector.cpp b/Engine/Source/Editor/UnrealEd/Private/Cooker/CookMPCollector.cpp
deleted file mode 100644
diff --git a/Engine/Source/Editor/UnrealEd/Private/Cooker/CookMPCollector.h b/Engine/Source/Editor/UnrealEd/Private/Cooker/CookMPCollector.h
deleted file mode 100644
diff --git a/Engine/Source/Editor/UnrealEd/Private/Cooker/CookTypes.h b/Engine/Source/Editor/UnrealEd/Private/Cooker/CookTypes.h
--- a/Engine/Source/Editor/UnrealEd/Private/Cooker/CookTypes.h
+++ b/Engine/Source/Editor/UnrealEd/Private/Cooker/CookTypes.h
@@ -7,6 +7,7 @@
 #include "Containers/Set.h"
 #include "Containers/UnrealString.h"
 #include "Cooker/CookProfiling.h"
+#include "Cooker/MPCollector.h"
 #include "CookOnTheSide/CookOnTheFlyServer.h" // ECookTickFlags
 #include "DerivedDataRequestOwner.h"
 #include "HAL/LowLevelMemTracker.h"
@@ -16,6 +17,7 @@
 #include "Misc/AssertionMacros.h"
 #include "Serialization/PackageWriter.h"
 #include "Templates/Function.h"
+#include "UObject/CookEnums.h"
 #include "UObject/NameTypes.h"
 #include "UObject/SavePackage.h"
 
@@ -106,23 +108,6 @@ namespace UE::Cook
 		Cook
 	};
 
-	/* The Result of a Cook */
-	enum class ECookResult : uint8
-	{
-		/* CookResults have not yet been set */
-		NotAttempted,
-		/* The package was saved with success. */
-		Succeeded,
-		/* The package was processed but SavePackage failed. */
-		Failed,
-		/** The package is a NeverCook package that needs to be added to cookresults for dependency tracking. */
-		NeverCookPlaceholder,
-		/** No information for this platform (used in CookWorker replication) */
-		Invalid,
-		Count,
-		NumBits= FPlatformMath::ConstExprCeilLogTwo(ECookResult::Count),
-	};
-
 	/** Return type for functions called reentrantly that can succeed,fail,or be incomplete */
 	enum class EPollStatus : uint8
 	{
@@ -364,40 +349,6 @@ namespace UE::Cook
 		friend bool ::LoadFromCompactBinary(FCbFieldView Field, UE::Cook::FBeginCookConfigSettings& Value);
 	};
 
-	/**
-	 *  Identifier for a CookWorker process launched from a Director process, or for the local process.
-	 *  A director can have multiple CookWorkers.
-	 */
-	struct FWorkerId
-	{
-	public:
-		FWorkerId() { Id = InvalidId; }
-		constexpr static FWorkerId Invalid() { return FWorkerId(InvalidId); }
-		constexpr static FWorkerId Local() { return FWorkerId(LocalId); }
-		static FWorkerId FromRemoteIndex(uint8 Index) { check(Index < InvalidId-1U);  return FWorkerId(Index + 1U); }
-		static FWorkerId FromLocalOrRemoteIndex(uint8 Index) { check(Index < InvalidId);  return FWorkerId(Index); }
-
-		bool IsValid() const { return Id != InvalidId; }
-		bool IsInvalid() const { return Id == InvalidId; }
-		bool IsLocal() const { return Id == LocalId; }
-		bool IsRemote() const { return Id != InvalidId && Id != LocalId; }
-		uint8 GetRemoteIndex() const { check(IsRemote()); return Id - 1U; }
-		uint8 GetLocalOrRemoteIndex() const { check(IsValid()); return Id; }
-		bool operator==(const FWorkerId& Other) const { return Id == Other.Id; }
-		bool operator!=(const FWorkerId& Other) const { return Id != Other.Id; }
-		bool operator<(const FWorkerId& Other) const { return Id < Other.Id; }
-		inline friend int32 GetTypeHash(const FWorkerId& WorkerId) { return WorkerId.Id; }
-
-	private:
-		constexpr explicit FWorkerId(uint8 InId) : Id(InId) {}
-
-	private:
-		uint8 Id;
-
-		constexpr static uint8 InvalidId = 255;
-		constexpr static uint8 LocalId = 0;
-	};
-
 	/** Report whether commandline/config has disabled use of timeouts throughout the cooker, useful for debugging. */
 	bool IsCookIgnoreTimeouts();
 }
diff --git a/Engine/Source/Editor/UnrealEd/Private/Cooker/CookWorkerClient.cpp b/Engine/Source/Editor/UnrealEd/Private/Cooker/CookWorkerClient.cpp
--- a/Engine/Source/Editor/UnrealEd/Private/Cooker/CookWorkerClient.cpp
+++ b/Engine/Source/Editor/UnrealEd/Private/Cooker/CookWorkerClient.cpp
@@ -2,13 +2,12 @@
 
 #include "CookWorkerClient.h"
 
-#include "CompactBinaryTCP.h"
-#include "CookDirector.h"
-#include "CookMPCollector.h"
-#include "CookPackageData.h"
-#include "CookPlatformManager.h"
-#include "CookTypes.h"
-#include "CookWorkerServer.h"
+#include "Cooker/CompactBinaryTCP.h"
+#include "Cooker/CookDirector.h"
+#include "Cooker/CookPackageData.h"
+#include "Cooker/CookPlatformManager.h"
+#include "Cooker/CookTypes.h"
+#include "Cooker/CookWorkerServer.h"
 #include "CoreGlobals.h"
 #include "HAL/PlatformTime.h"
 #include "Interfaces/ITargetPlatform.h"
diff --git a/Engine/Source/Editor/UnrealEd/Private/Cooker/CookWorkerClient.h b/Engine/Source/Editor/UnrealEd/Private/Cooker/CookWorkerClient.h
--- a/Engine/Source/Editor/UnrealEd/Private/Cooker/CookWorkerClient.h
+++ b/Engine/Source/Editor/UnrealEd/Private/Cooker/CookWorkerClient.h
@@ -4,9 +4,9 @@
 
 #include "CompactBinaryTCP.h"
 #include "Containers/Array.h"
-#include "CookMPCollector.h"
+#include "Cooker/CookTypes.h"
+#include "Cooker/MPCollector.h"
 #include "CookOnTheSide/CookOnTheFlyServer.h"
-#include "CookTypes.h"
 #include "HAL/CriticalSection.h"
 #include "IPAddress.h"
 #include "Misc/Guid.h"
diff --git a/Engine/Source/Editor/UnrealEd/Private/Cooker/CookWorkerServer.cpp b/Engine/Source/Editor/UnrealEd/Private/Cooker/CookWorkerServer.cpp
--- a/Engine/Source/Editor/UnrealEd/Private/Cooker/CookWorkerServer.cpp
+++ b/Engine/Source/Editor/UnrealEd/Private/Cooker/CookWorkerServer.cpp
@@ -5,10 +5,9 @@
 #include "Algo/Find.h"
 #include "Commandlets/AssetRegistryGenerator.h"
 #include "CompactBinaryTCP.h"
-#include "CookDirector.h"
-#include "CookMPCollector.h"
-#include "CookPackageData.h"
-#include "CookPlatformManager.h"
+#include "Cooker/CookDirector.h"
+#include "Cooker/CookPackageData.h"
+#include "Cooker/CookPlatformManager.h"
 #include "HAL/PlatformProcess.h"
 #include "Interfaces/ITargetPlatform.h"
 #include "Interfaces/ITargetPlatformManagerModule.h"
diff --git a/Engine/Source/Editor/UnrealEd/Private/Cooker/CookWorkerServer.h b/Engine/Source/Editor/UnrealEd/Private/Cooker/CookWorkerServer.h
--- a/Engine/Source/Editor/UnrealEd/Private/Cooker/CookWorkerServer.h
+++ b/Engine/Source/Editor/UnrealEd/Private/Cooker/CookWorkerServer.h
@@ -7,11 +7,11 @@
 #include "Containers/ArrayView.h"
 #include "Containers/RingBuffer.h"
 #include "Containers/Set.h"
+#include "Cooker/CookPackageData.h"
+#include "Cooker/CookSockets.h"
+#include "Cooker/CookTypes.h"
+#include "Cooker/MPCollector.h"
 #include "CookOnTheSide/CookOnTheFlyServer.h"
-#include "CookMPCollector.h"
-#include "CookPackageData.h"
-#include "CookSockets.h"
-#include "CookTypes.h"
 #include "HAL/CriticalSection.h"
 #include "Logging/LogVerbosity.h"
 #include "Misc/Guid.h"
diff --git a/Engine/Source/Editor/UnrealEd/Private/Cooker/PackageResultsMessage.h b/Engine/Source/Editor/UnrealEd/Private/Cooker/PackageResultsMessage.h
--- a/Engine/Source/Editor/UnrealEd/Private/Cooker/PackageResultsMessage.h
+++ b/Engine/Source/Editor/UnrealEd/Private/Cooker/PackageResultsMessage.h
@@ -5,8 +5,8 @@
 #include "Async/Future.h"
 #include "CompactBinaryTCP.h"
 #include "Containers/Array.h"
-#include "CookMPCollector.h"
-#include "CookTypes.h"
+#include "Cooker/CookTypes.h"
+#include "Cooker/MPCollector.h"
 #include "HAL/CriticalSection.h"
 #include "Templates/UniquePtr.h"
 #include "UObject/NameTypes.h"
diff --git a/Engine/Source/Editor/UnrealEd/Private/ShaderStatsCollector.h b/Engine/Source/Editor/UnrealEd/Private/ShaderStatsCollector.h
--- a/Engine/Source/Editor/UnrealEd/Private/ShaderStatsCollector.h
+++ b/Engine/Source/Editor/UnrealEd/Private/ShaderStatsCollector.h
@@ -2,7 +2,8 @@
 
 #pragma once
 
-#include "Cooker/CookMPCollector.h"
+#include "Cooker/MPCollector.h"
+#include "ShaderCompiler.h"
 #include "TickableEditorObject.h"
 
 class FShaderStatsAggregator : public UE::Cook::IMPCollector{
diff --git a/Engine/Source/Editor/UnrealEd/Public/Commandlets/AssetRegistryGenerator.h b/Engine/Source/Editor/UnrealEd/Public/Commandlets/AssetRegistryGenerator.h
--- a/Engine/Source/Editor/UnrealEd/Public/Commandlets/AssetRegistryGenerator.h
+++ b/Engine/Source/Editor/UnrealEd/Public/Commandlets/AssetRegistryGenerator.h
@@ -8,7 +8,7 @@
 #include "Containers/Map.h"
 #include "Containers/Set.h"
 #include "Containers/UnrealString.h"
-#include "Cooker/CookMPCollector.h"
+#include "Cooker/MPCollector.h"
 #include "Misc/AssetRegistryInterface.h"
 #include "Misc/Optional.h"
 #include "Templates/SharedPointer.h"
@@ -22,9 +22,11 @@ class IAssetRegistry;
 class ITargetPlatform;
 class IChunkDataGenerator;
 class UChunkDependencyInfo;
+class UCookOnTheFlyServer;
 struct FChunkDependencyTreeNode;
 struct FCookTagList;
 struct FSoftObjectPath;
+namespace UE::Cook { class FAssetRegistryMPCollector; }
 namespace UE::Cook { class FAssetRegistryPackageMessage; }
 namespace UE::Cook { class FCookSandbox; }
 namespace UE::Cook { class FCookWorkerClient; }
diff --git a/Engine/Source/Runtime/CoreUObject/Private/Cooker/MPCollector.cpp b/Engine/Source/Runtime/CoreUObject/Private/Cooker/MPCollector.cpp
new file mode 100644
--- /dev/null
+++ b/Engine/Source/Runtime/CoreUObject/Private/Cooker/MPCollector.cpp
@@ -0,0 +1,98 @@
+// Copyright Epic Games, Inc. All Rights Reserved.
+
+#include "Cooker/MPCollector.h"
+
+#if WITH_EDITOR
+
+namespace UE::Cook
+{
+
+static uint8 PlatformToIntImpl(const ITargetPlatform* Platform, TConstArrayView<const ITargetPlatform*> Platforms)
+{
+	int32 Result = Platforms.IndexOfByKey(Platform);
+	check(Result != INDEX_NONE);
+	uint8 NarrowResult = static_cast<uint8>(Result);
+	check(static_cast<int32>(NarrowResult) == Result);
+	return NarrowResult;
+}
+
+static const ITargetPlatform* IntToPlatformImpl(uint8 PlatformAsInt, TConstArrayView<const ITargetPlatform*> Platforms)
+{
+	int32 Index = static_cast<int32>(PlatformAsInt);
+	if (Index < Platforms.Num())
+	{
+		return Platforms[Index];
+	}
+	return nullptr;
+}
+
+void FMPCollectorClientTickContext::AddMessage(FCbObject Object)
+{
+	Messages.Add(Object);
+}
+
+uint8 FMPCollectorClientTickContext::PlatformToInt(const ITargetPlatform* Platform) const
+{
+	return PlatformToIntImpl(Platform, Platforms);
+}
+
+const ITargetPlatform* FMPCollectorClientTickContext::IntToPlatform(uint8 PlatformAsInt) const
+{
+	return IntToPlatformImpl(PlatformAsInt, Platforms);
+}
+
+void FMPCollectorClientTickPackageContext::AddMessage(FCbObject Object)
+{
+	Messages.Emplace(nullptr, MoveTemp(Object));
+}
+
+void FMPCollectorClientTickPackageContext::AddAsyncMessage(TFuture<FCbObject>&& ObjectFuture)
+{
+	AsyncMessages.Emplace(nullptr, MoveTemp(ObjectFuture));
+}
+
+void FMPCollectorClientTickPackageContext::AddPlatformMessage(const ITargetPlatform* TargetPlatform, FCbObject Object)
+{
+	check(TargetPlatform && Platforms.Contains(TargetPlatform));
+	Messages.Emplace(TargetPlatform, MoveTemp(Object));
+}
+
+void FMPCollectorClientTickPackageContext::AddAsyncPlatformMessage(const ITargetPlatform* TargetPlatform, TFuture<FCbObject>&& ObjectFuture)
+{
+	check(TargetPlatform && Platforms.Contains(TargetPlatform));
+	AsyncMessages.Emplace(TargetPlatform, MoveTemp(ObjectFuture));
+}
+
+uint8 FMPCollectorClientTickPackageContext::PlatformToInt(const ITargetPlatform* Platform) const
+{
+	return PlatformToIntImpl(Platform, Platforms);
+}
+
+const ITargetPlatform* FMPCollectorClientTickPackageContext::IntToPlatform(uint8 PlatformAsInt) const
+{
+	return IntToPlatformImpl(PlatformAsInt, Platforms);
+}
+
+uint8 FMPCollectorClientMessageContext::PlatformToInt(const ITargetPlatform* Platform) const
+{
+	return PlatformToIntImpl(Platform, Platforms);
+}
+
+const ITargetPlatform* FMPCollectorClientMessageContext::IntToPlatform(uint8 PlatformAsInt) const
+{
+	return IntToPlatformImpl(PlatformAsInt, Platforms);
+}
+
+uint8 FMPCollectorServerMessageContext::PlatformToInt(const ITargetPlatform* Platform) const
+{
+	return PlatformToIntImpl(Platform, Platforms);
+}
+
+const ITargetPlatform* FMPCollectorServerMessageContext::IntToPlatform(uint8 PlatformAsInt) const
+{
+	return IntToPlatformImpl(PlatformAsInt, Platforms);
+}
+
+} // namespace UE::Cook
+
+#endif // WITH_EDITOR
\ No newline at end of file
diff --git a/Engine/Source/Runtime/CoreUObject/Public/Cooker/MPCollector.h b/Engine/Source/Runtime/CoreUObject/Public/Cooker/MPCollector.h
new file mode 100644
--- /dev/null
+++ b/Engine/Source/Runtime/CoreUObject/Public/Cooker/MPCollector.h
@@ -0,0 +1,277 @@
+// Copyright Epic Games, Inc. All Rights Reserved.
+
+#pragma once
+
+#if WITH_EDITOR
+#include "Async/Future.h"
+#include "Containers/Array.h"
+#include "Containers/ArrayView.h"
+#include "HAL/Platform.h"
+#include "Misc/Guid.h"
+#include "Serialization/CompactBinary.h"
+#include "Templates/Function.h"
+#include "Templates/RefCounting.h"
+#include "UObject/CookEnums.h"
+#endif
+
+#if WITH_EDITOR
+
+class FCbObject;
+class FCbObjectView;
+class FCbWriter;
+class ITargetPlatform;
+
+namespace UE::Cook { class FCookWorkerClient; }
+namespace UE::Cook { class FCookWorkerServer; }
+
+namespace UE::Cook
+{
+
+/**
+ *  Identifier for a CookWorker process launched from a Director process, or for the local process.
+ *  A director can have multiple CookWorkers.
+ */
+struct FWorkerId
+{
+public:
+	FWorkerId() { Id = InvalidId; }
+	constexpr static FWorkerId Invalid() { return FWorkerId(InvalidId); }
+	constexpr static FWorkerId Local() { return FWorkerId(LocalId); }
+	static FWorkerId FromRemoteIndex(uint8 Index) { check(Index < InvalidId - 1U);  return FWorkerId(Index + 1U); }
+	static FWorkerId FromLocalOrRemoteIndex(uint8 Index) { check(Index < InvalidId);  return FWorkerId(Index); }
+
+	bool IsValid() const { return Id != InvalidId; }
+	bool IsInvalid() const { return Id == InvalidId; }
+	bool IsLocal() const { return Id == LocalId; }
+	bool IsRemote() const { return Id != InvalidId && Id != LocalId; }
+	uint8 GetRemoteIndex() const { check(IsRemote()); return Id - 1U; }
+	uint8 GetLocalOrRemoteIndex() const { check(IsValid()); return Id; }
+	bool operator==(const FWorkerId& Other) const { return Id == Other.Id; }
+	bool operator!=(const FWorkerId& Other) const { return Id != Other.Id; }
+	bool operator<(const FWorkerId& Other) const { return Id < Other.Id; }
+	inline friend int32 GetTypeHash(const FWorkerId& WorkerId) { return WorkerId.Id; }
+
+private:
+	constexpr explicit FWorkerId(uint8 InId) : Id(InId) {}
+
+private:
+	uint8 Id;
+
+	constexpr static uint8 InvalidId = 255;
+	constexpr static uint8 LocalId = 0;
+};
+
+class FMPCollectorClientTickContext
+{
+public:
+	TConstArrayView<const ITargetPlatform*> GetPlatforms() const { return Platforms; }
+	bool IsFlush()  const { return bFlush; }
+
+	COREUOBJECT_API void AddMessage(FCbObject Object);
+
+	COREUOBJECT_API uint8 PlatformToInt(const ITargetPlatform* Platform) const;
+	COREUOBJECT_API const ITargetPlatform* IntToPlatform(uint8 PlatformAsInt) const;
+
+private:
+	TConstArrayView<const ITargetPlatform*> Platforms;
+	TArray<FCbObject> Messages;
+	bool bFlush = false;
+
+	friend class FCookWorkerClient;
+};
+
+class FMPCollectorClientTickPackageContext
+{
+public:
+	struct FPlatformData
+	{
+		const ITargetPlatform* TargetPlatform = nullptr;
+		ECookResult CookResults = ECookResult::NotAttempted;
+	};
+	TConstArrayView<const ITargetPlatform*> GetPlatforms() const { return Platforms; }
+	TConstArrayView<FPlatformData> GetPlatformDatas() const { return PlatformDatas; }
+	FName GetPackageName() const { return PackageName; }
+
+	COREUOBJECT_API void AddMessage(FCbObject Object);
+	COREUOBJECT_API void AddAsyncMessage(TFuture<FCbObject>&& ObjectFuture);
+	COREUOBJECT_API void AddPlatformMessage(const ITargetPlatform* Platform, FCbObject Object);
+	COREUOBJECT_API void AddAsyncPlatformMessage(const ITargetPlatform* Platform, TFuture<FCbObject>&& ObjectFuture);
+
+	COREUOBJECT_API uint8 PlatformToInt(const ITargetPlatform* Platform) const;
+	COREUOBJECT_API const ITargetPlatform* IntToPlatform(uint8 PlatformAsInt) const;
+
+private:
+	TArray<TPair<const ITargetPlatform*, FCbObject>> Messages;
+	TArray<TPair<const ITargetPlatform*, TFuture<FCbObject>>> AsyncMessages;
+	TConstArrayView<const ITargetPlatform*> Platforms;
+	TConstArrayView<FPlatformData> PlatformDatas;
+	FName PackageName;
+
+	friend class FCookWorkerClient;
+};
+
+class FMPCollectorClientMessageContext
+{
+public:
+	TConstArrayView<const ITargetPlatform*> GetPlatforms() { return Platforms; }
+
+	COREUOBJECT_API uint8 PlatformToInt(const ITargetPlatform* Platform) const;
+	COREUOBJECT_API const ITargetPlatform* IntToPlatform(uint8 PlatformAsInt) const;
+
+private:
+	TConstArrayView<const ITargetPlatform*> Platforms;
+
+	friend class FCookWorkerClient;
+};
+
+class FMPCollectorServerMessageContext
+{
+public:
+	TConstArrayView<const ITargetPlatform*> GetPlatforms() { return Platforms; }
+
+	COREUOBJECT_API uint8 PlatformToInt(const ITargetPlatform* Platform) const;
+	COREUOBJECT_API const ITargetPlatform* IntToPlatform(uint8 PlatformAsInt) const;
+	FWorkerId GetWorkerId() const { return WorkerId; }
+	int32 GetProfileId() const { return ProfileId; }
+	FCookWorkerServer* GetCookWorkerServer() { return Server; }
+
+	bool HasPackageName() const { return !PackageName.IsNone(); }
+	FName GetPackageName() const { return PackageName; }
+	bool HasTargetPlatform() const { return TargetPlatform != nullptr; }
+	const ITargetPlatform* GetTargetPlatform() const { return TargetPlatform; }
+
+private:
+	TConstArrayView<const ITargetPlatform*> Platforms;
+	FName PackageName;
+	FCookWorkerServer* Server = nullptr;
+	const ITargetPlatform* TargetPlatform = nullptr;
+	int32 ProfileId;
+	FWorkerId WorkerId;
+
+	friend class FCookWorkerServer;
+};
+
+/**
+ * Interface used during cooking to send data collected from save/load on a remote CookWorker
+ * to the Director for aggregation into files saves at the end of the cook. 
+ */
+class IMPCollector : public FRefCountBase
+{
+public:
+	virtual ~IMPCollector() {}
+
+	virtual FGuid GetMessageType() const = 0;
+	virtual const TCHAR* GetDebugName() const = 0;
+
+	virtual void ClientTick(FMPCollectorClientTickContext& Context) {}
+	virtual void ClientTickPackage(FMPCollectorClientTickPackageContext& Context) {}
+	virtual void ClientReceiveMessage(FMPCollectorClientMessageContext& Context, FCbObjectView Message) {}
+	virtual void ServerReceiveMessage(FMPCollectorServerMessageContext& Context, FCbObjectView Message) {}
+};
+
+
+/** Baseclass for messages used by IMPCollectors that want to interpret their messages as c++ structs. */
+class IMPCollectorMessage
+{
+public:
+	virtual ~IMPCollectorMessage() {}
+
+	/** Marshall the message to a CompactBinaryObject. */
+	virtual void Write(FCbWriter& Writer) const = 0;
+	/** Unmarshall the message from a CompactBinaryObject. */
+	virtual bool TryRead(FCbObjectView Object) = 0;
+	/** Return the Guid that identifies the message to the remote connection. */
+	virtual FGuid GetMessageType() const = 0;
+	/** Return the debugname for diagnostics. */
+	virtual const TCHAR* GetDebugName() const = 0;
+};
+
+
+/** A subinterface of IMPCollector that uses a ICollectorMessage subclass to serialize the message. */
+template <typename MessageType>
+class IMPCollectorForMessage : public IMPCollector
+{
+public:
+	virtual void ClientReceiveMessage(FMPCollectorClientMessageContext& Context, bool bReadSuccessful,
+		MessageType&& Message) {}
+	virtual void ServerReceiveMessage(FMPCollectorServerMessageContext& Context, bool bReadSuccessful,
+		MessageType&& Message) {}
+
+	virtual FGuid GetMessageType() const override
+	{
+		return MessageType().GetMessageType();
+	}
+
+	virtual const TCHAR* GetDebugName() const override
+	{
+		return MessageType().GetDebugName();
+	}
+
+	virtual void ClientReceiveMessage(FMPCollectorClientMessageContext& Context, FCbObjectView Message) override
+	{
+		MessageType CBMessage;
+		bool bReadSuccessful = CBMessage.TryRead(Message);
+		ClientReceiveMessage(Context, bReadSuccessful, MoveTemp(CBMessage));
+	}
+
+	virtual void ServerReceiveMessage(FMPCollectorServerMessageContext& Context, FCbObjectView Message) override
+	{
+		MessageType CBMessage;
+		bool bReadSuccessful = CBMessage.TryRead(Message);
+		ServerReceiveMessage(Context, bReadSuccessful, MoveTemp(CBMessage));
+	}
+};
+
+/**
+ * An implementation of IMPCollector that uses an ICollectorMessage subclass to serialize the message,
+ * and directs messages received on the client to the given callback.
+ */
+template <typename MessageType>
+class TMPCollectorClientMessageCallback : public IMPCollectorForMessage<MessageType>
+{
+public:
+	TMPCollectorClientMessageCallback(TUniqueFunction<void(FMPCollectorClientMessageContext& Context, bool bReadSuccessful,
+		MessageType&& Message)>&& InCallback)
+		: Callback(MoveTemp(InCallback))
+	{
+		check(Callback);
+	}
+
+	virtual void ClientReceiveMessage(FMPCollectorClientMessageContext& Context, bool bReadSuccessful,
+		MessageType&& Message) override
+	{
+		Callback(Context, bReadSuccessful, MoveTemp(Message));
+	}
+
+private:
+	TUniqueFunction<void(FMPCollectorClientMessageContext& Context, bool bReadSuccessful, MessageType&& Message)> Callback;
+};
+
+/**
+ * An implementation of IMPCollector that uses an IMPCollectorMessage subclass to serialize the message,
+ * and directs messages received on the server to the given callback.
+ */
+template <typename MessageType>
+class TMPCollectorServerMessageCallback : public IMPCollectorForMessage<MessageType>
+{
+public:
+	TMPCollectorServerMessageCallback(TUniqueFunction<void(FMPCollectorServerMessageContext& Context, bool bReadSuccessful,
+		MessageType&& Message)>&& InCallback)
+		: Callback(MoveTemp(InCallback))
+	{
+		check(Callback);
+	}
+
+	virtual void ServerReceiveMessage(FMPCollectorServerMessageContext& Context, bool bReadSuccessful,
+		MessageType&& Message) override
+	{
+		Callback(Context, bReadSuccessful, MoveTemp(Message));
+	}
+
+private:
+	TUniqueFunction<void(FMPCollectorServerMessageContext& Context, bool bReadSuccessful, MessageType&& Message)> Callback;
+};
+
+} // namespace UE::Cook
+
+#endif // WITH_EDITOR
\ No newline at end of file
diff --git a/Engine/Source/Runtime/CoreUObject/Public/UObject/CookEnums.h b/Engine/Source/Runtime/CoreUObject/Public/UObject/CookEnums.h
--- a/Engine/Source/Runtime/CoreUObject/Public/UObject/CookEnums.h
+++ b/Engine/Source/Runtime/CoreUObject/Public/UObject/CookEnums.h
@@ -2,6 +2,8 @@
 
 #pragma once
 
+#include "HAL/PlatformMath.h"
+
 namespace UE::Cook
 {
 
@@ -19,5 +21,22 @@ enum class ECookingDLC
 	No,
 };
 
+/* The Result of the Cook-Save of a package */
+enum class ECookResult : uint8
+{
+	/* CookResults have not yet been set */
+	NotAttempted,
+	/* The package was saved with success. */
+	Succeeded,
+	/* The package was processed but SavePackage failed. */
+	Failed,
+	/** The package is a NeverCook package that needs to be added to cookresults for dependency tracking. */
+	NeverCookPlaceholder,
+	/** No information for this platform (used in CookWorker replication) */
+	Invalid,
+	Count,
+	NumBits = FPlatformMath::ConstExprCeilLogTwo(ECookResult::Count),
+};
+
 }
 
