Removing redundant check that results in race conditions after GC is done with reachability analysis

diff --git a/Engine/Source/Runtime/CoreUObject/Public/UObject/FastReferenceCollector.h b/Engine/Source/Runtime/CoreUObject/Public/UObject/FastReferenceCollector.h
--- a/Engine/Source/Runtime/CoreUObject/Public/UObject/FastReferenceCollector.h
+++ b/Engine/Source/Runtime/CoreUObject/Public/UObject/FastReferenceCollector.h
@@ -831,7 +831,6 @@ StoleARO:
 			CurrentObjects = MakeArrayView(Block->Objects, BlockSize);
 		} // while (true)
 		
-		check(!Context.ObjectsToSerialize.HasWork());
 		Processor.LogDetailedStatsSummary();
 	}
 
