Fix FastGetAsset that was looking in partially loaded packages

#rnx

diff --git a/Engine/Source/Runtime/CoreUObject/Public/AssetRegistry/AssetData.h b/Engine/Source/Runtime/CoreUObject/Public/AssetRegistry/AssetData.h
--- a/Engine/Source/Runtime/CoreUObject/Public/AssetRegistry/AssetData.h
+++ b/Engine/Source/Runtime/CoreUObject/Public/AssetRegistry/AssetData.h
@@ -510,6 +510,13 @@ PRAGMA_ENABLE_DEPRECATION_WARNINGS
 
 		// We load PackageName rather than using GetObjectPath because external assets may be saved in a different package to their loaded path. 
 		UPackage* FoundPackage = FindObjectFast<UPackage>(nullptr, PackageName);
+
+		// We could find an existing package in memory that is still pending load in the loader.
+		if (FoundPackage && !FoundPackage->IsFullyLoaded())
+		{
+			FoundPackage = nullptr;
+		}
+
 		if (FoundPackage == nullptr && bLoad)
 		{
 			FLinkerInstancingContext InstancingContext(MoveTemp(LoadTags));
