Remove old warning about calling StaticLoadObject from PostLoad on async threads, this case no longer crashes the engine. The issue was specific to EDL and the warning was very inconsistently displayed.

diff --git a/Engine/Source/Runtime/CoreUObject/Private/UObject/UObjectGlobals.cpp b/Engine/Source/Runtime/CoreUObject/Private/UObject/UObjectGlobals.cpp
--- a/Engine/Source/Runtime/CoreUObject/Private/UObject/UObjectGlobals.cpp
+++ b/Engine/Source/Runtime/CoreUObject/Private/UObject/UObjectGlobals.cpp
@@ -1424,16 +1424,6 @@ UObject* StaticLoadObjectInternal(UClass* ObjectClass, UObject* InOuter, const T
 
 UObject* StaticLoadObject(UClass* ObjectClass, UObject* InOuter, const TCHAR* InName, const TCHAR* Filename, uint32 LoadFlags, UPackageMap* Sandbox, bool bAllowObjectReconciliation, const FLinkerInstancingContext* InstancingContext)
 {
-	FUObjectThreadContext& ThreadContext = FUObjectThreadContext::Get();
-	if ((ThreadContext.SyncLoadUsingAsyncLoaderCount == 0) && ThreadContext.IsRoutingPostLoad && IsInAsyncLoadingThread())
-	{
-		UE_LOG(LogUObjectGlobals, Warning, TEXT("Calling StaticLoadObject(\"%s\", \"%s\", \"%s\") during PostLoad of %s is illegal and will crash in a cooked runtime"), 
-			*GetFullNameSafe(ObjectClass),
-			*GetFullNameSafe(InOuter),
-			InName,
-			*GetFullNameSafe(ThreadContext.CurrentlyPostLoadedObjectByALT));
-	}
-
 	UObject* Result = StaticLoadObjectInternal(ObjectClass, InOuter, InName, Filename, LoadFlags, Sandbox, bAllowObjectReconciliation, InstancingContext);
 	if (!Result)
 	{
