fix for FindUninitializedScriptStructMembers not actually working
https://p4-swarm.epicgames.net/reviews/29019597

change from eryk.dwornicki

diff --git a/Engine/Source/Runtime/CoreUObject/Private/UObject/Class.cpp b/Engine/Source/Runtime/CoreUObject/Private/UObject/Class.cpp
--- a/Engine/Source/Runtime/CoreUObject/Private/UObject/Class.cpp
+++ b/Engine/Source/Runtime/CoreUObject/Private/UObject/Class.cpp
@@ -3497,10 +3497,10 @@ static void FindUninitializedScriptStructMembers(UScriptStruct* ScriptStruct, ES
 
 		if (const FObjectProperty* ObjectPtrProperty = CastField<const FObjectProperty>(Property))
 		{
-			// Check any reflected pointer properties to make sure they got initialized
-			TObjectPtr<UObject> PropValue = ObjectPtrProperty->GetObjectPtrPropertyValue(WrapperFE.GetData());
-			const void** RawValue = reinterpret_cast<const void**>(&PropValue);
-			if (RawValue == BadPointer)
+			//using reinterpret_cast to avoid any methods of TObjectPtr being invoked
+			const TObjectPtr<UObject>* PropValue = ObjectPtrProperty->GetPropertyValuePtr_InContainer(WrapperFE.GetData());
+			const void* const* RawValue = reinterpret_cast<const void* const*>(PropValue);
+			if (*RawValue == BadPointer)
 			{
 				OutUninitializedProperties.Add(Property);
 			}
@@ -3639,18 +3639,20 @@ int32 FStructUtils::AttemptToFindUninitializedScriptStructMembers()
 
 		// First test if the tests aren't broken
 		FScriptStructTestWrapper WrapperFE(TestUninitializedScriptStructMembersTestStruct, 0xFE);
-		const FObjectPropertyBase* UninitializedProperty = CastFieldChecked<const FObjectPropertyBase>(TestUninitializedScriptStructMembersTestStruct->FindPropertyByName(TEXT("UninitializedObjectReference")));
-		const FObjectPropertyBase* InitializedProperty = CastFieldChecked<const FObjectPropertyBase>(TestUninitializedScriptStructMembersTestStruct->FindPropertyByName(TEXT("InitializedObjectReference")));
+		const FObjectProperty* UninitializedProperty = CastFieldChecked<const FObjectProperty>(TestUninitializedScriptStructMembersTestStruct->FindPropertyByName(TEXT("UninitializedObjectReference")));
+		const FObjectProperty* InitializedProperty = CastFieldChecked<const FObjectProperty>(TestUninitializedScriptStructMembersTestStruct->FindPropertyByName(TEXT("InitializedObjectReference")));
 		
-		TObjectPtr<UObject> UninitializedPropValue = UninitializedProperty->GetObjectPtrPropertyValue(WrapperFE.GetData());
-		const void** RawValue = reinterpret_cast<const void**>(&UninitializedPropValue);
+		//using reinterpret_cast to avoid any methods of TObjectPtr being invoked
+		const TObjectPtr<UObject>* UninitializedPropValue = UninitializedProperty->GetPropertyValuePtr_InContainer(WrapperFE.GetData());
+		const void* const* RawValue = reinterpret_cast<const void**>(&UninitializedPropValue);
 		if (*RawValue != BadPointer)
 		{
 			UE_LOG(LogClass, Warning, TEXT("ObjectProperty %s%s::%s seems to be initialized properly but it shouldn't be. Verify that AttemptToFindUninitializedScriptStructMembers() is working properly"), 
 				TestUninitializedScriptStructMembersTestStruct->GetPrefixCPP(), *TestUninitializedScriptStructMembersTestStruct->GetName(), *UninitializedProperty->GetNameCPP());
 		}
-		const UObject* InitializedPropValue = InitializedProperty->GetObjectPropertyValue_InContainer(WrapperFE.GetData());
-		if (InitializedPropValue != nullptr)
+		const TObjectPtr<UObject>* InitializedPropValue = InitializedProperty->GetPropertyValuePtr_InContainer(WrapperFE.GetData());
+		RawValue = reinterpret_cast<const void**>(&UninitializedPropValue);
+		if (*RawValue != nullptr)
 		{
 			UE_LOG(LogClass, Warning, TEXT("ObjectProperty %s%s::%s seems to be not initialized properly but it should be. Verify that AttemptToFindUninitializedScriptStructMembers() is working properly"),
 				TestUninitializedScriptStructMembersTestStruct->GetPrefixCPP(), *TestUninitializedScriptStructMembersTestStruct->GetName(), *InitializedProperty->GetNameCPP());
