Ensure that default values are preloaded during class regeneration.

Blueprint compilation expects all of its dependencies to be preloaded. Normally, RegenerateClass handles this, but we were still missing a few cases. If we don't correctly preload dependent exports, we can run into scenarios where postloading can occur prior to preloading. This typically occurs when we duplicate various objects during compilation.

We can leverage FReferenceFinder to recursively find dependent exports on the class default object to preload. Note that this only covers _native_ export dependencies. Preloading Blueprint-generated classes (and any instances of one) can get pretty thorny during this phase, so we intentionally skip those.

diff --git a/Engine/Source/Runtime/CoreUObject/Private/Blueprint/BlueprintSupport.cpp b/Engine/Source/Runtime/CoreUObject/Private/Blueprint/BlueprintSupport.cpp
--- a/Engine/Source/Runtime/CoreUObject/Private/Blueprint/BlueprintSupport.cpp
+++ b/Engine/Source/Runtime/CoreUObject/Private/Blueprint/BlueprintSupport.cpp
@@ -539,6 +539,33 @@ struct FPreloadMembersHelper
 		}
 	}
 
+	static void PreloadExternalNativeDependencies(UObject* InObject)
+	{
+		TArray<UObject*> MemberReferences;
+		FReferenceFinder ComponentCollector(MemberReferences, nullptr, false, true, true, true);
+		ComponentCollector.FindReferences(InObject);
+
+		for (UObject* CurrentObject : MemberReferences)
+		{
+			check(CurrentObject);
+
+			const bool bIsValidNativeDependency =
+				CurrentObject->HasAnyFlags(RF_NeedLoad) &&
+				!CurrentObject->IsA<UClass>() &&
+				!CurrentObject->GetClass()->HasAnyClassFlags(CLASS_CompiledFromBlueprint)
+			;
+
+			if (bIsValidNativeDependency)
+			{
+				if (FLinkerLoad* Linker = CurrentObject->GetLinker())
+				{
+					Linker->Preload(CurrentObject);
+					PreloadExternalNativeDependencies(CurrentObject);
+				}
+			}
+		}
+	}
+
 	static void PreloadObject(UObject* InObject)
 	{
 		if (InObject && !InObject->HasAnyFlags(RF_LoadCompleted))
@@ -733,6 +760,9 @@ bool FLinkerLoad::RegenerateBlueprintClass(UClass* LoadClass, UObject* ClassDefa
 			}
 		}
 
+		// The CDO may reference default values that live in other packages, which need to be preloaded.
+		FPreloadMembersHelper::PreloadExternalNativeDependencies(ClassDefaultObject);
+
 		{
 			// RegenerateClass largely performs redundant work since we already preloaded the remaining exports.
 			// However, in some circumstances, a re-preload might occur on an export that doesn't have RF_LoadCompleted set due to a cycle.
