Core Containers: DirectoryTree: Add GetChildren function for callers that want to use the DirectoryTree to keep track of which paths have been added.
#rnx

diff --git a/Engine/Source/Runtime/CoreUObject/Internal/Containers/DirectoryTree.h b/Engine/Source/Runtime/CoreUObject/Internal/Containers/DirectoryTree.h
--- a/Engine/Source/Runtime/CoreUObject/Internal/Containers/DirectoryTree.h
+++ b/Engine/Source/Runtime/CoreUObject/Internal/Containers/DirectoryTree.h
@@ -6,7 +6,9 @@
 #include "Containers/StringView.h"
 #include "Containers/UnrealString.h"
 #include "HAL/Platform.h"
+#include "Misc/EnumClassFlags.h"
 #include "Misc/PathViews.h"
+#include "Misc/ScopeExit.h"
 #include "Misc/StringBuilder.h"
 #include "Templates/TypeCompatibleBytes.h"
 #include "Templates/UniquePtr.h"
@@ -20,6 +22,36 @@ COREUOBJECT_API int32 FindInsertionIndex(int32 NumChildNodes, const TUniquePtr<F
 
 }
 
+enum class EDirectoryTreeGetFlags
+{
+	None = 0,
+	/**
+	 * If Recursive flag is present, GetChildren will return direct subpaths of a discovered directory and their
+	 * transitive subpaths. If false, it will return only the direct subpaths.
+	 *
+	 * Recursive=false and ImpliedChildren=false is an exception to this simple definition. In that case the
+	 * reported results for a requested directory will include the highest level childpaths under it that have been
+	 * added to the tree. These may be in transitive subpaths of the parent directory, and in that case their
+	 * parent directories in between the requested directory and their path will not be reported because they are
+	 * implied directories.
+	 */
+	Recursive = 0x1,
+	/**
+	 * If ImpliedParent flag is present, then the requested directory will return results even if it is an
+	 * implied directory (directory with child paths but not added itself, @see TDirectoryTree).
+	 * If not present, only directories that have been added to the tree will return non-empty results.
+	 */
+	ImpliedParent = 0x1 << 1,
+	/**
+	 * If ImpliedChildren is present, then all child paths discovered (either direct or recursive, depending on whether
+	 * Recursive flag is present) will be reported in the results, even if they are implied directories (directory
+	 * with child paths but not added itself, @see TDirectoryTree). If not present, only files and directories that
+	 * have been added to the tree will be returned in the results.
+	 */
+	ImpliedChildren = 0x1 << 2,
+};
+ENUM_CLASS_FLAGS(EDirectoryTreeGetFlags);
+
 /**
  * Container for path -> value that can efficiently report whether a parent directory of a given path exists.
  * Supports relative and absolute paths.
@@ -41,6 +73,10 @@ COREUOBJECT_API int32 FindInsertionIndex(int32 NumChildNodes, const TUniquePtr<F
  * 
  * For functions that return Values by reference or by pointer, that reference or pointer can be invalidated
  * by any functions that modify the tree, and should be discarded before calling any such functions.
+ * 
+ * Some functions that report results for directories behave differently for added directories versus implied
+ * directories. An added directory is one that was added specifically via FindOrAdd or other mutators. An implied
+ * directory is a directory that is not added, but that has a child path that is added to the tree.
  */
 template <typename ValueType>
 class TDirectoryTree
@@ -54,6 +90,8 @@ public:
 	 * This reference can be invalidated by any operations that modify the tree.
 	 */
 	ValueType& FindOrAdd(FStringView Path, bool* bOutExisted = nullptr);
+	/** Remove all paths and all memory usage from the tree. */
+	void Empty();
 	/** Remove a path from the tree and optionally report whether it existed. */
 	void Remove(FStringView Path, bool* bOutExisted = nullptr);
 	/** Free unused slack memory throughout the tree by reallocating containers tightly to their current size. */
@@ -105,6 +143,16 @@ public:
 	 */
 	bool TryFindClosestPath(FStringView Path, FString& OutPath, ValueType** OutValue);
 
+	/**
+	 * Report the children (optionally recursive or not, optionally implied or not) in the tree of a given Path
+	 * (optionally skipped if implied). @see EDirectoryTreeGetFlags.
+	 * Relative paths of discovered children will be appended to OutRelativeChildNames.
+	 * 
+	 * @return true iff (the path is found in the tree and either it is an added path or ImpliedParent was requested).
+	 */
+	bool TryGetChildren(FStringView Path, TArray<FString>& OutRelativeChildNames,
+		EDirectoryTreeGetFlags Flags = EDirectoryTreeGetFlags::None) const;
+
 private:
 
 	/**
@@ -148,6 +196,9 @@ private:
 		/** Return pointer to the Value stored in RelPath, if RelPath exists in the tree. */
 		ValueType* Find(FStringView InRelPath);
 
+		bool TryGetChildren(FStringBuilderBase& ReportedPathPrefix, TCHAR InPathSeparator, FStringView InRelPath,
+			TArray<FString>& OutRelativeChildNames, EDirectoryTreeGetFlags Flags) const;
+
 		/**
 		 * Recursively search the node's subtree to find the given RelPath. Return whether the path or any parent is
 		 * found, and if any is found and Outpath is present, append relative path of the discovered path into OutPath.
@@ -265,6 +316,15 @@ inline ValueType& TDirectoryTree<ValueType>::FindOrAdd(FStringView Path, bool* b
 	return Result;
 }
 
+template <typename ValueType>
+inline void TDirectoryTree<ValueType>::Empty()
+{
+	Root.Reset();
+	NumPaths = 0;
+	PathSeparator = '/';
+	bPathSeparatorInitialized = false;
+}
+
 template <typename ValueType>
 inline void TDirectoryTree<ValueType>::Remove(FStringView Path, bool* bOutExisted)
 {
@@ -447,6 +507,18 @@ inline ValueType* TDirectoryTree<ValueType>::TryFindClosestPathInternal(FStringV
 	return Root.HasValue() ? &Root.GetValue() : nullptr;
 }
 
+template <typename ValueType>
+inline bool TDirectoryTree<ValueType>::TryGetChildren(FStringView Path, TArray<FString>& OutRelativeChildNames,
+	EDirectoryTreeGetFlags Flags) const
+{
+	if (Path.IsEmpty() && !EnumHasAnyFlags(Flags, EDirectoryTreeGetFlags::ImpliedParent) && !Root.HasValue())
+	{
+		return false;
+	}
+	TStringBuilder<1024> ReportedPathPrefix;
+	return Root.TryGetChildren(ReportedPathPrefix, PathSeparator, Path, OutRelativeChildNames, Flags);
+}
+
 template <typename ValueType>
 TDirectoryTree<ValueType>::FTreeNode::FTreeNode(FTreeNode&& Other)
 {
@@ -760,6 +832,222 @@ inline ValueType* TDirectoryTree<ValueType>::FTreeNode::Find(FStringView InRelPa
 	}
 }
 
+template <typename ValueType>
+inline bool TDirectoryTree<ValueType>::FTreeNode::TryGetChildren(FStringBuilderBase& ReportPathPrefix,
+	TCHAR InPathSeparator, FStringView InRelPath, TArray<FString>& OutRelativeChildNames,
+	EDirectoryTreeGetFlags Flags) const
+{
+	if (InRelPath.IsEmpty())
+	{
+		// RelPath indicates this node, so append the children of this node to the list.
+		// Caller is responsible for not calling TryGetChildren on *this if results for *this
+		// should not be returned due to EDirectoryTreeGetFlags.
+
+		int32 NumChildNodes = GetNumChildNodes();
+		for (int32 Index = 0; Index < NumChildNodes; ++Index)
+		{
+			const FTreeNode& ChildNode = ChildNodes[Index];
+			const FString& ChildRelPath = RelPaths[Index];
+
+			if (EnumHasAnyFlags(Flags, EDirectoryTreeGetFlags::ImpliedChildren))
+			{
+				// When ImpliedChildren are supposed to be reported, iterate over every stored child
+				// and report the first component of its RelPath as a child.
+				// If recursive, also report the remaining components of its RelPath, and then
+				// forward the call to it to return its recursive children.
+				FStringView FirstComponent;
+				FStringView RemainingPath;
+				FPathViews::SplitFirstComponent(ChildRelPath, FirstComponent, RemainingPath);
+
+				int32 SavedLen = ReportPathPrefix.Len();
+				FPathViews::Append(ReportPathPrefix, FirstComponent);
+				UE::DirectoryTree::FixupPathSeparator(ReportPathPrefix, SavedLen, InPathSeparator);
+				OutRelativeChildNames.Add(*ReportPathPrefix);
+
+				if (EnumHasAnyFlags(Flags, EDirectoryTreeGetFlags::Recursive))
+				{
+					while (!RemainingPath.IsEmpty())
+					{
+						FPathViews::SplitFirstComponent(RemainingPath, FirstComponent, RemainingPath);
+
+						int32 SavedLenForSubPath = ReportPathPrefix.Len();
+						FPathViews::Append(ReportPathPrefix, FirstComponent);
+						UE::DirectoryTree::FixupPathSeparator(ReportPathPrefix, SavedLenForSubPath, InPathSeparator);
+						OutRelativeChildNames.Add(*ReportPathPrefix);
+					}
+
+					(void)ChildNode.TryGetChildren(ReportPathPrefix, InPathSeparator, FStringView(),
+						OutRelativeChildNames, Flags);
+				}
+				ReportPathPrefix.RemoveSuffix(ReportPathPrefix.Len() - SavedLen);
+			}
+			else
+			{
+				// When ImpliedChildren are not supposed to be reported, report each stored child by its full relpath,
+				// unless the child is an implied path. If recursive, also forward the call to it to return its
+				// recursive children.
+				if (ChildNode.HasValue() || EnumHasAnyFlags(Flags, EDirectoryTreeGetFlags::Recursive))
+				{
+					int32 SavedLen = ReportPathPrefix.Len();
+					FPathViews::Append(ReportPathPrefix, ChildRelPath);
+					UE::DirectoryTree::FixupPathSeparator(ReportPathPrefix, SavedLen, InPathSeparator);
+
+					if (ChildNode.HasValue())
+					{
+						OutRelativeChildNames.Add(*ReportPathPrefix);
+					}
+					if (EnumHasAnyFlags(Flags, EDirectoryTreeGetFlags::Recursive))
+					{
+						(void)ChildNode.TryGetChildren(ReportPathPrefix, InPathSeparator, FStringView(),
+							OutRelativeChildNames, Flags);
+					}
+					ReportPathPrefix.RemoveSuffix(ReportPathPrefix.Len() - SavedLen);
+				}
+			}
+		}
+
+		return true;
+	}
+	else
+	{
+		// We are still looking for the requested InRelPath and are not reporting results yet. Look for an existing
+		// stored child that has the same FirstComponent of its path as InRelPath does.
+		FStringView FirstComponent;
+		FStringView RemainingPath;
+		FPathViews::SplitFirstComponent(InRelPath, FirstComponent, RemainingPath);
+		bool bExists;
+		int32 InsertionIndex = FindInsertionIndex(FirstComponent, bExists);
+		if (!bExists)
+		{
+			// No child has the same FirstComponent as InRelPath, so InRelPath does not exist in the tree, not even as
+			// an implied path.
+			return false;
+		}
+
+		FTreeNode& ChildNode = ChildNodes[InsertionIndex];
+		FString& ChildRelPath = RelPaths[InsertionIndex];
+
+		FStringView ExistingFirstComponent;
+		FStringView ExistingRemainingPath;
+		FPathViews::SplitFirstComponent(ChildRelPath, ExistingFirstComponent, ExistingRemainingPath);
+		check(ExistingFirstComponent.Equals(FirstComponent, ESearchCase::IgnoreCase)); // Otherwise FindInsertionIndex would have returned bExists=false
+
+		for (int32 RunawayLoop = 0; RunawayLoop <= InRelPath.Len(); ++RunawayLoop)
+		{
+			if (ExistingRemainingPath.IsEmpty())
+			{
+				// We've reached the end of the existing path, so InRelPath is either equal to or a child of the
+				// ChildNode. If it is equal to the ChildNode, it is our responsibility to NOT call TryGetChildren
+				// if the ChildNode is an implied path and ImpliedParent flag is not requested.
+				if (RemainingPath.IsEmpty() && !EnumHasAnyFlags(Flags, EDirectoryTreeGetFlags::ImpliedParent) &&
+					!ChildNode.HasValue())
+				{
+					return false;
+				}
+
+				return ChildNode.TryGetChildren(ReportPathPrefix, InPathSeparator, RemainingPath,
+					OutRelativeChildNames, Flags);
+			}
+			else if (RemainingPath.IsEmpty())
+			{
+				// We've reached the end of the input path, but not the end of the existing path, so InRelPath is a
+				// parent of the existing path, and is an implied path rather than an added path.
+				if (!EnumHasAnyFlags(Flags, EDirectoryTreeGetFlags::ImpliedParent))
+				{
+					return false;
+				}
+
+				if (EnumHasAnyFlags(Flags, EDirectoryTreeGetFlags::ImpliedChildren))
+				{
+					// When implied children are supposed to be reported, add the next pathcomponent of the remaining
+					// child path as the first reported child of InRelPath. If recursive, also add all the remaining
+					// components of the existing child path and then forward on to the child path for all of its children.
+					FPathViews::SplitFirstComponent(ExistingRemainingPath, ExistingFirstComponent,
+						ExistingRemainingPath);
+					int32 SavedLen = ReportPathPrefix.Len();
+					FPathViews::Append(ReportPathPrefix, ExistingFirstComponent);
+					UE::DirectoryTree::FixupPathSeparator(ReportPathPrefix, SavedLen, InPathSeparator);
+					OutRelativeChildNames.Add(*ReportPathPrefix);
+
+					if (EnumHasAnyFlags(Flags, EDirectoryTreeGetFlags::Recursive))
+					{
+						while (!ExistingRemainingPath.IsEmpty())
+						{
+							FPathViews::SplitFirstComponent(ExistingRemainingPath, ExistingFirstComponent,
+								ExistingRemainingPath);
+
+							int32 SavedLenForSubPath = ReportPathPrefix.Len();
+							FPathViews::Append(ReportPathPrefix, ExistingFirstComponent);
+							UE::DirectoryTree::FixupPathSeparator(ReportPathPrefix, SavedLenForSubPath
+								, InPathSeparator);
+							OutRelativeChildNames.Add(*ReportPathPrefix);
+						}
+
+						(void) ChildNode.TryGetChildren(ReportPathPrefix, InPathSeparator, FStringView(),
+							OutRelativeChildNames, Flags);
+					}
+
+					ReportPathPrefix.RemoveSuffix(ReportPathPrefix.Len() - SavedLen);
+				}
+				else
+				{
+					// When implied children are not supposed to be reported, report the remaining components of the
+					// existing child path as a single string as the first reported child of InRelPath, but only if
+					// it is not an implied path. If recursive, forward to the childnode for all of its children.
+					if (ChildNode.HasValue() || EnumHasAnyFlags(Flags, EDirectoryTreeGetFlags::Recursive))
+					{
+						int32 SavedLenForSubPath = ReportPathPrefix.Len();
+						FPathViews::Append(ReportPathPrefix, ExistingRemainingPath);
+						UE::DirectoryTree::FixupPathSeparator(ReportPathPrefix, SavedLenForSubPath, InPathSeparator);
+
+						if (ChildNode.HasValue())
+						{
+							OutRelativeChildNames.Add(*ReportPathPrefix);
+						}
+						if (EnumHasAnyFlags(Flags, EDirectoryTreeGetFlags::Recursive))
+						{
+							(void)ChildNode.TryGetChildren(ReportPathPrefix, InPathSeparator, FStringView(),
+								OutRelativeChildNames, Flags);
+						}
+					}
+				}
+
+				return true;
+			}
+			else
+			{
+				// Both existing and remaining have more directory components
+				FStringView NextFirstComponent;
+				FStringView NextRemainingPath;
+				FPathViews::SplitFirstComponent(RemainingPath, NextFirstComponent, NextRemainingPath);
+				check(!NextFirstComponent.IsEmpty());
+				FStringView NextExistingFirstComponent;
+				FStringView NextExistingRemainingPath;
+				FPathViews::SplitFirstComponent(ExistingRemainingPath, NextExistingFirstComponent, NextExistingRemainingPath);
+				check(!NextExistingFirstComponent.IsEmpty());
+				if (NextFirstComponent == NextExistingFirstComponent)
+				{
+					// Next component is also a match, go to the next loop iteration to handle the new remainingpaths
+					RemainingPath = NextRemainingPath;
+					ExistingRemainingPath = NextExistingRemainingPath;
+
+					continue;
+				}
+				else
+				{
+					// The existing child diverges from the path components of InRelPath, so InRelPath does not exist
+					// in the tree, not even as an implied path.
+					return false;
+				}
+			}
+		}
+		checkf(false, TEXT("Infinite loop trying to split path %.*s into components."), InRelPath.Len(), InRelPath.GetData());
+	}
+
+	check(false); // Only way to get here is from the checkf in the else block.
+	return false;
+}
+
 template <typename ValueType>
 inline ValueType* TDirectoryTree<ValueType>::FTreeNode::TryFindClosestPath(FStringView InRelPath,
 	FStringBuilderBase* OutPath, TCHAR InPathSeparator)
diff --git a/Engine/Source/Runtime/CoreUObject/Private/Misc/DirectoryTree.cpp b/Engine/Source/Runtime/CoreUObject/Private/Misc/DirectoryTree.cpp
--- a/Engine/Source/Runtime/CoreUObject/Private/Misc/DirectoryTree.cpp
+++ b/Engine/Source/Runtime/CoreUObject/Private/Misc/DirectoryTree.cpp
@@ -3,6 +3,7 @@
 #include "Containers/DirectoryTree.h"
 
 #include "Algo/BinarySearch.h"
+#include "Algo/Sort.h"
 #include "Misc/AutomationTest.h"
 
 namespace UE::DirectoryTree
@@ -214,6 +215,42 @@ bool FDirectoryTreeTest::RunTest(const FString& Parameters)
 			PathType, Permutation, EditPermutationIndex, OtherPermutationIndex);
 	};
 
+	TArray<const TCHAR*> ScratchA;
+	TArray<const TCHAR*> ScratchB;
+	auto UnorderedEquals = [&ScratchA, &ScratchB](const TArray<FString>& A, TConstArrayView<const TCHAR*> B)
+		{
+			if (A.Num() != B.Num())
+			{
+				return false;
+			}
+			ScratchA.Reset(A.Num());
+			ScratchB.Reset(B.Num());
+			for (const FString& AStr : A)
+			{
+				ScratchA.Add(*AStr);
+			}
+			for (const TCHAR* BStr : B)
+			{
+				ScratchB.Add(BStr);
+			}
+			Algo::Sort(ScratchA, [](const TCHAR* StrA, const TCHAR* StrB)
+				{
+					return FCString::Stricmp(StrA, StrB) < 0;
+				});
+			Algo::Sort(ScratchB, [](const TCHAR* StrA, const TCHAR* StrB)
+				{
+					return FCString::Stricmp(StrA, StrB) < 0;
+				});
+			for (int32 Index = 0; Index < A.Num(); ++Index)
+			{
+				if (FCString::Stricmp(ScratchA[Index], ScratchB[Index]) != 0)
+				{
+					return false;
+				}
+			}
+			return true;
+		};
+
 	for (int32 PathType = 0; PathType < NumPathTypes; ++PathType)
 	{
 		int32 NumNonPaths = NonPathsByTypeAndIndex[PathType].Num();
@@ -473,6 +510,333 @@ bool FDirectoryTreeTest::RunTest(const FString& Parameters)
 		Value = Tree.Find(TEXTVIEW("/Root/PathZ"));
 		TestTrue(TEXT("MoveConstructOnlyValueZ correct"), Value && Value->Value == 437);
 	}
+
+	// Testing accessors
+	{
+		// GetChildren
+		TDirectoryTree<FMoveConstructOnly> Tree;
+		bool bExists;
+		TArray<FString> Children;
+
+		Children.Reset();
+		bExists = Tree.TryGetChildren(TEXTVIEW(""), Children,
+			EDirectoryTreeGetFlags::None);
+		TestTrue(TEXT("GetChildrenEmpty, Root, !ImpliedParent"),
+			bExists == false && Children.IsEmpty());
+
+		Children.Reset();
+		bExists = Tree.TryGetChildren(TEXTVIEW(""), Children,
+			EDirectoryTreeGetFlags::ImpliedParent |
+			EDirectoryTreeGetFlags::ImpliedChildren |
+			EDirectoryTreeGetFlags::Recursive);
+		TestTrue(TEXT("GetChildrenEmpty, Root, ImpliedParent"),
+			bExists == true && Children.IsEmpty());
+
+		Children.Reset();
+		bExists = Tree.TryGetChildren(TEXTVIEW("/SomePath"), Children,
+			EDirectoryTreeGetFlags::ImpliedParent |
+			EDirectoryTreeGetFlags::ImpliedChildren |
+			EDirectoryTreeGetFlags::Recursive);
+		TestTrue(TEXT("GetChildrenEmpty, Non-root"),
+			bExists == false && Children.IsEmpty());
+
+		Tree.FindOrAdd(TEXTVIEW("")).Value = 1;
+		Children.Reset();
+		bExists = Tree.TryGetChildren(TEXTVIEW(""), Children,
+			EDirectoryTreeGetFlags::None);
+		TestTrue(TEXT("GetChildrenRoot, !ImpliedParent, !ImpliedChildren"),
+			bExists == true && Children.IsEmpty());
+
+		Tree.Empty();
+		Tree.FindOrAdd(TEXTVIEW("/")).Value = 1;
+
+		bExists = Tree.TryGetChildren(TEXTVIEW(""), Children,
+			EDirectoryTreeGetFlags::ImpliedChildren |
+			EDirectoryTreeGetFlags::Recursive);
+		TestTrue(TEXT("GetChildrenRoot, !ImpliedParent, ImpliedChildren"),
+			bExists == false && Children.IsEmpty());
+
+		Children.Reset();
+		bExists = Tree.TryGetChildren(TEXTVIEW(""), Children,
+			EDirectoryTreeGetFlags::ImpliedParent |
+			EDirectoryTreeGetFlags::ImpliedChildren |
+			EDirectoryTreeGetFlags::Recursive);
+		TestTrue(TEXT("GetChildrenRoot, ImpliedParent, ImpliedChildren"),
+			bExists == true && UnorderedEquals(Children, { TEXT("/") }));
+
+		bExists = Tree.TryGetChildren(TEXTVIEW(""), Children,
+			EDirectoryTreeGetFlags::ImpliedParent |
+			EDirectoryTreeGetFlags::ImpliedChildren |
+			EDirectoryTreeGetFlags::Recursive);
+		TestTrue(TEXT("GetChildren appends to the outdir rather than resetting"),
+			bExists == true && Children.Num() == 2);
+
+		Children.Reset();
+		bExists = Tree.TryGetChildren(TEXTVIEW(""), Children,
+			EDirectoryTreeGetFlags::ImpliedParent |
+			EDirectoryTreeGetFlags::Recursive);
+		TestTrue(TEXT("GetChildrenRoot, ImpliedParent, !ImpliedChildren"),
+			bExists == true && UnorderedEquals(Children, { TEXT("/") }));
+
+		Tree.Empty();
+		Tree.FindOrAdd(TEXTVIEW("/Root/Child")).Value = 1;
+
+		Children.Reset();
+		bExists = Tree.TryGetChildren(TEXTVIEW(""), Children,
+			EDirectoryTreeGetFlags::ImpliedParent |
+			EDirectoryTreeGetFlags::ImpliedChildren);
+		TestTrue(TEXT("GetChildrenRootImpliedChild, ImpliedParent, ImpliedChildren, !Recursive"),
+			bExists == true && UnorderedEquals(Children, { TEXT("/") }));
+		Children.Reset();
+		bExists = Tree.TryGetChildren(TEXTVIEW(""), Children,
+			EDirectoryTreeGetFlags::ImpliedParent);
+		TestTrue(TEXT("GetChildrenRootImpliedChild, ImpliedParent, !ImpliedChildren, !Recursive"),
+			bExists == true && UnorderedEquals(Children, { TEXT("/Root/Child") }));
+		Children.Reset();
+		bExists = Tree.TryGetChildren(TEXTVIEW(""), Children,
+			EDirectoryTreeGetFlags::ImpliedParent |
+			EDirectoryTreeGetFlags::ImpliedChildren |
+			EDirectoryTreeGetFlags::Recursive);
+		TestTrue(TEXT("GetChildrenRootImpliedChild, ImpliedParent, ImpliedChildren, Recursive"),
+			bExists == true && UnorderedEquals(Children, { TEXT("/"), TEXT("/Root"), TEXT("/Root/Child") }));
+		Children.Reset();
+		bExists = Tree.TryGetChildren(TEXTVIEW(""), Children,
+			EDirectoryTreeGetFlags::ImpliedParent |
+			EDirectoryTreeGetFlags::Recursive);
+		TestTrue(TEXT("GetChildrenRootImpliedChild, ImpliedParent, !ImpliedChildren, Recursive"),
+			bExists == true && UnorderedEquals(Children, { TEXT("/Root/Child") }));
+
+		Tree.FindOrAdd(TEXTVIEW("/Root/Child2")).Value = 1;
+
+		Tree.Empty();
+		Tree.FindOrAdd(TEXTVIEW("/Stem/A_OtherChild")).Value = 1;
+		Tree.FindOrAdd(TEXTVIEW("/Stem/B_ImpliedRoot/AddedChild")).Value = 1;
+		Tree.FindOrAdd(TEXTVIEW("/Stem/B_ImpliedRoot/AddedChild/Child")).Value = 1;
+		Tree.FindOrAdd(TEXTVIEW("/Stem/B_ImpliedRoot/ImpliedChild/AddedChild")).Value = 1;
+		Tree.FindOrAdd(TEXTVIEW("/Stem/B_ImpliedRoot/ImpliedChild/AddedChild/AddedChild")).Value = 1;
+		Tree.FindOrAdd(TEXTVIEW("/Stem/B_ImpliedRoot/ImpliedChild/AddedChild/ImpliedChild/AddedChild")).Value = 1;
+		Tree.FindOrAdd(TEXTVIEW("/Stem/C_MiddleRoot/MiddlePath/ImpliedChild/AddedChild")).Value = 1;
+		Tree.FindOrAdd(TEXTVIEW("/Stem/C_MiddleRoot/MiddlePath/ImpliedChild/AddedChild/Child")).Value = 1;
+		Tree.FindOrAdd(TEXTVIEW("/Stem/D_MiddleRoot/MiddlePath/AddedChild")).Value = 1;
+		Tree.FindOrAdd(TEXTVIEW("/Stem/D_MiddleRoot/MiddlePath/AddedChild/Child")).Value = 1;
+		Tree.FindOrAdd(TEXTVIEW("/Stem/E_AddedRoot")).Value = 1;
+		Tree.FindOrAdd(TEXTVIEW("/Stem/E_AddedRoot/AddedChild")).Value = 1;
+		Tree.FindOrAdd(TEXTVIEW("/Stem/E_AddedRoot/ImpliedChild/AddedChild")).Value = 1;
+		Tree.FindOrAdd(TEXTVIEW("/Stem/E_AddedRoot/ImpliedChild/AddedChild/ImpliedChild/AddedChild")).Value = 1;
+		Tree.FindOrAdd(TEXTVIEW("/Stem/F_AddedRoot")).Value = 1;
+		Tree.FindOrAdd(TEXTVIEW("/Stem/F_AddedRoot/ImpliedChild/AddedChild")).Value = 1;
+
+
+		// Case: Requested path is an implied path that is a stored child in the tree.
+		Children.Reset();
+		bExists = Tree.TryGetChildren(TEXTVIEW("/Stem/B_ImpliedRoot"), Children,
+			EDirectoryTreeGetFlags::None);
+		TestTrue(TEXT("GetChildrenComplexA B_ImpliedRoot, !ImpliedParent, !ImpliedChildren, !Recursive"),
+			bExists == false && Children.IsEmpty());
+		Children.Reset();
+		bExists = Tree.TryGetChildren(TEXTVIEW("/Stem/B_ImpliedRoot"), Children,
+			EDirectoryTreeGetFlags::Recursive);
+		TestTrue(TEXT("GetChildrenComplexA B_ImpliedRoot, !ImpliedParent, !ImpliedChildren, Recursive"),
+			bExists == false && Children.IsEmpty());
+		Children.Reset();
+		bExists = Tree.TryGetChildren(TEXTVIEW("/Stem/B_ImpliedRoot"), Children,
+			EDirectoryTreeGetFlags::ImpliedChildren);
+		TestTrue(TEXT("GetChildrenComplexA B_ImpliedRoot, !ImpliedParent, ImpliedChildren, !Recursive"),
+			bExists == false && Children.IsEmpty());
+		Children.Reset();
+		bExists = Tree.TryGetChildren(TEXTVIEW("/Stem/B_ImpliedRoot"), Children,
+			EDirectoryTreeGetFlags::ImpliedChildren |
+			EDirectoryTreeGetFlags::Recursive);
+		TestTrue(TEXT("GetChildrenComplexA B_ImpliedRoot, !ImpliedParent, ImpliedChildren, Recursive"),
+			bExists == false && Children.IsEmpty());
+		Children.Reset();
+		bExists = Tree.TryGetChildren(TEXTVIEW("/Stem/B_ImpliedRoot"), Children,
+			EDirectoryTreeGetFlags::ImpliedParent);
+		TestTrue(TEXT("GetChildrenComplexA B_ImpliedRoot, ImpliedParent, !ImpliedChildren, !Recursive"),
+			bExists == true && UnorderedEquals(Children, { TEXT("AddedChild"), TEXT("ImpliedChild/AddedChild") }));
+		Children.Reset();
+		bExists = Tree.TryGetChildren(TEXTVIEW("/Stem/B_ImpliedRoot"), Children,
+			EDirectoryTreeGetFlags::ImpliedParent |
+			EDirectoryTreeGetFlags::Recursive);
+		TestTrue(TEXT("GetChildrenComplexA B_ImpliedRoot, ImpliedParent, !ImpliedChildren, Recursive"),
+			bExists == true && UnorderedEquals(Children, { TEXT("AddedChild"), TEXT("AddedChild/Child"),
+				TEXT("ImpliedChild/AddedChild"), TEXT("ImpliedChild/AddedChild/AddedChild"),
+				TEXT("ImpliedChild/AddedChild/ImpliedChild/AddedChild") }));
+		Children.Reset();
+		bExists = Tree.TryGetChildren(TEXTVIEW("/Stem/B_ImpliedRoot"), Children,
+			EDirectoryTreeGetFlags::ImpliedParent |
+			EDirectoryTreeGetFlags::ImpliedChildren);
+		TestTrue(TEXT("GetChildrenComplexA B_ImpliedRoot, ImpliedParent, ImpliedChildren, !Recursive"),
+			bExists == true && UnorderedEquals(Children, { TEXT("AddedChild"), TEXT("ImpliedChild") }));
+		Children.Reset();
+		bExists = Tree.TryGetChildren(TEXTVIEW("/Stem/B_ImpliedRoot"), Children,
+			EDirectoryTreeGetFlags::ImpliedParent |
+			EDirectoryTreeGetFlags::ImpliedChildren |
+			EDirectoryTreeGetFlags::Recursive);
+		TestTrue(TEXT("GetChildrenComplexA B_ImpliedRoot, ImpliedParent, ImpliedChildren, Recursive"),
+			bExists == true && UnorderedEquals(Children, { TEXT("AddedChild"), TEXT("AddedChild/Child"),
+				TEXT("ImpliedChild"), TEXT("ImpliedChild/AddedChild"), TEXT("ImpliedChild/AddedChild/AddedChild"),
+				TEXT("ImpliedChild/AddedChild/ImpliedChild"), TEXT("ImpliedChild/AddedChild/ImpliedChild/AddedChild") }));
+
+		// Case: Requested path is an implied path that is not a stored child in the tree - it is an in-between dir in a
+		// relpath - and it has an implied child
+		Children.Reset();
+		bExists = Tree.TryGetChildren(TEXTVIEW("/Stem/C_MiddleRoot/MiddlePath"), Children,
+			EDirectoryTreeGetFlags::None);
+		TestTrue(TEXT("GetChildrenComplexA C_MiddleRoot, !ImpliedParent, !ImpliedChildren, !Recursive"),
+			bExists == false && Children.IsEmpty());
+		Children.Reset();
+		bExists = Tree.TryGetChildren(TEXTVIEW("/Stem/C_MiddleRoot/MiddlePath"), Children,
+			EDirectoryTreeGetFlags::Recursive);
+		TestTrue(TEXT("GetChildrenComplexA C_MiddleRoot, !ImpliedParent, !ImpliedChildren, Recursive"),
+			bExists == false && Children.IsEmpty());
+		Children.Reset();
+		bExists = Tree.TryGetChildren(TEXTVIEW("/Stem/C_MiddleRoot/MiddlePath"), Children,
+			EDirectoryTreeGetFlags::ImpliedChildren);
+		TestTrue(TEXT("GetChildrenComplexA C_MiddleRoot, !ImpliedParent, ImpliedChildren, !Recursive"),
+			bExists == false && Children.IsEmpty());
+		Children.Reset();
+		bExists = Tree.TryGetChildren(TEXTVIEW("/Stem/C_MiddleRoot/MiddlePath"), Children,
+			EDirectoryTreeGetFlags::ImpliedChildren |
+			EDirectoryTreeGetFlags::Recursive);
+		TestTrue(TEXT("GetChildrenComplexA C_MiddleRoot, !ImpliedParent, ImpliedChildren, Recursive"),
+			bExists == false && Children.IsEmpty());
+		Children.Reset();
+		bExists = Tree.TryGetChildren(TEXTVIEW("/Stem/C_MiddleRoot/MiddlePath"), Children,
+			EDirectoryTreeGetFlags::ImpliedParent);
+		TestTrue(TEXT("GetChildrenComplexA C_MiddleRoot, ImpliedParent, !ImpliedChildren, !Recursive"),
+			bExists == true && UnorderedEquals(Children, { TEXT("ImpliedChild/AddedChild") }));
+		Children.Reset();
+		bExists = Tree.TryGetChildren(TEXTVIEW("/Stem/C_MiddleRoot/MiddlePath"), Children,
+			EDirectoryTreeGetFlags::ImpliedParent |
+			EDirectoryTreeGetFlags::Recursive);
+		TestTrue(TEXT("GetChildrenComplexA C_MiddleRoot, ImpliedParent, !ImpliedChildren, Recursive"),
+			bExists == true && UnorderedEquals(Children, { TEXT("ImpliedChild/AddedChild"),
+				TEXT("ImpliedChild/AddedChild/Child") }));
+		Children.Reset();
+		bExists = Tree.TryGetChildren(TEXTVIEW("/Stem/C_MiddleRoot/MiddlePath"), Children,
+			EDirectoryTreeGetFlags::ImpliedParent |
+			EDirectoryTreeGetFlags::ImpliedChildren);
+		TestTrue(TEXT("GetChildrenComplexA C_MiddleRoot, ImpliedParent, ImpliedChildren, !Recursive"),
+			bExists == true && UnorderedEquals(Children, { TEXT("ImpliedChild") }));
+		Children.Reset();
+		bExists = Tree.TryGetChildren(TEXTVIEW("/Stem/C_MiddleRoot/MiddlePath"), Children,
+			EDirectoryTreeGetFlags::ImpliedParent |
+			EDirectoryTreeGetFlags::ImpliedChildren |
+			EDirectoryTreeGetFlags::Recursive);
+		TestTrue(TEXT("GetChildrenComplexA C_MiddleRoot, ImpliedParent, ImpliedChildren, Recursive"),
+			bExists == true && UnorderedEquals(Children, { TEXT("ImpliedChild"), TEXT("ImpliedChild/AddedChild"),
+				TEXT("ImpliedChild/AddedChild/Child") }));
+
+		// Case: Requested path is a non-existent sibling path of an implied path that is not a stored
+		// path.
+		Children.Reset();
+		bExists = Tree.TryGetChildren(TEXTVIEW("/Stem/C_MiddleRoot/MiddlePathExceptItDoesNotExist"), Children,
+			EDirectoryTreeGetFlags::ImpliedParent |
+			EDirectoryTreeGetFlags::ImpliedChildren |
+			EDirectoryTreeGetFlags::Recursive);
+		TestTrue(TEXT("GetChildrenComplexA MiddlePathExceptItDoesNotExist, ImpliedParent, ImpliedChildren, Recursive"),
+			bExists == false && Children.IsEmpty());
+
+		// Case: Requested path is an implied path that is not a stored child in the tree - it is an in-between dir in a
+		// relpath - and it has an added child
+		Children.Reset();
+		bExists = Tree.TryGetChildren(TEXTVIEW("/Stem/D_MiddleRoot/MiddlePath"), Children,
+			EDirectoryTreeGetFlags::None);
+		TestTrue(TEXT("GetChildrenComplexA D_MiddleRoot, !ImpliedParent, !ImpliedChildren, !Recursive"),
+			bExists == false && Children.IsEmpty());
+		Children.Reset();
+		bExists = Tree.TryGetChildren(TEXTVIEW("/Stem/D_MiddleRoot/MiddlePath"), Children,
+			EDirectoryTreeGetFlags::Recursive);
+		TestTrue(TEXT("GetChildrenComplexA D_MiddleRoot, !ImpliedParent, !ImpliedChildren, Recursive"),
+			bExists == false && Children.IsEmpty());
+		Children.Reset();
+		bExists = Tree.TryGetChildren(TEXTVIEW("/Stem/D_MiddleRoot/MiddlePath"), Children,
+			EDirectoryTreeGetFlags::ImpliedChildren);
+		TestTrue(TEXT("GetChildrenComplexA D_MiddleRoot, !ImpliedParent, ImpliedChildren, !Recursive"),
+			bExists == false && Children.IsEmpty());
+		Children.Reset();
+		bExists = Tree.TryGetChildren(TEXTVIEW("/Stem/D_MiddleRoot/MiddlePath"), Children,
+			EDirectoryTreeGetFlags::ImpliedChildren |
+			EDirectoryTreeGetFlags::Recursive);
+		TestTrue(TEXT("GetChildrenComplexA D_MiddleRoot, !ImpliedParent, ImpliedChildren, Recursive"),
+			bExists == false && Children.IsEmpty());
+		Children.Reset();
+		bExists = Tree.TryGetChildren(TEXTVIEW("/Stem/D_MiddleRoot/MiddlePath"), Children,
+			EDirectoryTreeGetFlags::ImpliedParent);
+		TestTrue(TEXT("GetChildrenComplexA D_MiddleRoot, ImpliedParent, !ImpliedChildren, !Recursive"),
+			bExists == true && UnorderedEquals(Children, { TEXT("AddedChild") }));
+		Children.Reset();
+		bExists = Tree.TryGetChildren(TEXTVIEW("/Stem/D_MiddleRoot/MiddlePath"), Children,
+			EDirectoryTreeGetFlags::ImpliedParent |
+			EDirectoryTreeGetFlags::Recursive);
+		TestTrue(TEXT("GetChildrenComplexA D_MiddleRoot, ImpliedParent, !ImpliedChildren, Recursive"),
+			bExists == true && UnorderedEquals(Children, { TEXT("AddedChild"), TEXT("AddedChild/Child") }));
+		Children.Reset();
+		bExists = Tree.TryGetChildren(TEXTVIEW("/Stem/D_MiddleRoot/MiddlePath"), Children,
+			EDirectoryTreeGetFlags::ImpliedParent |
+			EDirectoryTreeGetFlags::ImpliedChildren);
+		TestTrue(TEXT("GetChildrenComplexA D_MiddleRoot, ImpliedParent, ImpliedChildren, !Recursive"),
+			bExists == true && UnorderedEquals(Children, { TEXT("AddedChild") }));
+		Children.Reset();
+		bExists = Tree.TryGetChildren(TEXTVIEW("/Stem/D_MiddleRoot/MiddlePath"), Children,
+			EDirectoryTreeGetFlags::ImpliedParent |
+			EDirectoryTreeGetFlags::ImpliedChildren |
+			EDirectoryTreeGetFlags::Recursive);
+		TestTrue(TEXT("GetChildrenComplexA D_MiddleRoot, ImpliedParent, ImpliedChildren, Recursive"),
+			bExists == true && UnorderedEquals(Children, { TEXT("AddedChild"), TEXT("AddedChild/Child") }));
+
+		// Case: Requested path is an added path and it has an added child and an implied child
+		Children.Reset();
+		bExists = Tree.TryGetChildren(TEXTVIEW("/Stem/E_AddedRoot"), Children,
+			EDirectoryTreeGetFlags::None);
+		TestTrue(TEXT("GetChildrenComplexA E_AddedRoot, !ImpliedParent, !ImpliedChildren, !Recursive"),
+			bExists == true && UnorderedEquals(Children, { TEXT("AddedChild"), TEXT("ImpliedChild/AddedChild")}));
+		Children.Reset();
+		bExists = Tree.TryGetChildren(TEXTVIEW("/Stem/E_AddedRoot"), Children,
+			EDirectoryTreeGetFlags::Recursive);
+		TestTrue(TEXT("GetChildrenComplexA E_AddedRoot, !ImpliedParent, !ImpliedChildren, Recursive"),
+			bExists == true && UnorderedEquals(Children, { TEXT("AddedChild"), TEXT("ImpliedChild/AddedChild"),
+				TEXT("ImpliedChild/AddedChild/ImpliedChild/AddedChild") }));
+		Children.Reset();
+		bExists = Tree.TryGetChildren(TEXTVIEW("/Stem/E_AddedRoot"), Children,
+			EDirectoryTreeGetFlags::ImpliedChildren);
+		TestTrue(TEXT("GetChildrenComplexA E_AddedRoot, !ImpliedParent, ImpliedChildren, !Recursive"),
+			bExists == true && UnorderedEquals(Children, { TEXT("AddedChild"), TEXT("ImpliedChild")}));
+		Children.Reset();
+		bExists = Tree.TryGetChildren(TEXTVIEW("/Stem/E_AddedRoot"), Children,
+			EDirectoryTreeGetFlags::ImpliedChildren |
+			EDirectoryTreeGetFlags::Recursive);
+		TestTrue(TEXT("GetChildrenComplexA E_MiddleRoot, !ImpliedParent, ImpliedChildren, Recursive"),
+			bExists == true && UnorderedEquals(Children, { TEXT("AddedChild"), TEXT("ImpliedChild"),
+				TEXT("ImpliedChild/AddedChild"), TEXT("ImpliedChild/AddedChild/ImpliedChild"),
+				TEXT("ImpliedChild/AddedChild/ImpliedChild/AddedChild") }));
+
+		// Case: Requested path is an added path and it has only an implied child
+		Children.Reset();
+		bExists = Tree.TryGetChildren(TEXTVIEW("/Stem/F_AddedRoot"), Children,
+			EDirectoryTreeGetFlags::None);
+		TestTrue(TEXT("GetChildrenComplexA F_AddedRoot, !ImpliedParent, !ImpliedChildren, !Recursive"),
+			bExists == true && UnorderedEquals(Children, { TEXT("ImpliedChild/AddedChild") }));
+		Children.Reset();
+		bExists = Tree.TryGetChildren(TEXTVIEW("/Stem/F_AddedRoot"), Children,
+			EDirectoryTreeGetFlags::Recursive);
+		TestTrue(TEXT("GetChildrenComplexA F_AddedRoot, !ImpliedParent, !ImpliedChildren, Recursive"),
+			bExists == true && UnorderedEquals(Children, { TEXT("ImpliedChild/AddedChild") }));
+		Children.Reset();
+		bExists = Tree.TryGetChildren(TEXTVIEW("/Stem/F_AddedRoot"), Children,
+			EDirectoryTreeGetFlags::ImpliedChildren);
+		TestTrue(TEXT("GetChildrenComplexA F_AddedRoot, !ImpliedParent, ImpliedChildren, !Recursive"),
+			bExists == true && UnorderedEquals(Children, { TEXT("ImpliedChild") }));
+		Children.Reset();
+		bExists = Tree.TryGetChildren(TEXTVIEW("/Stem/F_AddedRoot"), Children,
+			EDirectoryTreeGetFlags::ImpliedChildren |
+			EDirectoryTreeGetFlags::Recursive);
+		TestTrue(TEXT("GetChildrenComplexA F_AddedRoot, !ImpliedParent, ImpliedChildren, Recursive"),
+			bExists == true && UnorderedEquals(Children, { TEXT("ImpliedChild"),
+				TEXT("ImpliedChild/AddedChild") }));
+	}
+
 	return true;
 }
 
