Change check to ensure to allow some systems which directly call UObject::Serialize to work with a warning

diff --git a/Engine/Source/Runtime/CoreUObject/Private/UObject/LinkerSave.cpp b/Engine/Source/Runtime/CoreUObject/Private/UObject/LinkerSave.cpp
--- a/Engine/Source/Runtime/CoreUObject/Private/UObject/LinkerSave.cpp
+++ b/Engine/Source/Runtime/CoreUObject/Private/UObject/LinkerSave.cpp
@@ -295,16 +295,20 @@ FPackageIndex FLinkerSave::MapObject(TObjectPtr<const UObject> Object) const
 
 void FLinkerSave::MarkScriptSerializationStart(const UObject* Obj) 
 {
-	check(Obj == CurrentlySavingExportObject);
-	FObjectExport& Export = ExportMap[CurrentlySavingExport.ToExport()];
-	Export.ScriptSerializationStartOffset = Tell();
+	if (ensure(Obj == CurrentlySavingExportObject))
+	{
+		FObjectExport& Export = ExportMap[CurrentlySavingExport.ToExport()];
+		Export.ScriptSerializationStartOffset = Tell();
+	}
 }
 
 void FLinkerSave::MarkScriptSerializationEnd(const UObject* Obj) 
 {
-	check(Obj == CurrentlySavingExportObject);
-	FObjectExport& Export = ExportMap[CurrentlySavingExport.ToExport()];
-	Export.ScriptSerializationEndOffset = Tell();
+	if (ensure(Obj == CurrentlySavingExportObject))
+	{
+		FObjectExport& Export = ExportMap[CurrentlySavingExport.ToExport()];
+		Export.ScriptSerializationEndOffset = Tell();
+	}
 }
 
 void FLinkerSave::Seek( int64 InPos )
