UObjectBase: using TObjectPtr to avoid manually running barriers on OuterPrivate and ClassPrivate

diff --git a/Engine/Source/Runtime/CoreUObject/Private/UObject/Class.cpp b/Engine/Source/Runtime/CoreUObject/Private/UObject/Class.cpp
--- a/Engine/Source/Runtime/CoreUObject/Private/UObject/Class.cpp
+++ b/Engine/Source/Runtime/CoreUObject/Private/UObject/Class.cpp
@@ -2064,7 +2064,7 @@ void UStruct::AddReferencedObjects(UObject* InThis, FReferenceCollector& Collect
 
 void UStruct::SetSuperStruct(UStruct* NewSuperStruct)
 {
-	SuperStruct.Set(NewSuperStruct);
+	SuperStruct = NewSuperStruct;
 #if USTRUCT_FAST_ISCHILDOF_IMPL == USTRUCT_ISCHILDOF_STRUCTARRAY
 	this->ReinitializeBaseChainArray();
 #endif
diff --git a/Engine/Source/Runtime/CoreUObject/Private/UObject/GarbageCollection.cpp b/Engine/Source/Runtime/CoreUObject/Private/UObject/GarbageCollection.cpp
--- a/Engine/Source/Runtime/CoreUObject/Private/UObject/GarbageCollection.cpp
+++ b/Engine/Source/Runtime/CoreUObject/Private/UObject/GarbageCollection.cpp
@@ -516,6 +516,11 @@ namespace UE::GC
 		LogClassCountInfo(TEXT("objects of"), ClassToPurgeCountMap, 10, Total);
 #endif
 	}
+
+	void MarkAsReachable(const UObject* Obj)
+	{
+		Obj->MarkAsReachable();
+	}
 }
 
 namespace UE::GC::Private
diff --git a/Engine/Source/Runtime/CoreUObject/Private/UObject/ObjectFwd.cpp b/Engine/Source/Runtime/CoreUObject/Private/UObject/ObjectFwd.cpp
new file mode 100644
--- /dev/null
+++ b/Engine/Source/Runtime/CoreUObject/Private/UObject/ObjectFwd.cpp
@@ -0,0 +1,37 @@
+// Copyright Epic Games, Inc. All Rights Reserved.
+
+#include "UObject/ObjectFwd.h"
+#include "UObject/Object.h"
+
+namespace UE::CoreUObject::Private
+{
+UClass* GetClass(UObject* Obj)
+{
+	return Obj->GetClass();
+}
+
+FName GetFName(const UObject* Obj)
+{
+	return Obj->GetFName();
+}
+
+UObject* GetOuter(const UObject* Obj)
+{
+	return Obj->GetOuter();
+}
+
+FString GetPathName( const UObject* Obj, const UObject* StopOuter )
+{
+	return Obj->GetPathName(StopOuter);
+}
+
+UPackage* GetPackage(const UObject* Obj)
+{
+	return Obj->GetPackage();
+}
+
+FString GetFullName( const UObject* Obj, const UObject* StopOuter, EObjectFullNameFlags Flags)
+{
+	return Obj->GetFullName(StopOuter, Flags);
+}
+}
diff --git a/Engine/Source/Runtime/CoreUObject/Private/UObject/UObjectBase.cpp b/Engine/Source/Runtime/CoreUObject/Private/UObject/UObjectBase.cpp
--- a/Engine/Source/Runtime/CoreUObject/Private/UObject/UObjectBase.cpp
+++ b/Engine/Source/Runtime/CoreUObject/Private/UObject/UObjectBase.cpp
@@ -137,16 +137,7 @@ UObjectBase::UObjectBase(UClass* InClass,
 #if CSV_PROFILER && CSV_TRACK_UOBJECT_COUNT
 	UObjectStats::IncrementUObjectCount();
 #endif
-
-	if (UE::GC::Private::GIsIncrementalReachabilityPending)
-	{
-		ClassPrivate->MarkAsReachable();
-		if (OuterPrivate)
-		{
-			OuterPrivate->MarkAsReachable();
 		}	
-	}
-}
 
 
 /**
@@ -192,15 +183,6 @@ void UObjectBase::DeferredRegister(UClass *UClassStaticClass,const TCHAR* Packag
 	check(!ClassPrivate);
 	ClassPrivate = UClassStaticClass;
 
-	if (UE::GC::Private::GIsIncrementalReachabilityPending)
-	{
-		ClassPrivate->MarkAsReachable();
-		if (OuterPrivate)
-		{
-			OuterPrivate->MarkAsReachable();
-		}
-	}
-
 	// Add to the global object table.
 	AddObject(FName(InName), EInternalObjectFlags::None);
 	// At this point all compiled-in objects should have already been fully constructed so it's safe to remove the NotFullyConstructed flag
@@ -259,11 +241,7 @@ void UObjectBase::LowLevelRename(FName NewName,UObject *NewOuter)
 	if (NewOuter)
 	{
 		OuterPrivate = NewOuter;
-		if (UE::GC::Private::GIsIncrementalReachabilityPending)
-		{
-			OuterPrivate->MarkAsReachable();
 		}
-	}
 	HashObject(this);
 }
 
@@ -324,11 +302,6 @@ void UObjectBase::SetClass(UClass* NewClass)
 	ClassPrivate->DestroyPersistentUberGraphFrame((UObject*)this);
 #endif
 	ClassPrivate = NewClass;
-	if (UE::GC::Private::GIsIncrementalReachabilityPending)
-	{
-		checkf(ClassPrivate, TEXT("SetClass called on %s with a null class"), *GetFName().ToString());
-		ClassPrivate->MarkAsReachable();
-	}
 #if USE_UBER_GRAPH_PERSISTENT_FRAME
 	ClassPrivate->CreatePersistentUberGraphFrame((UObject*)this, /*bCreateOnlyIfEmpty =*/false, /*bSkipSuperClass =*/false, OldClass);
 #endif
diff --git a/Engine/Source/Runtime/CoreUObject/Private/UObject/UObjectGlobals.cpp b/Engine/Source/Runtime/CoreUObject/Private/UObject/UObjectGlobals.cpp
--- a/Engine/Source/Runtime/CoreUObject/Private/UObject/UObjectGlobals.cpp
+++ b/Engine/Source/Runtime/CoreUObject/Private/UObject/UObjectGlobals.cpp
@@ -1390,8 +1390,8 @@ UObject* StaticLoadObjectInternal(UClass* ObjectClass, UObject* InOuter, const T
 				{
 					if (UE::GC::Private::GIsIncrementalReachabilityPending)
 					{
-						Redirector->MarkAsReachable();
-						Redirector->DestinationObject->MarkAsReachable();
+						UE::GC::MarkAsReachable(Redirector);
+						UE::GC::MarkAsReachable(Redirector->DestinationObject);
 					}
 					return Redirector->DestinationObject;
 				}
@@ -1417,7 +1417,7 @@ UObject* StaticLoadObjectInternal(UClass* ObjectClass, UObject* InOuter, const T
 
 	if (Result && UE::GC::Private::GIsIncrementalReachabilityPending)
 	{
-		Result->MarkAsReachable();
+		UE::GC::MarkAsReachable(Result);
 	}
 	return Result;
 }
diff --git a/Engine/Source/Runtime/CoreUObject/Private/UObject/UObjectHash.cpp b/Engine/Source/Runtime/CoreUObject/Private/UObject/UObjectHash.cpp
--- a/Engine/Source/Runtime/CoreUObject/Private/UObject/UObjectHash.cpp
+++ b/Engine/Source/Runtime/CoreUObject/Private/UObject/UObjectHash.cpp
@@ -813,7 +813,7 @@ UObject* StaticFindObjectFastInternalThreadSafe(FUObjectHashTables& ThreadHash,
 	// Not found.
 	if (Result && UE::GC::Private::GIsIncrementalReachabilityPending)
 	{
-		Result->MarkAsReachable();
+		UE::GC::MarkAsReachable(Result);
 	}
 	return Result;
 }
@@ -1141,7 +1141,7 @@ void GetObjectsWithOuter(const class UObjectBase* Outer, TArray<UObject *>& Resu
 			{
 				if (UE::GC::Private::GIsIncrementalReachabilityPending)
 				{
-					Object->MarkAsReachable();
+					UE::GC::MarkAsReachable(Object);
 				}
 				Results.Add(Object);
 			}
@@ -1164,7 +1164,7 @@ void GetObjectsWithOuter(const class UObjectBase* Outer, TArray<UObject *>& Resu
 						{
 							if (UE::GC::Private::GIsIncrementalReachabilityPending)
 							{
-								Object->MarkAsReachable();
+								UE::GC::MarkAsReachable(Object);
 							}
 							Results.Add(Object);
 						}
@@ -1216,7 +1216,7 @@ void ForEachObjectWithOuterBreakable(const class UObjectBase* Outer, TFunctionRe
 			{
 				if (UE::GC::Private::GIsIncrementalReachabilityPending)
 				{
-					Object->MarkAsReachable();
+					UE::GC::MarkAsReachable(Object);
 				}
 				if (!Operation(Object))
 				{
@@ -1273,7 +1273,7 @@ UObjectBase* FindObjectWithOuter(const class UObjectBase* Outer, const class UCl
 			}
 			if (Result && UE::GC::Private::GIsIncrementalReachabilityPending)
 			{
-				Result->MarkAsReachable();
+				UE::GC::MarkAsReachable(Result);
 			}
 		}
 	}
@@ -1334,7 +1334,7 @@ void ForEachObjectWithPackage(const class UPackage* Package, TFunctionRef<bool(U
 			{
 				if (UE::GC::Private::GIsIncrementalReachabilityPending)
 				{
-					Object->MarkAsReachable();
+					UE::GC::MarkAsReachable(Object);
 				}
 				if (!Operation(Object))
 				{
@@ -1427,7 +1427,7 @@ FORCEINLINE void ForEachObjectOfClasses_Implementation(FUObjectHashTables& Threa
 				{
 					if (UE::GC::Private::GIsIncrementalReachabilityPending)
 					{
-						Object->MarkAsReachable();
+						UE::GC::MarkAsReachable(Object);
 					}
 					Operation(Object);
 				}
diff --git a/Engine/Source/Runtime/CoreUObject/Public/UObject/DynamicallyTypedValue.h b/Engine/Source/Runtime/CoreUObject/Public/UObject/DynamicallyTypedValue.h
--- a/Engine/Source/Runtime/CoreUObject/Public/UObject/DynamicallyTypedValue.h
+++ b/Engine/Source/Runtime/CoreUObject/Public/UObject/DynamicallyTypedValue.h
@@ -129,7 +129,7 @@ namespace UE
 				{
 					if (InObject)
 					{
-						InObject->MarkAsReachable();
+						UE::GC::MarkAsReachable(InObject);
 					}
 				}
 			};
diff --git a/Engine/Source/Runtime/CoreUObject/Public/UObject/Field.h b/Engine/Source/Runtime/CoreUObject/Public/UObject/Field.h
--- a/Engine/Source/Runtime/CoreUObject/Public/UObject/Field.h
+++ b/Engine/Source/Runtime/CoreUObject/Public/UObject/Field.h
@@ -275,7 +275,7 @@ class FFieldVariant
 	{
 		if (IsUObject() && ToUObjectUnsafe() && UE::GC::Private::GIsIncrementalReachabilityPending)
 		{
-			ToUObjectUnsafe()->MarkAsReachable();
+			UE::GC::MarkAsReachable(ToUObjectUnsafe());
 		}
 	}
 	
diff --git a/Engine/Source/Runtime/CoreUObject/Public/UObject/Object.h b/Engine/Source/Runtime/CoreUObject/Public/UObject/Object.h
--- a/Engine/Source/Runtime/CoreUObject/Public/UObject/Object.h
+++ b/Engine/Source/Runtime/CoreUObject/Public/UObject/Object.h
@@ -1827,22 +1827,6 @@ const T* GetValid(const T* Test)
 	return IsValid(Test) ? Test : nullptr;
 }
 
-/**
-* Helper function to mark object as reachable when performing Incremental GC
-*
-* @param	Object			Object to mark as reachable
-* @return	The same Object pointer that was passed to the function
-*/
-template <typename T>
-FORCEINLINE T* MarkAsReachable(T* Object)
-{
-	if (Object)
-	{
-		Object->MarkAsReachable();
-	}
-	return Object;
-}
-
 #if WITH_EDITOR
 /** Callback for editor object selection. This must be in core instead of editor for UObject::IsSelectedInEditor to work */
 extern COREUOBJECT_API TFunction<bool(const UObject*)> GIsObjectSelectedInEditor;
diff --git a/Engine/Source/Runtime/CoreUObject/Public/UObject/ObjectFwd.h b/Engine/Source/Runtime/CoreUObject/Public/UObject/ObjectFwd.h
new file mode 100644
--- /dev/null
+++ b/Engine/Source/Runtime/CoreUObject/Public/UObject/ObjectFwd.h
@@ -0,0 +1,43 @@
+// Copyright Epic Games, Inc. All Rights Reserved.
+
+#pragma once
+
+#include "Containers/UnrealString.h"
+#include "Misc/EnumClassFlags.h"
+
+/**
+* Enum which specifies the mode in which full object names are constructed
+*/
+enum class EObjectFullNameFlags
+{
+	// Standard object full name (i.e. "Type PackageName.ObjectName:SubobjectName")
+	None = 0,
+
+	// Adds package to the type portion (i.e. "TypePackage.TypeName PackageName.ObjectName:SubobjectName")
+	IncludeClassPackage = 1,
+};
+
+ENUM_CLASS_FLAGS(EObjectFullNameFlags);
+
+namespace UE::GC
+{
+/**
+* Marks the object as Reachable if it's currently marked as MaybeUnreachable by incremental GC.
+*/
+COREUOBJECT_API void MarkAsReachable(const UObject* Obj);
+}
+
+namespace UE::CoreUObject::Private
+{
+COREUOBJECT_API UPackage* GetPackage(const UObject* Obj);
+
+COREUOBJECT_API FString GetFullName(const UObject* Obj,
+																		const UObject* StopOuter=nullptr,
+																		EObjectFullNameFlags Flags = EObjectFullNameFlags::None);
+
+COREUOBJECT_API FString GetPathName( const UObject* Obj, const UObject* StopOuter=nullptr );
+
+COREUOBJECT_API UClass* GetClass(UObject* Obj);
+COREUOBJECT_API FName GetFName(const UObject* Obj);
+COREUOBJECT_API UObject* GetOuter(const UObject* Obj);
+}
diff --git a/Engine/Source/Runtime/CoreUObject/Public/UObject/ObjectHandle.h b/Engine/Source/Runtime/CoreUObject/Public/UObject/ObjectHandle.h
--- a/Engine/Source/Runtime/CoreUObject/Public/UObject/ObjectHandle.h
+++ b/Engine/Source/Runtime/CoreUObject/Public/UObject/ObjectHandle.h
@@ -10,7 +10,6 @@
 #include "Misc/AssertionMacros.h"
 #include "Templates/TypeHash.h"
 #include "UObject/NameTypes.h"
-#include "UObject/Object.h"
 #include "UObject/ObjectHandleTracking.h"
 #include "UObject/ObjectMacros.h"
 #include "UObject/ObjectRef.h"
@@ -235,7 +234,7 @@ namespace UE::CoreUObject::Private
 		if (IsObjectHandleResolved(Handle))
 		{
 			UObject* Obj = ReadObjectHandlePointerNoCheck(Handle);
-			return Obj != nullptr ? Obj->GetClass() : nullptr;
+			return Obj != nullptr ? UE::CoreUObject::Private::GetClass(Obj) : nullptr;
 		}
 		else
 		{
@@ -246,7 +245,7 @@ namespace UE::CoreUObject::Private
 		}
 #else
 		UObject* Obj = ReadObjectHandlePointerNoCheck(Handle);
-		return Obj != nullptr ? Obj->GetClass() : nullptr;
+		return Obj != nullptr ? UE::CoreUObject::Private::GetClass(Obj) : nullptr;
 #endif
 	}
 
@@ -351,3 +350,6 @@ namespace UE::CoreUObject::Private
 #endif
 }
 
+#if UE_ENABLE_INCLUDE_ORDER_DEPRECATED_IN_5_4
+#include "UObject/Object.h"
+#endif
diff --git a/Engine/Source/Runtime/CoreUObject/Public/UObject/ObjectPtr.h b/Engine/Source/Runtime/CoreUObject/Public/UObject/ObjectPtr.h
--- a/Engine/Source/Runtime/CoreUObject/Public/UObject/ObjectPtr.h
+++ b/Engine/Source/Runtime/CoreUObject/Public/UObject/ObjectPtr.h
@@ -5,8 +5,8 @@
 #include "HAL/Platform.h"
 #include "Serialization/StructuredArchive.h"
 #include "Templates/IsTObjectPtr.h"
-#include "UObject/Object.h"
 #include "UObject/ObjectHandle.h"
+#include "UObject/UObjectGlobals.h"
 #include "Templates/NonNullPointer.h"
 
 #include <type_traits>
@@ -218,13 +218,13 @@ public:
 	FString GetPathName() const
 	{
 		const UObject* ResolvedObject = Get();
-		return ResolvedObject ? ResolvedObject->GetPathName() : TEXT("None");
+		return ResolvedObject ? UE::CoreUObject::Private::GetPathName(ResolvedObject) : TEXT("None");
 	}
 
 	FName GetFName() const
 	{
 		const UObject* ResolvedObject = Get();
-		return ResolvedObject ? ResolvedObject->GetFName() : NAME_None;
+		return ResolvedObject ? UE::CoreUObject::Private::GetFName(ResolvedObject) : NAME_None;
 	}
 
 	FString GetName() const
@@ -235,13 +235,13 @@ public:
 	FObjectPtr GetOuter() const
 	{
 		const UObject* ResolvedObject = Get();
-		return ResolvedObject ? FObjectPtr(ResolvedObject->GetOuter()) : FObjectPtr(nullptr);
+		return ResolvedObject ? FObjectPtr(UE::CoreUObject::Private::GetOuter(ResolvedObject)) : FObjectPtr(nullptr);
 	}
 
 	FObjectPtr GetPackage() const
 	{
 		const UObject* ResolvedObject = Get();
-		return ResolvedObject ? FObjectPtr(ResolvedObject->GetPackage()) : FObjectPtr(nullptr);
+		return ResolvedObject ? FObjectPtr(UE::CoreUObject::Private::GetPackage(ResolvedObject)) : FObjectPtr(nullptr);
 	}
 	/**
 	 * Returns the fully qualified pathname for this object as well as the name of the class, in the format:
@@ -250,7 +250,7 @@ public:
 	FString GetFullName(EObjectFullNameFlags Flags = EObjectFullNameFlags::None) const
 	{
 		// UObjectBaseUtility::GetFullName is safe to call on null objects.
-		return Get()->GetFullName(nullptr, Flags);
+		return UE::CoreUObject::Private::GetFullName(Get(), nullptr, Flags);
 	}
 #endif
 
@@ -286,7 +286,7 @@ private:
 		{
 			if (UObject* Obj = UE::CoreUObject::Private::ReadObjectHandlePointerNoCheck(InPtr.GetHandleRef()))
 			{
-				Obj->MarkAsReachable();
+				UE::GC::MarkAsReachable(Obj);
 			}
 		}
 	}
@@ -294,7 +294,7 @@ private:
 	{
 		if (UE::GC::Private::GIsIncrementalReachabilityPending && InObj)
 		{
-			InObj->MarkAsReachable();
+			UE::GC::MarkAsReachable(InObj);
 		}
 	}
 #endif // UE_OBJECT_PTR_GC_BARRIER
@@ -731,11 +731,21 @@ namespace ObjectPtr_Private
 	{
 	public:
 		TNonAccessTrackedObjectPtr() = default;
-		explicit TNonAccessTrackedObjectPtr(T* Ptr) : ObjectPtr{Ptr} {}
 
-		void Set(T* Value)
+		explicit TNonAccessTrackedObjectPtr(ENoInit)
+			: ObjectPtr{NoInit}
+		{
+		}
+
+		explicit TNonAccessTrackedObjectPtr(T* Ptr)
+			: ObjectPtr{Ptr}
+		{
+		}
+
+		TNonAccessTrackedObjectPtr& operator=(T* Value)
 		{
 			ObjectPtr = Value;
+			return *this;
 		}
 	
 		T* Get() const
@@ -743,6 +753,11 @@ namespace ObjectPtr_Private
 			return ObjectPtr_Private::Friend::NoAccessTrackingGet(ObjectPtr);
 		}
 
+		explicit operator UPTRINT() const
+		{
+			return BitCast<UPTRINT>(Get());
+		}
+
 		TObjectPtr<T>& GetAccessTrackedObjectPtr() 
 		{
 			return ObjectPtr;
@@ -759,7 +774,7 @@ namespace ObjectPtr_Private
 		}		
 	
 	private:
-		TObjectPtr<T> ObjectPtr{};
+		TObjectPtr<T> ObjectPtr;
 	};
 }
 
@@ -1315,7 +1330,7 @@ namespace UE::Core::Private // private facilities; not for direct use
 #if UE_OBJECT_PTR_GC_BARRIER
 			if (UE::GC::Private::GIsIncrementalReachabilityPending && View)
 			{
-				View->MarkAsReachable();
+				UE::GC::MarkAsReachable(View);
 			}
 #endif // UE_OBJECT_PTR_GC_BARRIER
 		}
@@ -1334,7 +1349,7 @@ namespace UE::Core::Private // private facilities; not for direct use
 				{
 					if (Data[Index])
 					{
-						Data[Index]->MarkAsReachable();
+						UE::GC::MarkAsReachable(Data[Index]);
 					}
 				}
 			}
@@ -1354,7 +1369,7 @@ namespace UE::Core::Private // private facilities; not for direct use
 				{
 					if (Element)
 					{
-						Element->MarkAsReachable();
+						UE::GC::MarkAsReachable(Element);
 					}
 				}
 			}
@@ -1379,14 +1394,14 @@ namespace UE::Core::Private // private facilities; not for direct use
 					{
 						if (Pair.Key)
 						{
-							Pair.Key->MarkAsReachable();
+							UE::GC::MarkAsReachable(Pair.Key);
 						}
 					}
 					if constexpr (bValueReference)
 					{
 						if (Pair.Value)
 						{
-							Pair.Value->MarkAsReachable();
+							UE::GC::MarkAsReachable(Pair.Value);
 						}
 					}
 				}
@@ -1494,7 +1509,7 @@ namespace UE::Core::Private // private facilities; not for direct use
 		{
 			if (const UObject* Obj = Cast<UObject>(Ptr); Obj && UE::GC::Private::GIsIncrementalReachabilityPending)
 			{
-				Obj->MarkAsReachable();
+				UE::GC::MarkAsReachable(Obj);
 			}
 		}
 
@@ -1728,3 +1743,7 @@ private:
 #if UE_ENABLE_INCLUDE_ORDER_DEPRECATED_IN_5_3
 #include "UObject/Class.h"
 #endif
+
+#if UE_ENABLE_INCLUDE_ORDER_DEPRECATED_IN_5_4
+#include "UObject/Object.h"
+#endif
diff --git a/Engine/Source/Runtime/CoreUObject/Public/UObject/ObjectRef.h b/Engine/Source/Runtime/CoreUObject/Public/UObject/ObjectRef.h
--- a/Engine/Source/Runtime/CoreUObject/Public/UObject/ObjectRef.h
+++ b/Engine/Source/Runtime/CoreUObject/Public/UObject/ObjectRef.h
@@ -3,8 +3,9 @@
 #pragma once
 
 #include "Misc/StringBuilder.h"
-#include "UObject/UObjectBaseUtility.h"
 #include "UObject/ObjectHandleDefines.h"
+#include "UObject/ObjectFwd.h"
+#include "UObject/ObjectMacros.h"
 
 #if UE_WITH_OBJECT_HANDLE_LATE_RESOLVE || UE_WITH_OBJECT_HANDLE_TRACKING
 
@@ -121,4 +122,8 @@ private:
 UE_DEPRECATED(5.2, "Use FObjectRef::IsNull() instead")
 inline bool IsObjectRefNull(const FObjectRef& ObjectRef) { return ObjectRef.IsNull(); }
 
+#if UE_ENABLE_INCLUDE_ORDER_DEPRECATED_IN_5_4
+#include "UObject/UObjectBaseUtility.h"
+#endif
+
 #endif
\ No newline at end of file
diff --git a/Engine/Source/Runtime/CoreUObject/Public/UObject/UObjectBase.h b/Engine/Source/Runtime/CoreUObject/Public/UObject/UObjectBase.h
--- a/Engine/Source/Runtime/CoreUObject/Public/UObject/UObjectBase.h
+++ b/Engine/Source/Runtime/CoreUObject/Public/UObject/UObjectBase.h
@@ -18,6 +18,7 @@
 #include "UObject/ObjectMacros.h"
 #include "UObject/UObjectGlobals.h"
 #include "UObject/UnrealNames.h"
+#include "UObject/ObjectPtr.h"
 
 class UClass;
 class UEnum;
@@ -47,7 +48,9 @@ class UObjectBase
 		);
 protected:
 	UObjectBase() :
-		 NamePrivate(NoInit)  // screwy, but the name was already set and we don't want to set it again
+		ClassPrivate(NoInit),
+		NamePrivate(NoInit),  // screwy, but the name was already set and we don't want to set it again
+		OuterPrivate(NoInit)
 	{
 	}
 
@@ -249,13 +252,13 @@ private:
 	int32							InternalIndex;
 
 	/** Class the object belongs to. */
-	UClass*							ClassPrivate;
+	ObjectPtr_Private::TNonAccessTrackedObjectPtr<UClass>							ClassPrivate;
 
 	/** Name of this object */
 	FName							NamePrivate;
 
 	/** Object this object resides in. */
-	UObject*						OuterPrivate;
+	ObjectPtr_Private::TNonAccessTrackedObjectPtr<UObject>						OuterPrivate;
 	
 	friend class FBlueprintCompileReinstancer;
 	friend class FVerseObjectClassReplacer;
diff --git a/Engine/Source/Runtime/CoreUObject/Public/UObject/UObjectBaseUtility.h b/Engine/Source/Runtime/CoreUObject/Public/UObject/UObjectBaseUtility.h
--- a/Engine/Source/Runtime/CoreUObject/Public/UObject/UObjectBaseUtility.h
+++ b/Engine/Source/Runtime/CoreUObject/Public/UObject/UObjectBaseUtility.h
@@ -25,6 +25,7 @@
 #include "UObject/UObjectArray.h"
 #include "UObject/UObjectBase.h"
 #include "UObject/UObjectMarks.h"
+#include "UObject/ObjectFwd.h"
 
 class UClass;
 class UObject;
@@ -37,20 +38,6 @@ struct FGuid;
 	#endif
 #endif
 
-/**
-* Enum which specifies the mode in which full object names are constructed
-*/
-enum class EObjectFullNameFlags
-{
-	// Standard object full name (i.e. "Type PackageName.ObjectName:SubobjectName")
-	None = 0,
-
-	// Adds package to the type portion (i.e. "TypePackage.TypeName PackageName.ObjectName:SubobjectName")
-	IncludeClassPackage = 1,
-};
-
-ENUM_CLASS_FLAGS(EObjectFullNameFlags);
-
 /**
  * Provides utility functions for UObject, this class should not be used directly
  */
diff --git a/Engine/Source/Runtime/Engine/Public/TransactionCommon.h b/Engine/Source/Runtime/Engine/Public/TransactionCommon.h
--- a/Engine/Source/Runtime/Engine/Public/TransactionCommon.h
+++ b/Engine/Source/Runtime/Engine/Public/TransactionCommon.h
@@ -3,6 +3,7 @@
 #pragma once
 
 #include "CoreMinimal.h"
+#include "UObject/Object.h"
 #include "UObject/ObjectPtr.h"
 #include "UObject/WeakObjectPtr.h"
 #include "Serialization/ArchiveUObject.h"
