[Backout] - CL29251139 - CIS Compile Error
Original CL Desc
-----------------------------------------------------------------
UObjectBase: using TObjectPtr to avoid manually running barriers on OuterPrivate and ClassPrivate

diff --git a/Engine/Source/Runtime/CoreUObject/Private/UObject/Class.cpp b/Engine/Source/Runtime/CoreUObject/Private/UObject/Class.cpp
--- a/Engine/Source/Runtime/CoreUObject/Private/UObject/Class.cpp
+++ b/Engine/Source/Runtime/CoreUObject/Private/UObject/Class.cpp
@@ -2064,7 +2064,7 @@ void UStruct::AddReferencedObjects(UObject* InThis, FReferenceCollector& Collect
 
 void UStruct::SetSuperStruct(UStruct* NewSuperStruct)
 {
-	SuperStruct = NewSuperStruct;
+	SuperStruct.Set(NewSuperStruct);
 #if USTRUCT_FAST_ISCHILDOF_IMPL == USTRUCT_ISCHILDOF_STRUCTARRAY
 	this->ReinitializeBaseChainArray();
 #endif
diff --git a/Engine/Source/Runtime/CoreUObject/Private/UObject/GarbageCollection.cpp b/Engine/Source/Runtime/CoreUObject/Private/UObject/GarbageCollection.cpp
--- a/Engine/Source/Runtime/CoreUObject/Private/UObject/GarbageCollection.cpp
+++ b/Engine/Source/Runtime/CoreUObject/Private/UObject/GarbageCollection.cpp
@@ -516,11 +516,6 @@ namespace UE::GC
 		LogClassCountInfo(TEXT("objects of"), ClassToPurgeCountMap, 10, Total);
 #endif
 	}
-
-	void MarkAsReachable(const UObject* Obj)
-	{
-		Obj->MarkAsReachable();
-	}
 }
 
 namespace UE::GC::Private
diff --git a/Engine/Source/Runtime/CoreUObject/Private/UObject/ObjectFwd.cpp b/Engine/Source/Runtime/CoreUObject/Private/UObject/ObjectFwd.cpp
deleted file mode 100644
diff --git a/Engine/Source/Runtime/CoreUObject/Private/UObject/UObjectBase.cpp b/Engine/Source/Runtime/CoreUObject/Private/UObject/UObjectBase.cpp
--- a/Engine/Source/Runtime/CoreUObject/Private/UObject/UObjectBase.cpp
+++ b/Engine/Source/Runtime/CoreUObject/Private/UObject/UObjectBase.cpp
@@ -137,7 +137,16 @@ UObjectBase::UObjectBase(UClass* InClass,
 #if CSV_PROFILER && CSV_TRACK_UOBJECT_COUNT
 	UObjectStats::IncrementUObjectCount();
 #endif
+
+	if (UE::GC::Private::GIsIncrementalReachabilityPending)
+	{
+		ClassPrivate->MarkAsReachable();
+		if (OuterPrivate)
+		{
+			OuterPrivate->MarkAsReachable();
 		}	
+	}
+}
 
 
 /**
@@ -183,6 +192,15 @@ void UObjectBase::DeferredRegister(UClass *UClassStaticClass,const TCHAR* Packag
 	check(!ClassPrivate);
 	ClassPrivate = UClassStaticClass;
 
+	if (UE::GC::Private::GIsIncrementalReachabilityPending)
+	{
+		ClassPrivate->MarkAsReachable();
+		if (OuterPrivate)
+		{
+			OuterPrivate->MarkAsReachable();
+		}
+	}
+
 	// Add to the global object table.
 	AddObject(FName(InName), EInternalObjectFlags::None);
 	// At this point all compiled-in objects should have already been fully constructed so it's safe to remove the NotFullyConstructed flag
@@ -241,7 +259,11 @@ void UObjectBase::LowLevelRename(FName NewName,UObject *NewOuter)
 	if (NewOuter)
 	{
 		OuterPrivate = NewOuter;
+		if (UE::GC::Private::GIsIncrementalReachabilityPending)
+		{
+			OuterPrivate->MarkAsReachable();
 		}
+	}
 	HashObject(this);
 }
 
@@ -302,6 +324,11 @@ void UObjectBase::SetClass(UClass* NewClass)
 	ClassPrivate->DestroyPersistentUberGraphFrame((UObject*)this);
 #endif
 	ClassPrivate = NewClass;
+	if (UE::GC::Private::GIsIncrementalReachabilityPending)
+	{
+		checkf(ClassPrivate, TEXT("SetClass called on %s with a null class"), *GetFName().ToString());
+		ClassPrivate->MarkAsReachable();
+	}
 #if USE_UBER_GRAPH_PERSISTENT_FRAME
 	ClassPrivate->CreatePersistentUberGraphFrame((UObject*)this, /*bCreateOnlyIfEmpty =*/false, /*bSkipSuperClass =*/false, OldClass);
 #endif
diff --git a/Engine/Source/Runtime/CoreUObject/Private/UObject/UObjectGlobals.cpp b/Engine/Source/Runtime/CoreUObject/Private/UObject/UObjectGlobals.cpp
--- a/Engine/Source/Runtime/CoreUObject/Private/UObject/UObjectGlobals.cpp
+++ b/Engine/Source/Runtime/CoreUObject/Private/UObject/UObjectGlobals.cpp
@@ -1390,8 +1390,8 @@ UObject* StaticLoadObjectInternal(UClass* ObjectClass, UObject* InOuter, const T
 				{
 					if (UE::GC::Private::GIsIncrementalReachabilityPending)
 					{
-						UE::GC::MarkAsReachable(Redirector);
-						UE::GC::MarkAsReachable(Redirector->DestinationObject);
+						Redirector->MarkAsReachable();
+						Redirector->DestinationObject->MarkAsReachable();
 					}
 					return Redirector->DestinationObject;
 				}
@@ -1417,7 +1417,7 @@ UObject* StaticLoadObjectInternal(UClass* ObjectClass, UObject* InOuter, const T
 
 	if (Result && UE::GC::Private::GIsIncrementalReachabilityPending)
 	{
-		UE::GC::MarkAsReachable(Result);
+		Result->MarkAsReachable();
 	}
 	return Result;
 }
diff --git a/Engine/Source/Runtime/CoreUObject/Private/UObject/UObjectHash.cpp b/Engine/Source/Runtime/CoreUObject/Private/UObject/UObjectHash.cpp
--- a/Engine/Source/Runtime/CoreUObject/Private/UObject/UObjectHash.cpp
+++ b/Engine/Source/Runtime/CoreUObject/Private/UObject/UObjectHash.cpp
@@ -813,7 +813,7 @@ UObject* StaticFindObjectFastInternalThreadSafe(FUObjectHashTables& ThreadHash,
 	// Not found.
 	if (Result && UE::GC::Private::GIsIncrementalReachabilityPending)
 	{
-		UE::GC::MarkAsReachable(Result);
+		Result->MarkAsReachable();
 	}
 	return Result;
 }
@@ -1141,7 +1141,7 @@ void GetObjectsWithOuter(const class UObjectBase* Outer, TArray<UObject *>& Resu
 			{
 				if (UE::GC::Private::GIsIncrementalReachabilityPending)
 				{
-					UE::GC::MarkAsReachable(Object);
+					Object->MarkAsReachable();
 				}
 				Results.Add(Object);
 			}
@@ -1164,7 +1164,7 @@ void GetObjectsWithOuter(const class UObjectBase* Outer, TArray<UObject *>& Resu
 						{
 							if (UE::GC::Private::GIsIncrementalReachabilityPending)
 							{
-								UE::GC::MarkAsReachable(Object);
+								Object->MarkAsReachable();
 							}
 							Results.Add(Object);
 						}
@@ -1216,7 +1216,7 @@ void ForEachObjectWithOuterBreakable(const class UObjectBase* Outer, TFunctionRe
 			{
 				if (UE::GC::Private::GIsIncrementalReachabilityPending)
 				{
-					UE::GC::MarkAsReachable(Object);
+					Object->MarkAsReachable();
 				}
 				if (!Operation(Object))
 				{
@@ -1273,7 +1273,7 @@ UObjectBase* FindObjectWithOuter(const class UObjectBase* Outer, const class UCl
 			}
 			if (Result && UE::GC::Private::GIsIncrementalReachabilityPending)
 			{
-				UE::GC::MarkAsReachable(Result);
+				Result->MarkAsReachable();
 			}
 		}
 	}
@@ -1334,7 +1334,7 @@ void ForEachObjectWithPackage(const class UPackage* Package, TFunctionRef<bool(U
 			{
 				if (UE::GC::Private::GIsIncrementalReachabilityPending)
 				{
-					UE::GC::MarkAsReachable(Object);
+					Object->MarkAsReachable();
 				}
 				if (!Operation(Object))
 				{
@@ -1427,7 +1427,7 @@ FORCEINLINE void ForEachObjectOfClasses_Implementation(FUObjectHashTables& Threa
 				{
 					if (UE::GC::Private::GIsIncrementalReachabilityPending)
 					{
-						UE::GC::MarkAsReachable(Object);
+						Object->MarkAsReachable();
 					}
 					Operation(Object);
 				}
diff --git a/Engine/Source/Runtime/CoreUObject/Public/UObject/DynamicallyTypedValue.h b/Engine/Source/Runtime/CoreUObject/Public/UObject/DynamicallyTypedValue.h
--- a/Engine/Source/Runtime/CoreUObject/Public/UObject/DynamicallyTypedValue.h
+++ b/Engine/Source/Runtime/CoreUObject/Public/UObject/DynamicallyTypedValue.h
@@ -129,7 +129,7 @@ namespace UE
 				{
 					if (InObject)
 					{
-						UE::GC::MarkAsReachable(InObject);
+						InObject->MarkAsReachable();
 					}
 				}
 			};
diff --git a/Engine/Source/Runtime/CoreUObject/Public/UObject/Field.h b/Engine/Source/Runtime/CoreUObject/Public/UObject/Field.h
--- a/Engine/Source/Runtime/CoreUObject/Public/UObject/Field.h
+++ b/Engine/Source/Runtime/CoreUObject/Public/UObject/Field.h
@@ -275,7 +275,7 @@ class FFieldVariant
 	{
 		if (IsUObject() && ToUObjectUnsafe() && UE::GC::Private::GIsIncrementalReachabilityPending)
 		{
-			UE::GC::MarkAsReachable(ToUObjectUnsafe());
+			ToUObjectUnsafe()->MarkAsReachable();
 		}
 	}
 	
diff --git a/Engine/Source/Runtime/CoreUObject/Public/UObject/Object.h b/Engine/Source/Runtime/CoreUObject/Public/UObject/Object.h
--- a/Engine/Source/Runtime/CoreUObject/Public/UObject/Object.h
+++ b/Engine/Source/Runtime/CoreUObject/Public/UObject/Object.h
@@ -1827,6 +1827,22 @@ const T* GetValid(const T* Test)
 	return IsValid(Test) ? Test : nullptr;
 }
 
+/**
+* Helper function to mark object as reachable when performing Incremental GC
+*
+* @param	Object			Object to mark as reachable
+* @return	The same Object pointer that was passed to the function
+*/
+template <typename T>
+FORCEINLINE T* MarkAsReachable(T* Object)
+{
+	if (Object)
+	{
+		Object->MarkAsReachable();
+	}
+	return Object;
+}
+
 #if WITH_EDITOR
 /** Callback for editor object selection. This must be in core instead of editor for UObject::IsSelectedInEditor to work */
 extern COREUOBJECT_API TFunction<bool(const UObject*)> GIsObjectSelectedInEditor;
diff --git a/Engine/Source/Runtime/CoreUObject/Public/UObject/ObjectFwd.h b/Engine/Source/Runtime/CoreUObject/Public/UObject/ObjectFwd.h
deleted file mode 100644
diff --git a/Engine/Source/Runtime/CoreUObject/Public/UObject/ObjectHandle.h b/Engine/Source/Runtime/CoreUObject/Public/UObject/ObjectHandle.h
--- a/Engine/Source/Runtime/CoreUObject/Public/UObject/ObjectHandle.h
+++ b/Engine/Source/Runtime/CoreUObject/Public/UObject/ObjectHandle.h
@@ -10,6 +10,7 @@
 #include "Misc/AssertionMacros.h"
 #include "Templates/TypeHash.h"
 #include "UObject/NameTypes.h"
+#include "UObject/Object.h"
 #include "UObject/ObjectHandleTracking.h"
 #include "UObject/ObjectMacros.h"
 #include "UObject/ObjectRef.h"
@@ -234,7 +235,7 @@ namespace UE::CoreUObject::Private
 		if (IsObjectHandleResolved(Handle))
 		{
 			UObject* Obj = ReadObjectHandlePointerNoCheck(Handle);
-			return Obj != nullptr ? UE::CoreUObject::Private::GetClass(Obj) : nullptr;
+			return Obj != nullptr ? Obj->GetClass() : nullptr;
 		}
 		else
 		{
@@ -245,7 +246,7 @@ namespace UE::CoreUObject::Private
 		}
 #else
 		UObject* Obj = ReadObjectHandlePointerNoCheck(Handle);
-		return Obj != nullptr ? UE::CoreUObject::Private::GetClass(Obj) : nullptr;
+		return Obj != nullptr ? Obj->GetClass() : nullptr;
 #endif
 	}
 
@@ -350,6 +351,3 @@ namespace UE::CoreUObject::Private
 #endif
 }
 
-#if UE_ENABLE_INCLUDE_ORDER_DEPRECATED_IN_5_4
-#include "UObject/Object.h"
-#endif
diff --git a/Engine/Source/Runtime/CoreUObject/Public/UObject/ObjectPtr.h b/Engine/Source/Runtime/CoreUObject/Public/UObject/ObjectPtr.h
--- a/Engine/Source/Runtime/CoreUObject/Public/UObject/ObjectPtr.h
+++ b/Engine/Source/Runtime/CoreUObject/Public/UObject/ObjectPtr.h
@@ -5,8 +5,8 @@
 #include "HAL/Platform.h"
 #include "Serialization/StructuredArchive.h"
 #include "Templates/IsTObjectPtr.h"
+#include "UObject/Object.h"
 #include "UObject/ObjectHandle.h"
-#include "UObject/UObjectGlobals.h"
 #include "Templates/NonNullPointer.h"
 
 #include <type_traits>
@@ -218,13 +218,13 @@ public:
 	FString GetPathName() const
 	{
 		const UObject* ResolvedObject = Get();
-		return ResolvedObject ? UE::CoreUObject::Private::GetPathName(ResolvedObject) : TEXT("None");
+		return ResolvedObject ? ResolvedObject->GetPathName() : TEXT("None");
 	}
 
 	FName GetFName() const
 	{
 		const UObject* ResolvedObject = Get();
-		return ResolvedObject ? UE::CoreUObject::Private::GetFName(ResolvedObject) : NAME_None;
+		return ResolvedObject ? ResolvedObject->GetFName() : NAME_None;
 	}
 
 	FString GetName() const
@@ -235,13 +235,13 @@ public:
 	FObjectPtr GetOuter() const
 	{
 		const UObject* ResolvedObject = Get();
-		return ResolvedObject ? FObjectPtr(UE::CoreUObject::Private::GetOuter(ResolvedObject)) : FObjectPtr(nullptr);
+		return ResolvedObject ? FObjectPtr(ResolvedObject->GetOuter()) : FObjectPtr(nullptr);
 	}
 
 	FObjectPtr GetPackage() const
 	{
 		const UObject* ResolvedObject = Get();
-		return ResolvedObject ? FObjectPtr(UE::CoreUObject::Private::GetPackage(ResolvedObject)) : FObjectPtr(nullptr);
+		return ResolvedObject ? FObjectPtr(ResolvedObject->GetPackage()) : FObjectPtr(nullptr);
 	}
 	/**
 	 * Returns the fully qualified pathname for this object as well as the name of the class, in the format:
@@ -250,7 +250,7 @@ public:
 	FString GetFullName(EObjectFullNameFlags Flags = EObjectFullNameFlags::None) const
 	{
 		// UObjectBaseUtility::GetFullName is safe to call on null objects.
-		return UE::CoreUObject::Private::GetFullName(Get(), nullptr, Flags);
+		return Get()->GetFullName(nullptr, Flags);
 	}
 #endif
 
@@ -286,7 +286,7 @@ private:
 		{
 			if (UObject* Obj = UE::CoreUObject::Private::ReadObjectHandlePointerNoCheck(InPtr.GetHandleRef()))
 			{
-				UE::GC::MarkAsReachable(Obj);
+				Obj->MarkAsReachable();
 			}
 		}
 	}
@@ -294,7 +294,7 @@ private:
 	{
 		if (UE::GC::Private::GIsIncrementalReachabilityPending && InObj)
 		{
-			UE::GC::MarkAsReachable(InObj);
+			InObj->MarkAsReachable();
 		}
 	}
 #endif // UE_OBJECT_PTR_GC_BARRIER
@@ -731,21 +731,11 @@ namespace ObjectPtr_Private
 	{
 	public:
 		TNonAccessTrackedObjectPtr() = default;
+		explicit TNonAccessTrackedObjectPtr(T* Ptr) : ObjectPtr{Ptr} {}
 
-		explicit TNonAccessTrackedObjectPtr(ENoInit)
-			: ObjectPtr{NoInit}
-		{
-		}
-
-		explicit TNonAccessTrackedObjectPtr(T* Ptr)
-			: ObjectPtr{Ptr}
-		{
-		}
-
-		TNonAccessTrackedObjectPtr& operator=(T* Value)
+		void Set(T* Value)
 		{
 			ObjectPtr = Value;
-			return *this;
 		}
 	
 		T* Get() const
@@ -753,11 +743,6 @@ namespace ObjectPtr_Private
 			return ObjectPtr_Private::Friend::NoAccessTrackingGet(ObjectPtr);
 		}
 
-		explicit operator UPTRINT() const
-		{
-			return BitCast<UPTRINT>(Get());
-		}
-
 		TObjectPtr<T>& GetAccessTrackedObjectPtr() 
 		{
 			return ObjectPtr;
@@ -774,7 +759,7 @@ namespace ObjectPtr_Private
 		}		
 	
 	private:
-		TObjectPtr<T> ObjectPtr;
+		TObjectPtr<T> ObjectPtr{};
 	};
 }
 
@@ -1330,7 +1315,7 @@ namespace UE::Core::Private // private facilities; not for direct use
 #if UE_OBJECT_PTR_GC_BARRIER
 			if (UE::GC::Private::GIsIncrementalReachabilityPending && View)
 			{
-				UE::GC::MarkAsReachable(View);
+				View->MarkAsReachable();
 			}
 #endif // UE_OBJECT_PTR_GC_BARRIER
 		}
@@ -1349,7 +1334,7 @@ namespace UE::Core::Private // private facilities; not for direct use
 				{
 					if (Data[Index])
 					{
-						UE::GC::MarkAsReachable(Data[Index]);
+						Data[Index]->MarkAsReachable();
 					}
 				}
 			}
@@ -1369,7 +1354,7 @@ namespace UE::Core::Private // private facilities; not for direct use
 				{
 					if (Element)
 					{
-						UE::GC::MarkAsReachable(Element);
+						Element->MarkAsReachable();
 					}
 				}
 			}
@@ -1394,14 +1379,14 @@ namespace UE::Core::Private // private facilities; not for direct use
 					{
 						if (Pair.Key)
 						{
-							UE::GC::MarkAsReachable(Pair.Key);
+							Pair.Key->MarkAsReachable();
 						}
 					}
 					if constexpr (bValueReference)
 					{
 						if (Pair.Value)
 						{
-							UE::GC::MarkAsReachable(Pair.Value);
+							Pair.Value->MarkAsReachable();
 						}
 					}
 				}
@@ -1509,7 +1494,7 @@ namespace UE::Core::Private // private facilities; not for direct use
 		{
 			if (const UObject* Obj = Cast<UObject>(Ptr); Obj && UE::GC::Private::GIsIncrementalReachabilityPending)
 			{
-				UE::GC::MarkAsReachable(Obj);
+				Obj->MarkAsReachable();
 			}
 		}
 
@@ -1743,7 +1728,3 @@ private:
 #if UE_ENABLE_INCLUDE_ORDER_DEPRECATED_IN_5_3
 #include "UObject/Class.h"
 #endif
-
-#if UE_ENABLE_INCLUDE_ORDER_DEPRECATED_IN_5_4
-#include "UObject/Object.h"
-#endif
diff --git a/Engine/Source/Runtime/CoreUObject/Public/UObject/ObjectRef.h b/Engine/Source/Runtime/CoreUObject/Public/UObject/ObjectRef.h
--- a/Engine/Source/Runtime/CoreUObject/Public/UObject/ObjectRef.h
+++ b/Engine/Source/Runtime/CoreUObject/Public/UObject/ObjectRef.h
@@ -3,9 +3,8 @@
 #pragma once
 
 #include "Misc/StringBuilder.h"
+#include "UObject/UObjectBaseUtility.h"
 #include "UObject/ObjectHandleDefines.h"
-#include "UObject/ObjectFwd.h"
-#include "UObject/ObjectMacros.h"
 
 #if UE_WITH_OBJECT_HANDLE_LATE_RESOLVE || UE_WITH_OBJECT_HANDLE_TRACKING
 
@@ -122,8 +121,4 @@ private:
 UE_DEPRECATED(5.2, "Use FObjectRef::IsNull() instead")
 inline bool IsObjectRefNull(const FObjectRef& ObjectRef) { return ObjectRef.IsNull(); }
 
-#if UE_ENABLE_INCLUDE_ORDER_DEPRECATED_IN_5_4
-#include "UObject/UObjectBaseUtility.h"
-#endif
-
 #endif
\ No newline at end of file
diff --git a/Engine/Source/Runtime/CoreUObject/Public/UObject/UObjectBase.h b/Engine/Source/Runtime/CoreUObject/Public/UObject/UObjectBase.h
--- a/Engine/Source/Runtime/CoreUObject/Public/UObject/UObjectBase.h
+++ b/Engine/Source/Runtime/CoreUObject/Public/UObject/UObjectBase.h
@@ -18,7 +18,6 @@
 #include "UObject/ObjectMacros.h"
 #include "UObject/UObjectGlobals.h"
 #include "UObject/UnrealNames.h"
-#include "UObject/ObjectPtr.h"
 
 class UClass;
 class UEnum;
@@ -48,9 +47,7 @@ class UObjectBase
 		);
 protected:
 	UObjectBase() :
-		ClassPrivate(NoInit),
-		NamePrivate(NoInit),  // screwy, but the name was already set and we don't want to set it again
-		OuterPrivate(NoInit)
+		 NamePrivate(NoInit)  // screwy, but the name was already set and we don't want to set it again
 	{
 	}
 
@@ -252,13 +249,13 @@ private:
 	int32							InternalIndex;
 
 	/** Class the object belongs to. */
-	ObjectPtr_Private::TNonAccessTrackedObjectPtr<UClass>							ClassPrivate;
+	UClass*							ClassPrivate;
 
 	/** Name of this object */
 	FName							NamePrivate;
 
 	/** Object this object resides in. */
-	ObjectPtr_Private::TNonAccessTrackedObjectPtr<UObject>						OuterPrivate;
+	UObject*						OuterPrivate;
 	
 	friend class FBlueprintCompileReinstancer;
 	friend class FVerseObjectClassReplacer;
diff --git a/Engine/Source/Runtime/CoreUObject/Public/UObject/UObjectBaseUtility.h b/Engine/Source/Runtime/CoreUObject/Public/UObject/UObjectBaseUtility.h
--- a/Engine/Source/Runtime/CoreUObject/Public/UObject/UObjectBaseUtility.h
+++ b/Engine/Source/Runtime/CoreUObject/Public/UObject/UObjectBaseUtility.h
@@ -25,7 +25,6 @@
 #include "UObject/UObjectArray.h"
 #include "UObject/UObjectBase.h"
 #include "UObject/UObjectMarks.h"
-#include "UObject/ObjectFwd.h"
 
 class UClass;
 class UObject;
@@ -38,6 +37,20 @@ struct FGuid;
 	#endif
 #endif
 
+/**
+* Enum which specifies the mode in which full object names are constructed
+*/
+enum class EObjectFullNameFlags
+{
+	// Standard object full name (i.e. "Type PackageName.ObjectName:SubobjectName")
+	None = 0,
+
+	// Adds package to the type portion (i.e. "TypePackage.TypeName PackageName.ObjectName:SubobjectName")
+	IncludeClassPackage = 1,
+};
+
+ENUM_CLASS_FLAGS(EObjectFullNameFlags);
+
 /**
  * Provides utility functions for UObject, this class should not be used directly
  */
diff --git a/Engine/Source/Runtime/Engine/Public/TransactionCommon.h b/Engine/Source/Runtime/Engine/Public/TransactionCommon.h
--- a/Engine/Source/Runtime/Engine/Public/TransactionCommon.h
+++ b/Engine/Source/Runtime/Engine/Public/TransactionCommon.h
@@ -3,7 +3,6 @@
 #pragma once
 
 #include "CoreMinimal.h"
-#include "UObject/Object.h"
 #include "UObject/ObjectPtr.h"
 #include "UObject/WeakObjectPtr.h"
 #include "Serialization/ArchiveUObject.h"
