[ZenLoader] - Fix crash when reloading a world twice in PIE due to missing fixup in GlobalImportStore

Simple repro
  - PIE in EngineTest
  - 2x Open /Engine/Maps/Templates/OpenWorld

#rnx

diff --git a/Engine/Source/Runtime/CoreUObject/Private/Serialization/AsyncLoading2.cpp b/Engine/Source/Runtime/CoreUObject/Private/Serialization/AsyncLoading2.cpp
--- a/Engine/Source/Runtime/CoreUObject/Private/Serialization/AsyncLoading2.cpp
+++ b/Engine/Source/Runtime/CoreUObject/Private/Serialization/AsyncLoading2.cpp
@@ -1050,6 +1050,33 @@ public:
 				UPackage* FoundPackage = FindObjectFast<UPackage>(nullptr, PackageNameIfKnown);
 				if (FoundPackage)
 				{
+					// We need to maintain a 1:1 relationship between PackageId and Package object
+					// to avoid having a costly 1:N lookup during GC. If we find an existing package
+					// that already has a valid packageid in our table, it means it was renamed
+					// after it finished loading. Since our mapping is now "stale" we remove
+					// it before replacing the package ID in the package object, otherwise
+					// the mapping would leak during GC and we could end up crashing when
+					// trying to access the ref after the GC has run.
+					FPackageId OldPackageId = FoundPackage->GetPackageId();
+					if (OldPackageId.IsValid() && OldPackageId != PackageId)
+					{
+						FLoadedPackageRef* OldPackageRef = Packages.Find(OldPackageId);
+						if (OldPackageRef)
+						{
+							UE_LOG(LogStreaming, Display,
+								TEXT("FGlobalImportStore:AddPackageRef: Dropping stale reference to package %s (0x%llX) that has been renamed to %s (0x%llX)"),
+								*OldPackageRef->GetOriginalPackageName().ToString(),
+								OldPackageId.ValueForDebugging(),
+								*FoundPackage->GetName(),
+								PackageId.ValueForDebugging()
+							);
+							
+							check(OldPackageRef->GetRefCount() == 0);
+							RemoveCompletedRenamedPackage(*OldPackageRef);
+							RemovePackage(OldPackageId);
+						}
+					}
+
 					PackageRef.SetPackage(FoundPackage);
 					FoundPackage->SetCanBeImportedFlag(true);
 					FoundPackage->SetPackageId(PackageId);
@@ -1177,12 +1204,12 @@ public:
 			const FPackageId& PackageId = Item.PackageId;
 			if (PackageId.IsValid())
 			{
-				RemovePackage(PackageId, Item.ObjectName);
+				RemovePackage(PackageId);
 			}
 		}
 	}
 
-	void RemovePackage(FPackageId PackageId, FName PackageName)
+	void RemovePackage(FPackageId PackageId)
 	{
 		UE_ASYNC_PACKAGEID_DEBUG(PackageId);
 
