[Loader] - Add object flags in the error when detaching linker of an object that still needs loading.

diff --git a/Engine/Source/Runtime/CoreUObject/Private/UObject/UObjectLinker.cpp b/Engine/Source/Runtime/CoreUObject/Private/UObject/UObjectLinker.cpp
--- a/Engine/Source/Runtime/CoreUObject/Private/UObject/UObjectLinker.cpp
+++ b/Engine/Source/Runtime/CoreUObject/Private/UObject/UObjectLinker.cpp
@@ -116,9 +116,10 @@ void UObject::SetLinker( FLinkerLoad* LinkerLoad, int32 LinkerIndex, bool bShoul
 	if( Existing.Linker && bShouldDetachExisting )
 	{
 		UE_CLOG(HasAnyFlags(RF_NeedLoad|RF_NeedPostLoad), LogUObjectLinker, Error,
-			TEXT("Detaching from existing linker %s while object %s needs loading. Setting linker to %s."),
+			TEXT("Detaching from existing linker %s while object %s needs loading (%s). Setting linker to %s."),
 			*Existing.Linker->GetArchiveName(),
 			*GetFullName(),
+			*LexToString(GetFlags()),
 			LinkerLoad ? *LinkerLoad->GetDebugName() : TEXT("nullptr"));
 		check(Existing.Linker->ExportMap[Existing.LinkerIndex].Object!=nullptr);
 		check(Existing.Linker->ExportMap[Existing.LinkerIndex].Object==this);
