Fixed issue where the GC debug processor could stop being used during analysis.

diff --git a/Engine/Source/Runtime/CoreUObject/Private/UObject/GarbageCollection.cpp b/Engine/Source/Runtime/CoreUObject/Private/UObject/GarbageCollection.cpp
--- a/Engine/Source/Runtime/CoreUObject/Private/UObject/GarbageCollection.cpp
+++ b/Engine/Source/Runtime/CoreUObject/Private/UObject/GarbageCollection.cpp
@@ -4186,9 +4186,9 @@ private:
 		if (!GReachabilityState.CheckIfAnyContextIsSuspended())
 		{
 			GReachabilityState.ResetWorkers();
-			Stats = Context->Stats;
+			Stats.AddStats(Context->Stats);
+			GReachabilityState.UpdateStats(Context->Stats);
 			Pool.ReturnToPool(Context);
-			GReachabilityState.UpdateStats(Stats);
 		}
 		else if (bIsSingleThreaded)
 		{
