Initial property bag repository skeleton.

diff --git a/Engine/Source/Editor/Kismet/Private/BlueprintCompilationManager.cpp b/Engine/Source/Editor/Kismet/Private/BlueprintCompilationManager.cpp
--- a/Engine/Source/Editor/Kismet/Private/BlueprintCompilationManager.cpp
+++ b/Engine/Source/Editor/Kismet/Private/BlueprintCompilationManager.cpp
@@ -44,6 +44,7 @@
 #include "BlueprintEditorModule.h"
 #include "Animation/AnimBlueprint.h"
 #include "Stats/StatsHierarchical.h"
+#include "UObject/PropertyBagRepository.h"
 
 extern UNREALED_API UUnrealEdEngine* GUnrealEd;
 
@@ -2663,6 +2664,9 @@ void FBlueprintCompilationManagerImpl::ReinstanceBatch(TArray<FReinstancingJob>&
 		}
 	}
 
+	// Reassociate relevant property bags
+	UE::FPropertyBagRepository::Get().ReassociateObjects(OldArchetypeToNewArchetype);
+	
 	// 4. update known references to archetypes (e.g. component templates, WidgetTree). We don't want to run the normal 
 	// reference finder to update these because searching the entire object graph is time consuming. Instead we just replace
 	// all references in our UBlueprint and its generated class:
diff --git a/Engine/Source/Editor/UnrealEd/Private/Kismet2/KismetReinstanceUtilities.cpp b/Engine/Source/Editor/UnrealEd/Private/Kismet2/KismetReinstanceUtilities.cpp
--- a/Engine/Source/Editor/UnrealEd/Private/Kismet2/KismetReinstanceUtilities.cpp
+++ b/Engine/Source/Editor/UnrealEd/Private/Kismet2/KismetReinstanceUtilities.cpp
@@ -41,6 +41,7 @@
 #include "Engine/ScopedMovementUpdate.h"
 #include "InstancedReferenceSubobjectHelper.h"
 #include "UObject/PropertyOptional.h"
+#include "UObject/PropertyBagRepository.h"
 
 DECLARE_CYCLE_STAT(TEXT("Replace Instances"), EKismetReinstancerStats_ReplaceInstancesOfClass, STATGROUP_KismetReinstancer );
 DECLARE_CYCLE_STAT(TEXT("Find Referencers"), EKismetReinstancerStats_FindReferencers, STATGROUP_KismetReinstancer );
@@ -2852,6 +2853,9 @@ void FBlueprintCompileReinstancer::ReplaceInstancesOfClass_Inner(const TMap<UCla
 		}
 	}
 
+	// Reassociate relevant property bags 
+	UE::FPropertyBagRepository::Get().ReassociateObjects(OldToNewInstanceMap);
+
 	// Inform listeners of object reinstancing
 	FCoreUObjectDelegates::OnObjectsReinstanced.Broadcast(OldToNewInstanceMap);
 	
diff --git a/Engine/Source/Runtime/CoreUObject/Internal/UObject/PropertyBagRepository.h b/Engine/Source/Runtime/CoreUObject/Internal/UObject/PropertyBagRepository.h
new file mode 100644
--- /dev/null
+++ b/Engine/Source/Runtime/CoreUObject/Internal/UObject/PropertyBagRepository.h
@@ -0,0 +1,92 @@
+// Copyright Epic Games, Inc. All Rights Reserved.
+
+#pragma once
+
+#include "CoreTypes.h"
+#include "Containers/Map.h"
+#include "HAL/CriticalSection.h"
+
+class UObjectBase;
+class UObject;
+
+namespace UE
+{
+
+class FPropertyBag;
+
+// Singleton class tracking property bag association with objects
+class FPropertyBagRepository
+{
+	// TODO: Make private throughout and extend access permissions here or in wrapper classes? Don't want engine code modifying bags outside of serializers and details panels.
+	//friend UObjectBase;
+	//friend UStruct;
+
+private:
+	friend class FPropertyBagRepositoryLock;
+	mutable FCriticalSection CriticalSection;
+	
+	// Lifetimes/ownership:
+	// Managed within UObjectBase and synced with object lifetime. The repo tracks pointers to bags, not the bags themselves.
+	// FPropertyBag destruction handles destruction of sub-bags automatically, via the FPropertyBagProperty tracking them.
+
+	/** Map of objects/subobjects to their top level property bag. */
+	// TODO: Currently will only exist in editor world, but could tracking per world make some sense for teardown in future? We're relying on object destruction to occur properly to free these up. 
+	TMap<const UObjectBase*, FPropertyBag*> ObjectToPropertyBagMap;	// TODO: Ref bags via handle?
+	// /** Map of subobject/container/struct/etc. paths (needs FPropertyBagPath) to their property bag. Might not want to track subobjects here (they'll be in ObjectToPropertyBagMap already). */
+	// TMap<FSoftObjectPath, FPropertyBag*> ObjectPathToPropertySubBagMap;
+	
+	//TMap<const UObjectBase*, UObject*> ObjectToArchetypeMap;
+
+	FPropertyBagRepository() = default;
+
+public:
+	FPropertyBagRepository(const FPropertyBagRepository &) = delete;
+	FPropertyBagRepository& operator=(const FPropertyBagRepository&) = delete;
+	
+	// Singleton accessor
+	static COREUOBJECT_API FPropertyBagRepository& Get();
+
+	// Reclaim space - TODO: Hook up to GC.
+	void ShrinkMaps();
+	
+	// TODO: Restrict bag creation to actor creation and UStruct::SerializeVersionedTaggedProperties?
+	// Object owner is tracked internally
+	FPropertyBag* CreateOuterBag(UObjectBase* Owner);
+
+	// Future version for reworked archetypes - track archetype rather than bag (directly):
+	/**
+	 * Instantiate an archetype object representing all fields within the bag, tracked against the owner object.
+	 * @param Owner			- Associated in world object.
+	 * @return				- Custom archetype object, UClass derived from associated bag.
+	 */
+	//UObject* CreateArchetype(UObjectBase* Owner);
+
+	// TODO: Restrict property bag  destruction to within UObject::BeginDestroy() & FPropertyBagProperty destructor.
+	void DestroyOuterBag(UObjectBase* Owner);
+
+	/**
+	 * ReassociateObjects
+	 * @param ReplacedObjects - old/new owner object pairs. Reassigns archetypes/bags to the new owner.
+	 */
+	COREUOBJECT_API void ReassociateObjects(const TMap<UObject*, UObject*>& ReplacedObjects);
+	
+	// Accessors
+	FPropertyBag* FindBag(const UObjectBase* Owner);
+	const FPropertyBag* FindBag(const UObjectBase* Owner) const;
+	//UObject* FindArchetype(const UObjectBase* Owner);
+	
+	bool HasBag(const UObjectBase* Owner) const;
+	
+	// TODO: Cache floating property status - within bag?
+	//EMatchStatus MatchesBag(const UObjectBase* Object) const;
+	
+private:
+	void Lock() const { CriticalSection.Lock(); }
+	void Unlock() const { CriticalSection.Unlock(); }
+
+	// Not thread safe - internal use only.
+	FPropertyBag* FindBagUnsafe(const UObjectBase* Object);
+	const FPropertyBag* FindBagUnsafe(const UObjectBase* Object) const;
+};
+
+} // UE
diff --git a/Engine/Source/Runtime/CoreUObject/Private/UObject/Obj.cpp b/Engine/Source/Runtime/CoreUObject/Private/UObject/Obj.cpp
--- a/Engine/Source/Runtime/CoreUObject/Private/UObject/Obj.cpp
+++ b/Engine/Source/Runtime/CoreUObject/Private/UObject/Obj.cpp
@@ -67,6 +67,7 @@
 #include "ProfilingDebugging/AssetMetadataTrace.h"
 #include "Containers/VersePath.h"
 #include "Serialization/LoadTimeTracePrivate.h"
+#include "UObject/PropertyBagRepository.h"
 
 DEFINE_LOG_CATEGORY(LogObj);
 
@@ -984,6 +985,9 @@ void UObject::BeginDestroy()
 	// Remove any associated external package, at this point
 	SetExternalPackage(nullptr);
 
+	// Destroy any associated property bag.
+	UE::FPropertyBagRepository::Get().DestroyOuterBag(this);
+	
 	// ensure BeginDestroy has been routed back to UObject::BeginDestroy.
 #if !(UE_BUILD_SHIPPING || UE_BUILD_TEST)
 	DebugBeginDestroyed.RemoveSingle(this);
diff --git a/Engine/Source/Runtime/CoreUObject/Private/UObject/PropertyBagRepository.cpp b/Engine/Source/Runtime/CoreUObject/Private/UObject/PropertyBagRepository.cpp
new file mode 100644
--- /dev/null
+++ b/Engine/Source/Runtime/CoreUObject/Private/UObject/PropertyBagRepository.cpp
@@ -0,0 +1,136 @@
+// Copyright Epic Games, Inc. All Rights Reserved.
+
+#include "UObject/PropertyBagRepository.h"
+#include "UObject/GarbageCollection.h"
+#include "UObject/Object.h"
+#include "UObject/UObjectGlobals.h"
+#include "UObject/PropertyBag.h"
+
+
+DEFINE_LOG_CATEGORY_STATIC(LogPropertyBagRepository, Log, All);
+
+namespace UE
+{
+
+class FPropertyBagRepositoryLock
+{
+#if THREADSAFE_UOBJECTS
+	const FPropertyBagRepository* Repo;	// Technically a singleton, but just in case...
+#endif
+public:
+	FORCEINLINE FPropertyBagRepositoryLock(const FPropertyBagRepository* InRepo)
+	{
+#if THREADSAFE_UOBJECTS
+		if (!(IsGarbageCollectingAndLockingUObjectHashTables() && IsInGameThread()))	// Mirror object hash tables behaviour exactly for now
+		{
+			Repo = InRepo;
+			InRepo->Lock();
+		}
+		else
+		{
+			Repo = nullptr;
+		}
+#else
+		check(IsInGameThread());
+#endif
+	}
+	FORCEINLINE ~FPropertyBagRepositoryLock()
+	{
+#if THREADSAFE_UOBJECTS
+		if (Repo)
+		{
+			Repo->Unlock();
+		}
+#endif
+	}
+};
+
+FPropertyBagRepository& FPropertyBagRepository::Get()
+{
+	static FPropertyBagRepository Repo;
+	return Repo;
+}
+
+void FPropertyBagRepository::ReassociateObjects(const TMap<UObject*, UObject*>& ReplacedObjects)
+{
+	FPropertyBagRepositoryLock LockRepo(this);
+	FPropertyBag* PropertyBag;
+	for(const TPair<UObject*, UObject*>& Pair : ReplacedObjects)
+	{
+		if(ObjectToPropertyBagMap.RemoveAndCopyValue(Pair.Key, PropertyBag))
+		{
+			// We may see duplicate bags generated during TPS based duplication of the old object - these can be safely deleted/replaced, although ideally we shouldn't be creating new bags during duplication.
+			if(FPropertyBag* DuplicateBag = ObjectToPropertyBagMap.FindRef(Pair.Value))
+			{
+				UE_LOG(LogPropertyBagRepository, Warning, TEXT("Duplicate property bag detected for %s"), *Pair.Value->GetName());
+				delete DuplicateBag;
+			}
+			//UE_LOG(LogPropertyBagRepository, Log, TEXT("Bag fixup: %s (#%08x) -> %s (#%08x)"), *Pair.Key->GetName(), uint64(Pair.Key), *Pair.Value->GetName(), uint64(Pair.Value));
+			ObjectToPropertyBagMap.Emplace(Pair.Value, PropertyBag);
+		}
+	}
+}
+
+// TODO: Create these by class on construction?
+FPropertyBag* FPropertyBagRepository::CreateOuterBag(UObjectBase* Owner)
+{
+	FPropertyBagRepositoryLock LockRepo(this);
+	FPropertyBag* PropertyBag = FindBagUnsafe(Owner);
+	if(!PropertyBag)
+	{
+		PropertyBag = new FPropertyBag;
+		ObjectToPropertyBagMap.Emplace(Owner, PropertyBag);
+	}
+	return PropertyBag;
+}
+
+// TODO: Remove this? Bag destruction to be handled entirely via UObject::BeginDestroy() (+ FPropertyBagProperty destructor)?
+void FPropertyBagRepository::DestroyOuterBag(UObjectBase* Owner)
+{
+	FPropertyBagRepositoryLock LockRepo(this);
+	UE::FPropertyBag* PropertyBag = nullptr;
+	if(ObjectToPropertyBagMap.RemoveAndCopyValue(Owner, PropertyBag))
+	{
+		delete PropertyBag;
+	}
+}
+
+FPropertyBag* FPropertyBagRepository::FindBag(const UObjectBase* Object)
+{
+	FPropertyBagRepositoryLock LockRepo(this);
+	return FindBagUnsafe(Object);
+}
+
+const FPropertyBag* FPropertyBagRepository::FindBag(const UObjectBase* Object) const
+{
+	FPropertyBagRepositoryLock LockRepo(this);
+	return FindBagUnsafe(Object);
+}
+
+FPropertyBag* FPropertyBagRepository::FindBagUnsafe(const UObjectBase* Object)
+{
+	UE::FPropertyBag* PropertyBag = ObjectToPropertyBagMap.FindRef(Object);
+	// Evaluate property bag support via type traits? During Assign?
+	//UE_CLOG(!PropertyBag || Object->GetClass()->HasSerializeViaPropertyBag(), LogPropertyBagRepository, Warning, TEXT("Object %s PropertyBag is invalid: Doesn't serialize via property bags"), *static_cast<const UObjectBaseUtility*>(Object)->GetPathName());
+	return PropertyBag;
+}
+
+const FPropertyBag* FPropertyBagRepository::FindBagUnsafe(const UObjectBase* Object) const
+{
+	return const_cast<FPropertyBagRepository*>(this)->FindBagUnsafe(Object);
+}
+
+bool FPropertyBagRepository::HasBag(const UObjectBase* Object) const
+{
+	// TODO: Should be consistent across all objects of a given type, so handle via TStructOpsTypeTraits or similar?
+	return FindBag(Object) != nullptr;
+}
+
+// Not sure this is necessary.
+void FPropertyBagRepository::ShrinkMaps()
+{
+	FPropertyBagRepositoryLock LockRepo(this);
+	ObjectToPropertyBagMap.Compact();
+}
+
+} // UE
