Fixed missing include

#rnx

diff --git a/Engine/Source/Runtime/CoreUObject/Public/UObject/FastReferenceCollector.h b/Engine/Source/Runtime/CoreUObject/Public/UObject/FastReferenceCollector.h
--- a/Engine/Source/Runtime/CoreUObject/Public/UObject/FastReferenceCollector.h
+++ b/Engine/Source/Runtime/CoreUObject/Public/UObject/FastReferenceCollector.h
@@ -23,6 +23,10 @@
 #include "VerseVM/VVMWriteBarrier.h"
 #endif
 
+#if ENABLE_GC_HISTORY
+#include "UObject/ReferenceToken.h"
+#endif
+
 /*=============================================================================
 	FastReferenceCollector.h: Unreal realtime garbage collection helpers
 =============================================================================*/
diff --git a/Engine/Source/Runtime/CoreUObject/Public/UObject/GarbageCollectionHistory.h b/Engine/Source/Runtime/CoreUObject/Public/UObject/GarbageCollectionHistory.h
--- a/Engine/Source/Runtime/CoreUObject/Public/UObject/GarbageCollectionHistory.h
+++ b/Engine/Source/Runtime/CoreUObject/Public/UObject/GarbageCollectionHistory.h
@@ -13,6 +13,7 @@
 #include "UObject/GCObjectInfo.h"
 #include "UObject/GarbageCollection.h"
 #include "UObject/NameTypes.h"
+#include "UObject/ReferenceToken.h"
 
 class FGCObjectInfo;
 class UObject;
