Fix warning V570: The variable is assigned to itself.

diff --git a/Engine/Plugins/Runtime/RigVM/Source/RigVM/Public/RigVMCore/RigVM.h b/Engine/Plugins/Runtime/RigVM/Source/RigVM/Public/RigVMCore/RigVM.h
--- a/Engine/Plugins/Runtime/RigVM/Source/RigVM/Public/RigVMCore/RigVM.h
+++ b/Engine/Plugins/Runtime/RigVM/Source/RigVM/Public/RigVMCore/RigVM.h
@@ -821,7 +821,6 @@ public:
 			check(ExternalVariableDef.Property == InExternalVariable.Property);
 			check(bAllowNullMemory || InExternalVariable.Memory != nullptr);
 			
-			ExternalVariableDef.Property = ExternalVariableDef.Property;
 			Context.ExternalVariableRuntimeData.Add(FRigVMExternalVariableRuntimeData(InExternalVariable.Memory));
 		}
 	}
diff --git a/Engine/Source/Runtime/CoreUObject/Private/UObject/ObjectHandleIndex.cpp b/Engine/Source/Runtime/CoreUObject/Private/UObject/ObjectHandleIndex.cpp
--- a/Engine/Source/Runtime/CoreUObject/Private/UObject/ObjectHandleIndex.cpp
+++ b/Engine/Source/Runtime/CoreUObject/Private/UObject/ObjectHandleIndex.cpp
@@ -266,7 +266,7 @@ namespace UE::CoreUObject::Private
 			//ObjectPaths could have been modified when the read lock was released and the write
 			//lock was acquired, so we must check and see if the ObjectPath was added in that window.
 			LocalLockScope.ReleaseReadOnlyLockAndAcquireWriteLock_USE_WITH_CAUTION();
-			PathIndex = PathIndex = PackageData->ObjectDescriptors.IndexOfByPredicate([&ObjectPath](const FObjectDescriptor& Desc)
+			PathIndex = PackageData->ObjectDescriptors.IndexOfByPredicate([&ObjectPath](const FObjectDescriptor& Desc)
 				{
 					return Desc.ObjectPath == ObjectPath;
 
diff --git a/Engine/Source/Runtime/Renderer/Private/RayTracing/RayTracingDebug.cpp b/Engine/Source/Runtime/Renderer/Private/RayTracing/RayTracingDebug.cpp
--- a/Engine/Source/Runtime/Renderer/Private/RayTracing/RayTracingDebug.cpp
+++ b/Engine/Source/Runtime/Renderer/Private/RayTracing/RayTracingDebug.cpp
@@ -1443,7 +1443,6 @@ void FDeferredShadingSceneRenderer::RenderRayTracingDebug(FRDGBuilder& GraphBuil
 	else 
 	{
 		FRDGBufferDesc StatsBufferDesc = FRDGBufferDesc::CreateStructuredDesc(sizeof(uint32) * 2, 1);
-		StatsBufferDesc.Usage = EBufferUsageFlags(StatsBufferDesc.Usage);
 		StatsBuffer = GraphBuilder.CreateBuffer(StatsBufferDesc, TEXT("RayTracingDebug.HitStatsBuffer"));
 
 		FRayTracingDebugHitStatsUniformBufferParameters* DebugHitStatsUniformBufferParameters = GraphBuilder.AllocParameters<FRayTracingDebugHitStatsUniformBufferParameters>();
