Include the verse VM dependencies if WITH_VERSE_VM=WITH_COREUOBJECT is set.
Being in CoreUObject.build.cs should imply that WITH_COREUOBJECT=1, so this should be  equivalent to the existing WITH_VERSE_VM=1 check.
#rnx

diff --git a/Engine/Source/Runtime/CoreUObject/CoreUObject.Build.cs b/Engine/Source/Runtime/CoreUObject/CoreUObject.Build.cs
--- a/Engine/Source/Runtime/CoreUObject/CoreUObject.Build.cs
+++ b/Engine/Source/Runtime/CoreUObject/CoreUObject.Build.cs
@@ -39,7 +39,8 @@ public class CoreUObject : ModuleRules
 
 		// If using the new VM either by default or directly, then add in a dependency to 
 		// the core VerseVM which contains special compile flags not compatible with CoreUObject
-		if (!Target.bUseVerseBPVM || Target.GlobalDefinitions.Contains("WITH_VERSE_VM=1"))
+		if (!Target.bUseVerseBPVM || Target.GlobalDefinitions.Contains("WITH_VERSE_VM=1")
+			|| Target.GlobalDefinitions.Contains("WITH_VERSE_VM=WITH_COREUOBJECT"))
 		{
 			PublicDependencyModuleNames.AddRange(
 				new string[]
