Removing redundant marking paths from GC since all FUObjectItem::Set*/Clear*Flags functions internally redirect to atomic functions (ThisThreadAtomicallyClear/SetFlag)

diff --git a/Engine/Source/Runtime/CoreUObject/Private/UObject/GarbageCollection.cpp b/Engine/Source/Runtime/CoreUObject/Private/UObject/GarbageCollection.cpp
--- a/Engine/Source/Runtime/CoreUObject/Private/UObject/GarbageCollection.cpp
+++ b/Engine/Source/Runtime/CoreUObject/Private/UObject/GarbageCollection.cpp
@@ -961,84 +961,45 @@ static bool MarkClusterMutableObjectsAsReachable(FUObjectCluster& Cluster, Conta
 		{
 			FUObjectItem* ReferencedMutableObjectItem = GUObjectArray.IndexToObjectUnsafeForGC(ReferencedMutableObjectIndex);
 			UE::GC::GStats.IncClusterToObjectRefs(ReferencedMutableObjectItem);
-			if constexpr (IsParallel(Options))
-			{
-				PRAGMA_DISABLE_DEPRECATION_WARNINGS
-				if (!ReferencedMutableObjectItem->HasAnyFlags(EInternalObjectFlags::PendingKill | EInternalObjectFlags::Garbage))
-				PRAGMA_ENABLE_DEPRECATION_WARNINGS
-				{
-					if (ReferencedMutableObjectItem->IsMaybeUnreachable())
-					{
-						if (ReferencedMutableObjectItem->ThisThreadAtomicallyClearedMaybeUnreachable())
-						{
-							// Needs doing because this is either a normal unclustered object (clustered objects are never unreachable) or a cluster root
-							ObjectsToSerialize.Add(static_cast<UObject*>(ReferencedMutableObjectItem->Object));
-
-							// So is this a cluster root maybe?
-							if (ReferencedMutableObjectItem->GetOwnerIndex() < 0)
-							{
-								MarkReferencedClustersAsReachable<Options>(ReferencedMutableObjectItem->GetClusterIndex(), ObjectsToSerialize);
-							}
-						}
-					}
-					else if (ReferencedMutableObjectItem->GetOwnerIndex() > 0 && !ReferencedMutableObjectItem->HasAnyFlags(EInternalObjectFlags::ReachableInCluster))
-					{
-						// This is a clustered object that maybe hasn't been processed yet
-						if (ReferencedMutableObjectItem->ThisThreadAtomicallySetFlag(EInternalObjectFlags::ReachableInCluster))
-						{
-							// Needs doing, we need to get its cluster root and process it too
-							FUObjectItem* ReferencedMutableObjectsClusterRootItem = GUObjectArray.IndexToObjectUnsafeForGC(ReferencedMutableObjectItem->GetOwnerIndex());
-							if (ReferencedMutableObjectsClusterRootItem->IsMaybeUnreachable())
-							{
-								// The root is also maybe unreachable so process it and all the referenced clusters
-								if (ReferencedMutableObjectsClusterRootItem->ThisThreadAtomicallyClearedMaybeUnreachable())
-								{
-									MarkReferencedClustersAsReachable<Options>(ReferencedMutableObjectsClusterRootItem->GetClusterIndex(), ObjectsToSerialize);
-								}
-							}
-						}
-					}
-				}
-				else
-				{
-					// Pending kill support for clusters (multi-threaded case)
-					ReferencedMutableObjectIndex = -1;
-					bAddClusterObjectsToSerialize = true;
-				}
-			}
 			PRAGMA_DISABLE_DEPRECATION_WARNINGS
-			else if (!ReferencedMutableObjectItem->HasAnyFlags(EInternalObjectFlags::PendingKill | EInternalObjectFlags::Garbage))
+			if (!ReferencedMutableObjectItem->HasAnyFlags(EInternalObjectFlags::PendingKill | EInternalObjectFlags::Garbage))
 			PRAGMA_ENABLE_DEPRECATION_WARNINGS
 			{
 				if (ReferencedMutableObjectItem->IsMaybeUnreachable())
 				{
-					// Needs doing because this is either a normal unclustered object (clustered objects are never unreachable) or a cluster root
-					ReferencedMutableObjectItem->ClearFlags(EInternalObjectFlags::MaybeUnreachable);
-					ObjectsToSerialize.Add(static_cast<UObject*>(ReferencedMutableObjectItem->Object));
-						
-					// So is this a cluster root?
-					if (ReferencedMutableObjectItem->GetOwnerIndex() < 0)
+					if (ReferencedMutableObjectItem->ThisThreadAtomicallyClearedMaybeUnreachable())
 					{
-						MarkReferencedClustersAsReachable<Options>(ReferencedMutableObjectItem->GetClusterIndex(), ObjectsToSerialize);
+						// Needs doing because this is either a normal unclustered object (clustered objects are never unreachable) or a cluster root
+						ObjectsToSerialize.Add(static_cast<UObject*>(ReferencedMutableObjectItem->Object));
+
+						// So is this a cluster root maybe?
+						if (ReferencedMutableObjectItem->GetOwnerIndex() < 0)
+						{
+							MarkReferencedClustersAsReachable<Options>(ReferencedMutableObjectItem->GetClusterIndex(), ObjectsToSerialize);
+						}
 					}
 				}
 				else if (ReferencedMutableObjectItem->GetOwnerIndex() > 0 && !ReferencedMutableObjectItem->HasAnyFlags(EInternalObjectFlags::ReachableInCluster))
 				{
-					// This is a clustered object that hasn't been processed yet
-					ReferencedMutableObjectItem->SetFlags(EInternalObjectFlags::ReachableInCluster);
-						
-					// If the root is also unreachable, process it and all its referenced clusters
-					FUObjectItem* ReferencedMutableObjectsClusterRootItem = GUObjectArray.IndexToObjectUnsafeForGC(ReferencedMutableObjectItem->GetOwnerIndex());
-					if (ReferencedMutableObjectsClusterRootItem->IsMaybeUnreachable())
+					// This is a clustered object that maybe hasn't been processed yet
+					if (ReferencedMutableObjectItem->ThisThreadAtomicallySetFlag(EInternalObjectFlags::ReachableInCluster))
 					{
-						ReferencedMutableObjectsClusterRootItem->ClearFlags(EInternalObjectFlags::MaybeUnreachable);
-						MarkReferencedClustersAsReachable<Options>(ReferencedMutableObjectsClusterRootItem->GetClusterIndex(), ObjectsToSerialize);
+						// Needs doing, we need to get its cluster root and process it too
+						FUObjectItem* ReferencedMutableObjectsClusterRootItem = GUObjectArray.IndexToObjectUnsafeForGC(ReferencedMutableObjectItem->GetOwnerIndex());
+						if (ReferencedMutableObjectsClusterRootItem->IsMaybeUnreachable())
+						{
+							// The root is also maybe unreachable so process it and all the referenced clusters
+							if (ReferencedMutableObjectsClusterRootItem->ThisThreadAtomicallyClearedMaybeUnreachable())
+							{
+								MarkReferencedClustersAsReachable<Options>(ReferencedMutableObjectsClusterRootItem->GetClusterIndex(), ObjectsToSerialize);
+							}
+						}
 					}
 				}
 			}
 			else
 			{
-				// Pending kill support for clusters (single-threaded case)
+				// Pending kill support for clusters
 				ReferencedMutableObjectIndex = -1;
 				bAddClusterObjectsToSerialize = true;
 			}
@@ -1091,14 +1052,7 @@ static FORCENOINLINE void MarkReferencedClustersAsReachable(int32 ClusterIndex,
 			{
 				if (ReferencedClusterRootObjectItem->IsMaybeUnreachable())
 				{
-					if constexpr (IsParallel(Options))
-					{
-						ReferencedClusterRootObjectItem->ThisThreadAtomicallyClearedFlag(EInternalObjectFlags::MaybeUnreachable);
-					}
-					else
-					{
-						ReferencedClusterRootObjectItem->ClearFlags(EInternalObjectFlags::MaybeUnreachable);
-					}
+					ReferencedClusterRootObjectItem->ClearMaybeUnreachable();
 				}
 			}
 			else
@@ -2977,28 +2931,10 @@ public:
 				FUObjectItem* RootObjectItem = GUObjectArray.IndexToObjectUnsafeForGC(Metadata.ObjectItem->GetOwnerIndex());
 				checkSlow(RootObjectItem->HasAnyFlags(EInternalObjectFlags::ClusterRoot));
 
-				bool bNeedsDoing = true;
-				if constexpr (IsParallel(Options))
-				{
-					bNeedsDoing = Metadata.ObjectItem->ThisThreadAtomicallySetFlag(EInternalObjectFlags::ReachableInCluster);
-				}
-				else
-				{
-					Metadata.ObjectItem->SetFlags(EInternalObjectFlags::ReachableInCluster);
-				}
-				if (bNeedsDoing)
+				if (Metadata.ObjectItem->ThisThreadAtomicallySetFlag(EInternalObjectFlags::ReachableInCluster))
 				{
-					if constexpr (IsParallel(Options))
-					{
-						if (ClearMaybeUnreachableInterlocked(RootObjectItem->Flags))
-						{
-							// Make sure all referenced clusters are marked as reachable too
-							MarkReferencedClustersAsReachableThunk<Options>(RootObjectItem->GetClusterIndex(), Context.ObjectsToSerialize);
-						}
-					}
-					else if (RootObjectItem->IsMaybeUnreachable())
+					if (ClearMaybeUnreachableInterlocked(RootObjectItem->Flags))
 					{
-						RootObjectItem->ClearFlags(EInternalObjectFlags::MaybeUnreachable);
 						// Make sure all referenced clusters are marked as reachable too
 						MarkReferencedClustersAsReachableThunk<Options>(RootObjectItem->GetClusterIndex(), Context.ObjectsToSerialize);
 					}
@@ -4015,7 +3951,7 @@ public:
 						}
 						else
 						{
-							ObjectItem->SetFlags(EInternalObjectFlags::MaybeUnreachable);
+							ObjectItem->SetMaybeUnreachable();
 						}
 					}					
 				}
@@ -4070,9 +4006,9 @@ public:
 						FUObjectItem* RootObjectItem = GUObjectArray.IndexToObjectUnsafeForGC(OwnerIndex);
 						checkSlow(RootObjectItem->HasAnyFlags(EInternalObjectFlags::ClusterRoot));
 						// if it is reachable via keep flags we will do this below (or maybe already have)
-						if (RootObjectItem->HasAnyFlags(EInternalObjectFlags::MaybeUnreachable)) 
+						if (RootObjectItem->IsMaybeUnreachable()) 
 						{
-							RootObjectItem->ClearFlags(EInternalObjectFlags::MaybeUnreachable);
+							RootObjectItem->ClearMaybeUnreachable();
 							// Make sure all referenced clusters are marked as reachable too
 							MarkReferencedClustersAsReachable<EGCOptions::None>(RootObjectItem->GetClusterIndex(), InitialObjects);
 						}
diff --git a/Engine/Source/Runtime/CoreUObject/Public/UObject/UObjectArray.h b/Engine/Source/Runtime/CoreUObject/Public/UObject/UObjectArray.h
--- a/Engine/Source/Runtime/CoreUObject/Public/UObject/UObjectArray.h
+++ b/Engine/Source/Runtime/CoreUObject/Public/UObject/UObjectArray.h
@@ -197,7 +197,7 @@ struct
 	}
 	FORCEINLINE bool IsMaybeUnreachable() const
 	{
-		return !!(Flags & int32(EInternalObjectFlags::MaybeUnreachable));
+		return !!(GetFlagsInternal() & int32(EInternalObjectFlags::MaybeUnreachable));
 	}
 	FORCEINLINE bool ThisThreadAtomicallyClearedRFUnreachable()
 	{
