#UE Temporary logging for CDO creation

diff --git a/Engine/Source/Runtime/CoreUObject/Private/UObject/Class.cpp b/Engine/Source/Runtime/CoreUObject/Private/UObject/Class.cpp
--- a/Engine/Source/Runtime/CoreUObject/Private/UObject/Class.cpp
+++ b/Engine/Source/Runtime/CoreUObject/Private/UObject/Class.cpp
@@ -4342,6 +4342,8 @@ UObject* UClass::CreateDefaultObject()
 {
 	if ( ClassDefaultObject == NULL )
 	{
+		UE_LOG(LogClass, Display, TEXT("Creating CDO for %s"), *GetPathName());
+
 		ensureMsgf(!bLayoutChanging, TEXT("Class named %s creating its CDO while changing its layout"), *GetName());
 
 		UClass* ParentClass = GetSuperClass();
