#UE Move temporary logging

diff --git a/Engine/Source/Runtime/CoreUObject/Private/UObject/Class.cpp b/Engine/Source/Runtime/CoreUObject/Private/UObject/Class.cpp
--- a/Engine/Source/Runtime/CoreUObject/Private/UObject/Class.cpp
+++ b/Engine/Source/Runtime/CoreUObject/Private/UObject/Class.cpp
@@ -4342,8 +4342,6 @@ UObject* UClass::CreateDefaultObject()
 {
 	if ( ClassDefaultObject == NULL )
 	{
-		UE_LOG(LogClass, Display, TEXT("Creating CDO for %s"), *GetPathName());
-
 		ensureMsgf(!bLayoutChanging, TEXT("Class named %s creating its CDO while changing its layout"), *GetName());
 
 		UClass* ParentClass = GetSuperClass();
@@ -4405,6 +4403,7 @@ UObject* UClass::CreateDefaultObject()
 				EObjectInitializerOptions InitOptions = EObjectInitializerOptions::None;
 				if (!HasAnyClassFlags(CLASS_Native | CLASS_Intrinsic))
 				{
+					UE_LOG(LogClass, Display, TEXT("Creating non-native non-intrinsic CDO for %s"), *GetPathName());
 					// Blueprint CDOs have their properties always initialized.
 					InitOptions |= EObjectInitializerOptions::InitializeProperties;
 				}
