Fix ClearReturnValue destroying the value before clearing it, which assumes that the value being cleared is an initialized value.

diff --git a/Engine/Source/Runtime/CoreUObject/Private/UObject/ScriptCore.cpp b/Engine/Source/Runtime/CoreUObject/Private/UObject/ScriptCore.cpp
--- a/Engine/Source/Runtime/CoreUObject/Private/UObject/ScriptCore.cpp
+++ b/Engine/Source/Runtime/CoreUObject/Private/UObject/ScriptCore.cpp
@@ -1158,10 +1158,7 @@ void ClearReturnValue(FProperty* ReturnProp, RESULT_DECL)
 		uint8* Data = (uint8*)RESULT_PARAM;
 		for (int32 ArrayIdx = 0; ArrayIdx < ReturnProp->ArrayDim; ArrayIdx++, Data += ReturnProp->ElementSize)
 		{
-			// destroy old value if necessary
-			ReturnProp->DestroyValue(Data);
-
-			// copy zero value for return property into Result, or default construct as necessary
+			// Clear the property. This assumes that it has already been initialized, and that the caller will destroy it.
 			ReturnProp->ClearValue(Data);
 		}
 	}
