Save and ResetLoaders only flush loading requests specific to the package when there's already pending async requests

diff --git a/Engine/Source/Runtime/CoreUObject/Private/UObject/Linker.cpp b/Engine/Source/Runtime/CoreUObject/Private/UObject/Linker.cpp
--- a/Engine/Source/Runtime/CoreUObject/Private/UObject/Linker.cpp
+++ b/Engine/Source/Runtime/CoreUObject/Private/UObject/Linker.cpp
@@ -472,51 +472,6 @@ PRAGMA_ENABLE_DEPRECATION_WARNINGS
 	Global functions
 -----------------------------------------------------------------------------*/
 
-void ResetLinkerExports(UPackage* InPackage)
-{
-	FLinkerManager::Get().ResetLinkerExports(InPackage);
-}
-
-void ResetLoaders(UObject* InPkg)
-{
-	// This call into IsAsyncLoading calls an atomic, which we want to avoid in a transact and this log is not something we want to roll back so ignore
-	UE_AUTORTFM_OPEN({
-		if (IsAsyncLoading())
-		{
-			UE_LOG(LogLinker, Log, TEXT("ResetLoaders(%s) is flushing async loading"), *GetPathNameSafe(InPkg));
-		}
-	});
-
-	// Make sure we're not in the middle of loading something in the background.
-	FlushAsyncLoading();
-	FLinkerManager::Get().ResetLoaders(InPkg);
-}
-
-void ResetLoaders(TArrayView<UObject*> InOuters)
-{
-	// This call into IsAsyncLoading calls an atomic, which we want to avoid in a transact and this log is not something we want to roll back so ignore
-	UE_AUTORTFM_OPEN({
-		if (IsAsyncLoading())
-		{
-			UE_LOG(LogLinker, Log, TEXT("ResetLoaders is flushing async loading"));
-		}
-	});
-
-	// Make sure we're not in the middle of loading something in the background.
-	FlushAsyncLoading();
-	FLinkerManager::Get().ResetLoaders(InOuters);
-}
-
-void DeleteLoaders()
-{
-	FLinkerManager::Get().DeleteLinkers();
-}
-
-void DeleteLoader(FLinkerLoad* Loader)
-{
-	FLinkerManager::Get().RemoveLinker(Loader);
-}
-
 static void LogGetPackageLinkerError(FUObjectSerializeContext* LoadContext, const FPackagePath& PackagePath, const FText& InErrorMessage, UObject* InOuter, uint32 LoadFlags)
 {
 	static FName NAME_LoadErrors("LoadErrors");
@@ -925,15 +880,84 @@ FLinkerLoad* LoadPackageLinker(UPackage* InOuter, const TCHAR* InLongPackageName
 	return LoadPackageLinker(InOuter, GetPackagePath(InOuter, InLongPackageName), LoadFlags, Sandbox, InReaderOverride, [](FLinkerLoad* InLinker) {});
 }
 
+
+void ResetLinkerExports(UPackage* InPackage)
+{
+	FLinkerManager::Get().ResetLinkerExports(InPackage);
+}
+
+void ConditionalFlushAsyncLoadingForLinkers(TConstArrayView<FLinkerLoad*> InLinkers)
+{
+	// If there are currently pending async requests for any package, 
+	// it's possible that some of those pending requests are for the given Linkers that our caller wants to flush.
+	// Our caller wants to flush the Linkers because we must purge references to a linker from the async loader before the linker can be reset.
+	// But we don't have a way to inspect the pending async requests and see which linker they are for. 
+	// Flushing the new requests will also flush any other pending async requests for the same linkers.
+	if (InLinkers.Num() > 0 && IsAsyncLoading())
+	{
+		UE_LOG(LogLinker, Log, TEXT("Conditionally flushing loading for linker(s) (%s)"), *GetPathNameSafe(InLinkers[0]->LinkerRoot));
+
+		TArray<int32, TInlineAllocator<4>> RequestIds;
+		for (FLinkerLoad* Linker : InLinkers)
+		{
+			int32 Request = LoadPackageAsync(Linker->GetPackagePath(), FLoadPackageAsyncOptionalParams{ .PackagePriority = MAX_int32 });
+			RequestIds.Add(Request);
+		}
+		FlushAsyncLoading(RequestIds);
+	}
+}
+
+void ResetLoaders(UObject* InPkg)
+{
+	if (FLinkerLoad* Loader = FLinkerLoad::FindExistingLinkerForPackage(InPkg->GetPackage()))
+	{
+		// Make sure we're not in the middle of loading something in the background.
+		ConditionalFlushAsyncLoadingForLinkers(MakeArrayView({ Loader }));
+		// Detach all exports from the linker and dissociate the linker.
+		FLinkerManager::Get().ResetLoaders(MakeArrayView({ Loader }));
+	}
+}
+
+void ResetLoaders(TArrayView<UObject*> InOuters)
+{
+	TSet<FLinkerLoad*> LinkersToReset;
+	for (UObject* Object : InOuters)
+	{
+		if (UPackage* TopLevelPackage = Object->GetPackage())
+		{
+			if (FLinkerLoad* LinkerToReset = FLinkerLoad::FindExistingLinkerForPackage(TopLevelPackage))
+			{
+				LinkersToReset.Add(LinkerToReset);
+			}
+		}
+	}
+
+	if (LinkersToReset.Num())
+	{
+		// Make sure we're not in the middle of loading something in the background.
+		ConditionalFlushAsyncLoadingForLinkers(LinkersToReset.Array());
+		FLinkerManager::Get().ResetLoaders(LinkersToReset);
+	}
+}
+
+void ConditionalFlushAsyncLoadingForSave(UPackage* InPackage)
+{
+	if (FLinkerLoad* Loader = FLinkerLoad::FindExistingLinkerForPackage(InPackage->GetPackage()))
+	{
+		ConditionalFlushAsyncLoadingForLinkers(MakeArrayView({ Loader }));
+	}
+}
+
 void ResetLoadersForSave(UPackage* Package, const TCHAR* Filename)
 {
 	FLinkerLoad* Loader = FLinkerLoad::FindExistingLinkerForPackage(Package);
-	if( Loader )
+	if (Loader)
 	{
-		if ( FPackagePath::FromLocalPath(Filename) == Loader->GetPackagePath())
+		if (FPackagePath::FromLocalPath(Filename) == Loader->GetPackagePath())
 		{
 			// Detach all exports from the linker and dissociate the linker.
-			ResetLoaders( Package );
+			ConditionalFlushAsyncLoadingForLinkers(MakeArrayView({ Loader }));
+			FLinkerManager::Get().ResetLoaders(MakeArrayView({ Loader }));
 		}
 	}
 }
@@ -951,8 +975,22 @@ void ResetLoadersForSave(TArrayView<FPackageSaveInfo> InPackages)
 		{
 			return FLinkerLoad::FindExistingLinkerForPackage(InPackageSaveInfo.Package);
 		});
-	FlushAsyncLoading();
-	FLinkerManager::Get().ResetLoaders(LinkersToReset);
+
+	if (LinkersToReset.Num())
+	{
+		ConditionalFlushAsyncLoadingForLinkers(LinkersToReset.Array());
+		FLinkerManager::Get().ResetLoaders(LinkersToReset);
+	}
+}
+
+void DeleteLoaders()
+{
+	FLinkerManager::Get().DeleteLinkers();
+}
+
+void DeleteLoader(FLinkerLoad* Loader)
+{
+	FLinkerManager::Get().RemoveLinker(Loader);
 }
 
 void EnsureLoadingComplete(UPackage* Package)
diff --git a/Engine/Source/Runtime/CoreUObject/Private/UObject/LinkerManager.cpp b/Engine/Source/Runtime/CoreUObject/Private/UObject/LinkerManager.cpp
--- a/Engine/Source/Runtime/CoreUObject/Private/UObject/LinkerManager.cpp
+++ b/Engine/Source/Runtime/CoreUObject/Private/UObject/LinkerManager.cpp
@@ -179,22 +179,10 @@ void FLinkerManager::ResetLoaders(UObject* InPkg)
 	}
 }
 
-void FLinkerManager::ResetLoaders(TArrayView<UObject*> InPackages)
+void FLinkerManager::ResetLoaders(TConstArrayView<FLinkerLoad*> InLinkerLoad)
 {
-	TSet<FLinkerLoad*> LinkerLoads;
-	LinkerLoads.Reserve(InPackages.Num());
-
-	for (UObject* Object : InPackages)
-	{
-		if (UPackage* TopLevelPackage = Object->GetPackage())
-		{
-			if (FLinkerLoad* LinkerToReset = FLinkerLoad::FindExistingLinkerForPackage(TopLevelPackage))
-			{
-				LinkerLoads.Add(LinkerToReset);
-			}
-		}
-	}
-
+	TSet<FLinkerLoad*> LinkerLoads; 
+	LinkerLoads.Append(InLinkerLoad);
 	ResetLoaders(LinkerLoads);
 }
 
diff --git a/Engine/Source/Runtime/CoreUObject/Private/UObject/LinkerManager.h b/Engine/Source/Runtime/CoreUObject/Private/UObject/LinkerManager.h
--- a/Engine/Source/Runtime/CoreUObject/Private/UObject/LinkerManager.h
+++ b/Engine/Source/Runtime/CoreUObject/Private/UObject/LinkerManager.h
@@ -114,7 +114,7 @@ public:
 	void ResetLoaders(UObject* InPkg);
 
 	/** Empty the loaders */
-	void ResetLoaders(TArrayView<UObject*> InPackages);
+	void ResetLoaders(TConstArrayView<FLinkerLoad*> InLinkerLoad);
 
 	/** Empty the loaders from the specified set */
 	void ResetLoaders(const TSet<FLinkerLoad*>& InLinkerLoadSet);
diff --git a/Engine/Source/Runtime/CoreUObject/Private/UObject/SavePackage2.cpp b/Engine/Source/Runtime/CoreUObject/Private/UObject/SavePackage2.cpp
--- a/Engine/Source/Runtime/CoreUObject/Private/UObject/SavePackage2.cpp
+++ b/Engine/Source/Runtime/CoreUObject/Private/UObject/SavePackage2.cpp
@@ -2457,14 +2457,6 @@ ESavePackageResult FinalizeFile(FStructuredArchive::FRecord& StructuredArchiveRo
 {
 	SCOPED_SAVETIMER(UPackage_Save_FinalizeFile);
 
-	// In the concurrent case, it is called right after routing presave so it can be done in batch before going concurrent
-	if (!SaveContext.IsConcurrent())
-	{
-		// If we're writing to the existing file call ResetLoaders on the Package so that we drop the handle to the file on disk and can write to it
-		//COOK_STAT(FScopedDurationTimer SaveTimer(FSavePackageStats::ResetLoadersTimeSec));
-		ResetLoadersForSave(SaveContext.GetPackage(), SaveContext.GetFilename());
-	}
-
 	FSavePackageContext* SavePackageContext = SaveContext.GetSavePackageContext();
 	IPackageWriter* PackageWriter = SavePackageContext ? SavePackageContext->PackageWriter : nullptr;
 	if (PackageWriter || SaveContext.IsSaveToMemory())
@@ -3051,21 +3043,20 @@ FSavePackageResultStruct UPackage::Save2(UPackage* InPackage, UObject* InAsset,
 
 	// Ensures
 	SlowTask.EnterProgressFrame();
-	EnsurePackageLocalization(SaveContext.GetPackage());
 	{
-		// FullyLoad the package's Loader, so that anything we need to serialize (bulkdata, thumbnails) is available
-		//COOK_STAT(FScopedDurationTimer SaveTimer(FSavePackageStats::FullyLoadLoadersTimeSec));
-		EnsureLoadingComplete(SaveContext.GetPackage());
-
 		if (!SaveContext.IsConcurrent())
 		{
-			// We need to fulfill all pending streaming and async loading requests to then allow us to lock the global IO manager. 
-			// The latter implies flushing all file handles which is a pre-requisite of saving a package. The code basically needs 
-			// to be sure that we are not reading from a file that is about to be overwritten and that there is no way we might 
-			// start reading from the file till we are done overwriting it.
-			FlushAsyncLoading();
+			// We need to make sure to flush any pending request that may involve the existing linker of this package as we want to reset it
+			// to release any handle on the file prior to overwriting it.
+			UPackage* Package = SaveContext.GetPackage();
+			ConditionalFlushAsyncLoadingForSave(Package);
+			(*GFlushStreamingFunc)();
+
+			EnsurePackageLocalization(Package);
+
+			// FullyLoad the package's Loader, so that anything we need to serialize (bulkdata, thumbnails) is available
+			EnsureLoadingComplete(Package);
 		}
-		(*GFlushStreamingFunc)();
 	}
 
 	// PreSave Asset
@@ -3079,26 +3070,33 @@ FSavePackageResultStruct UPackage::Save2(UPackage* InPackage, UObject* InAsset,
 		SaveContext.SetPreSaveCleanup(ObjectSaveContext.bCleanupRequired);
 	}
 
-	// Route Presave only if not calling concurrently or if the PackageWriter claims already completed, in those case they should be handled separately already
 	SlowTask.EnterProgressFrame();
-	IPackageWriter* PackageWriter = SaveContext.GetPackageWriter();
-	if (!SaveContext.IsConcurrent() && (!PackageWriter || !PackageWriter->IsPreSaveCompleted()))
+	if (!SaveContext.IsConcurrent())
 	{
-		SaveContext.Result = RoutePresave(SaveContext);
+		// Route Presave only if not calling concurrently or if the PackageWriter claims already completed, in those case they should be handled separately already
+		IPackageWriter* PackageWriter = SaveContext.GetPackageWriter();
+		if (!PackageWriter || !PackageWriter->IsPreSaveCompleted())
+		{
+			SaveContext.Result = RoutePresave(SaveContext);
+			if (SaveContext.Result != ESavePackageResult::Success)
+			{
+				return SaveContext.Result;
+			}
+		}
+
+		// Trigger platform cooked data caching after PreSave but before package harvesting 
+		// After PreSave because objects can be created during PreSave and we need to cache them
+		// Before package harvesting because it might modify some property and hence affect the harvested
+		// property name of a tagged property for example
+		SaveContext.Result = BeginCachePlatformCookedData(SaveContext);
 		if (SaveContext.Result != ESavePackageResult::Success)
 		{
 			return SaveContext.Result;
 		}
-	}
 
-	// Trigger platform cooked data caching after PreSave but before package harvesting 
-	// After PreSave because objects can be created during PreSave and we need to cache them
-	// Before package harvesting because it might modify some property and hence affect the harvested
-	// property name of a tagged property for example
-	SaveContext.Result = BeginCachePlatformCookedData(SaveContext);
-	if (SaveContext.Result != ESavePackageResult::Success)
-	{
-		return SaveContext.Result;
+		// If we're writing to the existing file call ResetLoaders on the Package so that we drop the handle to the file on disk and can write to it
+		// This might end flushing async loading for this package
+		ResetLoadersForSave(SaveContext.GetPackage(), SaveContext.GetFilename());
 	}
 
 	SlowTask.EnterProgressFrame();
diff --git a/Engine/Source/Runtime/CoreUObject/Public/UObject/Linker.h b/Engine/Source/Runtime/CoreUObject/Public/UObject/Linker.h
--- a/Engine/Source/Runtime/CoreUObject/Public/UObject/Linker.h
+++ b/Engine/Source/Runtime/CoreUObject/Public/UObject/Linker.h
@@ -656,26 +656,6 @@ typedef uint32 ELazyLoaderFlags;
 	Global functions
 -----------------------------------------------------------------------------*/
 
-/**
- * Reset the linker exports associated with the package
- * @note, this might flush async loading if the linker is owned by the loading thread
- */
-COREUOBJECT_API void ResetLinkerExports(UPackage* InPackage);
-
-/**
- * Remove references to the linker for the given package and delete the linker. 
- * Can be called after the package has finished loading.
- * Flushes async loading.
- */
-COREUOBJECT_API void ResetLoaders(UObject* InOuter);
-COREUOBJECT_API void ResetLoaders(TArrayView<UObject*> InOuters);
-
-/** Deletes all linkers that have finished loading */
-COREUOBJECT_API void DeleteLoaders();
-
-/** Queues linker for deletion */
-COREUOBJECT_API void DeleteLoader(FLinkerLoad* Loader);
-
 /** 
  * Loads a linker for a package and returns it without loading any objects.
  * @param InOuter Package if known, can be null
@@ -713,9 +693,29 @@ COREUOBJECT_API FLinkerLoad* GetPackageLinker(UPackage* InOuter, const FPackageP
 UE_DEPRECATED(5.0, "Use version that takes a FPackagePath without a FGuid instead")
 COREUOBJECT_API FLinkerLoad* GetPackageLinker(UPackage* InOuter, const TCHAR* InLongPackageName, uint32 LoadFlags, UPackageMap* Sandbox, FGuid* CompatibleGuid, FArchive* InReaderOverride = nullptr, FUObjectSerializeContext** InOutLoadContext = nullptr, FLinkerLoad* ImportLinker = nullptr, const FLinkerInstancingContext* InstancingContext = nullptr);
 
+COREUOBJECT_API FString GetPrestreamPackageLinkerName(const TCHAR* InLongPackageName, bool bSkipIfExists = true);
+
+/**
+ * Reset the linker exports associated with the package
+ * @note, this might flush async loading if the linker is owned by the loading thread
+ */
+COREUOBJECT_API void ResetLinkerExports(UPackage* InPackage);
 
+/**
+ * Remove references to the linker for the given package and delete the linker.
+ * Can be called after the package has finished loading.
+ * Flushes async loading.
+ */
+COREUOBJECT_API void ResetLoaders(UObject* InOuter);
+COREUOBJECT_API void ResetLoaders(TArrayView<UObject*> InOuters);
 
-COREUOBJECT_API FString GetPrestreamPackageLinkerName(const TCHAR* InLongPackageName, bool bSkipIfExists = true);
+
+/**
+ *  Conditionally flush async loading for a specific package if there's any pending async requests
+ * 
+ * @param InPackage		The package to flush for
+ */
+COREUOBJECT_API void ConditionalFlushAsyncLoadingForSave(UPackage* InPackage);
 
 /**
  *
@@ -734,6 +734,12 @@ COREUOBJECT_API void ResetLoadersForSave(UPackage* Package, const TCHAR* Filenam
  */
 COREUOBJECT_API void ResetLoadersForSave(TArrayView<FPackageSaveInfo> InPackages);
 
+/** Deletes all linkers that have finished loading */
+COREUOBJECT_API void DeleteLoaders();
+
+/** Queues linker for deletion */
+COREUOBJECT_API void DeleteLoader(FLinkerLoad* Loader);
+
 /*
  * Ensure all data that can be loaded from the linker (thumbnails, bulk data) is loaded, in preparation for saving out the given package
  *
