#UE Remove temporary logging

diff --git a/Engine/Source/Runtime/CoreUObject/Private/UObject/Class.cpp b/Engine/Source/Runtime/CoreUObject/Private/UObject/Class.cpp
--- a/Engine/Source/Runtime/CoreUObject/Private/UObject/Class.cpp
+++ b/Engine/Source/Runtime/CoreUObject/Private/UObject/Class.cpp
@@ -4403,7 +4403,6 @@ UObject* UClass::CreateDefaultObject()
 				EObjectInitializerOptions InitOptions = EObjectInitializerOptions::None;
 				if (!HasAnyClassFlags(CLASS_Native | CLASS_Intrinsic))
 				{
-					UE_LOG(LogClass, Display, TEXT("Creating non-native non-intrinsic CDO for %s"), *GetPathName());
 					// Blueprint CDOs have their properties always initialized.
 					InitOptions |= EObjectInitializerOptions::InitializeProperties;
 				}
