Optimized UObject hash's ClassToObjectListMap

diff --git a/Engine/Source/Runtime/CoreUObject/Private/UObject/UObjectHash.cpp b/Engine/Source/Runtime/CoreUObject/Private/UObject/UObjectHash.cpp
--- a/Engine/Source/Runtime/CoreUObject/Private/UObject/UObjectHash.cpp
+++ b/Engine/Source/Runtime/CoreUObject/Private/UObject/UObjectHash.cpp
@@ -14,6 +14,11 @@
 #include "HAL/LowLevelMemStats.h"
 #include "UObject/AnyPackagePrivate.h"
 
+#if UE_STORE_OBJECT_LIST_INTERNAL_INDEX
+static_assert(sizeof(FName) == 4, "Internal object index optimization depends exploits 4 bytes padding after the FName");
+static_assert(sizeof(UObjectBase) == 40, "UObjectBase size has changed!");
+#endif
+
 DEFINE_LOG_CATEGORY_STATIC(LogUObjectHash, Log, All);
 
 DECLARE_CYCLE_STAT( TEXT( "GetObjectsOfClass" ), STAT_Hash_GetObjectsOfClass, STATGROUP_UObjectHash );
@@ -174,7 +179,7 @@ struct FHashBucket
 		return !!ElementsOrSetPtr[0] + !!ElementsOrSetPtr[1];
 	}
 	/** Returns the amount of memory allocated for and by Items TSet */
-	FORCEINLINE SIZE_T GetItemsSize() const
+	FORCEINLINE SIZE_T GetAllocatedSize() const
 	{
 		const TSet<UObjectBase*>* Items = GetSet();
 		if (Items)
@@ -183,7 +188,7 @@ struct FHashBucket
 		}
 		return 0;
 	}
-	void Compact()
+	void Shrink()
 	{
 		TSet<UObjectBase*>* Items = GetSet();
 		if (Items)
@@ -191,6 +196,8 @@ struct FHashBucket
 			Items->Compact();
 		}
 	}
+	FORCEINLINE struct FHashBucketIterator CreateIterator();
+
 private:
 	/** Gets an iterator for the TSet in this bucket or for the EmptyBucker if Items is null */
 	FORCEINLINE TSet<UObjectBase*>::TIterator GetIteratorForSet()
@@ -262,13 +269,18 @@ struct FHashBucketIterator
 	}
 };
 
+FORCEINLINE FHashBucketIterator FHashBucket::CreateIterator()
+{
+	return FHashBucketIterator(*this);
+}
+
 /**
 * Wrapper around a TMap with FHashBucket values that supports read only locks
 */
-template <typename T>
-class TBucketMap : private TMap<T, FHashBucket>
+template <typename T, typename K = FHashBucket>
+class TBucketMap : private TMap<T, K>
 {
-	typedef TMap<T, FHashBucket> Super;
+	typedef TMap<T, K> Super;
 #if !UE_BUILD_SHIPPING
 	int32 ReadOnlyLock = 0;
 #endif // !UE_BUILD_SHIPPING
@@ -319,7 +331,7 @@ public:
 		Super::Remove(Key);
 	}
 
-	FORCEINLINE FHashBucket& FindOrAdd(const T& Key)
+	FORCEINLINE K& FindOrAdd(const T& Key)
 	{
 #if !UE_BUILD_SHIPPING
 		UE_CLOG(ReadOnlyLock != 0, LogObj, Fatal, TEXT("Trying to modify UObject map (FindOrAdd) that is currently being iterated. Please make sure you're not creating new UObjects or Garbage Collecting while iterating UObject hash tables."));
@@ -327,7 +339,7 @@ public:
 		return Super::FindOrAdd(Key);
 	}
 
-	FORCEINLINE FHashBucket& FindOrAdd(T&& Key)
+	FORCEINLINE K& FindOrAdd(T&& Key)
 	{
 #if !UE_BUILD_SHIPPING
 		UE_CLOG(ReadOnlyLock != 0, LogObj, Fatal, TEXT("Trying to modify UObject map (FindOrAdd) that is currently being iterated. Please make sure you're not creating new UObjects or Garbage Collecting while iterating UObject hash tables."));
@@ -367,7 +379,11 @@ public:
 
 	/** Map of object to their outers, used to avoid an object iterator to find such things. **/
 	TBucketMap<UObjectBase*> ObjectOuterMap;
+#if UE_STORE_OBJECT_LIST_INTERNAL_INDEX
+	TBucketMap<UClass*, TArray<UObjectBase*> > ClassToObjectListMap;
+#else
 	TBucketMap<UClass*> ClassToObjectListMap;
+#endif
 	TMap<UClass*, TSet<UClass*> > ClassToChildListMap;
 	TAtomic<uint64> ClassToChildListMapVersion;
 
@@ -394,7 +410,7 @@ public:
 				Hash.Compact();
 				for (auto& Pair : Hash)
 				{
-					Pair.Value.Compact();
+					Pair.Value.Shrink();
 				}
 				break;
 			case 1:
@@ -404,28 +420,28 @@ public:
 				ObjectOuterMap.Compact();
 				for (auto& Pair : ObjectOuterMap)
 				{
-					Pair.Value.Compact();
+					Pair.Value.Shrink();
 				}
 				break;
 			case 3:
 				ClassToObjectListMap.Compact();
 				for (auto& Pair : ClassToObjectListMap)
 				{
-					Pair.Value.Compact();
+					Pair.Value.Shrink();
 				}
 				break;
 			case 4:
 				ClassToChildListMap.Compact();
 				for (auto& Pair : ClassToChildListMap)
 				{
-					Pair.Value.Compact();
+					Pair.Value.Shrink();
 				}
 				break;
 			case 5:
 				PackageToObjectListMap.Compact();
 				for (auto& Pair : PackageToObjectListMap)
 				{
-					Pair.Value.Compact();
+					Pair.Value.Shrink();
 				}
 				break;
 			case 6:
@@ -951,12 +967,14 @@ FORCEINLINE static void AddToOuterMap(FUObjectHashTables& ThreadHash, UObjectBas
 }
 
 // Assumes that ThreadHash's critical is already locked
-FORCEINLINE static void AddToClassMap(FUObjectHashTables& ThreadHash, UObjectBase* Object)
+FORCEINLINE void AddToClassMap(FUObjectHashTables& ThreadHash, UObjectBase* Object)
 {
 	{
 		check(Object->GetClass());
-		FHashBucket& ObjectList = ThreadHash.ClassToObjectListMap.FindOrAdd(Object->GetClass());
-		ObjectList.Add(Object);
+#if UE_STORE_OBJECT_LIST_INTERNAL_INDEX
+		Object->ObjectListInternalIndex =
+#endif
+		ThreadHash.ClassToObjectListMap.FindOrAdd(Object->GetClass()).Add(Object);
 	}
 
 	UObjectBaseUtility* ObjectWithUtility = static_cast<UObjectBaseUtility*>(Object);
@@ -1012,18 +1030,24 @@ FORCEINLINE static void RemoveFromOuterMap(FUObjectHashTables& ThreadHash, UObje
 }
 
 // Assumes that ThreadHash's critical is already locked
-FORCEINLINE static void RemoveFromClassMap(FUObjectHashTables& ThreadHash, UObjectBase* Object)
+FORCEINLINE void RemoveFromClassMap(FUObjectHashTables& ThreadHash, UObjectBase* Object)
 {
 	UObjectBaseUtility* ObjectWithUtility = static_cast<UObjectBaseUtility*>(Object);
 
 	{
-		FHashBucket& ObjectList = ThreadHash.ClassToObjectListMap.FindOrAdd(Object->GetClass());
+		auto& ObjectList = ThreadHash.ClassToObjectListMap.FindOrAdd(Object->GetClass());
+#if UE_STORE_OBJECT_LIST_INTERNAL_INDEX
+		ObjectList.Last()->ObjectListInternalIndex = Object->ObjectListInternalIndex;
+		ObjectList[Object->ObjectListInternalIndex] = ObjectList.Last();
+		ObjectList.Pop(false);
+#else
 		int32 NumRemoved = ObjectList.Remove(Object);
 
 		if (!LIKELY(NumRemoved == 1))
 		{
 			OnHashFailure((UObjectBaseUtility*)Object, TEXT("ClassMap"), TEXT("remove miscount"));
 		}
+#endif
 
 		if (!ObjectList.Num())
 		{
@@ -1417,10 +1441,10 @@ FORCEINLINE void ForEachObjectOfClasses_Implementation(FUObjectHashTables& Threa
 
 	for (const UClass* SearchClass : ClassesToLookFor)
 	{
-		FHashBucket* List = ThreadHash.ClassToObjectListMap.Find(SearchClass);
+		auto List = ThreadHash.ClassToObjectListMap.Find(SearchClass);
 		if (List)
 		{
-			for (FHashBucketIterator ObjectIt(*List); ObjectIt; ++ObjectIt)
+			for (auto ObjectIt = List->CreateIterator(); ObjectIt; ++ObjectIt)
 			{
 				UObject* Object = static_cast<UObject*>(*ObjectIt);
 				if (!Object->HasAnyFlags(ExcludeFlags) && !Object->HasAnyInternalFlags(ExclusionInternalFlags))
@@ -1505,10 +1529,10 @@ bool ClassHasInstancesAsyncLoading(const UClass* ClassToLookFor)
 
 	for (const UClass* SearchClass : ClassesToSearch)
 	{
-		FHashBucket* List = ThreadHash.ClassToObjectListMap.Find(SearchClass);
+		auto List = ThreadHash.ClassToObjectListMap.Find(SearchClass);
 		if (List)
 		{
-			for (FHashBucketIterator ObjectIt(*List); ObjectIt; ++ObjectIt)
+			for (auto ObjectIt = List->CreateIterator(); ObjectIt; ++ObjectIt)
 			{
 				UObject *Object = static_cast<UObject*>(*ObjectIt);
 				if (Object->HasAnyInternalFlags(EInternalObjectFlags::AsyncLoading))
@@ -1815,7 +1839,7 @@ void LogHashStatisticsInternal(TBucketMap<int32>& Hash, FOutputDevice& Ar, const
 	// Calculate the size of a all Allocations inside of the buckets (TSet Items)
 	for (auto& Pair : Hash)
 	{
-		HashtableAllocatedSize += Pair.Value.GetItemsSize();
+		HashtableAllocatedSize += Pair.Value.GetAllocatedSize();
 	}
 	Ar.Logf(TEXT("Total memory allocated for and by Object Hash: %" SIZE_T_FMT " bytes."), HashtableAllocatedSize);
 }
@@ -1842,7 +1866,7 @@ void LogHashOuterStatistics(FOutputDevice& Ar, const bool bShowHashBucketCollisi
 	SIZE_T HashOuterMapSize = 0;
 	for (TPair<UObjectBase*, FHashBucket>& OuterMapEntry : FUObjectHashTables::Get().ObjectOuterMap)
 	{
-		HashOuterMapSize += OuterMapEntry.Value.GetItemsSize();
+		HashOuterMapSize += OuterMapEntry.Value.GetAllocatedSize();
 	}
 	Ar.Logf(TEXT("Total memory allocated for Object Outer Map: %" SIZE_T_FMT " bytes."), HashOuterMapSize);
 	Ar.Logf(TEXT(""));
@@ -1862,7 +1886,7 @@ void LogHashMemoryOverheadStatistics(FOutputDevice& Ar, const bool bShowIndividu
 		SIZE_T Size = HashTables.Hash.GetAllocatedSize();
 		for (const TPair<int32, FHashBucket>& Pair : HashTables.Hash)
 		{
-			Size += Pair.Value.GetItemsSize();
+			Size += Pair.Value.GetAllocatedSize();
 		}
 		if (bShowIndividualStats)
 		{
@@ -1884,7 +1908,7 @@ void LogHashMemoryOverheadStatistics(FOutputDevice& Ar, const bool bShowIndividu
 		int64 Size = HashTables.ObjectOuterMap.GetAllocatedSize();
 		for (const TPair<UObjectBase*, FHashBucket>& Pair : HashTables.ObjectOuterMap)
 		{
-			Size += Pair.Value.GetItemsSize();
+			Size += Pair.Value.GetAllocatedSize();
 		}
 		if (bShowIndividualStats)
 		{
@@ -1895,9 +1919,9 @@ void LogHashMemoryOverheadStatistics(FOutputDevice& Ar, const bool bShowIndividu
 
 	{
 		SIZE_T Size = HashTables.ClassToObjectListMap.GetAllocatedSize();
-		for (const TPair<UClass*, FHashBucket>& Pair : HashTables.ClassToObjectListMap)
+		for (const auto& Pair : HashTables.ClassToObjectListMap)
 		{
-			Size += Pair.Value.GetItemsSize();
+			Size += Pair.Value.GetAllocatedSize();
 		}
 		if (bShowIndividualStats)
 		{
@@ -1923,7 +1947,7 @@ void LogHashMemoryOverheadStatistics(FOutputDevice& Ar, const bool bShowIndividu
 		SIZE_T Size = HashTables.PackageToObjectListMap.GetAllocatedSize();
 		for (const TPair<UPackage*, FHashBucket>& Pair : HashTables.PackageToObjectListMap)
 		{
-			Size += Pair.Value.GetItemsSize();
+			Size += Pair.Value.GetAllocatedSize();
 		}
 		if (bShowIndividualStats)
 		{
diff --git a/Engine/Source/Runtime/CoreUObject/Public/UObject/UObjectBase.h b/Engine/Source/Runtime/CoreUObject/Public/UObject/UObjectBase.h
--- a/Engine/Source/Runtime/CoreUObject/Public/UObject/UObjectBase.h
+++ b/Engine/Source/Runtime/CoreUObject/Public/UObject/UObjectBase.h
@@ -26,6 +26,12 @@ class UObject;
 class UPackage;
 class UScriptStruct;
 
+// If FName is 4 bytes than we can use padding after it to store internal object list index and use array instead of a hash map for lookup in UObjectHash.cpp
+// This might change the each UClass' object list iteration order
+#ifndef UE_STORE_OBJECT_LIST_INTERNAL_INDEX
+#	define UE_STORE_OBJECT_LIST_INTERNAL_INDEX 0
+#endif
+
 DECLARE_DWORD_COUNTER_STAT_EXTERN(TEXT("STAT_UObjectsStatGroupTester"), STAT_UObjectsStatGroupTester, STATGROUP_UObjects, COREUOBJECT_API);
 
 /** 
@@ -257,12 +263,22 @@ private:
 	/** Name of this object */
 	FName							NamePrivate;
 
+#if UE_STORE_OBJECT_LIST_INTERNAL_INDEX
+	/** Internal index into an array that stores all objects.
+	 It's used for registering and unregistering of UObjects in a global hash map.
+	 This optimization uses array instead of a hash map for reduced memory usage
+	*/
+	int32							ObjectListInternalIndex;
+#endif
+
 	/** Object this object resides in. */
 	ObjectPtr_Private::TNonAccessTrackedObjectPtr<UObject>						OuterPrivate;
 	
 	friend class FBlueprintCompileReinstancer;
 	friend class FVerseObjectClassReplacer;
 	friend class FContextObjectManager;
+	friend void AddToClassMap(class FUObjectHashTables& ThreadHash, UObjectBase* Object);
+	friend void RemoveFromClassMap(class FUObjectHashTables& ThreadHash, UObjectBase* Object);
 
 #if WITH_EDITOR
 	/** This is used by the reinstancer to re-class and re-archetype the current instances of a class before recompiling */
