Fixed Editor build failing on a UE_STORE_OBJECT_LIST_INTERNAL_INDEX define

diff --git a/Engine/Source/Runtime/CoreUObject/Public/UObject/UObjectBase.h b/Engine/Source/Runtime/CoreUObject/Public/UObject/UObjectBase.h
--- a/Engine/Source/Runtime/CoreUObject/Public/UObject/UObjectBase.h
+++ b/Engine/Source/Runtime/CoreUObject/Public/UObject/UObjectBase.h
@@ -28,7 +28,7 @@ class UScriptStruct;
 
 // If FName is 4 bytes than we can use padding after it to store internal object list index and use array instead of a hash map for lookup in UObjectHash.cpp
 // This might change the each UClass' object list iteration order
-#ifndef UE_STORE_OBJECT_LIST_INTERNAL_INDEX
+#if !defined(UE_STORE_OBJECT_LIST_INTERNAL_INDEX)
 #	define UE_STORE_OBJECT_LIST_INTERNAL_INDEX 0
 #endif
 
