Added a better explanation for UObjectBase size assert

diff --git a/Engine/Source/Runtime/CoreUObject/Private/UObject/UObjectHash.cpp b/Engine/Source/Runtime/CoreUObject/Private/UObject/UObjectHash.cpp
--- a/Engine/Source/Runtime/CoreUObject/Private/UObject/UObjectHash.cpp
+++ b/Engine/Source/Runtime/CoreUObject/Private/UObject/UObjectHash.cpp
@@ -16,7 +16,8 @@
 
 #if UE_STORE_OBJECT_LIST_INTERNAL_INDEX
 static_assert(sizeof(FName) == 4, "Internal object index optimization depends exploits 4 bytes padding after the FName");
-static_assert(sizeof(UObjectBase) == 40, "UObjectBase size has changed!");
+static_assert(sizeof(UObjectBase) == 40,
+	"UObjectBase size has changed! This optimization relies on a 4 bytes padding after NamePrivate.Please make sure that the new binary layout it optimal");
 #endif
 
 DEFINE_LOG_CATEGORY_STATIC(LogUObjectHash, Log, All);
