DirectoryTree: 1) Fix bug in GetChildren - added children of implied children were not reported in some cases.
2) Fix assertion in GetChildren when called on a DriveSpecifierWithoutRoot.
#rnx

diff --git a/Engine/Source/Runtime/CoreUObject/Internal/Containers/DirectoryTree.h b/Engine/Source/Runtime/CoreUObject/Internal/Containers/DirectoryTree.h
--- a/Engine/Source/Runtime/CoreUObject/Internal/Containers/DirectoryTree.h
+++ b/Engine/Source/Runtime/CoreUObject/Internal/Containers/DirectoryTree.h
@@ -228,6 +228,8 @@ private:
 		/** Set HasValue=false, and destruct any existing old value. */
 		void RemoveValue();
 
+		void FixupDirectChildrenPathSeparator(TCHAR OldSeparator, TCHAR NewSeparator);
+
 	private:
 		const static uint32 NumFlagBits = 1;
 		const static uint32 FlagsShift = (8 * sizeof(uint32) - NumFlagBits);
@@ -259,13 +261,17 @@ private:
 	};
 
 private:
+	bool NormalizePathForReading(FStringView& Path, FStringBuilderBase& NormalizeBuffer) const;
+	bool NormalizePathForWriting(FStringView& Path, FStringBuilderBase& NormalizeBuffer);
 	ValueType* TryFindClosestPathInternal(FStringView Path, FStringBuilderBase* OutPath);
+	void InitializePathSeparator(TCHAR InPathSeparator);
 
 private:
 	FTreeNode Root;
 	int32 NumPaths = 0;
 	TCHAR PathSeparator = '/';
 	bool bPathSeparatorInitialized = false;
+	bool bNeedDriveWithoutPathFixup = false;
 };
 
 
@@ -297,16 +303,17 @@ inline ValueType& TDirectoryTree<ValueType>::FindOrAdd(FStringView Path, bool* b
 		int32 UnusedIndex;
 		if (Path.FindChar('/', UnusedIndex))
 		{
-			PathSeparator = '/';
-			bPathSeparatorInitialized = true;
+			InitializePathSeparator('/');
 		}
 		else if (Path.FindChar('\\', UnusedIndex))
 		{
-			PathSeparator = '\\';
-			bPathSeparatorInitialized = true;
+			InitializePathSeparator('\\');
 		}
 	}
 
+	TStringBuilder<16> NormalizeBuffer;
+	NormalizePathForWriting(Path, NormalizeBuffer);
+
 	bool bExisted;
 	ValueType& Result = Root.FindOrAdd(Path, bExisted);
 	if (!bExisted)
@@ -339,6 +346,9 @@ inline void TDirectoryTree<ValueType>::Remove(FStringView Path, bool* bOutExiste
 	}
 	else
 	{
+		TStringBuilder<16> NormalizeBuffer;
+		NormalizePathForReading(Path, NormalizeBuffer);
+
 		Root.Remove(Path, bExisted);
 	}
 	if (bExisted)
@@ -395,6 +405,10 @@ inline ValueType* TDirectoryTree<ValueType>::Find(FStringView Path)
 	{
 		return Root.HasValue() ? &Root.GetValue() : nullptr;
 	}
+
+	TStringBuilder<16> NormalizeBuffer;
+	NormalizePathForReading(Path, NormalizeBuffer);
+
 	return Root.Find(Path);
 }
 
@@ -497,6 +511,9 @@ inline ValueType* TDirectoryTree<ValueType>::TryFindClosestPathInternal(FStringV
 	}
 	if (!Path.IsEmpty())
 	{
+		TStringBuilder<16> NormalizeBuffer;
+		NormalizePathForReading(Path, NormalizeBuffer);
+
 		ValueType* Result = Root.TryFindClosestPath(Path, OutPath, PathSeparator);
 		if (Result)
 		{
@@ -507,6 +524,67 @@ inline ValueType* TDirectoryTree<ValueType>::TryFindClosestPathInternal(FStringV
 	return Root.HasValue() ? &Root.GetValue() : nullptr;
 }
 
+template <typename ValueType>
+inline bool TDirectoryTree<ValueType>::NormalizePathForReading(FStringView& Path,
+	FStringBuilderBase& NormalizeBuffer) const
+{
+	// Drive specifiers without a root are a special case; they break our assumption that if
+	// if	FPathViews::IsParentPathOf(DriveSpecifier, PathInThatDrive)
+	// then DriveSpecifier == FirstComponentOfPathInThatDrive.
+	// 'D:' is a parent path of 'D:/Path' but FirstComponent of 'D:/Path' is 'D:/' != 'D:' 
+	// 
+	// In general usage on e.g. Windows, drive specifiers without a path are interpreted to mean the current
+	// working directory of the given drive. But we don't have that context so that meaning is not applicable.
+	// 
+	// We therefore instead interpret them to mean the root of the drive. Append the PathSeparator to make
+	// them the root.
+	if (FPathViews::IsDriveSpecifierWithoutRoot(Path))
+	{
+		FStringView Volume;
+		FStringView Remainder;
+		FPathViews::SplitVolumeSpecifier(Path, Volume, Remainder);
+		NormalizeBuffer << Volume << PathSeparator << Remainder;
+		Path = NormalizeBuffer.ToView();
+		return true;
+	}
+	return false;
+}
+
+template <typename ValueType>
+inline bool TDirectoryTree<ValueType>::NormalizePathForWriting(FStringView& Path, FStringBuilderBase& NormalizeBuffer)
+{
+	if (NormalizePathForReading(Path, NormalizeBuffer))
+	{
+		// If the call to NormalizePath came from a function that is adding paths to the tree, and
+		// !bPathSeparatorInitialized, then leave a marker that we might need to fix up the added paths
+		// when we encounter the user's desired path separator. 
+		bNeedDriveWithoutPathFixup = !bPathSeparatorInitialized;
+		return true;
+	}
+	return false;
+}
+
+template <typename ValueType>
+inline void TDirectoryTree<ValueType>::InitializePathSeparator(TCHAR InPathSeparator)
+{
+	check(!bPathSeparatorInitialized);
+
+	// If the requested PathSeparator is not the one we guessed it was when we had to
+	// normalize a drive without a path (e.g. 'D:' -> 'D:/'), then fixup all those
+	// drive children to have the desired separator.
+	bNeedDriveWithoutPathFixup &= (InPathSeparator != PathSeparator);
+	if (bNeedDriveWithoutPathFixup)
+	{
+		// The drives without paths will be direct children of the root so we only need
+		// to fixup direct children
+		Root.FixupDirectChildrenPathSeparator(PathSeparator, InPathSeparator);
+		bNeedDriveWithoutPathFixup = false;
+	}
+
+	PathSeparator = InPathSeparator;
+	bPathSeparatorInitialized = true;
+}
+
 template <typename ValueType>
 inline bool TDirectoryTree<ValueType>::TryGetChildren(FStringView Path, TArray<FString>& OutRelativeChildNames,
 	EDirectoryTreeGetFlags Flags) const
@@ -515,6 +593,9 @@ inline bool TDirectoryTree<ValueType>::TryGetChildren(FStringView Path, TArray<F
 	{
 		return false;
 	}
+	TStringBuilder<16> NormalizeBuffer;
+	NormalizePathForReading(Path, NormalizeBuffer);
+
 	TStringBuilder<1024> ReportedPathPrefix;
 	return Root.TryGetChildren(ReportedPathPrefix, PathSeparator, Path, OutRelativeChildNames, Flags);
 }
@@ -884,25 +965,23 @@ inline bool TDirectoryTree<ValueType>::FTreeNode::TryGetChildren(FStringBuilderB
 			else
 			{
 				// When ImpliedChildren are not supposed to be reported, report each stored child by its full relpath,
-				// unless the child is an implied path. If recursive, also forward the call to it to return its
-				// recursive children.
-				if (ChildNode.HasValue() || EnumHasAnyFlags(Flags, EDirectoryTreeGetFlags::Recursive))
-				{
-					int32 SavedLen = ReportPathPrefix.Len();
-					FPathViews::Append(ReportPathPrefix, ChildRelPath);
-					UE::DirectoryTree::FixupPathSeparator(ReportPathPrefix, SavedLen, InPathSeparator);
+				// unless the child is an implied path. If the child is an implied path, recursively ask the child
+				// to report its added children. Also, if user requested recursive, ask the child to return its
+				// recursive children even if it has a value.
+				int32 SavedLen = ReportPathPrefix.Len();
+				FPathViews::Append(ReportPathPrefix, ChildRelPath);
+				UE::DirectoryTree::FixupPathSeparator(ReportPathPrefix, SavedLen, InPathSeparator);
 
-					if (ChildNode.HasValue())
-					{
-						OutRelativeChildNames.Add(*ReportPathPrefix);
-					}
-					if (EnumHasAnyFlags(Flags, EDirectoryTreeGetFlags::Recursive))
-					{
-						(void)ChildNode.TryGetChildren(ReportPathPrefix, InPathSeparator, FStringView(),
-							OutRelativeChildNames, Flags);
-					}
-					ReportPathPrefix.RemoveSuffix(ReportPathPrefix.Len() - SavedLen);
+				if (ChildNode.HasValue())
+				{
+					OutRelativeChildNames.Add(*ReportPathPrefix);
 				}
+				if (!ChildNode.HasValue() || EnumHasAnyFlags(Flags, EDirectoryTreeGetFlags::Recursive))
+				{
+					(void)ChildNode.TryGetChildren(ReportPathPrefix, InPathSeparator, FStringView(),
+						OutRelativeChildNames, Flags);
+				}
+				ReportPathPrefix.RemoveSuffix(ReportPathPrefix.Len() - SavedLen);
 			}
 		}
 
@@ -930,7 +1009,7 @@ inline bool TDirectoryTree<ValueType>::FTreeNode::TryGetChildren(FStringBuilderB
 		FStringView ExistingFirstComponent;
 		FStringView ExistingRemainingPath;
 		FPathViews::SplitFirstComponent(ChildRelPath, ExistingFirstComponent, ExistingRemainingPath);
-		check(ExistingFirstComponent.Equals(FirstComponent, ESearchCase::IgnoreCase)); // Otherwise FindInsertionIndex would have returned bExists=false
+		check(FPathViews::Equals(FirstComponent, ExistingFirstComponent)); // Otherwise FindInsertionIndex would have returned bExists=false
 
 		for (int32 RunawayLoop = 0; RunawayLoop <= InRelPath.Len(); ++RunawayLoop)
 		{
@@ -1266,3 +1345,15 @@ inline void TDirectoryTree<ValueType>::FTreeNode::Realloc(int32 NewCapacity)
 	}
 	CapacityChildNodes = NewCapacity;
 }
+
+template <typename ValueType>
+inline void TDirectoryTree<ValueType>::FTreeNode::FixupDirectChildrenPathSeparator(TCHAR OldSeparator,
+	TCHAR NewSeparator)
+{
+	int32 Num = GetNumChildNodes();
+	for (int32 Index = 0; Index < Num; ++Index)
+	{
+		RelPaths[Index].ReplaceCharInline(OldSeparator, NewSeparator);
+	}
+}
+
diff --git a/Engine/Source/Runtime/CoreUObject/Private/Misc/DirectoryTree.cpp b/Engine/Source/Runtime/CoreUObject/Private/Misc/DirectoryTree.cpp
--- a/Engine/Source/Runtime/CoreUObject/Private/Misc/DirectoryTree.cpp
+++ b/Engine/Source/Runtime/CoreUObject/Private/Misc/DirectoryTree.cpp
@@ -511,6 +511,109 @@ bool FDirectoryTreeTest::RunTest(const FString& Parameters)
 		TestTrue(TEXT("MoveConstructOnlyValueZ correct"), Value && Value->Value == 437);
 	}
 
+	// Handling special case of drive specifiers without a path
+	{
+		TStringBuilder<16> FoundPath;
+		TArray<FString> ChildNames;
+		int* FoundValue = nullptr;
+		auto Reset = [&FoundPath, &ChildNames, &FoundValue]()
+			{
+				FoundPath.Reset();
+				ChildNames.Reset();
+				FoundValue = nullptr;
+			};
+
+		{
+			TDirectoryTree<int32> Tree;
+			Tree.FindOrAdd(TEXTVIEW("D:")) = 1;
+
+			Reset();
+			TestTrue(TEXT("DriveSpecifier: Before PathSep: Without PathSep: Tree.Contains"), Tree.Contains(TEXTVIEW("D:")));
+			TestTrue(TEXT("DriveSpecifier: Before PathSep: Without PathSep: Tree.Find"), Tree.Find(TEXTVIEW("D:")) != nullptr);
+			TestTrue(TEXT("DriveSpecifier: Before PathSep: Without PathSep: Tree.ContainsPathOrParent"), Tree.ContainsPathOrParent(TEXTVIEW("D:")));
+			TestTrue(TEXT("DriveSpecifier: Before PathSep: Without PathSep: Tree.FindClosestValue"), Tree.FindClosestValue(TEXTVIEW("D:")) != nullptr);
+			TestTrue(TEXT("DriveSpecifier: Before PathSep: Without PathSep: Tree.TryFindClosestPath"),
+				Tree.TryFindClosestPath(TEXTVIEW("D:"), FoundPath, &FoundValue) && FoundPath.Len() > 0 && FoundValue != nullptr);
+			TestTrue(TEXT("DriveSpecifier: Before PathSep: Without PathSep: Tree.TryGetChildren"), Tree.TryGetChildren(TEXTVIEW("D:"), ChildNames));
+
+			Tree.FindOrAdd(TEXTVIEW("D:/root")) = 1;
+
+			Reset();
+			TestTrue(TEXT("DriveSpecifier: After PathSep('/'): Without PathSep: Tree.Contains"), Tree.Contains(TEXTVIEW("D:")));
+			TestTrue(TEXT("DriveSpecifier: After PathSep('/'): Without PathSep: Tree.Find"), Tree.Find(TEXTVIEW("D:")) != nullptr);
+			TestTrue(TEXT("DriveSpecifier: After PathSep('/'): Without PathSep: Tree.ContainsPathOrParent"), Tree.ContainsPathOrParent(TEXTVIEW("D:")));
+			TestTrue(TEXT("DriveSpecifier: After PathSep('/'): Without PathSep: Tree.FindClosestValue"), Tree.FindClosestValue(TEXTVIEW("D:")) != nullptr);
+			TestTrue(TEXT("DriveSpecifier: After PathSep('/'): Without PathSep: Tree.TryFindClosestPath"),
+				Tree.TryFindClosestPath(TEXTVIEW("D:"), FoundPath, &FoundValue) && FoundPath.Len() > 0 && FoundValue != nullptr);
+			TestTrue(TEXT("DriveSpecifier: After PathSep('/'): Without PathSep: Tree.TryGetChildren"), Tree.TryGetChildren(TEXTVIEW("D:"), ChildNames));
+
+			Reset();
+			TestTrue(TEXT("DriveSpecifier: After PathSep('/'): With PathSep: Tree.Contains"), Tree.Contains(TEXTVIEW("D:/")));
+			TestTrue(TEXT("DriveSpecifier: After PathSep('/'): With PathSep: Tree.Find"), Tree.Find(TEXTVIEW("D:/")) != nullptr);
+			TestTrue(TEXT("DriveSpecifier: After PathSep('/'): With PathSep: Tree.ContainsPathOrParent"), Tree.ContainsPathOrParent(TEXTVIEW("D:/")));
+			TestTrue(TEXT("DriveSpecifier: After PathSep('/'): With PathSep: Tree.FindClosestValue"), Tree.FindClosestValue(TEXTVIEW("D:/")) != nullptr);
+			TestTrue(TEXT("DriveSpecifier: After PathSep('/'): With PathSep: Tree.TryFindClosestPath"),
+				Tree.TryFindClosestPath(TEXTVIEW("D:/"), FoundPath, &FoundValue) && FoundPath.Len() > 0 && FoundValue != nullptr);
+			TestTrue(TEXT("DriveSpecifier: After PathSep('/'): With PathSep: Tree.TryGetChildren"), Tree.TryGetChildren(TEXTVIEW("D:/"), ChildNames));
+		}
+		{
+			TDirectoryTree<int32> Tree;
+			Tree.FindOrAdd(TEXTVIEW("D:")) = 1;
+			Tree.FindOrAdd(TEXTVIEW("D:\\root")) = 1;
+
+			Reset();
+			TestTrue(TEXT("DriveSpecifier: After PathSep('\\'): Without PathSep: Tree.Contains"), Tree.Contains(TEXTVIEW("D:")));
+			TestTrue(TEXT("DriveSpecifier: After PathSep('\\'): Without PathSep: Tree.Find"), Tree.Find(TEXTVIEW("D:")) != nullptr);
+			TestTrue(TEXT("DriveSpecifier: After PathSep('\\'): Without PathSep: Tree.ContainsPathOrParent"), Tree.ContainsPathOrParent(TEXTVIEW("D:")));
+			TestTrue(TEXT("DriveSpecifier: After PathSep('\\'): Without PathSep: Tree.FindClosestValue"), Tree.FindClosestValue(TEXTVIEW("D:")) != nullptr);
+			TestTrue(TEXT("DriveSpecifier: After PathSep('\\'): Without PathSep: Tree.TryFindClosestPath"),
+				Tree.TryFindClosestPath(TEXTVIEW("D:"), FoundPath, &FoundValue) && FoundPath.Len() > 0 && FoundValue != nullptr);
+			TestTrue(TEXT("DriveSpecifier: After PathSep('\\'): Without PathSep: Tree.TryGetChildren"), Tree.TryGetChildren(TEXTVIEW("D:"), ChildNames));
+
+			Reset();
+			TestTrue(TEXT("DriveSpecifier: After PathSep('\\'): With PathSep: Tree.Contains"), Tree.Contains(TEXTVIEW("D:\\")));
+			TestTrue(TEXT("DriveSpecifier: After PathSep('\\'): With PathSep: Tree.Find"), Tree.Find(TEXTVIEW("D:\\")) != nullptr);
+			TestTrue(TEXT("DriveSpecifier: After PathSep('\\'): With PathSep: Tree.ContainsPathOrParent"), Tree.ContainsPathOrParent(TEXTVIEW("D:\\")));
+			TestTrue(TEXT("DriveSpecifier: After PathSep('\\'): With PathSep: Tree.FindClosestValue"), Tree.FindClosestValue(TEXTVIEW("D:\\")) != nullptr);
+			TestTrue(TEXT("DriveSpecifier: After PathSep('\\'): With PathSep: Tree.TryFindClosestPath"),
+				Tree.TryFindClosestPath(TEXTVIEW("D:\\"), FoundPath, &FoundValue) && FoundPath.Len() > 0 && FoundValue != nullptr);
+			TestTrue(TEXT("DriveSpecifier: After PathSep('\\'): With PathSep: Tree.TryGetChildren"), Tree.TryGetChildren(TEXTVIEW("D:\\"), ChildNames));
+		}
+		{
+			TDirectoryTree<int32> Tree;
+			Tree.FindOrAdd(TEXTVIEW("D:root")) = 1;
+
+			Reset();
+			TestTrue(TEXT("DriveSpecifierLong: Before PathSep: Without PathSep: Tree.Contains"), Tree.Contains(TEXTVIEW("D:root")));
+			TestTrue(TEXT("DriveSpecifierLong: Before PathSep: Without PathSep: Tree.Find"), Tree.Find(TEXTVIEW("D:root")) != nullptr);
+			TestTrue(TEXT("DriveSpecifierLong: Before PathSep: Without PathSep: Tree.ContainsPathOrParent"), Tree.ContainsPathOrParent(TEXTVIEW("D:root")));
+			TestTrue(TEXT("DriveSpecifierLong: Before PathSep: Without PathSep: Tree.FindClosestValue"), Tree.FindClosestValue(TEXTVIEW("D:root")) != nullptr);
+			TestTrue(TEXT("DriveSpecifierLong: Before PathSep: Without PathSep: Tree.TryFindClosestPath"),
+				Tree.TryFindClosestPath(TEXTVIEW("D:root"), FoundPath, &FoundValue) && FoundPath.Len() > 0 && FoundValue != nullptr);
+			TestTrue(TEXT("DriveSpecifierLong: Before PathSep: Without PathSep: Tree.TryGetChildren"), Tree.TryGetChildren(TEXTVIEW("D:root"), ChildNames));
+
+			Tree.FindOrAdd(TEXTVIEW("D:\\root\\path")) = 1;
+
+			Reset();
+			TestTrue(TEXT("DriveSpecifierLong: After PathSep('\\'): Without PathSep: Tree.Contains"), Tree.Contains(TEXTVIEW("D:root")));
+			TestTrue(TEXT("DriveSpecifierLong: After PathSep('\\'): Without PathSep: Tree.Find"), Tree.Find(TEXTVIEW("D:root")) != nullptr);
+			TestTrue(TEXT("DriveSpecifierLong: After PathSep('\\'): Without PathSep: Tree.ContainsPathOrParent"), Tree.ContainsPathOrParent(TEXTVIEW("D:root")));
+			TestTrue(TEXT("DriveSpecifierLong: After PathSep('\\'): Without PathSep: Tree.FindClosestValue"), Tree.FindClosestValue(TEXTVIEW("D:root")) != nullptr);
+			TestTrue(TEXT("DriveSpecifierLong: After PathSep('\\'): Without PathSep: Tree.TryFindClosestPath"),
+				Tree.TryFindClosestPath(TEXTVIEW("D:root"), FoundPath, &FoundValue) && FoundPath.Len() > 0 && FoundValue != nullptr);
+			TestTrue(TEXT("DriveSpecifierLong: After PathSep('\\'): Without PathSep: Tree.TryGetChildren"), Tree.TryGetChildren(TEXTVIEW("D:root"), ChildNames));
+
+			Reset();
+			TestTrue(TEXT("DriveSpecifierLong: After PathSep('\\'): With PathSep: Tree.Contains"), Tree.Contains(TEXTVIEW("D:\\root")));
+			TestTrue(TEXT("DriveSpecifierLong: After PathSep('\\'): With PathSep: Tree.Find"), Tree.Find(TEXTVIEW("D:\\root")) != nullptr);
+			TestTrue(TEXT("DriveSpecifierLong: After PathSep('\\'): With PathSep: Tree.ContainsPathOrParent"), Tree.ContainsPathOrParent(TEXTVIEW("D:\\root")));
+			TestTrue(TEXT("DriveSpecifierLong: After PathSep('\\'): With PathSep: Tree.FindClosestValue"), Tree.FindClosestValue(TEXTVIEW("D:\\root")) != nullptr);
+			TestTrue(TEXT("DriveSpecifierLong: After PathSep('\\'): With PathSep: Tree.TryFindClosestPath"),
+				Tree.TryFindClosestPath(TEXTVIEW("D:root"), FoundPath, &FoundValue) && FoundPath.Len() > 0 && FoundValue != nullptr);
+			TestTrue(TEXT("DriveSpecifierLong: After PathSep('\\'): With PathSep: Tree.TryGetChildren"), Tree.TryGetChildren(TEXTVIEW("D:\\root"), ChildNames));
+		}
+	}
+
 	// Testing accessors
 	{
 		// GetChildren
@@ -835,6 +938,19 @@ bool FDirectoryTreeTest::RunTest(const FString& Parameters)
 		TestTrue(TEXT("GetChildrenComplexA F_AddedRoot, !ImpliedParent, ImpliedChildren, Recursive"),
 			bExists == true && UnorderedEquals(Children, { TEXT("ImpliedChild"),
 				TEXT("ImpliedChild/AddedChild") }));
+
+		// Case: Requesting !ImpliedChildren and !Recursive on a path with an implied child, should report
+		// the added path children of the Implied child
+		Tree.Empty();
+		Tree.FindOrAdd(TEXTVIEW("/Root/Implied1/Added1")).Value = 1;
+		Tree.FindOrAdd(TEXTVIEW("/Root/Implied1/Added2")).Value = 1;
+		Tree.FindOrAdd(TEXTVIEW("/Root/Implied2/Added")).Value = 1;
+
+		Children.Reset();
+		bExists = Tree.TryGetChildren(TEXTVIEW("/Root"), Children,
+			EDirectoryTreeGetFlags::ImpliedParent);
+		TestTrue(TEXT("!ImpliedChildren, !Recursive, and direct child is implied."),
+			bExists = true && UnorderedEquals(Children, { TEXT("Implied1/Added1"), TEXT("Implied1/Added2"), TEXT("Implied2/Added") }));
 	}
 
 	return true;
