Silence a static analyzer warning.

#rnx

diff --git a/Engine/Source/Runtime/CoreUObject/Private/UObject/ExpandingChunkedList.h b/Engine/Source/Runtime/CoreUObject/Private/UObject/ExpandingChunkedList.h
--- a/Engine/Source/Runtime/CoreUObject/Private/UObject/ExpandingChunkedList.h
+++ b/Engine/Source/Runtime/CoreUObject/Private/UObject/ExpandingChunkedList.h
@@ -71,7 +71,7 @@ public:
 		{
 			Head->Items[Head->NumItems++] = Item;
 			Lock.Unlock();
-			FMemory::Free(NewChunk);
+			FMemory::Free(NewChunk); //-V611 There is no destructor and the memory came from Malloc...
 			return;
 		}
 
