Fix crash due to unimplemented function in old loader

diff --git a/Engine/Source/Runtime/CoreUObject/Private/UObject/Linker.cpp b/Engine/Source/Runtime/CoreUObject/Private/UObject/Linker.cpp
--- a/Engine/Source/Runtime/CoreUObject/Private/UObject/Linker.cpp
+++ b/Engine/Source/Runtime/CoreUObject/Private/UObject/Linker.cpp
@@ -900,7 +900,7 @@ void ConditionalFlushAsyncLoadingForLinkers(TConstArrayView<FLinkerLoad*> InLink
 		TArray<int32, TInlineAllocator<4>> RequestIds;
 		for (FLinkerLoad* Linker : InLinkers)
 		{
-			int32 Request = LoadPackageAsync(Linker->GetPackagePath(), FLoadPackageAsyncOptionalParams{ .PackagePriority = MAX_int32 });
+			int32 Request = LoadPackageAsync(Linker->GetPackagePath());
 			RequestIds.Add(Request);
 		}
 		FlushAsyncLoading(RequestIds);
