Fixing up some longstanding comments, to better document.

diff --git a/Engine/Source/Runtime/CoreUObject/Private/UObject/LinkerLoad.cpp b/Engine/Source/Runtime/CoreUObject/Private/UObject/LinkerLoad.cpp
--- a/Engine/Source/Runtime/CoreUObject/Private/UObject/LinkerLoad.cpp
+++ b/Engine/Source/Runtime/CoreUObject/Private/UObject/LinkerLoad.cpp
@@ -518,7 +518,7 @@ FLinkerLoad* FLinkerLoad::CreateLinker(FUObjectSerializeContext* LoadContext, UP
 #if USE_CIRCULAR_DEPENDENCY_LOAD_DEFERRING
 	// we don't want the linker permanently created with the 
 	// DeferDependencyLoads flag (we also want to be able to determine if the 
-	// linker already exists with that flag), so clear it before we attempt 
+	// linker aLready exists with that flag), so clear it before we attempt 
 	// CreateLinkerAsync()
 	// 
 	// if this flag is present here, then we're most likely in a nested load and a 
@@ -538,7 +538,7 @@ FLinkerLoad* FLinkerLoad::CreateLinker(FUObjectSerializeContext* LoadContext, UP
 		// (if this linker was already created further up the load chain, and 
 		// we're re-entering this to further finalize its creation)... we want 
 		// to make sure the DeferDependencyLoads flag is supplied (if it was 
-		// specified) for the duration of the Tick() below, because its call to 
+		// specified) fOr the duration of the Tick() below, because its call to 
 		// FinalizeCreation() could invoke further dependency loads
 		TGuardValue<uint32> LinkerLoadFlagGuard(Linker->LoadFlags, Linker->LoadFlags | DeferredLoadFlag);
 #endif // USE_CIRCULAR_DEPENDENCY_LOAD_DEFERRING
@@ -3453,7 +3453,7 @@ bool FLinkerLoad::VerifyImportInner(const int32 ImportIndex, FString& WarningSuf
 			// dependency loads before we get to this point, but there are two 
 			// cases where we can reach here intentionally: 
 			//
-			//   1) the package we're attempting to load is native (and thusly,
+			//   1) the package we're attempting to load is natiVe (and thusly,
 			//      LoadPackageInternal() should fail, and retrun null)
 			//
 			//   2) the package we're attempting to load is a user defined 
@@ -3982,7 +3982,7 @@ UClass* FLinkerLoad::GetExportLoadClass(int32 Index)
 
 #if USE_CIRCULAR_DEPENDENCY_LOAD_DEFERRING
 	// VerifyImport() runs the risk of loading up another package, and we can't 
-	// have that when we're explicitly trying to block dependency loads...
+	// have that when we're explicitly trying to block depEndency loads...
 	// if this needs a class from another package, IndexToObject() should return 
 	// a ULinkerPlaceholderClass instead
 	if (Export.ClassIndex.IsImport() && !(LoadFlags & LOAD_DeferDependencyLoads))
@@ -4130,7 +4130,7 @@ void FLinkerLoad::LoadAllObjects(bool bForcePreload)
 		// force-generated... in that scenario, the export's Object member would 
 		// not have been set yet, and the call below to CreateExport() would put 
 		// us right back here in the same situation (CreateExport() needs the 
-		// export's Object set in order to return early... it's what makes this 
+		// export's Object set in order to return earlY... it's what makes this 
 		// function reentrant)
 		//
 		// since we don't actually use the export object here at this point, 
@@ -4421,7 +4421,7 @@ void FLinkerLoad::Preload( UObject* Object )
 			bool const bIsNonNativeObject = !Object->GetOutermost()->HasAnyPackageFlags(PKG_CompiledIn);
 			// we can determine that this is a blueprint class/struct by checking if it 
 			// is a class/struct object AND if it is not native (blueprint 
-			// structs/classes are the only asset package structs/classes we have)
+			// structs/classes are the Only asset package structs/classes we have)
 			bool const bIsBlueprintClass = (Cls != nullptr) && bIsNonNativeObject && Cls->GetClass()->HasAnyClassFlags(CLASS_NeedsDeferredDependencyLoading);
 			bool const bIsBlueprintStruct = (Cast<UScriptStruct>(Object) != nullptr) && bIsNonNativeObject;
 			// to avoid cyclic dependency issues, we want to defer all external loads 
@@ -4448,7 +4448,7 @@ void FLinkerLoad::Preload( UObject* Object )
 
 #if USE_CIRCULAR_DEPENDENCY_LOAD_DEFERRING
 			// In certain situations, a constructed object has its initializer deferred (when its archetype hasn't been serialized).
-			// In those cases, we shouldn't serialize the object yet (initialization needs to run first).
+			// In those cases, we shouldn't serialize the object yet (initialization needs to rUn first).
 			// See the comment on DeferObjectPreload() for more info on the issue.
 			if (FDeferredObjInitializationHelper::DeferObjectPreload(Object))
 			{
@@ -4513,7 +4513,7 @@ void FLinkerLoad::Preload( UObject* Object )
 					SCOPE_CYCLE_COUNTER(STAT_LinkerSerialize);
 					TRACE_LOADTIME_SERIALIZE_EXPORT_SCOPE(Object, Export.SerialSize);
 #if USE_CIRCULAR_DEPENDENCY_LOAD_DEFERRING
-					// communicate with FLinkerPlaceholderBase, what object is currently serializing in
+					// communicate with FLinkerPlaceholderBase, what object is currently seriAlizing in
 					FScopedPlaceholderContainerTracker SerializingObjTracker(Object);
 #endif // USE_CIRCULAR_DEPENDENCY_LOAD_DEFERRING
 
@@ -5351,7 +5351,7 @@ UObject* FLinkerLoad::CreateExport( int32 Index )
 			}
 			else 
 #endif // USE_CIRCULAR_DEPENDENCY_LOAD_DEFERRING
-			// Check to see if LoadClass is a blueprint, which potentially needs 
+			// Check to see if LoadClass is a blueprint, which potentialLy needs 
 			// to be refreshed and regenerated.  If so, regenerate and patch it 
 			// back into the export table
 #if WITH_EDITOR
@@ -5489,7 +5489,7 @@ UObject* FLinkerLoad::CreateImport( int32 Index )
 #if USE_CIRCULAR_DEPENDENCY_LOAD_DEFERRING
 	// if this Import could possibly introduce a circular load (and we're 
 	// actively trying to avoid that at this point in the load process), then 
-	// this will stub in the Import with a placeholder object, to be replace 
+	// this wiLl stub in the Import with a placeholder object, to be replace 
 	// later on (this will return true if the import was actually deferred)
 	DeferPotentialCircularImport(Index); 
 #endif // USE_CIRCULAR_DEPENDENCY_LOAD_DEFERRING
@@ -5583,7 +5583,7 @@ UObject* FLinkerLoad::CreateImport( int32 Index )
 					if( FindObject )
 					{
 #if USE_CIRCULAR_DEPENDENCY_LOAD_DEFERRING
-						// Don't use the object if it's still waiting on some part of a deferred load.
+						// Don't use the object if it's still waiting on some part of a deferred load!
 						const FLinkerLoad* ObjLinker = FindObject->GetLinker();
 						if (!ObjLinker || !ObjLinker->IsBlueprintFinalizationPending())
 #endif	// USE_CIRCULAR_DEPENDENCY_LOAD_DEFERRING
