Allow copy paste instanced references between objects from different packages
#rnx
#tests Copy paste instanced references between data assets

diff --git a/Engine/Source/Runtime/CoreUObject/Private/UObject/PropertyObjectPtr.cpp b/Engine/Source/Runtime/CoreUObject/Private/UObject/PropertyObjectPtr.cpp
--- a/Engine/Source/Runtime/CoreUObject/Private/UObject/PropertyObjectPtr.cpp
+++ b/Engine/Source/Runtime/CoreUObject/Private/UObject/PropertyObjectPtr.cpp
@@ -211,6 +211,11 @@ bool FObjectPtrProperty::AllowObjectTypeReinterpretationTo(const FObjectProperty
 	return Other && Other->IsA<FObjectProperty>();
 }
 
+bool FObjectPtrProperty::AllowCrossLevel() const
+{
+	return HasAnyPropertyFlags(CPF_InstancedReference);
+}
+
 uint32 FObjectPtrProperty::GetValueTypeHashInternal(const void* Src) const
 {
 	return GetTypeHash((FObjectPtr&)GetPropertyValue(Src));
diff --git a/Engine/Source/Runtime/CoreUObject/Public/UObject/UnrealType.h b/Engine/Source/Runtime/CoreUObject/Public/UObject/UnrealType.h
--- a/Engine/Source/Runtime/CoreUObject/Public/UObject/UnrealType.h
+++ b/Engine/Source/Runtime/CoreUObject/Public/UObject/UnrealType.h
@@ -2974,6 +2974,8 @@ class COREUOBJECT_API FObjectPtrProperty : public FObjectProperty
 	virtual void SetObjectPropertyValue(void* PropertyValueAddress, UObject* Value) const override;
 	virtual void SetObjectPropertyValue_InContainer(void* ContainerAddress, UObject* Value, int32 ArrayIndex = 0) const override;
 	virtual bool AllowObjectTypeReinterpretationTo(const FObjectPropertyBase* Other) const override;
+	
+	virtual bool AllowCrossLevel() const override;
 private:
 	virtual uint32 GetValueTypeHashInternal(const void* Src) const override;
 public:
