[ZenLoader] - Fix crash when trying to mix cooked and non-cooked content

#rnx

diff --git a/Engine/Source/Runtime/CoreUObject/Private/Serialization/AsyncLoading2.cpp b/Engine/Source/Runtime/CoreUObject/Private/Serialization/AsyncLoading2.cpp
--- a/Engine/Source/Runtime/CoreUObject/Private/Serialization/AsyncLoading2.cpp
+++ b/Engine/Source/Runtime/CoreUObject/Private/Serialization/AsyncLoading2.cpp
@@ -5058,6 +5058,9 @@ void FAsyncPackage2::ImportPackagesRecursiveInner(FAsyncLoadingThreadState2& Thr
 
 		if (bIsZenPackage != bIsZenPackageImport)
 		{
+			UE_ASYNC_PACKAGE_LOG(VeryVerbose, Desc, TEXT("ImportPackages: AddDependency"),
+				TEXT("Adding package dependency to %s import '%s'."), bIsZenPackageImport ? TEXT("cooked") : TEXT("non-cooked"), *ImportedPackage->Desc.UPackageName.ToString());
+
 			// When importing a linker load package from a zen package or vice versa we need to wait for all the exports in the imported package to be created
 			// and serialized before we can start processing our own exports
 			GetPackageNode(Package_DependenciesReady).DependsOn(&ImportedPackage->GetPackageNode(Package_ExportsSerialized));
@@ -5542,14 +5545,30 @@ bool FAsyncPackage2::ResolveLinkerLoadImports(FAsyncLoadingThreadState2& ThreadS
 	for (int32 ImportIndex = 0; ImportIndex < ImportedPackagesCount; ++ImportIndex)
 	{
 		FAsyncPackage2* ImportedPackage = Data.ImportedAsyncPackages[ImportIndex];
-		if (ImportedPackage && ImportedPackage->LinkerLoadState.IsSet())
+		if (ImportedPackage)
 		{
-			if (ImportedPackage->AsyncPackageLoadingState < EAsyncPackageLoadingState2::WaitingForLinkerLoadDependencies)
+			if (ImportedPackage->LinkerLoadState.IsSet())
 			{
-				check(ThreadState.PackagesOnStack.Contains(ImportedPackage));
-				UE_LOG(LogStreaming, Warning, TEXT("Package %s might be missing an import from package %s because of a circular dependency between them."),
-					*Desc.UPackageName.ToString(),
-					*ImportedPackage->Desc.UPackageName.ToString());
+				if (ImportedPackage->AsyncPackageLoadingState < EAsyncPackageLoadingState2::WaitingForLinkerLoadDependencies)
+				{
+					check(ThreadState.PackagesOnStack.Contains(ImportedPackage));
+					UE_LOG(LogStreaming, Warning, TEXT("Package %s might be missing an import from package %s because of a circular dependency between them."),
+						*Desc.UPackageName.ToString(),
+						*ImportedPackage->Desc.UPackageName.ToString());
+				}
+			}
+			else
+			{
+				// A dependency is added for zen imports in ImportPackagesRecursiveInner that should prevent
+				// us from getting a zen package in a state before its exports are ready. 
+				// Just verify that it's working as intended.
+				if (ImportedPackage->AsyncPackageLoadingState < EAsyncPackageLoadingState2::ExportsDone)
+				{
+					check(ThreadState.PackagesOnStack.Contains(ImportedPackage));
+					UE_LOG(LogStreaming, Warning, TEXT("Package %s might be missing an import from cooked package %s because it's exports are not yet ready."),
+						*Desc.UPackageName.ToString(),
+						*ImportedPackage->Desc.UPackageName.ToString());
+				}
 			}
 		}
 	}
@@ -5809,8 +5828,6 @@ void FAsyncPackage2::ConditionalBeginResolveLinkerLoadImports(FAsyncLoadingThrea
 			else
 			{
 				// Don't advance state of cooked package, nodes already have dependencies setup.
-				// Just validate that the dependency is doing its job and that the zen package we import have reached a proper state.
-				check(Package->AsyncPackageLoadingState >= EAsyncPackageLoadingState2::ExportsDone);
 			}
 		}
 	);
