[ZenLoader] - Add PackageId reverse lookup in editor so warnings can contain useful information instead of just hashes

#rnx

diff --git a/Engine/Source/Runtime/Core/Private/IO/PackageId.cpp b/Engine/Source/Runtime/Core/Private/IO/PackageId.cpp
--- a/Engine/Source/Runtime/Core/Private/IO/PackageId.cpp
+++ b/Engine/Source/Runtime/Core/Private/IO/PackageId.cpp
@@ -4,8 +4,20 @@
 
 #include "Hash/CityHash.h"
 #include "Misc/Char.h"
+#include "HAL/LowLevelMemTracker.h"
 #include "Serialization/StructuredArchiveAdapters.h"
 #include "Serialization/StructuredArchiveSlots.h"
+#include "Misc/ScopeRWLock.h"
+
+#if WITH_PACKAGEID_NAME_MAP
+LLM_DEFINE_TAG(PackageId_ReverseMapping);
+
+namespace PackageIdImpl
+{
+	FRWLock Lock;
+	TMap<uint64, FName> Entries;
+}
+#endif
 
 FPackageId FPackageId::FromName(const FName& Name)
 {
@@ -18,9 +30,26 @@ FPackageId FPackageId::FromName(const FName& Name)
 	}
 	uint64 Hash = CityHash64(reinterpret_cast<const char*>(NameStr), NameLen * sizeof(TCHAR));
 	checkf(Hash != InvalidId, TEXT("Package name hash collision \"%s\" and InvalidId"), NameStr);
+
+#if WITH_PACKAGEID_NAME_MAP
+	{
+		FWriteScopeLock ScopeWriteLock(PackageIdImpl::Lock);
+		LLM_SCOPE_BYTAG(PackageId_ReverseMapping);
+		FName EntryName = PackageIdImpl::Entries.FindOrAdd(Hash, Name);
+		checkf(EntryName == Name, TEXT("FPackageId collision: %llu for both %s and %s"), Hash, *Name.ToString(), *EntryName.ToString());
+	}
+#endif
 	return FPackageId(Hash);
 }
 
+#if WITH_PACKAGEID_NAME_MAP
+FName FPackageId::GetName() const
+{
+	FReadScopeLock ScopeReadLock(PackageIdImpl::Lock);
+	return PackageIdImpl::Entries.FindRef(Id);
+}
+#endif
+
 FArchive& operator<<(FArchive& Ar, FPackageId& Value)
 {
 	FStructuredArchiveFromArchive(Ar).GetSlot() << Value;
diff --git a/Engine/Source/Runtime/Core/Public/IO/PackageId.h b/Engine/Source/Runtime/Core/Public/IO/PackageId.h
--- a/Engine/Source/Runtime/Core/Public/IO/PackageId.h
+++ b/Engine/Source/Runtime/Core/Public/IO/PackageId.h
@@ -10,6 +10,10 @@
 class FArchive;
 class FStructuredArchiveSlot;
 
+#ifndef WITH_PACKAGEID_NAME_MAP
+#define WITH_PACKAGEID_NAME_MAP WITH_EDITOR
+#endif
+
 class FPackageId
 {
 	static constexpr uint64 InvalidId = 0;
@@ -66,4 +70,8 @@ public:
 	CORE_API friend FArchive& operator<<(FArchive& Ar, FPackageId& Value);
 
 	CORE_API friend void operator<<(FStructuredArchiveSlot Slot, FPackageId& Value);
+
+#if WITH_PACKAGEID_NAME_MAP
+	CORE_API FName GetName() const;
+#endif
 };
diff --git a/Engine/Source/Runtime/CoreUObject/Private/Serialization/AsyncLoading2.cpp b/Engine/Source/Runtime/CoreUObject/Private/Serialization/AsyncLoading2.cpp
--- a/Engine/Source/Runtime/CoreUObject/Private/Serialization/AsyncLoading2.cpp
+++ b/Engine/Source/Runtime/CoreUObject/Private/Serialization/AsyncLoading2.cpp
@@ -420,6 +420,15 @@ TRACE_DECLARE_ATOMIC_INT_COUNTER(AsyncLoadingPackagesWithRemainingWork, TEXT("As
 TRACE_DECLARE_ATOMIC_INT_COUNTER(AsyncLoadingPendingIoRequests, TEXT("AsyncLoading/PendingIoRequests"));
 TRACE_DECLARE_ATOMIC_MEMORY_COUNTER(AsyncLoadingTotalLoaded, TEXT("AsyncLoading/TotalLoaded"));
 
+FString FormatPackageId(FPackageId PackageId)
+{
+#if WITH_PACKAGEID_NAME_MAP
+	return FString::Printf(TEXT("0x%llX (%s)"), PackageId.ValueForDebugging(), *PackageId.GetName().ToString());
+#else
+	return FString::Printf(TEXT("0x%llX"), PackageId.ValueForDebugging());
+#endif
+}
+
 struct FAsyncPackage2;
 class FAsyncLoadingThread2;
 
@@ -1156,9 +1165,9 @@ public:
 			FPackageId& PackageId = Pair.Key;
 			FLoadedPackageRef& Ref = Pair.Value;
 			ensureMsgf(Ref.GetRefCount() == 0,
-				TEXT("PackageId '0x%llX' with ref count %d should not have a ref count now")
+				TEXT("PackageId '%s' with ref count %d should not have a ref count now")
 				TEXT(", or this check is incorrectly reached during active loading."),
-				PackageId.ValueForDebugging(),
+				*FormatPackageId(PackageId),
 				Ref.GetRefCount());
 		}
 	}
@@ -1268,33 +1277,33 @@ public:
 			{
 				UObject* ExistingObject = FindPublicExportObjectUnchecked(*PublicExportKey);
 				UE_CLOG(!ExistingObject, LogStreaming, Fatal,
-					TEXT("FGlobalImportStore::VerifyObjectForRemoval: The loaded public export object '%s' with flags (ObjectFlags=%x, InternalObjectFlags=%x) and id 0x%llX:0x%llX is missing in GlobalImportStore. ")
+					TEXT("FGlobalImportStore::VerifyObjectForRemoval: The loaded public export object '%s' with flags (ObjectFlags=%x, InternalObjectFlags=%x) and id %s:0x%llX is missing in GlobalImportStore. ")
 					TEXT("Reason unknown. Double delete? Bug or hash collision?"),
 					*GCObject->GetFullName(),
 					GCObject->GetFlags(),
 					int(GCObject->GetInternalFlags()),
-					PackageId.ValueForDebugging(),
+					*FormatPackageId(PackageId),
 					PublicExportKey->GetExportHash());
 
 				UE_CLOG(ExistingObject != GCObject, LogStreaming, Fatal,
-					TEXT("FGlobalImportStore::VerifyObjectForRemoval: The loaded public export object '%s' with flags (ObjectFlags=%x, InternalObjectFlags=%x) and id 0x%llX:0x%llX is not matching the object '%s' in GlobalImportStore. ")
+					TEXT("FGlobalImportStore::VerifyObjectForRemoval: The loaded public export object '%s' with flags (ObjectFlags=%x, InternalObjectFlags=%x) and id %s:0x%llX is not matching the object '%s' in GlobalImportStore. ")
 					TEXT("Reason unknown. Overwritten after it was added? Bug or hash collision?"),
 					*GCObject->GetFullName(),
 					GCObject->GetFlags(),
 					int(GCObject->GetInternalFlags()),
-					PackageId.ValueForDebugging(),
+					*FormatPackageId(PackageId),
 					PublicExportKey->GetExportHash(),
 					*ExistingObject->GetFullName());
 			}
 			else
 			{
 				UE_LOG(LogStreaming, Warning,
-					TEXT("FGlobalImportStore::VerifyObjectForRemoval: The package for the serialized GC object '%s' with flags (ObjectFlags=%x, InternalObjectFlags=%x) and id 0x%llX:0x%llX is missing in GlobalImportStore. ")
+					TEXT("FGlobalImportStore::VerifyObjectForRemoval: The package for the serialized GC object '%s' with flags (ObjectFlags=%x, InternalObjectFlags=%x) and id %s:0x%llX is missing in GlobalImportStore. ")
 					TEXT("Most likely this object has been moved into this package after it was loaded, while the original package is still around."),
 					*GCObject->GetFullName(),
 					GCObject->GetFlags(),
 					int(GCObject->GetInternalFlags()),
-					PackageId.ValueForDebugging(),
+					*FormatPackageId(PackageId),
 					PublicExportKey->GetExportHash());
 			}
 		}
@@ -1318,28 +1327,28 @@ public:
 			{
 				UObject* Object = static_cast<UObject*>(GUObjectArray.IndexToObject(ObjectIndex)->Object);
 				ensureMsgf(!Object->HasAnyInternalFlags(EInternalObjectFlags::LoaderImport) || GUObjectArray.IsDisregardForGC(Object),
-						TEXT("FGlobalImportStore::VerifyPackageForRemoval: The loaded public export object '%s' with flags (ObjectFlags=%x, InternalObjectFlags=%x) and id 0x%llX is probably still referenced by the loader."),
+						TEXT("FGlobalImportStore::VerifyPackageForRemoval: The loaded public export object '%s' with flags (ObjectFlags=%x, InternalObjectFlags=%x) and id %s is probably still referenced by the loader."),
 						*Object->GetFullName(),
 						Object->GetFlags(),
 						Object->GetInternalFlags(),
-						PackageId.ValueForDebugging());
+						*FormatPackageId(PackageId));
 
 				FPublicExportKey* PublicExportKey = ObjectIndexToPublicExport.Find(ObjectIndex);
 				UE_CLOG(!PublicExportKey, LogStreaming, Fatal,
-					TEXT("FGlobalImportStore::VerifyPackageForRemoval: %s (0x%llX) - ")
+					TEXT("FGlobalImportStore::VerifyPackageForRemoval: %s (%s) - ")
 					TEXT("The loaded public export object '%s' is missing in GlobalImportStore."),
 					*Package->GetName(),
-					PackageId.ValueForDebugging(),
+					*FormatPackageId(PackageId),
 					*Object->GetFullName());
 
 				FPackageId ObjectPackageId = PublicExportKey->GetPackageId();
 				UE_CLOG(ObjectPackageId != PackageId, LogStreaming, Fatal,
-					TEXT("FGlobalImportStore::VerifyPackageForRemoval: %s (0x%llX) - ")
-					TEXT("The loaded public export object '%s' has a mismatching package id 0x%llX in GlobalImportStore."),
+					TEXT("FGlobalImportStore::VerifyPackageForRemoval: %s (%s) - ")
+					TEXT("The loaded public export object '%s' has a mismatching package id %s in GlobalImportStore."),
 					*Package->GetName(),
-					PackageId.ValueForDebugging(),
+					*FormatPackageId(PackageId),
 					*Object->GetFullName(),
-					ObjectPackageId.ValueForDebugging());
+					*FormatPackageId(ObjectPackageId));
 
 				VerifyObjectForRemoval(Object);
 			}
@@ -1384,11 +1393,11 @@ public:
 			int32 ExistingObjectIndex = GUObjectArray.ObjectToIndex(ExistingObject);
 
 			UE_LOG(LogStreaming, Display,
-				TEXT("FGlobalImportStore::StoreGlobalObject: The constructed public export object '%s' with index %d and id 0x%llX:0x%llX collides with object '%s' (ObjectFlags=%X, InternalObjectFlags=%x) with index %d in GlobalImportStore. ")
+				TEXT("FGlobalImportStore::StoreGlobalObject: The constructed public export object '%s' with index %d and id %s:0x%llX collides with object '%s' (ObjectFlags=%X, InternalObjectFlags=%x) with index %d in GlobalImportStore. ")
 				TEXT("The existing object will be replaced since it or its package was most likely renamed after it was loaded the first time."),
 				Object ? *Object->GetFullName() : TEXT("null"),
 				ObjectIndex,
-				Key.GetPackageId().Value(), Key.GetExportHash(),
+				*FormatPackageId(Key.GetPackageId()), Key.GetExportHash(),
 				*ExistingObject->GetFullName(),
 				ExistingObject->GetFlags(),
 				int(ExistingObject->GetInternalFlags()),
@@ -1399,14 +1408,14 @@ public:
 		}
 
 		FPublicExportKey* ExistingKey = ObjectIndexToPublicExport.Find(ObjectIndex);
-		if (ExistingKey)
+		if (ExistingKey && (*ExistingKey != Key))
 		{
-			UE_CLOG(*ExistingKey != Key, LogStreaming, Fatal,
-				TEXT("FGlobalImportStore::StoreGlobalObject: The constructed public export object '%s' with index %d and id 0x%llX:0x%llX already exists in GlobalImportStore but with a different key 0x%llX:0x%llX."),
+			UE_LOG(LogStreaming, Fatal,
+				TEXT("FGlobalImportStore::StoreGlobalObject: The constructed public export object '%s' with index %d and id %s:0x%llX already exists in GlobalImportStore but with a different key %s:0x%llX."),
 				Object ? *Object->GetFullName() : TEXT("null"),
 				ObjectIndex,
-				Key.GetPackageId().Value(), Key.GetExportHash(),
-				ExistingKey->GetPackageId().Value(), ExistingKey->GetExportHash());
+				*FormatPackageId(Key.GetPackageId()), Key.GetExportHash(),
+				*FormatPackageId(ExistingKey->GetPackageId()), ExistingKey->GetExportHash());
 		}
 
 		FLoadedPackageRef& PackageRef = FindPackageRefChecked(Key.GetPackageId());
@@ -4996,7 +5005,7 @@ void FAsyncPackage2::ImportPackagesRecursiveInner(FAsyncLoadingThreadState2& Thr
 			{
 				ImportedPackageRef.SetHasFailed();
 				UE_ASYNC_PACKAGE_LOG(Warning, Desc, TEXT("ImportPackages: SkipPackage"),
-					TEXT("Failed to load uncooked imported package with id '0x%llX' ('%s')"), ImportedPackageId.Value(), *ImportedPackageNameToLoad.ToString());
+					TEXT("Failed to load uncooked imported package with id '0x%llX' ('%s')"), ImportedPackageId.ValueForDebugging(), *ImportedPackageNameToLoad.ToString());
 			}
 			continue;
 		}
@@ -5027,7 +5036,7 @@ void FAsyncPackage2::ImportPackagesRecursiveInner(FAsyncLoadingThreadState2& Thr
 			if (!ImportedPackageRef.HasPackage()) // If we found a package it's not actually missing but we can't load it anyway
 			{
 				UE_ASYNC_PACKAGE_CLOG(!ImportedPackageUPackageName.IsNone(), Display, Desc, TEXT("ImportPackages: SkipPackage"),
-					TEXT("Skipping non mounted imported package %s (0x%llX)"), *ImportedPackageNameToLoad.ToString(), ImportedPackageId.Value());
+					TEXT("Skipping non mounted imported package %s (0x%llX)"), *ImportedPackageNameToLoad.ToString(), ImportedPackageId.ValueForDebugging());
 				ImportedPackageRef.SetIsMissingPackage();
 			}
 			continue;
@@ -5038,20 +5047,20 @@ void FAsyncPackage2::ImportPackagesRecursiveInner(FAsyncLoadingThreadState2& Thr
 			ImportedPackage = AsyncLoadingThread.FindOrInsertPackage(ThreadState, PackageDesc, bInserted, this);
 		}
 
-		checkf(ImportedPackage, TEXT("Failed to find or insert imported package with id '0x%llX'"), ImportedPackageId.Value());
+		checkf(ImportedPackage, TEXT("Failed to find or insert imported package with id '%s'"), *FormatPackageId(ImportedPackageId));
 		TRACE_LOADTIME_ASYNC_PACKAGE_IMPORT_DEPENDENCY(this, ImportedPackage);
 
 		if (bInserted)
 		{
 			UE_ASYNC_PACKAGE_LOG(Verbose, Desc, TEXT("ImportPackages: AddPackage"),
-			TEXT("Start loading imported package with id '0x%llX'"), ImportedPackageId.ValueForDebugging());
+			TEXT("Start loading imported package with id '%s'"), *FormatPackageId(ImportedPackageId));
 			++AsyncLoadingThread.PackagesWithRemainingWorkCounter;
 			TRACE_COUNTER_SET(AsyncLoadingPackagesWithRemainingWork, AsyncLoadingThread.PackagesWithRemainingWorkCounter);
 		}
 		else
 		{
 			UE_ASYNC_PACKAGE_LOG_VERBOSE(VeryVerbose, Desc, TEXT("ImportPackages: UpdatePackage"),
-				TEXT("Imported package with id '0x%llX' is already being loaded."), ImportedPackageId.ValueForDebugging());
+				TEXT("Imported package with id '%s' is already being loaded."), *FormatPackageId(ImportedPackageId));
 		}
 		ImportedPackage->AddRef();
 		Header.ImportedAsyncPackagesView[LocalImportedPackageIndex] = ImportedPackage;
@@ -6604,8 +6613,8 @@ void FAsyncPackage2::EventDrivenCreateExport(const FAsyncPackageHeaderData& Head
 		FlagsToSet |= EInternalObjectFlags::LoaderImport;
 		ImportStore.StoreGlobalObject(Desc.UPackageId, Export.PublicExportHash, Object);
 
-		UE_ASYNC_PACKAGE_LOG_VERBOSE(VeryVerbose, Desc, TEXT("CreateExport"), TEXT("Created %s export %s. Tracked as 0x%llX:0x%llX"),
-			Object->HasAnyFlags(RF_Public) ? TEXT("public") : TEXT("private"), *Object->GetPathName(), Desc.UPackageId.Value(), Export.PublicExportHash);
+		UE_ASYNC_PACKAGE_LOG_VERBOSE(VeryVerbose, Desc, TEXT("CreateExport"), TEXT("Created %s export %s. Tracked as %s:0x%llX"),
+			Object->HasAnyFlags(RF_Public) ? TEXT("public") : TEXT("private"), *Object->GetPathName(), *FormatPackageId(Desc.UPackageId), Export.PublicExportHash);
 	}
 	else
 	{
