When using zenloader, mark gc'ed object as invalid in the linker
either the entire asset is getting gc'ed and they will be recreated on reload
or they are truly invalid and shouldn't be recreated once gc'ed now that the package is marked fully loaded

diff --git a/Engine/Source/Runtime/CoreUObject/Private/Serialization/AsyncLoading2.cpp b/Engine/Source/Runtime/CoreUObject/Private/Serialization/AsyncLoading2.cpp
--- a/Engine/Source/Runtime/CoreUObject/Private/Serialization/AsyncLoading2.cpp
+++ b/Engine/Source/Runtime/CoreUObject/Private/Serialization/AsyncLoading2.cpp
@@ -8381,25 +8381,14 @@ void FAsyncLoadingThread2::CollectUnreachableObjects(
 			// Normally done from UObject::BeginDestroy but we need to do it already here
 			if (FLinkerLoad* ObjectLinker = Object->GetLinker())
 			{
-#if WITH_EDITORONLY_DATA
-				// Make sure the linker entry stays as 'bExportLoadFailed' if the entry was marked as such, 
-				// Not doing this will cause crashes when resolving circular imports!
 				const int32 CachedLinkerIndex = Object->GetLinkerIndex();
-				bool bLinkerEntryWasInvalid = false;
-				if (ObjectLinker->ExportMap.IsValidIndex(CachedLinkerIndex))
-				{
-					FObjectExport& ObjExport = ObjectLinker->ExportMap[CachedLinkerIndex];
-					bLinkerEntryWasInvalid = ObjExport.bExportLoadFailed;
-				}
-#endif // WITH_EDITORONLY_DATA
 				Object->SetLinker(nullptr, INDEX_NONE);
-#if WITH_EDITORONLY_DATA
-				if (bLinkerEntryWasInvalid)
-				{
-					FObjectExport& ObjExport = ObjectLinker->ExportMap[CachedLinkerIndex];
-					ObjExport.bExportLoadFailed = true;
-				}
-#endif // WITH_EDITORONLY_DATA
+				// As we are garbaging the object, mark it as invalid in the linker
+				// Either it is now truly invalid to access this object 
+				// (i.e the asset ran upgraded, migrated away from this object and doesn't hold any references to it anymore.)
+				// or we are gc'ing the entire asset in which case the linker will eventually get purged and those entry won't be marked invalid anymore on recreation
+				FObjectExport& ObjExport = ObjectLinker->ExportMap[CachedLinkerIndex];
+				ObjExport.bExportLoadFailed = true;
 			}
 #endif // ALT2_ENABLE_LINKERLOAD_SUPPORT
 		});
