Better error when a parent goes missing

diff --git a/Engine/Source/Runtime/CoreUObject/Private/UObject/LinkerLoad.cpp b/Engine/Source/Runtime/CoreUObject/Private/UObject/LinkerLoad.cpp
--- a/Engine/Source/Runtime/CoreUObject/Private/UObject/LinkerLoad.cpp
+++ b/Engine/Source/Runtime/CoreUObject/Private/UObject/LinkerLoad.cpp
@@ -4831,6 +4831,24 @@ bool FLinkerLoad::WillTextureBeLoaded( UClass* Class, int32 ExportIndex )
 	}
 }
 
+#if WITH_EDITORONLY_DATA
+namespace UE::Private {
+static FString GetPackageObjectFullName(FLinkerLoad* Linker, const FPackageIndex Index)
+{
+	if (Index.IsImport())
+	{
+		return Linker->GetImportFullName(Index);
+	}
+	else if (Index.IsExport())
+	{
+		return Linker->GetExportFullName(Index);
+	}
+
+	return TEXT("none");
+}
+}
+#endif
+
 bool FLinkerLoad::IsPackageReferenceAllowed(UPackage* InPackage)
 {
 	if (InPackage && !InPackage->IsExternallyReferenceable())
@@ -4943,7 +4961,8 @@ UObject* FLinkerLoad::CreateExport( int32 Index )
 					
 					if (!bFailedToLoadGeneratedStruct && !FLinkerLoad::IsKnownMissingPackage(*GetExportFullName(Index)))
 					{
-						UE_ASSET_LOG(LogLinker, Warning, PackagePath, TEXT("CreateExport: Failed to load Parent for %s"), *GetExportFullName(Index));
+						using namespace UE::Private;
+						UE_ASSET_LOG(LogLinker, Warning, PackagePath, TEXT("CreateExport: Failed to load %s as Parent for %s - both will fail to load"), *GetPackageObjectFullName(this, Export.SuperIndex), *GetExportFullName(Index));
 					}
 #endif
 					return nullptr;
