Deferred initialization on editor load when a non-native default subobject type override has a circular load dependency with the outer object's underlying non-native type.

diff --git a/Engine/Source/Runtime/CoreUObject/Private/Blueprint/BlueprintSupport.cpp b/Engine/Source/Runtime/CoreUObject/Private/Blueprint/BlueprintSupport.cpp
--- a/Engine/Source/Runtime/CoreUObject/Private/Blueprint/BlueprintSupport.cpp
+++ b/Engine/Source/Runtime/CoreUObject/Private/Blueprint/BlueprintSupport.cpp
@@ -2883,6 +2883,12 @@ FObjectInitializer* FDeferredObjInitializationHelper::DeferObjectInitializerIfNe
 			return true;
 		};
 
+		auto IsObjectLoadPending = [](const UObject* InObject)
+		{
+			return InObject &&
+				(InObject->HasAnyFlags(RF_NeedLoad) || (InObject->HasAnyFlags(RF_WasLoaded) && !InObject->HasAnyFlags(RF_LoadCompleted)));
+		};
+
 		const bool bIsCDO = TargetObj->HasAnyFlags(RF_ClassDefaultObject);
 		if (bIsCDO)
 		{
@@ -2937,17 +2943,37 @@ FObjectInitializer* FDeferredObjInitializationHelper::DeferObjectInitializerIfNe
 			if (!DeferredInitializerCopy)
 			{
 				UObject* Archetype = DeferringInitializer.GetArchetype();
-			
-				const bool bArchetypeLoadPending = Archetype &&
-					( Archetype->HasAnyFlags(RF_NeedLoad) || (Archetype->HasAnyFlags(RF_WasLoaded) && !Archetype->HasAnyFlags(RF_LoadCompleted)) );
-
-				if (bArchetypeLoadPending)
+				if (IsObjectLoadPending(Archetype))
 				{
 					FDeferredSubObjInitializationTracker& SubObjInitDeferalSys = FDeferredSubObjInitializationTracker::Get();
 					DeferredInitializerCopy = SubObjInitDeferalSys.Add(Archetype, DeferringInitializer);
 				}
 			}
 		}
+		else if (TargetObj->HasAnyFlags(RF_DefaultSubObject))
+		{
+			// Since users can override default subobject types with non-native subtypes from the editor side, we need to
+			// ensure its non-native CDO has been fully serialized before we can allow those subobjects to be initialized.
+			// Deferral can occur e.g. when the non-native subtype contains a strong reference to a non-native owner type,
+			// resulting in a circular load dependency that can manifest if the non-native subobject type is loaded first.
+			// In that case, we'll then defer that subobject's initialization until after we've serialized its type's CDO.
+			const UClass* SubobjectClass = TargetObj->GetClass();
+			if (ensure(SubobjectClass) && !SubobjectClass->IsNative())
+			{
+				DEFERRED_DEPENDENCY_CHECK(SubobjectClass->HasAnyClassFlags(CLASS_CompiledFromBlueprint));
+
+				// Grab the subobject type's CDO and verify that it's what we expect. At this point, any placeholder export
+				// should at least have been created/resolved to an actual object, but may not be fully serialized just yet.
+				UObject* SubobjectCDO = SubobjectClass->GetDefaultObject(false);
+				DEFERRED_DEPENDENCY_CHECK(SubobjectCDO && SubobjectCDO->HasAnyFlags(RF_ClassDefaultObject));
+
+				if (IsObjectLoadPending(SubobjectCDO))
+				{
+					FDeferredSubObjInitializationTracker& SubObjInitDeferalSys = FDeferredSubObjInitializationTracker::Get();
+					DeferredInitializerCopy = SubObjInitDeferalSys.Add(SubobjectCDO, DeferringInitializer);
+				}
+			}
+		}
 	}
 
 	return DeferredInitializerCopy;
