Cleaning up and fixing Garbage Collector log output related to Reachability Analysis performance

diff --git a/Engine/Source/Runtime/CoreUObject/Private/UObject/GarbageCollection.cpp b/Engine/Source/Runtime/CoreUObject/Private/UObject/GarbageCollection.cpp
--- a/Engine/Source/Runtime/CoreUObject/Private/UObject/GarbageCollection.cpp
+++ b/Engine/Source/Runtime/CoreUObject/Private/UObject/GarbageCollection.cpp
@@ -5122,7 +5122,8 @@ void PreCollectGarbageImpl(EObjectFlags KeepFlags)
 		}
 
 
-		UE_LOG(LogGarbage, Log, TEXT("%sCollecting garbage%s"),
+		UE_LOG(LogGarbage, Log, TEXT("%s%sCollecting garbage%s"),
+			GIsIncrementalReachabilityPending ? TEXT("Resuming ") : TEXT(""),
 			(!bPerformFullPurge && GAllowIncrementalReachability) ? TEXT("Incrementally ") : TEXT(""),
 			IsAsyncLoading() ? TEXT(" while async loading") : TEXT(""));
 
@@ -5159,7 +5160,7 @@ void PreCollectGarbageImpl(EObjectFlags KeepFlags)
 					VerifyClustersAssumptions();
 				}
 				VerifyObjectFlagMirroring();
-				UE_LOG(LogGarbage, Log, TEXT("%f ms for Verify GC Assumptions"), (FPlatformTime::Seconds() - StartTime) * 1000);
+				UE_LOG(LogGarbage, Log, TEXT("%.2f ms for Verify GC Assumptions"), (FPlatformTime::Seconds() - StartTime) * 1000);
 			}
 #endif
 		}
@@ -5361,14 +5362,9 @@ void FReachabilityAnalysisState::PerformReachabilityAnalysisAndConditionallyPurg
 		{
 			if (!GIsIncrementalReachabilityPending)
 			{
-				UE_LOG(LogGarbage, Log, TEXT("Starting incremental GC"));
 				ReferenceProcessingTotalTime = 0.0;
 				IncrementalMarkPhaseTotalTime = 0.0;
 			}
-			else
-			{
-				UE_LOG(LogGarbage, Log, TEXT("Resuming incremental GC"));
-			}
 
 			PerformReachabilityAnalysis();
 		}
@@ -5381,19 +5377,18 @@ void FReachabilityAnalysisState::PerformReachabilityAnalysisAndConditionallyPurg
 		ReferenceProcessingTotalTime += ReferenceProcessingElapsedTime;
 		IncrementalMarkPhaseTotalTime += ElapsedTime;
 
-		if (GIsIncrementalReachabilityPending)
+		if (UE_LOG_ACTIVE(LogGarbage, Log))
 		{
-			UE_LOG(LogGarbage, Log, TEXT("GC Reachability Analysis iteration time: %f.2 ms (%f.2 ms on reference traversal)"), ElapsedTime * 1000, ReferenceProcessingElapsedTime * 1000);
-			const double SuspendLatency = CurrentTime - (IterationStartTime + IterationTimeLimit);
-			UE_LOG(LogGarbage, Log, TEXT("GC suspend latency: %f ms"), SuspendLatency * 1000.0);
-		}
-		else
-		{
-			const double ReferenceProcessingTotalTimeMs = ReferenceProcessingTotalTime * 1000;
-			const double IncrementalMarkPhaseTotalTimeMs = IncrementalMarkPhaseTotalTime * 1000;			
-			UE_LOG(LogGarbage, Log, TEXT("GC Reachability Analysis total time: %f.2 ms (%f.2 ms on reference traversal)"), IncrementalMarkPhaseTotalTimeMs, ReferenceProcessingTotalTimeMs);
-			if (UE_LOG_ACTIVE(LogGarbage, Log))
+			if (GIsIncrementalReachabilityPending)
+			{
+				const double SuspendLatency = FMath::Max(0.0, CurrentTime - (IterationStartTime + IterationTimeLimit));
+				UE_LOG(LogGarbage, Log, TEXT("GC Reachability iteration time: %.2f ms (%.2f ms on reference traversal, latency: %.3f ms)"), ElapsedTime * 1000, ReferenceProcessingElapsedTime * 1000, SuspendLatency * 1000.0);
+			}
+			else
 			{
+				const double ReferenceProcessingTotalTimeMs = ReferenceProcessingTotalTime * 1000;
+				const double IncrementalMarkPhaseTotalTimeMs = IncrementalMarkPhaseTotalTime * 1000;
+				UE_LOG(LogGarbage, Log, TEXT("GC Reachability Analysis total time: %.2f ms (%.2f ms on reference traversal)"), IncrementalMarkPhaseTotalTimeMs, ReferenceProcessingTotalTimeMs);
 				FString ExtraDetail = WITH_VERSE_VM ? FString::Printf(TEXT("and %d verse cells"), Stats.NumVerseCells) : FString();
 				UE_LOG(LogGarbage, Log, TEXT("%.2f ms for %sGC - %d refs/ms while processing %d references from %d objects %s with %d clusters"),
 					IncrementalMarkPhaseTotalTimeMs,
