Implemented the legacy loader LoadPackage with a param struct
Reintroduce the usage in ConditionalFlushAsyncLoadingForLinkers

diff --git a/Engine/Source/Runtime/CoreUObject/Private/Serialization/AsyncLoading.cpp b/Engine/Source/Runtime/CoreUObject/Private/Serialization/AsyncLoading.cpp
--- a/Engine/Source/Runtime/CoreUObject/Private/Serialization/AsyncLoading.cpp
+++ b/Engine/Source/Runtime/CoreUObject/Private/Serialization/AsyncLoading.cpp
@@ -7437,6 +7437,21 @@ bool FAsyncLoadingThread::ShouldAlwaysLoadPackageAsync(const FPackagePath& InPac
 	return FPlatformProperties::RequiresCookedData() && GEventDrivenLoaderEnabled && EVENT_DRIVEN_ASYNC_LOAD_ACTIVE_AT_RUNTIME;
 }
 
+int32 FAsyncLoadingThread::LoadPackage(const FPackagePath& PackagePath, FLoadPackageAsyncOptionalParams OptionalParams)
+{
+	if (OptionalParams.ProgressDelegate.IsValid())
+	{
+		UE_LOG(LogStreaming, Warning, TEXT("Progress delegate is only supported for zenloader. A CompletionDelegate should be used instead for this loader."));
+	}
+
+	FLoadPackageAsyncDelegate CompletionDelegate;
+	if (OptionalParams.CompletionDelegate.IsValid())
+	{
+		CompletionDelegate = MoveTemp(*OptionalParams.CompletionDelegate.Get());
+	}
+	return LoadPackage(PackagePath, OptionalParams.CustomPackageName, MoveTemp(CompletionDelegate), OptionalParams.PackageFlags, OptionalParams.PIEInstanceID, OptionalParams.PackagePriority, OptionalParams.InstancingContext, OptionalParams.LoadFlags);
+}
+
 int32 FAsyncLoadingThread::LoadPackage(const FPackagePath& InPackagePath, FName InCustomName, FLoadPackageAsyncDelegate InCompletionDelegate, EPackageFlags InPackageFlags, int32 InPIEInstanceID, int32 InPackagePriority, const FLinkerInstancingContext* InInstancingContext, uint32 InLoadFlags)
 {
 	checkf(IsInGameThread(), TEXT("LoadPackageAsync is only thread-safe when using the zenloader (i.e. AsyncLoading2)."));
diff --git a/Engine/Source/Runtime/CoreUObject/Private/Serialization/AsyncLoadingThread.h b/Engine/Source/Runtime/CoreUObject/Private/Serialization/AsyncLoadingThread.h
--- a/Engine/Source/Runtime/CoreUObject/Private/Serialization/AsyncLoadingThread.h
+++ b/Engine/Source/Runtime/CoreUObject/Private/Serialization/AsyncLoadingThread.h
@@ -381,6 +381,8 @@ public:
 
 	bool ShouldAlwaysLoadPackageAsync(const FPackagePath& InPackagePath) override;
 
+	int32 LoadPackage(const FPackagePath& PackagePath, FLoadPackageAsyncOptionalParams OptionalParams) override;
+
 	int32 LoadPackage(
 			const FPackagePath& InPackagePath,
 			FName InCustomName,
diff --git a/Engine/Source/Runtime/CoreUObject/Private/UObject/Linker.cpp b/Engine/Source/Runtime/CoreUObject/Private/UObject/Linker.cpp
--- a/Engine/Source/Runtime/CoreUObject/Private/UObject/Linker.cpp
+++ b/Engine/Source/Runtime/CoreUObject/Private/UObject/Linker.cpp
@@ -900,7 +900,7 @@ void ConditionalFlushAsyncLoadingForLinkers(TConstArrayView<FLinkerLoad*> InLink
 		TArray<int32, TInlineAllocator<4>> RequestIds;
 		for (FLinkerLoad* Linker : InLinkers)
 		{
-			int32 Request = LoadPackageAsync(Linker->GetPackagePath());
+			int32 Request = LoadPackageAsync(Linker->GetPackagePath(), FLoadPackageAsyncOptionalParams{ .PackagePriority = MAX_int32 });
 			RequestIds.Add(Request);
 		}
 		FlushAsyncLoading(RequestIds);
