FlushAsyncLoading prior to disallowing any new requests as requests done during the scope of the flush are going to get flushed as well.
Directly use the loader to issue sync load from LoadPackage

diff --git a/Engine/Source/Runtime/CoreUObject/Private/UObject/UObjectGlobals.cpp b/Engine/Source/Runtime/CoreUObject/Private/UObject/UObjectGlobals.cpp
--- a/Engine/Source/Runtime/CoreUObject/Private/UObject/UObjectGlobals.cpp
+++ b/Engine/Source/Runtime/CoreUObject/Private/UObject/UObjectGlobals.cpp
@@ -1661,9 +1661,15 @@ UPackage* LoadPackageInternal(UPackage* InOuter, const FPackagePath& PackagePath
 			PackageFlags |= PKG_PlayInEditor;
 		}
 #endif
-		constexpr int32 PIEInstanceID = INDEX_NONE;
-		constexpr int32 Priority = INT32_MAX;
-		int32 RequestID = LoadPackageAsync(PackagePath, PackageName, FLoadPackageAsyncDelegate(), PackageFlags, PIEInstanceID, Priority, InstancingContext, LoadFlags);
+		FLoadPackageAsyncOptionalParams OptionalParams
+		{
+			.CustomPackageName = PackageName,
+			.PackageFlags = PackageFlags,
+			.PackagePriority = INT32_MAX,
+			.InstancingContext = InstancingContext,
+			.LoadFlags = LoadFlags
+		};
+		int32 RequestID = LoadPackageAsync(PackagePath, MoveTemp(OptionalParams));
 
 		if (RequestID != INDEX_NONE)
 		{
diff --git a/Engine/Source/Runtime/Launch/Private/LaunchEngineLoop.cpp b/Engine/Source/Runtime/Launch/Private/LaunchEngineLoop.cpp
--- a/Engine/Source/Runtime/Launch/Private/LaunchEngineLoop.cpp
+++ b/Engine/Source/Runtime/Launch/Private/LaunchEngineLoop.cpp
@@ -5053,9 +5053,6 @@ void FEngineLoop::Exit()
 
 	// Make sure we're not in the middle of loading something.
 	{
-		// From now on it's not allowed to request new async loads
-		SetAsyncLoadingAllowed(false);
-
 		bool bFlushOnExit = true;
 		if (GConfig)
 		{
@@ -5070,6 +5067,9 @@ void FEngineLoop::Exit()
 		{
 			CancelAsyncLoading();
 		}
+		// From now on it's not allowed to request new async loads
+		// any new requests done during the scope of the flush should have been flushed as well, now prevent any new requests
+		SetAsyncLoadingAllowed(false);
 	}
 
 	// Block till all outstanding resource streaming requests are fulfilled.
