Set serialized bulk data flags after save

diff --git a/Engine/Source/Runtime/CoreUObject/Private/UObject/LinkerSave.cpp b/Engine/Source/Runtime/CoreUObject/Private/UObject/LinkerSave.cpp
--- a/Engine/Source/Runtime/CoreUObject/Private/UObject/LinkerSave.cpp
+++ b/Engine/Source/Runtime/CoreUObject/Private/UObject/LinkerSave.cpp
@@ -795,7 +795,7 @@ void FLinkerSave::OnPostSaveBulkData()
 		{
 			FBulkData& BulkData = *Kv.Key;
 			const FObjectDataResource& DataResource = DataResourceMap[Kv.Value];
-			BulkData.SetFlagsFromDiskWrittenValues(static_cast<EBulkDataFlags>(DataResource.Flags), DataResource.SerialOffset, DataResource.SerialSize, Summary.BulkDataStartOffset);
+			BulkData.SetFlagsFromDiskWrittenValues(static_cast<EBulkDataFlags>(DataResource.LegacyBulkDataFlags), DataResource.SerialOffset, DataResource.SerialSize, Summary.BulkDataStartOffset);
 		}
 	}
 #endif
