Allow copy paste instanced references between objects from different packages
Support Object and ObjectPtr
Move FObjectPtrProperty::AllowCrossLevel() to FObjectProperty::AllowCrossLevel()

diff --git a/Engine/Source/Runtime/CoreUObject/Private/UObject/PropertyObject.cpp b/Engine/Source/Runtime/CoreUObject/Private/UObject/PropertyObject.cpp
--- a/Engine/Source/Runtime/CoreUObject/Private/UObject/PropertyObject.cpp
+++ b/Engine/Source/Runtime/CoreUObject/Private/UObject/PropertyObject.cpp
@@ -92,6 +92,11 @@ EConvertFromTypeResult FObjectProperty::ConvertFromType(const FPropertyTag& Tag,
 	return EConvertFromTypeResult::UseSerializeItem;
 }
 
+bool FObjectProperty::AllowCrossLevel() const
+{
+	return HasAnyPropertyFlags(CPF_InstancedReference);
+}
+
 void FObjectProperty::SerializeItem(FStructuredArchive::FSlot Slot, void* Value, void const* Defaults) const
 {
 	FArchive& UnderlyingArchive = Slot.GetUnderlyingArchive();
diff --git a/Engine/Source/Runtime/CoreUObject/Private/UObject/PropertyObjectPtr.cpp b/Engine/Source/Runtime/CoreUObject/Private/UObject/PropertyObjectPtr.cpp
--- a/Engine/Source/Runtime/CoreUObject/Private/UObject/PropertyObjectPtr.cpp
+++ b/Engine/Source/Runtime/CoreUObject/Private/UObject/PropertyObjectPtr.cpp
@@ -211,11 +211,6 @@ bool FObjectPtrProperty::AllowObjectTypeReinterpretationTo(const FObjectProperty
 	return Other && Other->IsA<FObjectProperty>();
 }
 
-bool FObjectPtrProperty::AllowCrossLevel() const
-{
-	return HasAnyPropertyFlags(CPF_InstancedReference);
-}
-
 uint32 FObjectPtrProperty::GetValueTypeHashInternal(const void* Src) const
 {
 	return GetTypeHash((FObjectPtr&)GetPropertyValue(Src));
