Disable s.AsyncLoadingAlwaysProcessPackages by default for the legacy loader as it is a workaround for undiagnosed editor sync load deadlock
It however can cause issues when loading a mix of cooked and uncooked content.

diff --git a/Engine/Source/Runtime/CoreUObject/Private/Serialization/AsyncLoading.cpp b/Engine/Source/Runtime/CoreUObject/Private/Serialization/AsyncLoading.cpp
--- a/Engine/Source/Runtime/CoreUObject/Private/Serialization/AsyncLoading.cpp
+++ b/Engine/Source/Runtime/CoreUObject/Private/Serialization/AsyncLoading.cpp
@@ -571,7 +571,7 @@ static FAutoConsoleVariableRef CVarAsyncLoadingThreadEnabled(
 	ECVF_Default
 	);
 
-static int32 GAsyncLoadingAlwaysProcessPackages = 1;
+static int32 GAsyncLoadingAlwaysProcessPackages = 0;
 static FAutoConsoleVariableRef CVarAsyncLoadingAlwaysProcessPackages(
 	TEXT("s.AsyncLoadingAlwaysProcessPackages"),
 	GAsyncLoadingAlwaysProcessPackages,
