MetaData is not loaded in -game, the old implementation of metadata preloading confusing made it look like it was always loaded. This null check prevents a crash in -game, it's likely we don't need to preload metadata for any good reason

diff --git a/Engine/Source/Runtime/CoreUObject/Private/Blueprint/BlueprintSupport.cpp b/Engine/Source/Runtime/CoreUObject/Private/Blueprint/BlueprintSupport.cpp
--- a/Engine/Source/Runtime/CoreUObject/Private/Blueprint/BlueprintSupport.cpp
+++ b/Engine/Source/Runtime/CoreUObject/Private/Blueprint/BlueprintSupport.cpp
@@ -729,14 +729,20 @@ bool FLinkerLoad::RegenerateBlueprintClass(UClass* LoadClass, UObject* ClassDefa
 		ForcePreloadObject(*ClassSourceObject);
 
 #if WITH_EDITORONLY_DATA
+		// We likely don't need to load meta data here at all! but we have been doing so 
+		// since 2080292 - subtly the code from 2080292 wouldn't assert about missing metadata
+		// but only because UPackage::GetMetaData creates a dummy UMetaData object that is
+		// tagged as RF_LoadCompleted. Consider removing the forced metadata creation.
 		int32 MetadataIndex = LoadMetaDataFromExportMap(true);
 
 		// Older content may not have a metadata object in its package.
 		if (MetadataIndex != INDEX_NONE)
 		{
 			const FObjectExport& MetadataExport = Exp(FPackageIndex::FromExport(MetadataIndex));
-			check(MetadataExport.Object);
-			ForcePreloadObject(*MetadataExport.Object);
+			if(MetadataExport.Object) // metadata not loaded in -game, has not been at least since UE4
+			{
+				ForcePreloadObject(*MetadataExport.Object);
+			}
 		}
 #endif
 
