Fixed up metaprogramming in TObjectPtr.

diff --git a/Engine/Source/Runtime/Core/Public/Templates/NonNullPointer.h b/Engine/Source/Runtime/Core/Public/Templates/NonNullPointer.h
--- a/Engine/Source/Runtime/Core/Public/Templates/NonNullPointer.h
+++ b/Engine/Source/Runtime/Core/Public/Templates/NonNullPointer.h
@@ -12,37 +12,6 @@
 class FArchive;
 enum class EDefaultConstructNonNullPtr { UnsafeDoNotUse }; // So we can construct TNonNullPtrs
 
-namespace UE::Core::Private::NonNullPtr {
-/**
- * Version of `::TPointerIsConvertibleFromTo` that produces an incomplete type
- * when either `From` or `To` are incomplete types
- */
-template <typename, typename, typename = void>
-struct TPointerIsConvertibleFromTo;
-
-/**
- * Specialization of
- * `UE::Core::Private::NonNullPtr::TPointerIsConvertibleFromTo` for complete
- * non-function types
- */
-template <typename From, typename To>
-struct TPointerIsConvertibleFromTo<From, To, std::void_t<decltype(sizeof(From)), decltype(sizeof(To))>>
-{
-	static constexpr bool Value = std::is_convertible_v<From*, To*>;
-};
-
-/**
- * Specialization of
- * `UE::Core::Private::NonNullPtr::TPointerIsConvertibleFromTo` for function
- * types, which are always complete types
- */
-template <typename Result1, typename... Args1, typename Result2, typename... Args2>
-struct TPointerIsConvertibleFromTo<Result1(Args1...), Result2(Args2...)>
-{
-	static constexpr bool Value = std::is_convertible_v<Result1(*)(Args1...), Result2(*)(Args2...)>;
-};
-}
-
 /**
  * TNonNullPtr is a non-nullable, non-owning, raw/naked/unsafe pointer.
  */
@@ -81,8 +50,8 @@ public:
 	 * Constructs a non-null pointer from another non-null pointer
 	 */
 	template <
-		typename OtherObjectType
-		UE_REQUIRES(UE::Core::Private::NonNullPtr::TPointerIsConvertibleFromTo<OtherObjectType, ObjectType>::Value)
+		typename OtherObjectType,
+		decltype(ImplicitConv<ObjectType*>((OtherObjectType*)nullptr))* = nullptr
 	>
 	FORCEINLINE TNonNullPtr(const TNonNullPtr<OtherObjectType>& Other)
 		: Object(Other.Object)
@@ -113,8 +82,8 @@ public:
 	 * Assignment operator taking another TNonNullPtr
 	 */
 	template <
-		typename OtherObjectType
-		UE_REQUIRES(UE::Core::Private::NonNullPtr::TPointerIsConvertibleFromTo<OtherObjectType, ObjectType>::Value)
+		typename OtherObjectType,
+		decltype(ImplicitConv<ObjectType*>((OtherObjectType*)nullptr))* = nullptr
 	>
 	FORCEINLINE TNonNullPtr& operator=(const TNonNullPtr<OtherObjectType>& Other)
 	{
diff --git a/Engine/Source/Runtime/CoreUObject/Public/Templates/IsTObjectPtr.h b/Engine/Source/Runtime/CoreUObject/Public/Templates/IsTObjectPtr.h
--- a/Engine/Source/Runtime/CoreUObject/Public/Templates/IsTObjectPtr.h
+++ b/Engine/Source/Runtime/CoreUObject/Public/Templates/IsTObjectPtr.h
@@ -5,13 +5,14 @@
 template <typename T>
 struct TObjectPtr;
 
+template <typename T> constexpr bool TIsTObjectPtr_V                               = false;
+template <typename T> constexpr bool TIsTObjectPtr_V<               TObjectPtr<T>> = true;
+template <typename T> constexpr bool TIsTObjectPtr_V<const          TObjectPtr<T>> = true;
+template <typename T> constexpr bool TIsTObjectPtr_V<      volatile TObjectPtr<T>> = true;
+template <typename T> constexpr bool TIsTObjectPtr_V<const volatile TObjectPtr<T>> = true;
+
 template <typename T>
 struct TIsTObjectPtr
 {
-	enum { Value = false };
+	static constexpr bool Value = TIsTObjectPtr_V<T>;
 };
-
-template <typename T> struct TIsTObjectPtr<               TObjectPtr<T>> { enum { Value = true }; };
-template <typename T> struct TIsTObjectPtr<const          TObjectPtr<T>> { enum { Value = true }; };
-template <typename T> struct TIsTObjectPtr<      volatile TObjectPtr<T>> { enum { Value = true }; };
-template <typename T> struct TIsTObjectPtr<const volatile TObjectPtr<T>> { enum { Value = true }; };
diff --git a/Engine/Source/Runtime/CoreUObject/Public/UObject/ObjectPtr.h b/Engine/Source/Runtime/CoreUObject/Public/UObject/ObjectPtr.h
--- a/Engine/Source/Runtime/CoreUObject/Public/UObject/ObjectPtr.h
+++ b/Engine/Source/Runtime/CoreUObject/Public/UObject/ObjectPtr.h
@@ -329,13 +329,10 @@ namespace ObjectPtr_Private
 	/** Coerce to pointer through implicit conversion to CommonPointerType where CommonPointerType is deduced, and must be a C++ pointer, not a wrapper type. */
 	template <
 		typename T,
-		typename U,
-		typename CommonPointerType =  decltype(std::declval<bool>() ? std::declval<const T*>() : std::declval<U>()),
-		std::enable_if_t<
-			std::is_pointer<CommonPointerType>::value
-		>* = nullptr
+		typename U
+		UE_REQUIRES(std::is_pointer_v<std::common_type_t<const T*, U>>)
 	>
-	FORCEINLINE auto CoerceToPointer(U&& Other) -> CommonPointerType
+	FORCEINLINE std::common_type_t<const T*, U> CoerceToPointer(const U& Other)
 	{
 		return Other;
 	}
@@ -343,13 +340,10 @@ namespace ObjectPtr_Private
 	/** Coerce to pointer through the use of a ".Get()" member, which is the convention within Unreal smart pointer types. */
 	template <
 		typename T,
-		typename U,
-		std::enable_if_t<
-			!TIsTObjectPtr<std::decay_t<U>>::Value,
-			decltype(std::declval<U>().Get())
-		>* = nullptr
+		typename U
+		UE_REQUIRES(!TIsTObjectPtr_V<U>)
 	>
-	FORCEINLINE auto CoerceToPointer(U&& Other) -> decltype(std::declval<U>().Get())
+	FORCEINLINE auto CoerceToPointer(const U& Other) -> decltype(Other.Get())
 	{
 		return Other.Get();
 	}
@@ -373,13 +367,10 @@ namespace ObjectPtr_Private
 	/** Perform shallow equality check between a TObjectPtr and another (non TObjectPtr) type that we can coerce to a pointer. */
 	template <
 		typename T,
-		typename U,
-		std::enable_if_t<
-			!TIsTObjectPtr<std::decay_t<U>>::Value,
-			decltype(CoerceToPointer<T>(std::declval<U>()) == std::declval<const T*>())
-		>* = nullptr
+		typename U
+		UE_REQUIRES(!TIsTObjectPtr_V<U>)
 	>
-	bool IsObjectPtrEqual(const TObjectPtr<T>& Ptr, U&& Other)
+	auto IsObjectPtrEqual(const TObjectPtr<T>& Ptr, const U& Other) -> decltype(CoerceToPointer<T>(Other) == std::declval<const T*>())
 	{
 		// This function deliberately avoids the tracking code path as we are only doing
 		// a shallow pointer comparison.
@@ -417,7 +408,7 @@ struct TObjectPtr
 	// This means that the following are invalid and must fail to compile:
 	// - TObjectPtr<int>
 	// - TObjectPtr<IInterface>
-	static_assert(std::disjunction<std::negation<std::bool_constant<sizeof(ObjectPtr_Private::ResolveTypeIsComplete<T>(1)) == 2>>, std::is_base_of<UObject, T>>::value, "TObjectPtr<T> can only be used with types derived from UObject");
+	static_assert(std::disjunction<std::bool_constant<sizeof(ObjectPtr_Private::ResolveTypeIsComplete<T>(1)) != 2>, std::is_base_of<UObject, T>>::value, "TObjectPtr<T> can only be used with types derived from UObject");
 #endif
 
 public:
@@ -469,12 +460,10 @@ public:
 
 	template <
 		typename U,
-		std::enable_if_t<
-			!TIsTObjectPtr<std::decay_t<U>>::Value,
-			decltype(ImplicitConv<T*>(std::declval<U>()))
-		>* = nullptr
+		decltype(ImplicitConv<T*>(std::declval<U>()))* = nullptr
+		UE_REQUIRES(!TIsTObjectPtr_V<std::decay_t<U>>)
 	>
-	FORCEINLINE TObjectPtr(U&& Object)
+	FORCEINLINE TObjectPtr(const U& Object)
 		: ObjectPtr(const_cast<std::remove_const_t<T>*>(ImplicitConv<T*>(Object)))
 	{
 	}
@@ -524,10 +513,8 @@ public:
 
 	template <
 		typename U,
-		std::enable_if_t<
-			!TIsTObjectPtr<std::decay_t<U>>::Value,
-			decltype(ImplicitConv<T*>(std::declval<U>()))
-		>* = nullptr
+		decltype(ImplicitConv<T*>(std::declval<U>()))* = nullptr
+		UE_REQUIRES(!TIsTObjectPtr_V<std::decay_t<U>>)
 	>
 	FORCEINLINE TObjectPtr<T>& operator=(U&& Object)
 	{
@@ -544,7 +531,7 @@ public:
 	// Equality/Inequality comparisons against other TObjectPtr
 	template <
 		typename U,
-		typename Base = std::decay_t<decltype(false ? std::declval<std::decay_t<T*>>() : std::declval<std::decay_t<U*>>())>
+		typename Base = std::common_type_t<T*, U*>
 	>
 	FORCEINLINE bool operator==(const TObjectPtr<U>& Other) const
 	{
@@ -560,35 +547,18 @@ public:
 	// Equality/Inequality comparisons against another type that can be implicitly converted to the pointer type kept in a TObjectPtr
 	template <
 		typename U,
-		typename = decltype(ObjectPtr_Private::IsObjectPtrEqual(std::declval<const TObjectPtr<T>&>(), std::declval<U&&>()))
+		typename = decltype(ObjectPtr_Private::IsObjectPtrEqual(std::declval<const TObjectPtr<T>&>(), std::declval<const U&>()))
 	>
-	FORCEINLINE bool operator==(U&& Other) const
+	FORCEINLINE bool operator==(const U& Other) const
 	{
 		return ObjectPtr_Private::IsObjectPtrEqual(*this, Other);
 	}
 
 #if __cplusplus < 202002L
-	template <
-		typename U,
-		typename Base = std::decay_t<decltype(false ? std::declval<std::decay_t<T*>>() : std::declval<std::decay_t<U*>>())>
-	>
-	FORCEINLINE bool operator!=(const TObjectPtr<U>& Other) const
-	{
-		return ObjectPtr != Other.ObjectPtr;
-	}
-
-	FORCEINLINE bool operator!=(TYPE_OF_NULLPTR) const
-	{
-		return ObjectPtr.operator bool();
-	}
-
-	template <
-		typename U,
-		typename = decltype(ObjectPtr_Private::IsObjectPtrEqual(std::declval<const TObjectPtr<T>&>(), std::declval<U&&>()))
-	>
-	FORCEINLINE bool operator!=(U&& Other) const
+	template <typename U>
+	FORCEINLINE auto operator!=(const U& Other) const -> decltype(!(*this == Other))
 	{
-		return !ObjectPtr_Private::IsObjectPtrEqual(*this, Other);
+		return !(*this == Other);
 	}
 #endif
 
@@ -801,18 +771,18 @@ FORCEINLINE void operator<<(FStructuredArchiveSlot Slot, TObjectPtr<T>& InObject
 template <
 	typename T,
 	typename U,
-	decltype(ObjectPtr_Private::IsObjectPtrEqual(std::declval<const TObjectPtr<T>&>(), std::declval<U&&>()))* = nullptr
+	decltype(ObjectPtr_Private::IsObjectPtrEqual(std::declval<const TObjectPtr<T>&>(), std::declval<const U&>()))* = nullptr
 >
-FORCEINLINE bool operator==(U&& Other, const TObjectPtr<T>& Ptr)
+FORCEINLINE bool operator==(const U& Other, const TObjectPtr<T>& Ptr)
 {
 	return ObjectPtr_Private::IsObjectPtrEqual(Ptr, Other);
 }
 template <
 	typename T,
 	typename U,
-	decltype(ObjectPtr_Private::IsObjectPtrEqual(std::declval<const TObjectPtr<T>&>(), std::declval<U&&>()))* = nullptr
+	decltype(ObjectPtr_Private::IsObjectPtrEqual(std::declval<const TObjectPtr<T>&>(), std::declval<const U&>()))* = nullptr
 >
-FORCEINLINE bool operator!=(U&& Other, const TObjectPtr<T>& Ptr)
+FORCEINLINE bool operator!=(const U& Other, const TObjectPtr<T>& Ptr)
 {
 	return !ObjectPtr_Private::IsObjectPtrEqual(Ptr, Other);
 }
@@ -904,8 +874,8 @@ FORCEINLINE T** ToRawPtrArrayUnsafe(T** ArrayOfPtr)
 PRAGMA_DISABLE_DEPRECATION_WARNINGS
 template <
 	typename ArrayType,
-	typename ArrayTypeNoRef = std::remove_reference_t<ArrayType>,
-	std::enable_if_t<TIsTArray_V<ArrayTypeNoRef>>* = nullptr
+	typename ArrayTypeNoRef = std::remove_reference_t<ArrayType>
+	UE_REQUIRES(TIsTArray_V<ArrayTypeNoRef>)
 >
 #if UE_DEPRECATE_MUTABLE_TOBJECTPTR
 const auto&
@@ -1383,8 +1353,8 @@ namespace UE::Core::Private // private facilities; not for direct use
 		static void Close(ViewType& View)
 		{
 #if UE_OBJECT_PTR_GC_BARRIER
-			static constexpr bool bKeyReference = TIsTObjectPtr<K>::Value;
-			static constexpr bool bValueReference = TIsTObjectPtr<V>::Value;
+			static constexpr bool bKeyReference = TIsTObjectPtr_V<K>;
+			static constexpr bool bValueReference = TIsTObjectPtr_V<V>;
 			static_assert(bKeyReference || bValueReference);
 			if (UE::GC::Private::GIsIncrementalReachabilityPending)
 			{
@@ -1648,7 +1618,7 @@ public:
 	 */
 	template <
 		typename OtherObjectType,
-		typename = typename TEnableIf<UE::Core::Private::NonNullPtr::TPointerIsConvertibleFromTo<OtherObjectType, ObjectType>::Value>::Type
+		decltype(ImplicitConv<ObjectType*>((OtherObjectType*)nullptr))* = nullptr
 	>
 	FORCEINLINE TNonNullPtr(const TNonNullPtr<OtherObjectType>& Other)
 		: Object(Other.Object)
@@ -1678,8 +1648,11 @@ public:
 	/**
 	 * Assignment operator taking another TNonNullPtr
 	 */
-	template <typename OtherObjectType>
-	FORCEINLINE typename TEnableIf<UE::Core::Private::NonNullPtr::TPointerIsConvertibleFromTo<OtherObjectType, ObjectType>::Value, TNonNullPtr&>::Type operator=(const TNonNullPtr<OtherObjectType>& Other)
+	template <
+		typename OtherObjectType,
+		decltype(ImplicitConv<ObjectType*>((OtherObjectType*)nullptr))* = nullptr
+	>
+	FORCEINLINE TNonNullPtr& operator=(const TNonNullPtr<OtherObjectType>& Other)
 	{
 		Object = Other.Object;
 		return *this;
