Renamed parameters and variables from 'Index' to 'InternalIndex' in FScriptSetHelper and FScriptMapHelper in preparation for another change.

diff --git a/Engine/Source/Runtime/CoreUObject/Public/UObject/UnrealType.h b/Engine/Source/Runtime/CoreUObject/Public/UObject/UnrealType.h
--- a/Engine/Source/Runtime/CoreUObject/Public/UObject/UnrealType.h
+++ b/Engine/Source/Runtime/CoreUObject/Public/UObject/UnrealType.h
@@ -3892,7 +3892,7 @@ public:
 	virtual void EmitReferenceInfo(UE::GC::FSchemaBuilder& Schema, int32 BaseOffset, TArray<const FStructProperty*>& EncounteredStructProps, UE::GC::FPropertyStack& DebugPath) override;
 	virtual bool SameType(const FProperty* Other) const override;
 	virtual EConvertFromTypeResult ConvertFromType(const FPropertyTag& Tag, FStructuredArchive::FSlot Slot, uint8* Data, UStruct* DefaultsStruct, const uint8* Defaults) override;
-	virtual void* GetValueAddressAtIndex_Direct(const FProperty* Inner, void* InValueAddress, int32 Index) const override;
+	virtual void* GetValueAddressAtIndex_Direct(const FProperty* Inner, void* InValueAddress, int32 LogicalIndex) const override;
 	// End of FProperty interface
 
 	FString GetCPPTypeCustom(FString* ExtendedTypeText, uint32 CPPExportFlags, const FString& KeyTypeText, const FString& InKeyExtendedTypeText, const FString& ValueTypeText, const FString& InValueExtendedTypeText) const;
@@ -3919,18 +3919,18 @@ public:
 	 * Helper function to check if the specified index of a key/value pair in the underlying set is valid.
 	 * Used by the garbage collector where for performance reasons the provided map pointer is not guarded
 	 */
-	bool IsValidIndex(void* InMap, int32 Index) const
+	bool IsValidIndex(void* InMap, int32 InternalIndex) const
 	{
-		return WithScriptMap(InMap, [Index](auto* Map) { return Map->IsValidIndex(Index); });
+		return WithScriptMap(InMap, [InternalIndex](auto* Map) { return Map->IsValidIndex(InternalIndex); });
 	}
 
 	/*
 	 * Helper function to get the pointer to a key/value pair at the specified index.
 	 * Used by the garbage collector where for performance reasons the provided map pointer is not guarded
 	 */
-	uint8* GetPairPtr(void* InMap, int32 Index) const
+	uint8* GetPairPtr(void* InMap, int32 InternalIndex) const
 	{
-		return WithScriptMap(InMap, [this, Index](auto* Map) { return (uint8*)Map->GetData(Index, MapLayout); });
+		return WithScriptMap(InMap, [this, InternalIndex](auto* Map) { return (uint8*)Map->GetData(InternalIndex, MapLayout); });
 	}
 
 	const FProperty* GetKeyProperty() const
@@ -4018,7 +4018,7 @@ public:
 	virtual void EmitReferenceInfo(UE::GC::FSchemaBuilder& Schema, int32 BaseOffset, TArray<const FStructProperty*>& EncounteredStructProps, UE::GC::FPropertyStack& DebugPath) override;
 	virtual bool SameType(const FProperty* Other) const override;
 	virtual EConvertFromTypeResult ConvertFromType(const FPropertyTag& Tag, FStructuredArchive::FSlot Slot, uint8* Data, UStruct* DefaultsStruct, const uint8* Defaults) override;
-	virtual void* GetValueAddressAtIndex_Direct(const FProperty* Inner, void* InValueAddress, int32 Index) const override;
+	virtual void* GetValueAddressAtIndex_Direct(const FProperty* Inner, void* InValueAddress, int32 LogicalIndex) const override;
 	// End of FProperty interface
 
 	FString GetCPPTypeCustom(FString* ExtendedTypeText, uint32 CPPExportFlags, const FString& ElementTypeText, const FString& InElementExtendedTypeText) const;
@@ -4046,20 +4046,20 @@ public:
 	 * Helper function to check if the specified index of an element is valid.
 	 * Used by the garbage collector where for performance reasons the provided set pointer is not guarded
 	 */
-	bool IsValidIndex(void* InSet, int32 Index) const
+	bool IsValidIndex(void* InSet, int32 InternalIndex) const
 	{
 		FScriptSet* Set = (FScriptSet*)InSet;
-		return Set->IsValidIndex(Index);
+		return Set->IsValidIndex(InternalIndex);
 	}
 
 	/*
 	 * Helper function to get the pointer to an element at the specified index.
 	 * Used by the garbage collector where for performance reasons the provided set pointer is not guarded
 	 */
-	uint8* GetElementPtr(void* InSet, int32 Index) const
+	uint8* GetElementPtr(void* InSet, int32 InternalIndex) const
 	{
 		FScriptSet* Set = (FScriptSet*)InSet;
-		return (uint8*)Set->GetData(Index, SetLayout);
+		return (uint8*)Set->GetData(InternalIndex, SetLayout);
 	}
 #if WITH_EDITORONLY_DATA
 	virtual void AppendSchemaHash(FBlake3& Builder, bool bSkipEditorOnly) const override;
@@ -4500,13 +4500,13 @@ public:
 	/**
 	 * Index range check
 	 *
-	 * @param  Index  Index to check
+	 * @param InternalIndex Index to check
 	 *
 	 * @return true if accessing this element is legal.
 	 */
-	FORCEINLINE bool IsValidIndex(int32 Index) const
+	FORCEINLINE bool IsValidIndex(int32 InternalIndex) const
 	{
-		return WithScriptMap([Index](auto* Map) { return Map->IsValidIndex(Index); });
+		return WithScriptMap([InternalIndex](auto* Map) { return Map->IsValidIndex(InternalIndex); });
 	}
 
 	/**
@@ -4539,22 +4539,22 @@ public:
 	/**
 	 * Returns a uint8 pointer to the pair in the map
 	 *
-	 * @param  Index  index of the item to return a pointer to.
+	 * @param InternalIndex index of the item to return a pointer to.
 	 *
 	 * @return Pointer to the pair, or nullptr if the map is empty.
 	 */
-	FORCEINLINE uint8* GetPairPtr(int32 Index)
+	FORCEINLINE uint8* GetPairPtr(int32 InternalIndex)
 	{
-		return WithScriptMap([this, Index](auto* Map) -> uint8*
+		return WithScriptMap([this, InternalIndex](auto* Map) -> uint8*
 		{
 			if (Map->Num() == 0)
 			{
-				checkSlow(!Index);
+				checkSlow(!InternalIndex);
 				return nullptr;
 			}
 
-			checkSlow(Map->IsValidIndex(Index));
-			return (uint8*)Map->GetData(Index, MapLayout);
+			checkSlow(Map->IsValidIndex(InternalIndex));
+			return (uint8*)Map->GetData(InternalIndex, MapLayout);
 		});
 	}
 
@@ -4563,57 +4563,57 @@ public:
 	 * identical to GetPairPtr, but provides clarity of purpose and avoids exposing
 	 * implementation details of TMap.
 	 *
-	 * @param  Index  index of the item to return a pointer to.
+	 * @param InternalIndex index of the item to return a pointer to.
 	 *
 	 * @return Pointer to the key, or nullptr if the map is empty.
 	 */
-	FORCEINLINE uint8* GetKeyPtr(int32 Index)
+	FORCEINLINE uint8* GetKeyPtr(int32 InternalIndex)
 	{
-		return WithScriptMap([this, Index](auto* Map) -> uint8*
+		return WithScriptMap([this, InternalIndex](auto* Map) -> uint8*
 		{
 			if (Map->Num() == 0)
 			{
-				checkSlow(!Index);
+				checkSlow(!InternalIndex);
 				return nullptr;
 			}
 		
-			checkSlow(Map->IsValidIndex(Index));
-			return (uint8*)Map->GetData(Index, MapLayout);
+			checkSlow(Map->IsValidIndex(InternalIndex));
+			return (uint8*)Map->GetData(InternalIndex, MapLayout);
 		});
 	}
 
 	/**
 	 * Returns a uint8 pointer to the Value (second element) in the map.
 	 *
-	 * @param  Index  index of the item to return a pointer to.
+	 * @param InternalIndex index of the item to return a pointer to.
 	 *
 	 * @return Pointer to the value, or nullptr if the map is empty.
 	 */
-	FORCEINLINE uint8* GetValuePtr(int32 Index)
+	FORCEINLINE uint8* GetValuePtr(int32 InternalIndex)
 	{
-		return WithScriptMap([this, Index](auto* Map) -> uint8*
+		return WithScriptMap([this, InternalIndex](auto* Map) -> uint8*
 		{
 			if (Map->Num() == 0)
 			{
-				checkSlow(!Index);
+				checkSlow(!InternalIndex);
 				return nullptr;
 			}
 		
-			checkSlow(Map->IsValidIndex(Index));
-			return (uint8*)Map->GetData(Index, MapLayout) + MapLayout.ValueOffset;
+			checkSlow(Map->IsValidIndex(InternalIndex));
+			return (uint8*)Map->GetData(InternalIndex, MapLayout) + MapLayout.ValueOffset;
 		});
 	}
 
 	/**
 	 * Returns a uint8 pointer to the pair in the map.
 	 *
-	 * @param  Index  index of the item to return a pointer to.
+	 * @param InternalIndex index of the item to return a pointer to.
 	 *
 	 * @return Pointer to the pair, or nullptr if the map is empty.
 	 */
-	FORCEINLINE const uint8* GetPairPtr(int32 Index) const
+	FORCEINLINE const uint8* GetPairPtr(int32 InternalIndex) const
 	{
-		return const_cast<FScriptMapHelper*>(this)->GetPairPtr(Index);
+		return const_cast<FScriptMapHelper*>(this)->GetPairPtr(InternalIndex);
 	}
 	
 	/**
@@ -4624,10 +4624,10 @@ public:
 	*/
 	uint8* FindNthPairPtr(int32 N)
 	{
-		const int32 Index = FindInternalIndex(N);
+		const int32 InternalIndex = FindInternalIndex(N);
 		
-		checkSlow(IsValidIndex(Index));
-		return (Index != INDEX_NONE) ? GetPairPtr(Index) : nullptr;
+		checkSlow(IsValidIndex(InternalIndex));
+		return (InternalIndex != INDEX_NONE) ? GetPairPtr(InternalIndex) : nullptr;
 	}
 	
 	/**
@@ -4638,8 +4638,8 @@ public:
 	*/
 	uint8* FindNthKeyPtr(int32 N)
 	{
-		const int32 Index = FindInternalIndex(N);
-		return (Index != INDEX_NONE) ? GetKeyPtr(Index) : nullptr;
+		const int32 InternalIndex = FindInternalIndex(N);
+		return (InternalIndex != INDEX_NONE) ? GetKeyPtr(InternalIndex) : nullptr;
 	}
 	
 	/**
@@ -4650,8 +4650,8 @@ public:
 	*/
 	uint8* FindNthValuePtr(int32 N)
 	{
-		const int32 Index = FindInternalIndex(N);
-		return (Index != INDEX_NONE) ? GetValuePtr(Index) : nullptr;
+		const int32 InternalIndex = FindInternalIndex(N);
+		return (InternalIndex != INDEX_NONE) ? GetValuePtr(InternalIndex) : nullptr;
 	}
 	
 	/**
@@ -4662,8 +4662,8 @@ public:
 	*/
 	const uint8* FindNthPairPtr(int32 N) const
 	{
-		const int32 Index = FindInternalIndex(N);
-		return (Index != INDEX_NONE) ? GetPairPtr(Index) : nullptr;
+		const int32 InternalIndex = FindInternalIndex(N);
+		return (InternalIndex != INDEX_NONE) ? GetPairPtr(InternalIndex) : nullptr;
 	}
 
 	/**
@@ -4762,16 +4762,16 @@ public:
 	/**
 	 * Removes an element at the specified index, destroying it.
 	 *
-	 * @param  Index  The index of the element to remove.
+	 * @param InternalIndex The index of the element to remove.
 	 */
-	void RemoveAt(int32 Index, int32 Count = 1)
+	void RemoveAt(int32 InternalIndex, int32 Count = 1)
 	{
-		return WithScriptMap([this, Index, Count](auto* Map)
+		return WithScriptMap([this, InternalIndex, Count](auto* Map)
 		{
-			check(Map->IsValidIndex(Index));
+			check(Map->IsValidIndex(InternalIndex));
 
-			DestructItems(Index, Count);
-			for (int32 LocalCount = Count, LocalIndex = Index; LocalCount; ++LocalIndex)
+			DestructItems(InternalIndex, Count);
+			for (int32 LocalCount = Count, LocalIndex = InternalIndex; LocalCount; ++LocalIndex)
 			{
 				if (Map->IsValidIndex(LocalIndex))
 				{
@@ -4852,25 +4852,25 @@ public:
 
 			FProperty* LocalKeyProp = this->KeyProp; // prevent aliasing in loop below
 
-			int32 Index = IndexHint;
+			int32 InternalIndex = IndexHint;
 			for (;;)
 			{
-				if (Map->IsValidIndex(Index))
+				if (Map->IsValidIndex(InternalIndex))
 				{
-					const void* PairToSearch = Map->GetData(Index, MapLayout);
+					const void* PairToSearch = Map->GetData(InternalIndex, MapLayout);
 					if (LocalKeyProp->Identical(PairWithKeyToFind, PairToSearch))
 					{
-						return Index;
+						return InternalIndex;
 					}
 				}
 
-				++Index;
-				if (Index == MapMax)
+				++InternalIndex;
+				if (InternalIndex == MapMax)
 				{
-					Index = 0;
+					InternalIndex = 0;
 				}
 
-				if (Index == IndexHint)
+				if (InternalIndex == IndexHint)
 				{
 					return INDEX_NONE;
 				}
@@ -4888,15 +4888,15 @@ public:
 	 */
 	FORCEINLINE uint8* FindMapPairPtrWithKey(const void* PairWithKeyToFind, int32 IndexHint = 0)
 	{
-		int32 Index = FindMapIndexWithKey(PairWithKeyToFind, IndexHint);
-		uint8* Result = (Index >= 0) ? GetPairPtr(Index) : nullptr;
+		int32 InternalIndex = FindMapIndexWithKey(PairWithKeyToFind, IndexHint);
+		uint8* Result = (InternalIndex >= 0) ? GetPairPtr(InternalIndex) : nullptr;
 		return Result;
 	}
 
 	/** Finds the associated pair from hash, rather than linearly searching */
 	int32 FindMapPairIndexFromHash(const void* KeyPtr)
 	{
-		int32 Index = WithScriptMap([this, KeyPtr, LocalKeyPropForCapture = this->KeyProp](auto* Map)
+		int32 InternalIndex = WithScriptMap([this, KeyPtr, LocalKeyPropForCapture = this->KeyProp](auto* Map)
 		{
 			return Map->FindPairIndex(
 				KeyPtr,
@@ -4905,14 +4905,14 @@ public:
 				[LocalKeyPropForCapture](const void* A, const void* B) { return LocalKeyPropForCapture->Identical(A, B); }
 			);
 		});
-		return Index;
+		return InternalIndex;
 	}
 
 	/** Finds the associated pair from hash, rather than linearly searching */
 	uint8* FindMapPairPtrFromHash(const void* KeyPtr)
 	{
-		int32 Index = FindMapPairIndexFromHash(KeyPtr);
-		uint8* Result = (Index >= 0) ? GetPairPtr(Index) : nullptr;
+		int32 InternalIndex = FindMapPairIndexFromHash(KeyPtr);
+		uint8* Result = (InternalIndex >= 0) ? GetPairPtr(InternalIndex) : nullptr;
 		return Result;
 	}
 
@@ -5127,17 +5127,17 @@ private:
 	/**
 	 * Internal function to call into the property system to construct / initialize elements.
 	 *
-	 * @param  Index  First item to construct.
-	 * @param  Count  Number of items to construct.
+	 * @param InternalIndex First item to construct.
+	 * @param Count Number of items to construct.
 	 */
-	void ConstructItem(int32 Index)
+	void ConstructItem(int32 InternalIndex)
 	{
-		check(IsValidIndex(Index));
+		check(IsValidIndex(InternalIndex));
 
 		bool bZeroKey   = !!(KeyProp  ->PropertyFlags & CPF_ZeroConstructor);
 		bool bZeroValue = !!(ValueProp->PropertyFlags & CPF_ZeroConstructor);
 
-		void* Dest = WithScriptMap([this, Index](auto* Map) { return Map->GetData(Index, MapLayout); });
+		void* Dest = WithScriptMap([this, InternalIndex](auto* Map) { return Map->GetData(InternalIndex, MapLayout); });
 
 		if (bZeroKey || bZeroValue)
 		{
@@ -5159,9 +5159,9 @@ private:
 	/**
 	 * Internal function to call into the property system to destruct elements.
 	 */
-	void DestructItems(int32 Index, int32 Count)
+	void DestructItems(int32 InternalIndex, int32 Count)
 	{
-		check(Index >= 0);
+		check(InternalIndex >= 0);
 		check(Count >= 0);
 
 		if (Count == 0)
@@ -5175,14 +5175,14 @@ private:
 		if (bDestroyKeys || bDestroyValues)
 		{
 			uint32 Stride  = MapLayout.SetLayout.Size;
-			uint8* PairPtr = WithScriptMap([this, Index](auto* Map) { return (uint8*)Map->GetData(Index, MapLayout); });
+			uint8* PairPtr = WithScriptMap([this, InternalIndex](auto* Map) { return (uint8*)Map->GetData(InternalIndex, MapLayout); });
 			if (bDestroyKeys)
 			{
 				if (bDestroyValues)
 				{
-					for (; Count; ++Index)
+					for (; Count; ++InternalIndex)
 					{
-						if (IsValidIndex(Index))
+						if (IsValidIndex(InternalIndex))
 						{
 							KeyProp  ->DestroyValue_InContainer(PairPtr);
 							ValueProp->DestroyValue_InContainer(PairPtr);
@@ -5193,9 +5193,9 @@ private:
 				}
 				else
 				{
-					for (; Count; ++Index)
+					for (; Count; ++InternalIndex)
 					{
-						if (IsValidIndex(Index))
+						if (IsValidIndex(InternalIndex))
 						{
 							KeyProp->DestroyValue_InContainer(PairPtr);
 							--Count;
@@ -5206,9 +5206,9 @@ private:
 			}
 			else
 			{
-				for (; Count; ++Index)
+				for (; Count; ++InternalIndex)
 				{
-					if (IsValidIndex(Index))
+					if (IsValidIndex(InternalIndex))
 					{
 						ValueProp->DestroyValue_InContainer(PairPtr);
 						--Count;
@@ -5222,25 +5222,25 @@ private:
 	/**
 	 * Returns a uint8 pointer to the pair in the array without checking the index.
 	 *
-	 * @param  Index  index of the item to return a pointer to.
+	 * @param InternalIndex index of the item to return a pointer to.
 	 *
 	 * @return Pointer to the pair, or nullptr if the map is empty.
 	 */
-	FORCEINLINE uint8* GetPairPtrWithoutCheck(int32 Index)
+	FORCEINLINE uint8* GetPairPtrWithoutCheck(int32 InternalIndex)
 	{
-		return WithScriptMap([this, Index](auto* Map) { return (uint8*)Map->GetData(Index, MapLayout); });
+		return WithScriptMap([this, InternalIndex](auto* Map) { return (uint8*)Map->GetData(InternalIndex, MapLayout); });
 	}
 
 	/**
 	 * Returns a uint8 pointer to the pair in the array without checking the index.
 	 *
-	 * @param  Index  index of the item to return a pointer to.
+	 * @param InternalIndex index of the item to return a pointer to.
 	 *
 	 * @return Pointer to the pair, or nullptr if the map is empty.
 	 */
-	FORCEINLINE const uint8* GetPairPtrWithoutCheck(int32 Index) const
+	FORCEINLINE const uint8* GetPairPtrWithoutCheck(int32 InternalIndex) const
 	{
-		return const_cast<FScriptMapHelper*>(this)->GetPairPtrWithoutCheck(Index);
+		return const_cast<FScriptMapHelper*>(this)->GetPairPtrWithoutCheck(InternalIndex);
 	}
 
 public:
@@ -5289,13 +5289,13 @@ public:
 	/**
 	* Index range check
 	*
-	* @param  Index  Index to check
+	* @param InternalIndex Index to check
 	*
 	* @return true if accessing this element is legal.
 	*/
-	FORCEINLINE bool IsValidIndex(int32 Index) const
+	FORCEINLINE bool IsValidIndex(int32 InternalIndex) const
 	{
-		return Set->IsValidIndex(Index);
+		return Set->IsValidIndex(InternalIndex);
 	}
 
 	/**
@@ -5339,32 +5339,32 @@ public:
 	/**
 	* Returns a uint8 pointer to the element in the set.
 	*
-	* @param  Index  index of the item to return a pointer to.
+	* @param InternalIndex index of the item to return a pointer to.
 	*
 	* @return Pointer to the element, or nullptr if the set is empty.
 	*/
-	FORCEINLINE uint8* GetElementPtr(int32 Index)
+	FORCEINLINE uint8* GetElementPtr(int32 InternalIndex)
 	{
 		if (Num() == 0)
 		{
-			checkSlow(!Index);
+			checkSlow(!InternalIndex);
 			return nullptr;
 		}
 
-		checkSlow(IsValidIndex(Index));
-		return (uint8*)Set->GetData(Index, SetLayout);
+		checkSlow(IsValidIndex(InternalIndex));
+		return (uint8*)Set->GetData(InternalIndex, SetLayout);
 	}
 
 	/**
 	* Returns a uint8 pointer to the element in the set.
 	*
-	* @param  Index  index of the item to return a pointer to.
+	* @param InternalIndex index of the item to return a pointer to.
 	*
 	* @return Pointer to the element, or nullptr if the set is empty.
 	*/
-	FORCEINLINE const uint8* GetElementPtr(int32 Index) const
+	FORCEINLINE const uint8* GetElementPtr(int32 InternalIndex) const
 	{
-		return const_cast<FScriptSetHelper*>(this)->GetElementPtr(Index);
+		return const_cast<FScriptSetHelper*>(this)->GetElementPtr(InternalIndex);
 	}
 
 	/**
@@ -5375,8 +5375,8 @@ public:
 	*/
 	uint8* FindNthElementPtr(int32 N)
 	{
-		const int32 Index = FindInternalIndex(N);
-		return (Index != INDEX_NONE) ? GetElementPtr(Index) : nullptr;
+		const int32 InternalIndex = FindInternalIndex(N);
+		return (InternalIndex != INDEX_NONE) ? GetElementPtr(InternalIndex) : nullptr;
 	}
 
 	/**
@@ -5387,8 +5387,8 @@ public:
 	*/
 	const uint8* FindNthElementPtr(int32 N) const
 	{
-		const int32 Index = FindInternalIndex(N);
-		return (Index != INDEX_NONE) ? GetElementPtr(Index) : nullptr;
+		const int32 InternalIndex = FindInternalIndex(N);
+		return (InternalIndex != INDEX_NONE) ? GetElementPtr(InternalIndex) : nullptr;
 	}
 
 	/**
@@ -5462,18 +5462,18 @@ public:
 	/**
 	* Removes an element at the specified index, destroying it.
 	*
-	* @param  Index  The index of the element to remove.
+	* @param InternalIndex The index of the element to remove.
 	*/
-	void RemoveAt(int32 Index, int32 Count = 1)
+	void RemoveAt(int32 InternalIndex, int32 Count = 1)
 	{
-		check(IsValidIndex(Index));
+		check(IsValidIndex(InternalIndex));
 
-		DestructItems(Index, Count);
-		for (; Count; ++Index)
+		DestructItems(InternalIndex, Count);
+		for (; Count; ++InternalIndex)
 		{
-			if (IsValidIndex(Index))
+			if (IsValidIndex(InternalIndex))
 			{
-				Set->RemoveAt(Index, SetLayout);
+				Set->RemoveAt(InternalIndex, SetLayout);
 				--Count;
 			}
 		}
@@ -5543,25 +5543,25 @@ public:
 
 		FProperty* LocalKeyProp = this->ElementProp; // prevent aliasing in loop below
 
-		int32 Index = IndexHint;
+		int32 InternalIndex = IndexHint;
 		for (;;)
 		{
-			if (IsValidIndex(Index))
+			if (IsValidIndex(InternalIndex))
 			{
-				const void* ElementToCheck = GetElementPtrWithoutCheck(Index);
+				const void* ElementToCheck = GetElementPtrWithoutCheck(InternalIndex);
 				if (LocalKeyProp->Identical(ElementToFind, ElementToCheck))
 				{
-					return Index;
+					return InternalIndex;
 				}
 			}
 
-			++Index;
-			if (Index == SetMax)
+			++InternalIndex;
+			if (InternalIndex == SetMax)
 			{
-				Index = 0;
+				InternalIndex = 0;
 			}
 
-			if (Index == IndexHint)
+			if (InternalIndex == IndexHint)
 			{
 				return INDEX_NONE;
 			}
@@ -5578,8 +5578,8 @@ public:
 	*/
 	FORCEINLINE uint8* FindElementPtr(const void* ElementToFind, int32 IndexHint = 0)
 	{
-		const int32 Index = FindElementIndex(ElementToFind, IndexHint);
-		uint8* Result = (Index >= 0 ? GetElementPtr(Index) : nullptr);
+		const int32 InternalIndex = FindElementIndex(ElementToFind, IndexHint);
+		uint8* Result = (InternalIndex >= 0 ? GetElementPtr(InternalIndex) : nullptr);
 		return Result;
 	}
 
@@ -5598,8 +5598,8 @@ public:
 	/** Finds element pointer from hash, rather than linearly searching */
 	FORCEINLINE uint8* FindElementPtrFromHash(const void* ElementToFind)
 	{
-		const int32 Index = FindElementIndexFromHash(ElementToFind);
-		uint8* Result = (Index >= 0 ? GetElementPtr(Index) : nullptr);
+		const int32 InternalIndex = FindElementIndexFromHash(ElementToFind);
+		uint8* Result = (InternalIndex >= 0 ? GetElementPtr(InternalIndex) : nullptr);
 		return Result;
 	}
 
@@ -5716,15 +5716,15 @@ private:
 	/**
 	* Internal function to call into the property system to construct / initialize elements.
 	*
-	* @param  Index  First item to construct.
-	* @param  Count  Number of items to construct.
+	* @param InternalIndex First item to construct.
+	* @param Count Number of items to construct.
 	*/
-	void ConstructItem(int32 Index)
+	void ConstructItem(int32 InternalIndex)
 	{
-		check(IsValidIndex(Index));
+		check(IsValidIndex(InternalIndex));
 
 		bool bZeroElement = !!(ElementProp->PropertyFlags & CPF_ZeroConstructor);
-		uint8* Dest = GetElementPtrWithoutCheck(Index);
+		uint8* Dest = GetElementPtrWithoutCheck(InternalIndex);
 
 		if (bZeroElement)
 		{
@@ -5741,9 +5741,9 @@ private:
 	/**
 	* Internal function to call into the property system to destruct elements.
 	*/
-	void DestructItems(int32 Index, int32 Count)
+	void DestructItems(int32 InternalIndex, int32 Count)
 	{
-		check(Index >= 0);
+		check(InternalIndex >= 0);
 		check(Count >= 0);
 
 		if (Count == 0)
@@ -5756,11 +5756,11 @@ private:
 		if (bDestroyElements)
 		{
 			uint32 Stride = SetLayout.Size;
-			uint8* ElementPtr = GetElementPtrWithoutCheck(Index);
+			uint8* ElementPtr = GetElementPtrWithoutCheck(InternalIndex);
 
-			for (; Count; ++Index)
+			for (; Count; ++InternalIndex)
 			{
-				if (IsValidIndex(Index))
+				if (IsValidIndex(InternalIndex))
 				{
 					ElementProp->DestroyValue_InContainer(ElementPtr);
 					--Count;
@@ -5773,25 +5773,25 @@ private:
 	/**
 	* Returns a uint8 pointer to the element in the array without checking the index.
 	*
-	* @param  Index  index of the item to return a pointer to.
+	* @param InternalIndex index of the item to return a pointer to.
 	*
 	* @return Pointer to the element, or nullptr if the array is empty.
 	*/
-	FORCEINLINE uint8* GetElementPtrWithoutCheck(int32 Index)
+	FORCEINLINE uint8* GetElementPtrWithoutCheck(int32 InternalIndex)
 	{
-		return (uint8*)Set->GetData(Index, SetLayout);
+		return (uint8*)Set->GetData(InternalIndex, SetLayout);
 	}
 
 	/**
 	* Returns a uint8 pointer to the element in the array without checking the index.
 	*
-	* @param  Index  index of the item to return a pointer to.
+	* @param InternalIndex index of the item to return a pointer to.
 	*
 	* @return Pointer to the pair, or nullptr if the array is empty.
 	*/
-	FORCEINLINE const uint8* GetElementPtrWithoutCheck(int32 Index) const
+	FORCEINLINE const uint8* GetElementPtrWithoutCheck(int32 InternalIndex) const
 	{
-		return const_cast<FScriptSetHelper*>(this)->GetElementPtrWithoutCheck(Index);
+		return const_cast<FScriptSetHelper*>(this)->GetElementPtrWithoutCheck(InternalIndex);
 	}
 
 public:
