Fixed TScriptInterface not being able to access the Serialize function from FScriptInterface

#rnx

diff --git a/Engine/Source/Runtime/CoreUObject/Public/UObject/ScriptInterface.h b/Engine/Source/Runtime/CoreUObject/Public/UObject/ScriptInterface.h
--- a/Engine/Source/Runtime/CoreUObject/Public/UObject/ScriptInterface.h
+++ b/Engine/Source/Runtime/CoreUObject/Public/UObject/ScriptInterface.h
@@ -31,10 +31,11 @@ private:
 	 */
 	void*		InterfacePointer = nullptr;
 
+protected:
 	/**
 	 * Serialize ScriptInterface
 	 */
-	FArchive& Serialize(FArchive& Ar, class UClass* InterfaceType);
+	COREUOBJECT_API FArchive& Serialize(FArchive& Ar, class UClass* InterfaceType);
 
 public:
 	/**
