Implemented the LoadPackage function with optional param struct for the cooked editor legacy loader

diff --git a/Engine/Source/Runtime/CoreUObject/Private/Serialization/EditorPackageLoader.cpp b/Engine/Source/Runtime/CoreUObject/Private/Serialization/EditorPackageLoader.cpp
--- a/Engine/Source/Runtime/CoreUObject/Private/Serialization/EditorPackageLoader.cpp
+++ b/Engine/Source/Runtime/CoreUObject/Private/Serialization/EditorPackageLoader.cpp
@@ -66,6 +66,21 @@ public:
 		return !bDoesUncookedPackageExist;
 	}
 
+	virtual int32 LoadPackage(const FPackagePath& PackagePath, FLoadPackageAsyncOptionalParams OptionalParams) override
+	{
+		if (OptionalParams.ProgressDelegate.IsValid())
+		{
+			UE_LOG(LogStreaming, Warning, TEXT("Progress delegate is only supported for zenloader. A CompletionDelegate should be used instead for this loader."));
+		}
+
+		FLoadPackageAsyncDelegate CompletionDelegate;
+		if (OptionalParams.CompletionDelegate.IsValid())
+		{
+			CompletionDelegate = MoveTemp(*OptionalParams.CompletionDelegate.Get());
+		}
+		return LoadPackage(PackagePath, OptionalParams.CustomPackageName, MoveTemp(CompletionDelegate), OptionalParams.PackageFlags, OptionalParams.PIEInstanceID, OptionalParams.PackagePriority, OptionalParams.InstancingContext, OptionalParams.LoadFlags);
+	}
+
 	virtual int32 LoadPackage(
 		const FPackagePath& PackagePath,
 		FName CustomPackageName,
