combining FObjectPropertyPtr and FObjectProperty and FClassPropertyPtr and FClassProperty together
simplifies usage of Object Properties and fixes a few edge cases

see zousar.shaker for issues

diff --git a/Engine/Plugins/Experimental/AnimNext/Source/AnimNext/Private/Param/ParamTypeHandle.cpp b/Engine/Plugins/Experimental/AnimNext/Source/AnimNext/Private/Param/ParamTypeHandle.cpp
--- a/Engine/Plugins/Experimental/AnimNext/Source/AnimNext/Private/Param/ParamTypeHandle.cpp
+++ b/Engine/Plugins/Experimental/AnimNext/Source/AnimNext/Private/Param/ParamTypeHandle.cpp
@@ -394,7 +394,7 @@ FParamTypeHandle FParamTypeHandle::FromProperty(const FProperty* InProperty)
 			Handle.SetParameterType(EParamType::AnimNextGraphReferencePose);
 		}
 	}
-	else if(InProperty->IsA<FObjectProperty>() || InProperty->IsA<FObjectPtrProperty>())
+	else if(InProperty->IsA<FObjectProperty>())
 	{
 		const UClass* Class = CastField<FObjectPropertyBase>(InProperty)->PropertyClass;
 		if (Class == UObject::StaticClass())
@@ -423,7 +423,7 @@ FParamTypeHandle FParamTypeHandle::FromProperty(const FProperty* InProperty)
 		else if (InProperty->IsA<FObjectPropertyBase>())
 		{
 			UClass* Class = CastField<FObjectPropertyBase>(InProperty)->PropertyClass;
-			if (InProperty->IsA<FObjectProperty>() || InProperty->IsA<FObjectPtrProperty>())
+			if (InProperty->IsA<FObjectProperty>())
 			{
 				Handle.SetParameterType(EParamType::Custom);
 				Handle.SetCustomTypeIndex(GetOrAllocateCustomTypeIndex(FAnimNextParamType::EValueType::Object, FAnimNextParamType::EContainerType::None, Class));
@@ -496,7 +496,7 @@ FParamTypeHandle FParamTypeHandle::FromProperty(const FProperty* InProperty)
 			else if (InnerProperty->IsA<FObjectPropertyBase>())
 			{
 				UClass* Class = CastField<FObjectPropertyBase>(InnerProperty)->PropertyClass;
-				if (InnerProperty->IsA<FObjectProperty>() || InnerProperty->IsA<FObjectPtrProperty>())
+				if (InnerProperty->IsA<FObjectProperty>())
 				{
 					Handle.SetParameterType(EParamType::Custom);
 					Handle.SetCustomTypeIndex(GetOrAllocateCustomTypeIndex(FAnimNextParamType::EValueType::Object, FAnimNextParamType::EContainerType::Array, Class));
diff --git a/Engine/Plugins/Experimental/PCG/Source/PCG/Private/Tests/Accessors/PCGAttributeExtractorTest.cpp b/Engine/Plugins/Experimental/PCG/Source/PCG/Private/Tests/Accessors/PCGAttributeExtractorTest.cpp
--- a/Engine/Plugins/Experimental/PCG/Source/PCG/Private/Tests/Accessors/PCGAttributeExtractorTest.cpp
+++ b/Engine/Plugins/Experimental/PCG/Source/PCG/Private/Tests/Accessors/PCGAttributeExtractorTest.cpp
@@ -424,7 +424,7 @@ bool FPCGAttributePropertyMultipleDepthTest::RunTest(const FString& Parameters)
 	TestStruct.DepthStruct.FloatValue = 0.324f;
 	TestStruct.DepthStruct.Depth2Struct.IntValue = 5;
 
-	const FObjectPtrProperty* ObjectProperty = CastField<FObjectPtrProperty>(FPCGAttributeExtractorTestStruct::StaticStruct()->FindPropertyByName(GET_MEMBER_NAME_CHECKED(FPCGAttributeExtractorTestStruct, Object)));
+	const FObjectProperty* ObjectProperty = CastField<FObjectProperty>(FPCGAttributeExtractorTestStruct::StaticStruct()->FindPropertyByName(GET_MEMBER_NAME_CHECKED(FPCGAttributeExtractorTestStruct, Object)));
 	const FProperty* DepthStructProperty = FPCGAttributeExtractorTestStruct::StaticStruct()->FindPropertyByName(GET_MEMBER_NAME_CHECKED(FPCGAttributeExtractorTestStruct, DepthStruct));
 	const FProperty* Depth2StructProperty = FPCGAttributeExtractorTestStructDepth1::StaticStruct()->FindPropertyByName(GET_MEMBER_NAME_CHECKED(FPCGAttributeExtractorTestStructDepth1, Depth2Struct));
 
diff --git a/Engine/Plugins/Experimental/StructUtils/Source/StructUtils/Private/PropertyBag.cpp b/Engine/Plugins/Experimental/StructUtils/Source/StructUtils/Private/PropertyBag.cpp
--- a/Engine/Plugins/Experimental/StructUtils/Source/StructUtils/Private/PropertyBag.cpp
+++ b/Engine/Plugins/Experimental/StructUtils/Source/StructUtils/Private/PropertyBag.cpp
@@ -163,10 +163,6 @@ namespace UE::StructUtils::Private
 
 			return EPropertyBagPropertyType::Object;
 		}
-		if (CastField<FObjectPtrProperty>(InSourceProperty))
-		{
-			return EPropertyBagPropertyType::Object;
-		}
 		if (CastField<FSoftObjectProperty>(InSourceProperty))
 		{
 			if (CastField<FSoftClassProperty>(InSourceProperty))
@@ -212,10 +208,6 @@ namespace UE::StructUtils::Private
 
 			return ObjectProperty->PropertyClass;
 		}
-		if (const FObjectPtrProperty* ObjectProperty = CastField<FObjectPtrProperty>(InSourceProperty))
-		{
-			return ObjectProperty->PropertyClass;
-		}
 		if (const FSoftObjectProperty* SoftObjectProperty = CastField<FSoftObjectProperty>(InSourceProperty))
 		{
 			if (const FSoftClassProperty* SoftClassProperty = CastField<FSoftClassProperty>(InSourceProperty))
@@ -405,13 +397,13 @@ namespace UE::StructUtils::Private
 		case EPropertyBagPropertyType::Object:
 			if (const UClass* Class = Cast<UClass>(Desc.ValueTypeObject))
 			{
-				FObjectPtrProperty* Prop = new FObjectPtrProperty(PropertyScope, Desc.Name, RF_Public);
+				FObjectProperty* Prop = new FObjectProperty(PropertyScope, Desc.Name, RF_Public);
 				if (Class->HasAnyClassFlags(CLASS_DefaultToInstanced))
 				{
 					Prop->SetPropertyFlags(CPF_InstancedReference);
 				}
 				Prop->SetPropertyClass(const_cast<UClass*>(Class));
-				Prop->SetPropertyFlags(CPF_HasGetValueTypeHash);
+				Prop->SetPropertyFlags(CPF_HasGetValueTypeHash | CPF_TObjectPtr);
 				return Prop;
 			}
 			break;
diff --git a/Engine/Plugins/FX/Niagara/Source/NiagaraEditorWidgets/Private/SNiagaraSystemEffectTypeBar.cpp b/Engine/Plugins/FX/Niagara/Source/NiagaraEditorWidgets/Private/SNiagaraSystemEffectTypeBar.cpp
--- a/Engine/Plugins/FX/Niagara/Source/NiagaraEditorWidgets/Private/SNiagaraSystemEffectTypeBar.cpp
+++ b/Engine/Plugins/FX/Niagara/Source/NiagaraEditorWidgets/Private/SNiagaraSystemEffectTypeBar.cpp
@@ -167,12 +167,12 @@ TSharedPtr<SWidget> SNiagaraSystemEffectTypeBar::CreatePropertyValueWidget(FProp
 			ValueTextBlock->SetText(BoolValue ? FText::FromString("True") : FText::FromString("False"));
 		}
 				
-		if(FObjectPtrProperty* ObjectProperty = CastField<FObjectPtrProperty>(Property))
+		if(FObjectProperty* ObjectProperty = CastField<FObjectProperty>(Property))
 		{
-			FObjectPtr& ObjectPtr = (FObjectPtr&) ObjectProperty->GetPropertyValue_InContainer(EffectType.Get());
+			const TObjectPtr<UObject>& ObjectPtr = ObjectProperty->GetPropertyValue_InContainer(EffectType.Get());
 			
 			if(ObjectPtr)
-			{				
+			{
 				if(ObjectProperty->HasMetaData("DisplayClassDisplayName"))
 				{
 					ValueTextBlock->SetText(GetClassDisplayName(ObjectPtr.Get()));
@@ -180,7 +180,7 @@ TSharedPtr<SWidget> SNiagaraSystemEffectTypeBar::CreatePropertyValueWidget(FProp
 				else
 				{
 					ValueTextBlock->SetText(GetObjectName(ObjectPtr.Get()));
-				}						
+				}
 			}
 			else
 			{
diff --git a/Engine/Plugins/Tests/CQTest/Source/CQTest/Public/ObjectBuilder.h b/Engine/Plugins/Tests/CQTest/Source/CQTest/Public/ObjectBuilder.h
--- a/Engine/Plugins/Tests/CQTest/Source/CQTest/Public/ObjectBuilder.h
+++ b/Engine/Plugins/Tests/CQTest/Source/CQTest/Public/ObjectBuilder.h
@@ -182,7 +182,7 @@ private:
 		}
 		else if constexpr (TIsTObjectPtr<Raw>::Value)
 		{
-			if (auto AsProperty = CastField<FObjectPtrProperty>(&Prop))
+			if (auto AsProperty = CastField<FObjectProperty>(&Prop))
 			{
 				return Raw::ElementType::StaticClass()->IsChildOf(AsProperty->PropertyClass);
 			}
diff --git a/Engine/Plugins/VirtualProduction/RemoteControl/Source/WebRemoteControl/Private/RemoteControlModels.h b/Engine/Plugins/VirtualProduction/RemoteControl/Source/WebRemoteControl/Private/RemoteControlModels.h
--- a/Engine/Plugins/VirtualProduction/RemoteControl/Source/WebRemoteControl/Private/RemoteControlModels.h
+++ b/Engine/Plugins/VirtualProduction/RemoteControl/Source/WebRemoteControl/Private/RemoteControlModels.h
@@ -149,7 +149,7 @@ struct FRCPropertyDescription
 		//Write the type name
 		Type = ValueProperty->GetCPPType();
 
-		if (const FObjectProperty* ObjectProperty = CastField<FObjectPtrProperty>(ValueProperty))
+		if (const FObjectProperty* ObjectProperty = CastField<FObjectProperty>(ValueProperty))
 		{
 			if (UClass* Class = ObjectProperty->PropertyClass)
 			{
diff --git a/Engine/Source/Editor/KismetCompiler/Private/KismetCompilerMisc.cpp b/Engine/Source/Editor/KismetCompiler/Private/KismetCompilerMisc.cpp
--- a/Engine/Source/Editor/KismetCompiler/Private/KismetCompilerMisc.cpp
+++ b/Engine/Source/Editor/KismetCompiler/Private/KismetCompilerMisc.cpp
@@ -1254,13 +1254,15 @@ FProperty* FKismetCompilerUtilities::CreatePrimitiveProperty(FFieldVariant Prope
 				{
 					NewPropertyObj = new FWeakObjectProperty(PropertyScope, ValidatedPropertyName, ObjectFlags);
 				}
-				else if (FLinkerLoad::IsImportLazyLoadEnabled())
-				{
-					NewPropertyObj = new FObjectPtrProperty(PropertyScope, ValidatedPropertyName, ObjectFlags);
-				}
 				else
 				{
 					NewPropertyObj = new FObjectProperty(PropertyScope, ValidatedPropertyName, ObjectFlags);
+					// If lazy load is enabled make the object property a TObjectPtr property
+					// to allow for unresolved UObjects
+					if (FLinkerLoad::IsImportLazyLoadEnabled())
+					{
+						NewPropertyObj->SetPropertyFlags(CPF_TObjectPtrWrapper);
+					}
 				}
 
 				// Is the property a reference to something that should default to instanced?
diff --git a/Engine/Source/Editor/MovieSceneTools/Private/Channels/BuiltInChannelEditors.cpp b/Engine/Source/Editor/MovieSceneTools/Private/Channels/BuiltInChannelEditors.cpp
--- a/Engine/Source/Editor/MovieSceneTools/Private/Channels/BuiltInChannelEditors.cpp
+++ b/Engine/Source/Editor/MovieSceneTools/Private/Channels/BuiltInChannelEditors.cpp
@@ -783,8 +783,8 @@ UMovieSceneKeyStructType* InstanceGeneratedStruct(FMovieSceneObjectPathChannel*
 		return nullptr;
 	}
 
-	FObjectPtrProperty* NewValueProperty = new FObjectPtrProperty(NewStruct, "Value", RF_NoFlags);
-	NewValueProperty->SetPropertyFlags(CPF_Edit);
+	FObjectProperty* NewValueProperty = new FObjectProperty(NewStruct, "Value", RF_NoFlags);
+	NewValueProperty->SetPropertyFlags(CPF_Edit | CPF_TObjectPtrWrapper);
 	NewValueProperty->SetMetaData("Category", TEXT("Key"));
 	NewValueProperty->PropertyClass = PropertyClass;
 	NewValueProperty->ArrayDim = 1;
@@ -805,7 +805,7 @@ void PostConstructKeyInstance(const TMovieSceneChannelHandle<FMovieSceneObjectPa
 
 	uint8* StructMemory = Struct->GetStructMemory();
 
-	FObjectPtrProperty* ValueProperty = CastFieldChecked<FObjectPtrProperty>(GeneratedStructType->DestValueProperty.Get());
+	FObjectProperty* ValueProperty = CastFieldChecked<FObjectProperty>(GeneratedStructType->DestValueProperty.Get());
 	FStructProperty*     TimeProperty  = CastFieldChecked<FStructProperty>(GeneratedStructType->DestTimeProperty.Get());
 
 	const FFrameNumber*  TimeAddress   = TimeProperty->ContainerPtrToValuePtr<FFrameNumber>(StructMemory);
diff --git a/Engine/Source/Programs/Shared/EpicGames.Core/UnrealEngineTypes.cs b/Engine/Source/Programs/Shared/EpicGames.Core/UnrealEngineTypes.cs
--- a/Engine/Source/Programs/Shared/EpicGames.Core/UnrealEngineTypes.cs
+++ b/Engine/Source/Programs/Shared/EpicGames.Core/UnrealEngineTypes.cs
@@ -388,8 +388,6 @@ namespace EpicGames.Core
 		FMulticastInlineDelegateProperty = 0x0004000000000000,
 		FMulticastSparseDelegateProperty = 0x0008000000000000,
 		FFieldPathProperty = 0x0010000000000000,
-		FObjectPtrProperty = 0x0020000000000000,
-		FClassPtrProperty = 0x0040000000000000,
 		FLargeWorldCoordinatesRealProperty = 0x0080000000000000,
 		FOptionalProperty = 0x0100000000000000,
 		FVerseValueProperty = 0x0200000000000000,
@@ -1196,6 +1194,16 @@ namespace EpicGames.Core
 		/// </summary>
 		SkipSerialization = 0x0080000000000000,
 
+		/// <summary>
+		/// Property is a TObjectPtr<T> instead of a USomething*. Need to differentiate between TObjectclassOf and TObjectPtr
+		/// </summary>
+		TObjectPtr = 0x0100000000000000,
+
+		/// <summary>
+		/// TObjectPtr properties are both wrapped and TObjectPtr
+		/// </summary>
+		TObjectPtrWrapper = TObjectPtr | UObjectWrapper,
+
 		/// <summary>
 		/// All Native Access Specifier flags
 		/// </summary>
@@ -1209,7 +1217,7 @@ namespace EpicGames.Core
 		/// <summary>
 		/// Flags that are propagated to properties inside array container
 		/// </summary>
-		PropagateToArrayInner =	ExportObject | PersistentInstance | InstancedReference | ContainsInstancedReference | Config | EditConst | Deprecated | EditorOnly | AutoWeak | UObjectWrapper,
+		PropagateToArrayInner =	ExportObject | PersistentInstance | InstancedReference | ContainsInstancedReference | Config | EditConst | Deprecated | EditorOnly | AutoWeak | UObjectWrapper | TObjectPtr,
 
 		/// <summary>
 		/// Flags that are propagated to properties inside optional container
@@ -1219,17 +1227,17 @@ namespace EpicGames.Core
 		/// <summary>
 		/// Flags that are propagated to value properties inside map container
 		/// </summary>
-		PropagateToMapValue = ExportObject | PersistentInstance | InstancedReference | ContainsInstancedReference | Config | EditConst | Deprecated | EditorOnly | AutoWeak | UObjectWrapper | Edit,
+		PropagateToMapValue = ExportObject | PersistentInstance | InstancedReference | ContainsInstancedReference | Config | EditConst | Deprecated | EditorOnly | AutoWeak | UObjectWrapper | TObjectPtr | Edit,
 
 		/// <summary>
 		/// Flags that are propagated to key properties inside map container
 		/// </summary>
-		PropagateToMapKey = ExportObject | PersistentInstance | InstancedReference | ContainsInstancedReference | Config | EditConst | Deprecated | EditorOnly | AutoWeak | UObjectWrapper | Edit,
+		PropagateToMapKey = ExportObject | PersistentInstance | InstancedReference | ContainsInstancedReference | Config | EditConst | Deprecated | EditorOnly | AutoWeak | UObjectWrapper | TObjectPtr | Edit,
 
 		/// <summary>
 		/// Flags that are propagated to properties inside set container
 		/// </summary>
-		PropagateToSetElement = ExportObject | PersistentInstance | InstancedReference | ContainsInstancedReference | Config | EditConst | Deprecated | EditorOnly | AutoWeak | UObjectWrapper | Edit,
+		PropagateToSetElement = ExportObject | PersistentInstance | InstancedReference | ContainsInstancedReference | Config | EditConst | Deprecated | EditorOnly | AutoWeak | UObjectWrapper | TObjectPtr | Edit,
 
 		/// <summary>
 		/// The flags that should never be set on interface properties
diff --git a/Engine/Source/Programs/Shared/EpicGames.UHT/Types/Properties/UhtArrayProperty.cs b/Engine/Source/Programs/Shared/EpicGames.UHT/Types/Properties/UhtArrayProperty.cs
--- a/Engine/Source/Programs/Shared/EpicGames.UHT/Types/Properties/UhtArrayProperty.cs
+++ b/Engine/Source/Programs/Shared/EpicGames.UHT/Types/Properties/UhtArrayProperty.cs
@@ -39,7 +39,7 @@ namespace EpicGames.UHT.Types
 		public UhtArrayProperty(UhtPropertySettings propertySettings, UhtProperty value) : base(propertySettings, value)
 		{
 			// If the creation of the value property set more flags, then copy those flags to ourselves
-			PropertyFlags |= ValueProperty.PropertyFlags & EPropertyFlags.UObjectWrapper;
+			PropertyFlags |= ValueProperty.PropertyFlags & (EPropertyFlags.UObjectWrapper | EPropertyFlags.TObjectPtr);
 
 			if (ValueProperty.MetaData.ContainsKey(UhtNames.NativeConst))
 			{
diff --git a/Engine/Source/Programs/Shared/EpicGames.UHT/Types/Properties/UhtClassPtrProperties.cs b/Engine/Source/Programs/Shared/EpicGames.UHT/Types/Properties/UhtClassPtrProperties.cs
--- a/Engine/Source/Programs/Shared/EpicGames.UHT/Types/Properties/UhtClassPtrProperties.cs
+++ b/Engine/Source/Programs/Shared/EpicGames.UHT/Types/Properties/UhtClassPtrProperties.cs
@@ -15,7 +15,7 @@ namespace EpicGames.UHT.Types
 	public class UhtClassPtrProperty : UhtClassProperty
 	{
 		/// <inheritdoc/>
-		public override string EngineClassName => "ClassPtrProperty";
+		public override string EngineClassName => "ClassProperty";
 
 		/// <inheritdoc/>
 		protected override string CppTypeText => "ObjectPtr";
@@ -36,7 +36,7 @@ namespace EpicGames.UHT.Types
 		public UhtClassPtrProperty(UhtPropertySettings propertySettings, UhtClass classObj, UhtClass metaClass, EPropertyFlags extraFlags = EPropertyFlags.None)
 			: base(propertySettings, classObj, metaClass)
 		{
-			PropertyFlags |= extraFlags | EPropertyFlags.UObjectWrapper;
+			PropertyFlags |= extraFlags | EPropertyFlags.TObjectPtr | EPropertyFlags.UObjectWrapper;
 		}
 
 		/// <inheritdoc/>
@@ -54,13 +54,13 @@ namespace EpicGames.UHT.Types
 		/// <inheritdoc/>
 		public override StringBuilder AppendMemberDecl(StringBuilder builder, IUhtPropertyMemberContext context, string name, string nameSuffix, int tabs)
 		{
-			return AppendMemberDecl(builder, context, name, nameSuffix, tabs, "FClassPtrPropertyParams");
+			return AppendMemberDecl(builder, context, name, nameSuffix, tabs, "FClassPropertyParams");
 		}
 
 		/// <inheritdoc/>
 		public override StringBuilder AppendMemberDef(StringBuilder builder, IUhtPropertyMemberContext context, string name, string nameSuffix, string? offset, int tabs)
 		{
-			AppendMemberDefStart(builder, context, name, nameSuffix, offset, tabs, "FClassPtrPropertyParams",
+			AppendMemberDefStart(builder, context, name, nameSuffix, offset, tabs, "FClassPropertyParams",
 				"UECodeGen_Private::EPropertyGenFlags::Class | UECodeGen_Private::EPropertyGenFlags::ObjectPtr");
 			AppendMemberDefRef(builder, context, Class, false);
 			AppendMemberDefRef(builder, context, MetaClass, false);
diff --git a/Engine/Source/Programs/Shared/EpicGames.UHT/Types/Properties/UhtMapProperty.cs b/Engine/Source/Programs/Shared/EpicGames.UHT/Types/Properties/UhtMapProperty.cs
--- a/Engine/Source/Programs/Shared/EpicGames.UHT/Types/Properties/UhtMapProperty.cs
+++ b/Engine/Source/Programs/Shared/EpicGames.UHT/Types/Properties/UhtMapProperty.cs
@@ -47,7 +47,7 @@ namespace EpicGames.UHT.Types
 			KeyProperty = key;
 
 			// If the creation of the value property set more flags, then copy those flags to ourselves
-			PropertyFlags |= ValueProperty.PropertyFlags & EPropertyFlags.UObjectWrapper;
+			PropertyFlags |= ValueProperty.PropertyFlags & (EPropertyFlags.UObjectWrapper | EPropertyFlags.TObjectPtr);
 
 			// Make sure the 'UObjectWrapper' flag is maintained so that both 'TMap<TSubclassOf<...>, ...>' and 'TMap<UClass*, TSubclassOf<...>>' works correctly
 			KeyProperty.PropertyFlags = (ValueProperty.PropertyFlags & ~EPropertyFlags.UObjectWrapper) | (KeyProperty.PropertyFlags & EPropertyFlags.UObjectWrapper);
diff --git a/Engine/Source/Programs/Shared/EpicGames.UHT/Types/Properties/UhtObjectPtrProperty.cs b/Engine/Source/Programs/Shared/EpicGames.UHT/Types/Properties/UhtObjectPtrProperty.cs
--- a/Engine/Source/Programs/Shared/EpicGames.UHT/Types/Properties/UhtObjectPtrProperty.cs
+++ b/Engine/Source/Programs/Shared/EpicGames.UHT/Types/Properties/UhtObjectPtrProperty.cs
@@ -18,7 +18,7 @@ namespace EpicGames.UHT.Types
 	public class UhtObjectPtrProperty : UhtObjectProperty
 	{
 		/// <inheritdoc/>
-		public override string EngineClassName => "ObjectPtrProperty";
+		public override string EngineClassName => "ObjectProperty";
 
 		/// <inheritdoc/>
 		protected override string CppTypeText => "ObjectPtr";
@@ -38,7 +38,7 @@ namespace EpicGames.UHT.Types
 		public UhtObjectPtrProperty(UhtPropertySettings propertySettings, UhtClass classObj, EPropertyFlags extraFlags = EPropertyFlags.None)
 			: base(propertySettings, classObj)
 		{
-			PropertyFlags |= extraFlags | EPropertyFlags.UObjectWrapper;
+			PropertyFlags |= extraFlags | EPropertyFlags.TObjectPtr | EPropertyFlags.UObjectWrapper;
 			PropertyCaps |= UhtPropertyCaps.PassCppArgsByRef;
 		}
 
@@ -76,13 +76,13 @@ namespace EpicGames.UHT.Types
 		/// <inheritdoc/>
 		public override StringBuilder AppendMemberDecl(StringBuilder builder, IUhtPropertyMemberContext context, string name, string nameSuffix, int tabs)
 		{
-			return AppendMemberDecl(builder, context, name, nameSuffix, tabs, "FObjectPtrPropertyParams");
+			return AppendMemberDecl(builder, context, name, nameSuffix, tabs, "FObjectPropertyParams");
 		}
 
 		/// <inheritdoc/>
 		public override StringBuilder AppendMemberDef(StringBuilder builder, IUhtPropertyMemberContext context, string name, string nameSuffix, string? offset, int tabs)
 		{
-			AppendMemberDefStart(builder, context, name, nameSuffix, offset, tabs, "FObjectPtrPropertyParams", "UECodeGen_Private::EPropertyGenFlags::Object | UECodeGen_Private::EPropertyGenFlags::ObjectPtr");
+			AppendMemberDefStart(builder, context, name, nameSuffix, offset, tabs, "FObjectPropertyParams", "UECodeGen_Private::EPropertyGenFlags::Object | UECodeGen_Private::EPropertyGenFlags::ObjectPtr");
 			AppendMemberDefRef(builder, context, Class, false);
 			AppendMemberDefEnd(builder, context, name, nameSuffix);
 			return builder;
diff --git a/Engine/Source/Programs/Shared/EpicGames.UHT/Types/Properties/UhtSetProperty.cs b/Engine/Source/Programs/Shared/EpicGames.UHT/Types/Properties/UhtSetProperty.cs
--- a/Engine/Source/Programs/Shared/EpicGames.UHT/Types/Properties/UhtSetProperty.cs
+++ b/Engine/Source/Programs/Shared/EpicGames.UHT/Types/Properties/UhtSetProperty.cs
@@ -38,7 +38,7 @@ namespace EpicGames.UHT.Types
 		public UhtSetProperty(UhtPropertySettings propertySettings, UhtProperty value) : base(propertySettings, value)
 		{
 			// If the creation of the value property set more flags, then copy those flags to ourselves
-			PropertyFlags |= ValueProperty.PropertyFlags & EPropertyFlags.UObjectWrapper;
+			PropertyFlags |= ValueProperty.PropertyFlags & (EPropertyFlags.UObjectWrapper | EPropertyFlags.TObjectPtr);
 
 			PropertyCaps |= UhtPropertyCaps.PassCppArgsByRef;
 			PropertyCaps &= ~(UhtPropertyCaps.CanBeContainerValue | UhtPropertyCaps.CanBeContainerKey);
diff --git a/Engine/Source/Runtime/CoreUObject/Private/Blueprint/BlueprintSupport.cpp b/Engine/Source/Runtime/CoreUObject/Private/Blueprint/BlueprintSupport.cpp
--- a/Engine/Source/Runtime/CoreUObject/Private/Blueprint/BlueprintSupport.cpp
+++ b/Engine/Source/Runtime/CoreUObject/Private/Blueprint/BlueprintSupport.cpp
@@ -226,16 +226,13 @@ void FBlueprintSupport::RepairDeferredDependenciesInObject(UObject* Object)
 	for (TPropertyValueIterator<const FObjectProperty> It(Object->GetClass(), Object); It; ++It)
 	{
 		const FObjectProperty* Property = It.Key();
-		void* PropertyValue = (void*)It.Value();
-		if (Property->IsA<FObjectPtrProperty>())
+		const TObjectPtr<UObject>& PropertyValue = Property->GetPropertyValue(It.Value());
+		if (!PropertyValue.IsResolved())
 		{
-			FObjectPtr* PropertyValueAsObjectPtr = ((FObjectPtr*)PropertyValue);
-			if (!PropertyValueAsObjectPtr->IsResolved())
-			{
-				continue;
-			}
+			continue;
 		}
-		UObject* PropertyValueAsObj = *((UObject**)PropertyValue);
+
+		UObject* PropertyValueAsObj = PropertyValue.Get();
 
 		FLinkerPlaceholderBase* Placeholder = nullptr;
 
diff --git a/Engine/Source/Runtime/CoreUObject/Private/Serialization/PropertyLocalizationDataGathering.cpp b/Engine/Source/Runtime/CoreUObject/Private/Serialization/PropertyLocalizationDataGathering.cpp
--- a/Engine/Source/Runtime/CoreUObject/Private/Serialization/PropertyLocalizationDataGathering.cpp
+++ b/Engine/Source/Runtime/CoreUObject/Private/Serialization/PropertyLocalizationDataGathering.cpp
@@ -527,9 +527,9 @@ void FPropertyLocalizationDataGatherer::GatherLocalizationDataFromChildTextPrope
 			// if the object is a FObjectPtr it might not be resolved
 			// if unresolved there is no need to resolve it as IsObjectValidForGather would return false anyways
 			const UObject* InnerObject = nullptr;
-			if (auto ObjectPropertyPtr = CastField<FObjectPtrProperty>(ObjectProperty))
+			if (auto ObjectPropertyPtr = CastField<FObjectProperty>(ObjectProperty))
 			{
-				const FObjectPtr& ObjectPtr = ObjectPropertyPtr->GetObjectPropertyValueAsPtr(ElementValueAddress);
+				const TObjectPtr<UObject>& ObjectPtr = ObjectPropertyPtr->GetPropertyValue(ElementValueAddress);
 				if (ObjectPtr.IsResolved())
 				{
 					InnerObject = ObjectPtr.Get();
diff --git a/Engine/Source/Runtime/CoreUObject/Private/UObject/Field.cpp b/Engine/Source/Runtime/CoreUObject/Private/UObject/Field.cpp
--- a/Engine/Source/Runtime/CoreUObject/Private/UObject/Field.cpp
+++ b/Engine/Source/Runtime/CoreUObject/Private/UObject/Field.cpp
@@ -1037,13 +1037,11 @@ FField* FField::CreateFromUField(UField* InField)
 	}
 	else if (UFieldClass == UObjectProperty::StaticClass())
 	{
+		FObjectProperty* ObjectProperty = new FObjectProperty(InField);
+		NewField = ObjectProperty;
 		if (FLinkerLoad::IsImportLazyLoadEnabled())
 		{
-			NewField = new FObjectPtrProperty(InField);
-		}
-		else
-		{
-			NewField = new FObjectProperty(InField);
+			ObjectProperty->SetPropertyFlags(CPF_TObjectPtrWrapper);
 		}
 	}
 	else if (UFieldClass == UWeakObjectProperty::StaticClass())
diff --git a/Engine/Source/Runtime/CoreUObject/Private/UObject/GarbageCollection.cpp b/Engine/Source/Runtime/CoreUObject/Private/UObject/GarbageCollection.cpp
--- a/Engine/Source/Runtime/CoreUObject/Private/UObject/GarbageCollection.cpp
+++ b/Engine/Source/Runtime/CoreUObject/Private/UObject/GarbageCollection.cpp
@@ -5905,7 +5905,7 @@ void FArrayProperty::EmitReferenceInfo(UE::GC::FSchemaBuilder& Schema, int32 Bas
 		// Structs and nested arrays share the same implementation on the Garbage Collector side
 		// as arrays of structs already push the array memory into the GC stack and process its tokens
 		// which is exactly what is required for nested arrays to work
-		if( Inner->IsA(FObjectProperty::StaticClass()) || Inner->IsA(FObjectPtrProperty::StaticClass()) )
+		if( Inner->IsA(FObjectProperty::StaticClass()) )
 		{
 			Type = bUsesFreezableAllocator ?  EMemberType::FreezableReferenceArray :  EMemberType::ReferenceArray;
 		}
diff --git a/Engine/Source/Runtime/CoreUObject/Private/UObject/PropertyAccessUtil.cpp b/Engine/Source/Runtime/CoreUObject/Private/UObject/PropertyAccessUtil.cpp
--- a/Engine/Source/Runtime/CoreUObject/Private/UObject/PropertyAccessUtil.cpp
+++ b/Engine/Source/Runtime/CoreUObject/Private/UObject/PropertyAccessUtil.cpp
@@ -171,22 +171,6 @@ bool ArePropertiesCompatible(const FProperty* InSrcProp, const FProperty* InDest
 		return true;
 	}
 
-	// Compare the classes as these must be an *exact* match as the access is low-level and without property coercion
-	// HOWEVER:
-	// Object properties can either be an ObjectProperty, or an ObjectPtrProperty
-	// We allow coercion between these two types
-	if (InSrcProp->GetClass() != InDestProp->GetClass())
-	{
-		const bool bSrcIsObject = CastField<FObjectProperty>(InSrcProp) || CastField<FObjectPtrProperty>(InSrcProp);
-		const bool bDestIsObject = CastField<FObjectProperty>(InDestProp) || CastField<FObjectPtrProperty>(InDestProp);
-		if (bSrcIsObject && bDestIsObject)
-		{
-			return true;
-		}
-
-		return false;
-	}
-
 	// Containers also need to check their inner types
 	if (const FArrayProperty* SrcArrayProp = CastField<FArrayProperty>(InSrcProp))
 	{
@@ -274,18 +258,6 @@ bool CopySinglePropertyValue(const FProperty* InSrcProp, const void* InSrcValue,
 		return true;
 	}
 
-	// If the source is an ObjectPtr and the destination is an Object, handle the copy in a custom way to ensure we always resolve
-	// the source before copying it to the destination.
-	if (const FObjectPtrProperty* SrcObjectPtrProp = CastField<FObjectPtrProperty>(InSrcProp))
-	{
-		if (const FObjectProperty* DestObjectProp = CastField<FObjectProperty>(InDestProp))
-		{
-			UObject* SrcObject = SrcObjectPtrProp->GetObjectPropertyValue(InSrcValue);
-			DestObjectProp->SetObjectPropertyValue(InDestValue, SrcObject);
-			return true;
-		}
-	}
-	
 	if (IsRealNumberConversion(InSrcProp, InDestProp))
 	{
 		ConvertRealNumber(InSrcProp, InSrcValue, InDestProp, InDestValue, 1);
@@ -373,24 +345,6 @@ bool CopyCompletePropertyValue(const FProperty* InSrcProp, const void* InSrcValu
 		}
 		return true;
 	}
-
-	// If the source is an ObjectPtr and the destination is an Object, handle the copy in a custom way to ensure we always resolve
-	// the source before copying it to the destination.
-	if (const FObjectPtrProperty* SrcObjectPtrProp = CastField<FObjectPtrProperty>(InSrcProp))
-	{
-		if (const FObjectProperty* DestObjectProp = CastField<FObjectProperty>(InDestProp))
-		{
-			for (int32 Idx = 0; Idx < InSrcProp->ArrayDim; ++Idx)
-			{
-				const void* SrcElemValue = static_cast<const uint8*>(InSrcValue) + (InSrcProp->ElementSize * Idx);
-				void* DestElemValue = static_cast<uint8*>(InDestValue) + (InDestProp->ElementSize * Idx);
-
-				UObject* SrcObject = SrcObjectPtrProp->GetObjectPropertyValue(SrcElemValue);
-				DestObjectProp->SetObjectPropertyValue(DestElemValue, SrcObject);
-			}
-			return true;
-		}
-	}
 	
 	if (IsRealNumberConversion(InSrcProp, InDestProp))
 	{
diff --git a/Engine/Source/Runtime/CoreUObject/Private/UObject/PropertyClass.cpp b/Engine/Source/Runtime/CoreUObject/Private/UObject/PropertyClass.cpp
--- a/Engine/Source/Runtime/CoreUObject/Private/UObject/PropertyClass.cpp
+++ b/Engine/Source/Runtime/CoreUObject/Private/UObject/PropertyClass.cpp
@@ -158,15 +158,20 @@ FString FClassProperty::GetCPPType(FString* ExtendedTypeText, uint32 CPPExportFl
 
 FString FClassProperty::GetCPPTypeCustom(FString* ExtendedTypeText, uint32 CPPExportFlags, const FString& InnerNativeTypeName) const
 {
-	if (PropertyFlags & CPF_UObjectWrapper)
+	if (EnumHasAnyFlags(PropertyFlags, CPF_TObjectPtr))
 	{
-		ensure(!InnerNativeTypeName.IsEmpty());
-		return FString::Printf(TEXT("TSubclassOf<%s> "), *InnerNativeTypeName);
+		if (!EnumHasAnyFlags((EPropertyExportCPPFlags)CPPExportFlags, CPPF_NoTObjectPtr))
+		{
+			ensure(!InnerNativeTypeName.IsEmpty());
+			return FString::Printf(TEXT("TObjectPtr<%s>"), *InnerNativeTypeName);
+		}
 	}
-	else
+	else if (EnumHasAnyFlags(PropertyFlags, CPF_UObjectWrapper))
 	{
-		return TEXT("UClass*");
+		ensure(!InnerNativeTypeName.IsEmpty());
+		return FString::Printf(TEXT("TSubclassOf<%s>"), *InnerNativeTypeName);
 	}
+	return TEXT("UClass*");
 }
 
 FString FClassProperty::GetCPPTypeForwardDeclaration() const
@@ -176,6 +181,11 @@ FString FClassProperty::GetCPPTypeForwardDeclaration() const
 
 FString FClassProperty::GetCPPMacroType( FString& ExtendedTypeText ) const
 {
+	if (PropertyFlags & CPF_TObjectPtr)
+	{
+		ExtendedTypeText = FString::Printf(TEXT("TObjectPtr<%s%s>"), PropertyClass->GetPrefixCPP(), *PropertyClass->GetName());
+		return TEXT("OBJECTPTR");
+	}
 	ExtendedTypeText = TEXT("UClass");
 	return TEXT("OBJECT");
 }
diff --git a/Engine/Source/Runtime/CoreUObject/Private/UObject/PropertyClassPtr.cpp b/Engine/Source/Runtime/CoreUObject/Private/UObject/PropertyClassPtr.cpp
deleted file mode 100644
diff --git a/Engine/Source/Runtime/CoreUObject/Private/UObject/PropertyObject.cpp b/Engine/Source/Runtime/CoreUObject/Private/UObject/PropertyObject.cpp
--- a/Engine/Source/Runtime/CoreUObject/Private/UObject/PropertyObject.cpp
+++ b/Engine/Source/Runtime/CoreUObject/Private/UObject/PropertyObject.cpp
@@ -23,6 +23,10 @@ FObjectProperty::FObjectProperty(FFieldVariant InOwner, const UECodeGen_Private:
 FString FObjectProperty::GetCPPTypeCustom(FString* ExtendedTypeText, uint32 CPPExportFlags, const FString& InnerNativeTypeName)  const
 {
 	ensure(!InnerNativeTypeName.IsEmpty());
+	if (HasAnyPropertyFlags(CPF_TObjectPtr) && !(CPPExportFlags & CPPF_NoTObjectPtr))
+	{
+		return FString::Printf(TEXT("TObjectPtr<%s>"), *InnerNativeTypeName);
+	}
 	return FString::Printf(TEXT("%s*"), *InnerNativeTypeName);
 }
 
@@ -33,6 +37,11 @@ FString FObjectProperty::GetCPPTypeForwardDeclaration() const
 
 FString FObjectProperty::GetCPPMacroType( FString& ExtendedTypeText ) const
 {
+	if (HasAnyPropertyFlags(CPF_TObjectPtr))
+	{
+		ExtendedTypeText = FString::Printf(TEXT("TObjectPtr<%s%s>"), PropertyClass->GetPrefixCPP(), *PropertyClass->GetName());
+		return TEXT("OBJECTPTR");
+	}
 	ExtendedTypeText = FString::Printf(TEXT("%s%s"), PropertyClass->GetPrefixCPP(), *PropertyClass->GetName());
 	return TEXT("OBJECT");
 }
@@ -91,6 +100,47 @@ EConvertFromTypeResult FObjectProperty::ConvertFromType(const FPropertyTag& Tag,
 
 	return EConvertFromTypeResult::UseSerializeItem;
 }
+bool FObjectProperty::Identical(const void* A, const void* B, uint32 PortFlags) const
+{
+	// We never return Identical when duplicating for PIE because we want to be sure to serialize everything. An example is the LevelScriptActor being serialized against its CDO,
+	// which contains actor references. We want to serialize those references so they are fixed up.
+	if ((PortFlags & PPF_DuplicateForPIE) != 0)
+	{
+		return false;
+	}
+
+	TObjectPtr<UObject> ObjectA = A ? GetPropertyValue(A) : nullptr;
+	TObjectPtr<UObject> ObjectB = B ? GetPropertyValue(B) : nullptr;
+
+	if (ObjectA == ObjectB)
+	{
+		return true;
+	}
+
+	if (!ObjectA || !ObjectB)
+	{
+		return false;
+	}
+
+	// If a deep comparison is required, resolve the object handles and run the deep comparison logic
+	// If a deep comparison is not required, avoid resolving the object handles because resolving declares
+	// a cook dependency.
+	if ((PortFlags & (PPF_DeepCompareInstances | PPF_DeepComparison)) != 0)
+	{
+		bool bPerformDeepComparison = true;
+		if ((PortFlags & PPF_DeepCompareDSOsOnly) != 0)
+		{
+			UClass* ClassA = ObjectA.GetClass();
+			UObject* DSO = ClassA ? ClassA->GetDefaultSubobjectByName(ObjectA.GetFName()) : nullptr;
+			bPerformDeepComparison = DSO != nullptr;
+		}
+		if (bPerformDeepComparison && ObjectA.GetClass() == ObjectB.GetClass() && ObjectA.GetFName() == ObjectB.GetFName())
+		{
+			return FObjectPropertyBase::StaticIdentical(ObjectA.Get(), ObjectB.Get(), PortFlags);
+		}
+	}
+	return false;
+}
 
 bool FObjectProperty::AllowCrossLevel() const
 {
@@ -100,33 +150,88 @@ bool FObjectProperty::AllowCrossLevel() const
 void FObjectProperty::SerializeItem(FStructuredArchive::FSlot Slot, void* Value, void const* Defaults) const
 {
 	FArchive& UnderlyingArchive = Slot.GetUnderlyingArchive();
-
+	TObjectPtr<UObject>* ObjectPtr = GetPropertyValuePtr(Value);
 	if (UnderlyingArchive.IsObjectReferenceCollector())
 	{
-		// Serialize in place
-		TObjectPtr<UObject>* ObjectPtr = GetPropertyValuePtr(Value);
-		Slot << (*ObjectPtr);
-
+		TObjectPtr<UObject> CurrentValue = *ObjectPtr;
+#if UE_WITH_OBJECT_HANDLE_LATE_RESOLVE
+		if (HasAnyPropertyFlags(CPF_TObjectPtr))
+		{
+			FObjectPtr* Ptr = (FObjectPtr*)ObjectPtr;
+			Slot << *Ptr;
+		}
+		else
+#endif
+		{
+			// Serialize in place
+			Slot << (*ObjectPtr);
+		}
 #if !(UE_BUILD_TEST || UE_BUILD_SHIPPING) 
 		if (!UnderlyingArchive.IsSaving())
 		{
-			CheckValidObject(ObjectPtr, *ObjectPtr);
+			CheckValidObject(Value, CurrentValue);
 		}
 #endif
 	}
 	else
 	{
-		TObjectPtr<UObject> ObjectValuePtr = GetObjectPtrPropertyValue(Value);
-		check(ObjectValuePtr.IsResolved());
-		UObject* ObjectValue = UE::CoreUObject::Private::ReadObjectHandlePointerNoCheck(ObjectValuePtr.GetHandle());
+#if UE_WITH_OBJECT_HANDLE_LATE_RESOLVE
+		if (HasAnyPropertyFlags(CPF_TObjectPtr))
+		{
+			FObjectHandle OriginalHandle = ObjectPtr->GetHandle();
+			Slot << *ObjectPtr;
 
-		Slot << ObjectValue;
+			FObjectHandle CurrentHandle = ObjectPtr->GetHandle();
+			if ((OriginalHandle != CurrentHandle) && IsObjectHandleResolved(CurrentHandle))
+			{
+#if USE_CIRCULAR_DEPENDENCY_LOAD_DEFERRING
+				if (ObjectPtr->IsA<ULinkerPlaceholderExportObject>())
+				{
+					//resolve the handle with no read to avoid trigger a handle read. 
+					ULinkerPlaceholderExportObject* PlaceholderVal = static_cast<ULinkerPlaceholderExportObject*>(UE::CoreUObject::Private::ReadObjectHandlePointerNoCheck(CurrentHandle));
+					PlaceholderVal->AddReferencingPropertyValue(this, Value);
+				}
+				else if (ObjectPtr->IsA<ULinkerPlaceholderClass>())
+				{
+					//resolve the handle with no read to avoid trigger a handle read. 
+					ULinkerPlaceholderClass* PlaceholderClass = static_cast<ULinkerPlaceholderClass*>(UE::CoreUObject::Private::ReadObjectHandlePointerNoCheck(CurrentHandle));
+					PlaceholderClass->AddReferencingPropertyValue(this, Value);
+				}
+				// NOTE: we don't remove this from CurrentValue if it is a 
+				//       ULinkerPlaceholderExportObject; this is because this property 
+				//       could be an array inner, and another member of that array (also 
+				//       referenced through this property)... if this becomes a problem,
+				//       then we could inc/decrement a ref count per referencing property 
+				//
+				// @TODO: if this becomes problematic (because ObjectValue doesn't match 
+				//        this property's PropertyClass), then we could spawn another
+				//        placeholder object (of PropertyClass's type), or use null; but
+				//        we'd have to modify ULinkerPlaceholderExportObject::ReplaceReferencingObjectValues()
+				//        to accommodate this (as it depends on finding itself as the set value)
+#endif // USE_CIRCULAR_DEPENDENCY_LOAD_DEFERRING
+
+				TObjectPtr<UObject> CurrentValue(*ObjectPtr);
+				CheckValidObject(Value, CurrentValue);
+			}
+		}
+		else
+#endif
+		{
+			TObjectPtr<UObject> ObjectValuePtr = GetObjectPtrPropertyValue(Value);
+			if (!ObjectValuePtr.IsResolved())
+			{
+				printf("");
+			}
+			check(ObjectValuePtr.IsResolved());
+			UObject* ObjectValue = UE::CoreUObject::Private::ReadObjectHandlePointerNoCheck(ObjectValuePtr.GetHandle());
 
-		TObjectPtr<UObject> CurrentValuePtr = GetObjectPtrPropertyValue(Value);
-		check(CurrentValuePtr.IsResolved());
-		UObject* CurrentValue = UE::CoreUObject::Private::ReadObjectHandlePointerNoCheck(CurrentValuePtr.GetHandle());
+			Slot << ObjectValue;
 
-		PostSerializeObjectItem(UnderlyingArchive, Value, CurrentValue, ObjectValue, EObjectPropertyOptions::None);
+			TObjectPtr<UObject> CurrentValuePtr = GetObjectPtrPropertyValue(Value);
+			check(CurrentValuePtr.IsResolved());
+			UObject* CurrentValue = UE::CoreUObject::Private::ReadObjectHandlePointerNoCheck(CurrentValuePtr.GetHandle());
+			PostSerializeObjectItem(UnderlyingArchive, Value, CurrentValue, ObjectValue, EObjectPropertyOptions::None);
+		}
 	}
 }
 
@@ -299,3 +404,24 @@ void FObjectProperty::CopyCompleteValueFromScriptVM_InContainer(void* OutContain
 {
 	SetWrappedUObjectPtrValues<FObjectPtr>(OutContainer, EPropertyMemoryAccess::InContainer, (UObject**)InValue, 0, ArrayDim);
 }
+
+void FObjectProperty::CopyValuesInternal(void* Dest, void const* Src, int32 Count) const
+{
+#if UE_WITH_OBJECT_HANDLE_LATE_RESOLVE
+	if (HasAllPropertyFlags(CPF_TObjectPtr))
+	{
+		for (int32 Index = 0; Index < Count; Index++)
+		{
+			GetPropertyValuePtr(Dest)[Index] = GetPropertyValuePtr(Src)[Index];
+		}
+	}
+	else
+#endif	
+	{
+		for (int32 Index = 0; Index < Count; Index++)
+		{
+			GetPropertyValuePtr(Dest)[Index] = GetPropertyValuePtr(Src)[Index].Get();
+		}
+	}
+	
+}
diff --git a/Engine/Source/Runtime/CoreUObject/Private/UObject/PropertyObjectPtr.cpp b/Engine/Source/Runtime/CoreUObject/Private/UObject/PropertyObjectPtr.cpp
deleted file mode 100644
diff --git a/Engine/Source/Runtime/CoreUObject/Private/UObject/UObjectGlobals.cpp b/Engine/Source/Runtime/CoreUObject/Private/UObject/UObjectGlobals.cpp
--- a/Engine/Source/Runtime/CoreUObject/Private/UObject/UObjectGlobals.cpp
+++ b/Engine/Source/Runtime/CoreUObject/Private/UObject/UObjectGlobals.cpp
@@ -4758,8 +4758,7 @@ static constexpr EClassCastFlags WeakCastFlags =			CASTCLASS_FWeakObjectProperty
 															CASTCLASS_FDelegateProperty |
 															CASTCLASS_FMulticastDelegateProperty;
 
-static constexpr EClassCastFlags ObjectCastFlags =			CASTCLASS_FObjectProperty |
-															CASTCLASS_FObjectPtrProperty;
+static constexpr EClassCastFlags ObjectCastFlags =			CASTCLASS_FObjectProperty;
 
 static constexpr EClassCastFlags OtherStrongCastFlags =		CASTCLASS_FInterfaceProperty |
 															CASTCLASS_FFieldPathProperty;
@@ -4846,11 +4845,13 @@ void CollectArrayReferences(FReferenceCollector& Collector, FArrayProperty& Prop
 		{
 			if (bIsReferenceArray)
 			{
-				if (UE_WITH_OBJECT_HANDLE_LATE_RESOLVE && EnumHasAnyFlags(InnerCastFlags, CASTCLASS_FObjectPtrProperty))
+#if UE_WITH_OBJECT_HANDLE_LATE_RESOLVE
+				if (InnerProperty.HasAnyPropertyFlags(CPF_TObjectPtr))
 				{
 					Collector.AddReferencedObjects(*reinterpret_cast<TArray<TObjectPtr<UObject>>*>(Instance), Referencer, &Property);
 				}
 				else
+#endif
 				{
 					FReferenceCollector::AROPrivate::AddReferencedObjects(Collector, *reinterpret_cast<TArray<UObject*>*>(Instance), Referencer, &Property);
 				}
@@ -5842,13 +5843,10 @@ namespace UECodeGen_Private
 
 			case EPropertyGenFlags::Object:
 			{
+				NewProp = NewFProperty<FObjectProperty, FObjectPropertyParams>(Outer, *PropBase);
 				if (EnumHasAllFlags(PropBase->Flags, EPropertyGenFlags::ObjectPtr))
 				{
-					NewProp = NewFProperty<FObjectPtrProperty, FObjectPtrPropertyParams>(Outer, *PropBase);
-				}
-				else
-				{
-					NewProp = NewFProperty<FObjectProperty, FObjectPropertyParams>(Outer, *PropBase);
+					NewProp->SetPropertyFlags(CPF_TObjectPtrWrapper);
 				}
 			}
 			break;
@@ -5873,13 +5871,10 @@ namespace UECodeGen_Private
 
 			case EPropertyGenFlags::Class:
 			{
+				NewProp = NewFProperty<FClassProperty, FClassPropertyParams>(Outer, *PropBase);
 				if (EnumHasAllFlags(PropBase->Flags, EPropertyGenFlags::ObjectPtr))
 				{
-					NewProp = NewFProperty<FClassPtrProperty, FClassPtrPropertyParams>(Outer, *PropBase);
-				}
-				else
-				{
-					NewProp = NewFProperty<FClassProperty, FClassPropertyParams>(Outer, *PropBase);
+					NewProp->SetPropertyFlags(CPF_TObjectPtrWrapper);
 				}
 			}
 			break;
diff --git a/Engine/Source/Runtime/CoreUObject/Public/Templates/Casts.h b/Engine/Source/Runtime/CoreUObject/Public/Templates/Casts.h
--- a/Engine/Source/Runtime/CoreUObject/Public/Templates/Casts.h
+++ b/Engine/Source/Runtime/CoreUObject/Public/Templates/Casts.h
@@ -432,7 +432,6 @@ DECLARE_CAST_BY_FLAG(UClass)							\
 DECLARE_CAST_BY_FLAG(FProperty)							\
 DECLARE_CAST_BY_FLAG(FObjectPropertyBase)				\
 DECLARE_CAST_BY_FLAG(FObjectProperty)					\
-DECLARE_CAST_BY_FLAG(FObjectPtrProperty)				\
 DECLARE_CAST_BY_FLAG(FWeakObjectProperty)				\
 DECLARE_CAST_BY_FLAG(FLazyObjectProperty)				\
 DECLARE_CAST_BY_FLAG(FSoftObjectProperty)				\
diff --git a/Engine/Source/Runtime/CoreUObject/Public/UObject/Class.h b/Engine/Source/Runtime/CoreUObject/Public/UObject/Class.h
--- a/Engine/Source/Runtime/CoreUObject/Public/UObject/Class.h
+++ b/Engine/Source/Runtime/CoreUObject/Public/UObject/Class.h
@@ -1911,7 +1911,7 @@ public:
 	{
 		//@TODO: UCREMOVAL: CPF_ConstParm added as a hack to get blueprints compiling with a const DamageType parameter.
 		const uint64 IgnoreFlags = CPF_PersistentInstance | CPF_ExportObject | CPF_InstancedReference 
-			| CPF_ContainsInstancedReference | CPF_ComputedFlags | CPF_ConstParm | CPF_UObjectWrapper
+			| CPF_ContainsInstancedReference | CPF_ComputedFlags | CPF_ConstParm | CPF_UObjectWrapper | CPF_TObjectPtr
 			| CPF_NativeAccessSpecifiers | CPF_AdvancedDisplay | CPF_BlueprintVisible | CPF_BlueprintReadOnly;
 		return IgnoreFlags;
 	}
diff --git a/Engine/Source/Runtime/CoreUObject/Public/UObject/ObjectMacros.h b/Engine/Source/Runtime/CoreUObject/Public/UObject/ObjectMacros.h
--- a/Engine/Source/Runtime/CoreUObject/Public/UObject/ObjectMacros.h
+++ b/Engine/Source/Runtime/CoreUObject/Public/UObject/ObjectMacros.h
@@ -373,8 +373,6 @@ enum EClassCastFlags : uint64
 	CASTCLASS_FMulticastInlineDelegateProperty	= 0x0004000000000000,
 	CASTCLASS_FMulticastSparseDelegateProperty	= 0x0008000000000000,
 	CASTCLASS_FFieldPathProperty			= 0x0010000000000000,
-	CASTCLASS_FObjectPtrProperty			= 0x0020000000000000,
-	CASTCLASS_FClassPtrProperty				= 0x0040000000000000,
 	CASTCLASS_FLargeWorldCoordinatesRealProperty = 0x0080000000000000,
 	CASTCLASS_FOptionalProperty				= 0x0100000000000000,
 	CASTCLASS_FVerseValueProperty			= 0x0200000000000000,
@@ -453,6 +451,7 @@ enum EPropertyFlags : uint64
 	CPF_NativeAccessSpecifierProtected	= 0x0020000000000000,	///< Protected native access specifier
 	CPF_NativeAccessSpecifierPrivate	= 0x0040000000000000,	///< Private native access specifier
 	CPF_SkipSerialization				= 0x0080000000000000,	///< Property shouldn't be serialized, can still be exported to text
+	CPF_TObjectPtr						= 0x0100000000000000,	///< Property is a TObjectPtr<T> instead of a USomething*. Need to differentiate between TObjectclassOf and TObjectPtr
 };
 
 /** All Native Access Specifier flags */
@@ -477,6 +476,7 @@ enum EPropertyFlags : uint64
 /** All the properties that should never be loaded or saved */
 #define CPF_ComputedFlags			(CPF_IsPlainOldData | CPF_NoDestructor | CPF_ZeroConstructor | CPF_HasGetValueTypeHash)
 
+#define CPF_TObjectPtrWrapper 		(CPF_UObjectWrapper | CPF_TObjectPtr)
 /** Mask of all property flags */
 #define CPF_AllFlags				((EPropertyFlags)0xFFFFFFFFFFFFFFFF)
 
diff --git a/Engine/Source/Runtime/CoreUObject/Public/UObject/UnrealType.h b/Engine/Source/Runtime/CoreUObject/Public/UObject/UnrealType.h
--- a/Engine/Source/Runtime/CoreUObject/Public/UObject/UnrealType.h
+++ b/Engine/Source/Runtime/CoreUObject/Public/UObject/UnrealType.h
@@ -96,7 +96,7 @@ enum EPropertyExportCPPFlags
 	CPPF_ArgumentOrReturnValue		=	0x00000002,
 	/** Indicates thet we are exporting this property's CPP text for C++ definition of a function. */
 	CPPF_Implementation				=	0x00000004,
-	/** Indicates thet we are exporting this property's CPP text with an custom type name */
+	/** Indicates that we are exporting this property's CPP text with an custom type name */
 	CPPF_CustomTypeName				=	0x00000008,
 	/** No 'const' keyword */
 	CPPF_NoConst					=	0x00000010,
@@ -106,6 +106,8 @@ enum EPropertyExportCPPFlags
 	CPPF_NoStaticArray				=	0x00000040,
 	/** Blueprint compiler generated C++ code */
 	CPPF_BlueprintCppBackend		=	0x00000080,
+	/** Indicates to not use TObjectPtr but use USomething* instead */
+	CPPF_NoTObjectPtr				=	0x00000100,
 };
 
 namespace EExportedDeclaration
@@ -2887,6 +2889,7 @@ PRAGMA_ENABLE_DEPRECATION_WARNINGS
 	virtual bool AllowCrossLevel() const override;
 private:
 	virtual uint32 GetValueTypeHashInternal(const void* Src) const override;
+	virtual void CopyValuesInternal(void* Dest, void const* Src, int32 Count) const override;
 public:
 	virtual void CopySingleValueToScriptVM( void* Dest, void const* Src ) const override;
 	virtual void CopySingleValueFromScriptVM( void* Dest, void const* Src ) const override;
@@ -2894,6 +2897,7 @@ public:
 	virtual void CopyCompleteValueFromScriptVM( void* Dest, void const* Src ) const override;
 	virtual void CopyCompleteValueToScriptVM_InContainer( void* OutValue, void const* InContainer ) const override;
 	virtual void CopyCompleteValueFromScriptVM_InContainer( void* OutContainer, void const* InValue ) const override;
+	virtual bool Identical(const void* A, const void* B, uint32 PortFlags) const override;
 	// End of FProperty interface
 
 	// FObjectPropertyBase interface
@@ -2926,61 +2930,7 @@ public:
 	}
 };
 
-//
-// Describes a reference variable to another object which may be nil.
-//
-class COREUOBJECT_API FObjectPtrProperty : public FObjectProperty
-{
-	DECLARE_FIELD(FObjectPtrProperty, FObjectProperty, CASTCLASS_FObjectPtrProperty)
-
-	using Super::Super;
-
-#if WITH_EDITORONLY_DATA
-	explicit FObjectPtrProperty(UField* InField)
-		: FObjectProperty(InField)
-	{
-	}
-#endif // WITH_EDITORONLY_DATA
-
-	// UHT interface
-	virtual FString GetCPPMacroType( FString& ExtendedTypeText ) const  override;
-	virtual FString GetCPPType( FString* ExtendedTypeText, uint32 CPPExportFlags ) const override;
-	// End of UHT interface
-
-	// FProperty interface
-	virtual bool SameType(const FProperty* Other) const override;
-	virtual bool Identical(const void* A, const void* B, uint32 PortFlags) const override;
-	virtual void SerializeItem(FStructuredArchive::FSlot Slot, void* Value, void const* Defaults) const override;
-	virtual void CopySingleValueToScriptVM( void* Dest, void const* Src ) const override;
-	virtual void CopySingleValueFromScriptVM( void* Dest, void const* Src ) const override;
-	virtual void CopyCompleteValueToScriptVM( void* Dest, void const* Src ) const override;
-	virtual void CopyCompleteValueFromScriptVM( void* Dest, void const* Src ) const override;
-	virtual void CopyCompleteValueToScriptVM_InContainer( void* OutValue, void const* InContainer ) const override;
-	virtual void CopyCompleteValueFromScriptVM_InContainer( void* OutContainer, void const* InValue ) const override;
-	// End of FProperty interface
-
-	// Helper method for sharing code with FClassPtrProperty even though one doesn't inherit from the other
-	static void StaticSerializeItem(const FObjectPropertyBase* ObjectProperty, FStructuredArchive::FSlot Slot, void* Value, void const* Defaults);
-
-	/// @brief Get the property value as FObjectPtr
-	/// @param PropertyValueAddress address of the property
-	/// @return reference to the FObjectPtr value
-	FObjectPtr& GetObjectPropertyValueAsPtr(const void* PropertyValueAddress) const;
-
-	// FObjectProperty interface
-	virtual TObjectPtr<UObject> GetObjectPtrPropertyValue(const void* PropertyValueAddress) const override;
-	virtual UObject* GetObjectPropertyValue(const void* PropertyValueAddress) const override;
-	virtual UObject* GetObjectPropertyValue_InContainer(const void* ContainerAddress, int32 ArrayIndex = 0) const override;
-	virtual void SetObjectPropertyValue(void* PropertyValueAddress, UObject* Value) const override;
-	virtual void SetObjectPropertyValue_InContainer(void* ContainerAddress, UObject* Value, int32 ArrayIndex = 0) const override;
-	virtual bool AllowObjectTypeReinterpretationTo(const FObjectPropertyBase* Other) const override;
-
-
-private:
-	virtual uint32 GetValueTypeHashInternal(const void* Src) const override;
-public:
-	// End of FObjectProperty interface
-};
+using FObjectPtrProperty UE_DEPRECATED(5.4, "FObjectPtrProperty is deprecated using FObjectProperty instead.")  = FObjectProperty;
 
 //
 // Describes a reference variable to another object which may be nil, and may turn nil at any point
@@ -3291,50 +3241,7 @@ PRAGMA_ENABLE_DEPRECATION_WARNINGS
 #endif // USE_CIRCULAR_DEPENDENCY_LOAD_DEFERRING
 };
 
-//
-// Describes a reference variable to another object which may be nil.
-//
-class COREUOBJECT_API FClassPtrProperty : public FClassProperty
-{
-	DECLARE_FIELD(FClassPtrProperty, FClassProperty, CASTCLASS_FClassPtrProperty)
-
-	using Super::Super;
-
-#if WITH_EDITORONLY_DATA
-	explicit FClassPtrProperty(UField* InField)
-		: FClassProperty(InField)
-	{
-	}
-#endif // WITH_EDITORONLY_DATA
-
-	// UHT interface
-	virtual FString GetCPPMacroType( FString& ExtendedTypeText ) const  override;
-	virtual FString GetCPPType( FString* ExtendedTypeText, uint32 CPPExportFlags ) const override;
-	// End of UHT interface
-
-	// FProperty interface
-	virtual bool SameType(const FProperty* Other) const override;
-	virtual bool Identical(const void* A, const void* B, uint32 PortFlags) const override;
-	virtual void SerializeItem(FStructuredArchive::FSlot Slot, void* Value, void const* Defaults) const override;
-	virtual void CopySingleValueToScriptVM( void* Dest, void const* Src ) const override;
-	virtual void CopySingleValueFromScriptVM( void* Dest, void const* Src ) const override;
-	virtual void CopyCompleteValueToScriptVM( void* Dest, void const* Src ) const override;
-	virtual void CopyCompleteValueFromScriptVM( void* Dest, void const* Src ) const override;
-	virtual void CopyCompleteValueToScriptVM_InContainer( void* OutValue, void const* InContainer ) const override;
-	virtual void CopyCompleteValueFromScriptVM_InContainer( void* OutContainer, void const* InValue ) const override;
-	// End of FProperty interface
-
-	// FObjectProperty interface
-	virtual UObject* GetObjectPropertyValue(const void* PropertyValueAddress) const override;
-	virtual UObject* GetObjectPropertyValue_InContainer(const void* ContainerAddress, int32 ArrayIndex = 0) const override;
-	virtual void SetObjectPropertyValue(void* PropertyValueAddress, UObject* Value) const override;
-	virtual void SetObjectPropertyValue_InContainer(void* ContainerAddress, UObject* Value, int32 ArrayIndex = 0) const override;
-private:
-	virtual uint32 GetValueTypeHashInternal(const void* Src) const override;
-public:
-	// End of FObjectProperty interface
-};
-
+using FClassPtrProperty UE_DEPRECATED(5.4, "FClassPtrProperty is deprecated use FClassProperty instead.") = FClassProperty;
 /*-----------------------------------------------------------------------------
 	FSoftClassProperty.
 -----------------------------------------------------------------------------*/
diff --git a/Engine/Source/Runtime/CoreUObject/Tests/ObjectPropertyTest.cpp b/Engine/Source/Runtime/CoreUObject/Tests/ObjectPropertyTest.cpp
--- a/Engine/Source/Runtime/CoreUObject/Tests/ObjectPropertyTest.cpp
+++ b/Engine/Source/Runtime/CoreUObject/Tests/ObjectPropertyTest.cpp
@@ -150,17 +150,24 @@ TEST_CASE("UE::CoreUObject::FObjectProperty::CheckValidAddressNonNullable")
 class FMockArchive : public FArchive
 {
 public:
-	TFunction<void(FObjectPtr& Value)> OnFObjectPtr;
-	virtual FArchive& operator<<(FObjectPtr& Value)
+	
+	FObjectPtr ObjectPtrValue;
+	virtual FArchive& operator<<(FObjectPtr& Value) override
+	{
+		Value = ObjectPtrValue;
+		return *this;
+	}
+
+	UObject* ObjectValue;
+	virtual FArchive& operator<<(UObject*& Value) override
 	{
-		if (OnFObjectPtr)
-			OnFObjectPtr(Value);
+		Value = ObjectValue;
 		return *this;
 	}
 };
 
 template<typename T>
-static void TestSerializeItem(FName ObjectName)
+static void TestSerializeItem(FName ObjectName, TFunctionRef<void(FMockArchive&, UObject*)> SetArchiveValue)
 {
 	UClass* Class = T::StaticClass();
 	FObjectProperty* Property = CastField<FObjectProperty>(Class->FindPropertyByName(TEXT("ObjectPtr")));
@@ -192,15 +199,13 @@ static void TestSerializeItem(FName ObjectName)
 	{
 		//verify that if the property is null no reads are triggered
 		FMockArchive MockArchive;
-		MockArchive.OnFObjectPtr = [&](FObjectPtr& Value)
-		{
-			Value = PlaceHolderExport;
-		};
+		SetArchiveValue(MockArchive, PlaceHolderExport);
+		
 
 		FBinaryArchiveFormatter Formatter(MockArchive);
 		FStructuredArchive Ar(Formatter);
 		FStructuredArchiveSlot Slot = Ar.Open();
-		FObjectPtrProperty::StaticSerializeItem(Property, Slot, &Obj->ObjectPtr, nullptr);
+		Property->SerializeItem(Slot, &Obj->ObjectPtr, nullptr);
 		CHECK(ResolveCount == 0);
 		CHECK(*reinterpret_cast<ULinkerPlaceholderExportObject**>(&Obj->ObjectPtr) == PlaceHolderExport);
 
@@ -208,16 +213,12 @@ static void TestSerializeItem(FName ObjectName)
 	{
 		//verify that if the property is not null no reads are triggered
 		FMockArchive MockArchive;
-		MockArchive.OnFObjectPtr = [&](FObjectPtr& Value)
-		{
-			Value = PlaceHolderClass;
-		};
-
+		SetArchiveValue(MockArchive, PlaceHolderClass);
 		FBinaryArchiveFormatter Formatter(MockArchive);
 		FStructuredArchive Ar(Formatter);
 		FStructuredArchiveSlot Slot = Ar.Open();
 		Obj->ObjectPtr = Other;
-		FObjectPtrProperty::StaticSerializeItem(Property, Slot, &Obj->ObjectPtr, nullptr);
+		Property->SerializeItem(Slot, &Obj->ObjectPtr, nullptr);
 		CHECK(ResolveCount == 0);
 		CHECK(*reinterpret_cast<ULinkerPlaceholderClass**>(&Obj->ObjectPtr) == PlaceHolderClass);
 	}
@@ -225,12 +226,18 @@ static void TestSerializeItem(FName ObjectName)
 
 TEST_CASE("UE::CoreUObject::FObjectPtrProperty::StaticSerializeItem")
 {
-	TestSerializeItem<UObjectPtrTestClassWithRef>(TEXT("Object1"));
+	TestSerializeItem<UObjectPtrTestClassWithRef>(TEXT("Object1"), [](FMockArchive& Ar, UObject* Value)
+		{
+			Ar.ObjectPtrValue = Value;
+		});
 }
 
 TEST_CASE("UE::CoreUObject::FObjectProperty::StaticSerializeItem")
 {
-	TestSerializeItem<UObjectWithRawProperty>(TEXT("Object2"));
+	TestSerializeItem<UObjectWithRawProperty>(TEXT("Object2"), [](FMockArchive& Ar, UObject* Value)
+		{
+			Ar.ObjectValue = Value;
+		});
 }
 
 
@@ -384,7 +391,7 @@ TEST_CASE("UE::CoreUObject::FObjectProperty::ParseObjectPropertyValue")
 #endif
 
 
-TEST_CASE("UE::FObjectPtrProperty::Identical")
+TEST_CASE("UE::FObjectProperty::Identical::ObjectPtr")
 {
 	int ResolveCount = 0;
 #if UE_WITH_OBJECT_HANDLE_TRACKING
@@ -441,4 +448,103 @@ TEST_CASE("UE::FObjectPtrProperty::Identical")
 #endif
 	
 }
+
+TEST_CASE("UE::FObjectProperty::Identical::Object")
+{
+	UClass* Class = UObjectWithRawProperty::StaticClass();
+	FObjectProperty* Property = CastField<FObjectProperty>(Class->FindPropertyByName(TEXT("ObjectPtr")));
+	REQUIRE(Property != nullptr);
+
+	UPackage* TestPackage = NewObject<UPackage>(nullptr, TEXT("Test/TestPackageName"), RF_Transient);
+	TestPackage->AddToRoot();
+	UPackage* TestPackage2 = NewObject<UPackage>(nullptr, TEXT("Test/TestPackageName2"), RF_Transient);
+	TestPackage2->AddToRoot();
+	ON_SCOPE_EXIT
+	{
+		TestPackage->RemoveFromRoot();
+		TestPackage2->RemoveFromRoot();
+	};
+	TObjectPtr<UObject> ObjWithRef = NewObject<UObjectWithRawProperty>(TestPackage, TEXT("UObjectWithRawProperty"));
+	TObjectPtr<UObject> Obj1 = NewObject<UObjectPtrTestClass>(TestPackage, TEXT("UObjectPtrTestClass"));
+	TObjectPtr<UObject> Obj2 = NewObject<UObjectPtrTestClass>(TestPackage2, TEXT("UObjectPtrTestClass"));
+
+	CHECK(Property->Identical(&Obj1, &Obj1, 0u));
+	CHECK(!Property->Identical(&Obj1, nullptr, 0u));
+	CHECK(!Property->Identical(nullptr, &Obj1, 0u));
+	CHECK(!Property->Identical(&ObjWithRef, &Obj2, 0u));
+	CHECK(!Property->Identical(&ObjWithRef, &Obj2, 0u));
+	CHECK(!Property->Identical(&Obj1, &ObjWithRef, PPF_DeepComparison));
+
+	CHECK(Property->Identical(&Obj1, &Obj2, PPF_DeepComparison));
+}
+
+TEST_CASE("UE::FObjectProperty::CopySingleValue")
+{
+	FObjectProperty* RawProperty = CastField<FObjectProperty>(UObjectWithRawProperty::StaticClass()->FindPropertyByName(TEXT("ObjectPtr")));
+	REQUIRE(RawProperty != nullptr);
+
+	FObjectProperty* PtrProperty = CastField<FObjectProperty>(UObjectPtrTestClassWithRef::StaticClass()->FindPropertyByName(TEXT("ObjectPtr")));
+	REQUIRE(PtrProperty != nullptr);
+
+	UPackage* TestPackage = NewObject<UPackage>(nullptr, TEXT("Test/CopySingleValue"), RF_Transient);
+	TestPackage->AddToRoot();
+	UPackage* TestPackage2 = NewObject<UPackage>(nullptr, TEXT("Test/CopySingleValue2"), RF_Transient);
+	TestPackage2->AddToRoot();
+	ON_SCOPE_EXIT
+	{
+		TestPackage->RemoveFromRoot();
+		TestPackage2->RemoveFromRoot();
+	};
+	TObjectPtr<UObject> ObjWithRawRef = NewObject<UObjectWithRawProperty>(TestPackage, TEXT("UObjectWithRawProperty"));
+	TObjectPtr<UObject> ObjWithPtrRef = NewObject<UObjectPtrTestClassWithRef>(TestPackage, TEXT("UObjectWithPtrProperty"));
+	UObject* Obj1 = NewObject<UObjectPtrTestClass>(TestPackage, TEXT("UObjectPtrTestClass"));
+	UObject* Obj2 = NewObject<UObjectPtrTestClass>(TestPackage2, TEXT("UObjectPtrTestClass2"));
+	UObject* RawPtr = Obj1;
+	TObjectPtr<UObject> PtrObj2 = Obj2;
+
+#if UE_WITH_OBJECT_HANDLE_TRACKING
+	FObjectHandle Handle = MakeUnresolvedHandle(Obj2);
+	PtrObj2 = TObjectPtr<UObject>(FObjectPtr(Handle) );
+#endif
+	//copy an unresolved TObjectPtr to an UObject* pointer. this should resolve the pointer
+	RawProperty->CopySingleValue(&RawPtr, &PtrObj2);
+	CHECK(RawPtr == Obj2);
+
+	PtrObj2 = nullptr;
+	//copy an UObject* to a TObjectPtr
+	PtrProperty->CopySingleValue(&PtrObj2, &RawPtr);
+	CHECK(PtrObj2 == RawPtr);
+
+	CHECK(RawProperty->GetClass() == PtrProperty->GetClass());
+
+}
+
+
+TEST_CASE("UE::FObjectProperty::GetCPPType")
+{
+	FObjectProperty* RawProperty = CastField<FObjectProperty>(UObjectWithRawProperty::StaticClass()->FindPropertyByName(TEXT("ObjectPtr")));
+	REQUIRE(RawProperty != nullptr);
+
+	FObjectProperty* PtrProperty = CastField<FObjectProperty>(UObjectPtrTestClassWithRef::StaticClass()->FindPropertyByName(TEXT("ObjectPtr")));
+	REQUIRE(PtrProperty != nullptr);
+
+	FString RawType = RawProperty->GetCPPType(nullptr, 0u);
+	CHECK(RawType == TEXT("UObjectPtrTestClass*"));
+
+	RawType = RawProperty->GetCPPType(nullptr, EPropertyExportCPPFlags::CPPF_NoTObjectPtr);
+	CHECK(RawType == TEXT("UObjectPtrTestClass*"));
+
+	FString PtrType = PtrProperty->GetCPPType(nullptr, 0u);
+	CHECK(PtrType == TEXT("TObjectPtr<UObjectPtrTestClass>"));
+
+	PtrType = PtrProperty->GetCPPType(nullptr, EPropertyExportCPPFlags::CPPF_NoTObjectPtr);
+	CHECK(PtrType == TEXT("UObjectPtrTestClass*")); 
+}
+
+TEST_CASE("UE::FObjectProperty::ArrayProperty")
+{
+
+	FArrayProperty* PtrProperty = CastField<FArrayProperty>(UObjectPtrTestClassWithRef::StaticClass()->FindPropertyByName(TEXT("ArrayObjPtr")));
+	REQUIRE(PtrProperty != nullptr);
+}
 #endif
diff --git a/Engine/Source/Runtime/CoreUObject/Tests/ObjectPtrTestClass.cpp b/Engine/Source/Runtime/CoreUObject/Tests/ObjectPtrTestClass.cpp
--- a/Engine/Source/Runtime/CoreUObject/Tests/ObjectPtrTestClass.cpp
+++ b/Engine/Source/Runtime/CoreUObject/Tests/ObjectPtrTestClass.cpp
@@ -35,10 +35,10 @@ IMPLEMENT_CORE_INTRINSIC_CLASS(UObjectPtrTestClassWithRef, UObject,
 			UECodeGen_Private::FObjectPropertyParams Params = { };
 			Params.NameUTF8 = "ObjectPtr";
 			Params.Offset = STRUCT_OFFSET(UObjectPtrTestClassWithRef, ObjectPtr);
-			Params.PropertyFlags = CPF_None;
+			Params.PropertyFlags = CPF_TObjectPtrWrapper;
 			Params.ObjectFlags = RF_Public | RF_Transient | RF_MarkAsNative;
 			Params.ClassFunc = nullptr;
-			auto Property = new FObjectPtrProperty(Class, Params );
+			auto Property = new FObjectProperty(Class, Params );
 			Property->PropertyClass = UObjectPtrTestClass::StaticClass();
 		}
 
@@ -47,10 +47,10 @@ IMPLEMENT_CORE_INTRINSIC_CLASS(UObjectPtrTestClassWithRef, UObject,
 			UECodeGen_Private::FObjectPropertyParams Params = { };
 			Params.NameUTF8 = "ObjectPtrNonNullable";
 			Params.Offset = STRUCT_OFFSET(UObjectPtrTestClassWithRef, ObjectPtrNonNullable);
-			Params.PropertyFlags = EPropertyFlags::CPF_NonNullable;
+			Params.PropertyFlags = EPropertyFlags::CPF_NonNullable | CPF_TObjectPtrWrapper;
 			Params.ObjectFlags = RF_Public | RF_Transient | RF_MarkAsNative;
 			Params.ClassFunc = nullptr;
-			auto Property = new FObjectPtrProperty(Class, Params);
+			auto Property = new FObjectProperty(Class, Params);
 			Property->PropertyClass = UObjectPtrTestClass::StaticClass();
 		}
 
@@ -62,7 +62,7 @@ IMPLEMENT_CORE_INTRINSIC_CLASS(UObjectPtrTestClassWithRef, UObject,
 			Params.PropertyFlags = CPF_None;
 			Params.ObjectFlags = RF_Public | RF_Transient | RF_MarkAsNative;
 			auto Property = new FArrayProperty(Class, Params);
-			auto InnerProperty = new FObjectPtrProperty(Property, TEXT("Inner"), EObjectFlags::RF_NoFlags);
+			auto InnerProperty = new FObjectProperty(Property, TEXT("Inner"), EObjectFlags::RF_NoFlags);
 			InnerProperty->PropertyClass = UObjectPtrTestClass::StaticClass();
 			Property->AddCppProperty(InnerProperty);
 		}
@@ -81,10 +81,10 @@ IMPLEMENT_CORE_INTRINSIC_CLASS(UObjectPtrTestClassWithRef, UObject,
 			UECodeGen_Private::FObjectPropertyParams Params = { };
 			Params.NameUTF8 = "ObjectPtr";
 			Params.Offset = STRUCT_OFFSET(UObjectPtrTestClassWithRef, ObjectPtr);
-			Params.PropertyFlags = CPF_None;
+			Params.PropertyFlags = CPF_TObjectPtrWrapper;
 			Params.ObjectFlags = RF_Public | RF_Transient | RF_MarkAsNative;
 			Params.ClassFunc = nullptr;
-			auto Property = new FObjectPtrProperty(Class, Params);
+			auto Property = new FObjectProperty(Class, Params);
 			Property->PropertyClass = UObjectPtrTestClass::StaticClass();
 		}
 
@@ -92,10 +92,10 @@ IMPLEMENT_CORE_INTRINSIC_CLASS(UObjectPtrTestClassWithRef, UObject,
 			UECodeGen_Private::FObjectPropertyParams Params = { };
 			Params.NameUTF8 = "ObjectPtrNonNullable";
 			Params.Offset = STRUCT_OFFSET(UObjectPtrTestClassWithRef, ObjectPtrNonNullable);
-			Params.PropertyFlags = EPropertyFlags::CPF_NonNullable;
+			Params.PropertyFlags = EPropertyFlags::CPF_NonNullable | CPF_TObjectPtrWrapper;
 			Params.ObjectFlags = RF_Public | RF_Transient | RF_MarkAsNative;
 			Params.ClassFunc = nullptr;
-			auto Property = new FObjectPtrProperty(Class, Params);
+			auto Property = new FObjectProperty(Class, Params);
 			Property->PropertyClass = UObjectPtrTestClass::StaticClass();
 		}
 
@@ -106,7 +106,7 @@ IMPLEMENT_CORE_INTRINSIC_CLASS(UObjectPtrTestClassWithRef, UObject,
 			Params.PropertyFlags = CPF_None;
 			Params.ObjectFlags = RF_Public | RF_Transient | RF_MarkAsNative;
 			auto Property = new FArrayProperty(Class, Params);
-			auto InnerProperty = new FObjectPtrProperty(Property, TEXT("Inner"), EObjectFlags::RF_NoFlags);
+			auto InnerProperty = new FObjectProperty(Property, TEXT("Inner"), EObjectFlags::RF_NoFlags);
 			InnerProperty->PropertyClass = UObjectPtrTestClass::StaticClass();
 			Property->AddCppProperty(InnerProperty);
 		}
@@ -115,14 +115,50 @@ IMPLEMENT_CORE_INTRINSIC_CLASS(UObjectPtrTestClassWithRef, UObject,
 
 IMPLEMENT_CORE_INTRINSIC_CLASS(UObjectWithClassProperty, UObject,
 	{
-		UECodeGen_Private::FClassPropertyParams Params = { };
-		Params.NameUTF8 = "ClassPtr";
-		Params.Offset = STRUCT_OFFSET(UObjectWithClassProperty, ClassPtr);
-		Params.PropertyFlags = CPF_None;
-		Params.ObjectFlags = RF_Public | RF_Transient | RF_MarkAsNative;
-		Params.ClassFunc = nullptr;
-		auto Property = new FClassPtrProperty(Class, Params);
-		Property->PropertyClass = UObjectWithClassProperty::StaticClass();
+		{
+			UECodeGen_Private::FClassPropertyParams Params = { };
+			Params.NameUTF8 = "ClassPtr";
+			Params.Offset = STRUCT_OFFSET(UObjectWithClassProperty, ClassPtr);
+			Params.PropertyFlags = CPF_TObjectPtrWrapper;
+			Params.ObjectFlags = RF_Public | RF_Transient | RF_MarkAsNative;
+			Params.ClassFunc = nullptr;
+			Params.MetaClassFunc = []()
+			{
+				return UClass::StaticClass();
+			};
+			auto Property = new FClassProperty(Class, Params);
+			Property->PropertyClass = UObject::StaticClass();
+		}
+	
+		{
+			UECodeGen_Private::FClassPropertyParams Params = { };
+			Params.NameUTF8 = "ClassRaw";
+			Params.Offset = STRUCT_OFFSET(UObjectWithClassProperty, ClassRaw);
+			Params.PropertyFlags = EPropertyFlags::CPF_None;
+			Params.ObjectFlags = RF_Public | RF_Transient | RF_MarkAsNative;
+			Params.ClassFunc = nullptr;
+			Params.MetaClassFunc = []()
+			{
+				return UClass::StaticClass();
+			};
+			auto Property = new FClassProperty(Class, Params);
+			Property->PropertyClass = UObject::StaticClass();
+		}
+
+		{
+			UECodeGen_Private::FClassPropertyParams Params = { };
+			Params.NameUTF8 = "SubClass";
+			Params.Offset = STRUCT_OFFSET(UObjectWithClassProperty, SubClass);
+			Params.PropertyFlags = EPropertyFlags::CPF_UObjectWrapper;
+			Params.ObjectFlags = RF_Public | RF_Transient | RF_MarkAsNative;
+			Params.ClassFunc = nullptr;
+			Params.MetaClassFunc = []()
+			{
+				return UObjectPtrTestClass::StaticClass();
+			};
+			auto Property = new FClassProperty(Class, Params);
+			Property->PropertyClass = UObject::StaticClass();
+		}
 	}
 );
 
@@ -136,7 +172,7 @@ IMPLEMENT_CORE_INTRINSIC_CLASS(UObjectWithRawProperty, UObject,
 			Params.ObjectFlags = RF_Public | RF_Transient | RF_MarkAsNative;
 			Params.ClassFunc = nullptr;
 			auto Property = new FObjectProperty(Class, Params);
-			Property->PropertyClass = UObjectWithRawProperty::StaticClass();
+			Property->PropertyClass = UObjectPtrTestClass::StaticClass();
 		}
 
 		{
@@ -147,7 +183,7 @@ IMPLEMENT_CORE_INTRINSIC_CLASS(UObjectWithRawProperty, UObject,
 			Params.ObjectFlags = RF_Public | RF_Transient | RF_MarkAsNative;
 			Params.ClassFunc = nullptr;
 			auto Property = new FObjectProperty(Class, Params);
-			Property->PropertyClass = UObjectWithRawProperty::StaticClass();
+			Property->PropertyClass = UObjectPtrTestClass::StaticClass();
 		}
 	}
 );
diff --git a/Engine/Source/Runtime/CoreUObject/Tests/ObjectPtrTestClass.h b/Engine/Source/Runtime/CoreUObject/Tests/ObjectPtrTestClass.h
--- a/Engine/Source/Runtime/CoreUObject/Tests/ObjectPtrTestClass.h
+++ b/Engine/Source/Runtime/CoreUObject/Tests/ObjectPtrTestClass.h
@@ -5,6 +5,7 @@
 
 #include "UObject/Object.h"
 #include "UObject/ObjectPtr.h"
+#include "Templates/SubclassOf.h"
 
 //simple test class for testing TObjectPtr resolve behavior
 class UObjectPtrTestClass : public UObject
@@ -34,6 +35,8 @@ class UObjectWithClassProperty : public UObject
 
 public:
 	TObjectPtr<UClass> ClassPtr;
+	TSubclassOf<UObjectPtrTestClass> SubClass;
+	UClass* ClassRaw;
 };
 
 //test class with raw pointer
diff --git a/Engine/Source/Runtime/CoreUObject/Tests/PropertyClassTest.cpp b/Engine/Source/Runtime/CoreUObject/Tests/PropertyClassTest.cpp
--- a/Engine/Source/Runtime/CoreUObject/Tests/PropertyClassTest.cpp
+++ b/Engine/Source/Runtime/CoreUObject/Tests/PropertyClassTest.cpp
@@ -12,6 +12,7 @@
 #include "UObject/LinkerPlaceholderClass.h"
 #include "UObject/ObjectHandleTracking.h"
 #include "LowLevelTestsRunner/WarnFilterScope.h"
+#include "ObjectRefTrackingTestBase.h"
 
 TEST_CASE("UE::CoreUObject::FClassProperty::Identical")
 {
@@ -36,9 +37,47 @@ TEST_CASE("UE::CoreUObject::FClassProperty::Identical")
 	{
 		TestPackage->RemoveFromRoot();
 	};
-	UObjectWithClassProperty* Obj = NewObject<UObjectWithClassProperty>(TestPackage, TEXT("UObjectWithClassProperty"));
+	UObjectWithClassProperty* Obj1 = NewObject<UObjectWithClassProperty>(TestPackage, TEXT("UObjectWithClassProperty1"));
+	UObjectWithClassProperty* Obj2 = NewObject<UObjectWithClassProperty>(TestPackage, TEXT("UObjectWithClassProperty2"));
+	Obj1->ClassPtr = UObjectWithClassProperty::StaticClass();
 
-	CHECK(Property->Identical(Obj, Obj, 0u));
+	const void* Value1 = Property->ContainerPtrToValuePtr<void>(Obj1);
+	const void* Value2 = Property->ContainerPtrToValuePtr<void>(Obj2);
+	
+	//NonNull != nullptr
+	CHECK(!Property->Identical(Value1, Value2, 0u));
+	CHECK(!Property->Identical(Value2, Value1, 0u));
 
+	Obj2->ClassPtr = Obj1->ClassPtr;
+	CHECK(Property->Identical(Value1, Value2, 0u));
+	CHECK(Property->Identical(Value2, Value1, 0u));
+}
+
+
+TEST_CASE("UE::CoreUObject::FClassProperty::GetCPPType")
+{
+	FObjectProperty* PtrProperty = CastField<FObjectProperty>(UObjectWithClassProperty::StaticClass()->FindPropertyByName(TEXT("ClassPtr")));
+	FObjectProperty* RawProperty = CastField<FObjectProperty>(UObjectWithClassProperty::StaticClass()->FindPropertyByName(TEXT("ClassRaw")));
+	FObjectProperty* SubClassProperty = CastField<FObjectProperty>(UObjectWithClassProperty::StaticClass()->FindPropertyByName(TEXT("SubClass")));
+
+	CHECK(PtrProperty->GetClass() == RawProperty->GetClass());
+
+	FString RawType = RawProperty->GetCPPType(nullptr, 0u);
+	CHECK(RawType == TEXT("UClass*"));
+
+	RawType = RawProperty->GetCPPType(nullptr, EPropertyExportCPPFlags::CPPF_NoTObjectPtr);
+	CHECK(RawType == TEXT("UClass*"));
+
+	FString PtrType = PtrProperty->GetCPPType(nullptr, 0u);
+	CHECK(PtrType == TEXT("TObjectPtr<UClass>"));
+
+	PtrType = PtrProperty->GetCPPType(nullptr, EPropertyExportCPPFlags::CPPF_NoTObjectPtr);
+	CHECK(PtrType == TEXT("UClass*"));
+
+	FString SubClassType = SubClassProperty->GetCPPType(nullptr, EPropertyExportCPPFlags::CPPF_NoTObjectPtr);
+	CHECK(SubClassType == TEXT("TSubclassOf<UObjectPtrTestClass>"));
+
+	SubClassType = SubClassProperty->GetCPPType(nullptr, 0u);
+	CHECK(SubClassType == TEXT("TSubclassOf<UObjectPtrTestClass>"));
 }
 #endif
diff --git a/Engine/Source/Runtime/Engine/Private/RepLayout.cpp b/Engine/Source/Runtime/Engine/Private/RepLayout.cpp
--- a/Engine/Source/Runtime/Engine/Private/RepLayout.cpp
+++ b/Engine/Source/Runtime/Engine/Private/RepLayout.cpp
@@ -947,19 +947,9 @@ static uint32 GetRepLayoutCmdCompatibleChecksum(
 	uint32 CompatibleChecksum = FCrc::StrCrc32(*Property->GetName().ToLower(), InChecksum);	
 	
 	// Evolve by property type
-	const FObjectPtrProperty* const ObjectPtrProperty = CastField<const FObjectPtrProperty>(Property);
-
-	FString CPPType;
-	if (ObjectPtrProperty)
-	{
-		// To remain compatible with TObjectPtr, use the underlying pointer type in the checksum since the net-serialized data is compatible.
-		CPPType = ObjectPtrProperty->FObjectProperty::GetCPPType(nullptr, 0).ToLower();
-	}
-	else
-	{
-		CPPType = Property->GetCPPType(nullptr, 0).ToLower();
-	}
+	const FObjectProperty* const ObjectPtrProperty = CastField<const FObjectProperty>(Property);
 
+	FString CPPType = Property->GetCPPType(nullptr, EPropertyExportCPPFlags::CPPF_NoTObjectPtr).ToLower();
 	CompatibleChecksum = FCrc::StrCrc32(*CPPType, CompatibleChecksum);
 	
 	// Evolve by StaticArrayIndex (to make all unrolled static array elements unique)
