Downgrade of error to warning until properly investigated with additional diagnostic

diff --git a/Engine/Source/Runtime/CoreUObject/Private/Serialization/AsyncLoading2.cpp b/Engine/Source/Runtime/CoreUObject/Private/Serialization/AsyncLoading2.cpp
--- a/Engine/Source/Runtime/CoreUObject/Private/Serialization/AsyncLoading2.cpp
+++ b/Engine/Source/Runtime/CoreUObject/Private/Serialization/AsyncLoading2.cpp
@@ -1086,6 +1086,16 @@ public:
 						}
 					}
 
+					if (PackageRef.bIsMissing)
+					{
+						PackageRef.bIsMissing = false;
+						UE_LOG(LogStreaming, Warning,
+							TEXT("FGlobalImportStore:AddPackageRef: Found reference to previously missing package %s (0x%llX)"),
+							*FoundPackage->GetName(),
+							PackageId.ValueForDebugging()
+						);
+					}
+
 					PackageRef.SetPackage(FoundPackage);
 					FoundPackage->SetCanBeImportedFlag(true);
 					FoundPackage->SetPackageId(PackageId);
