Various fix and preparation of supporting overridable serialization

#okforversepublic

diff --git a/Engine/Source/Runtime/CoreUObject/Public/UObject/UnrealType.h b/Engine/Source/Runtime/CoreUObject/Public/UObject/UnrealType.h
--- a/Engine/Source/Runtime/CoreUObject/Public/UObject/UnrealType.h
+++ b/Engine/Source/Runtime/CoreUObject/Public/UObject/UnrealType.h
@@ -4732,6 +4732,40 @@ public:
 		});
 	}
 
+	/** 
+	 * Maps have gaps in their indices, so this function translates a internal index
+	 * to an logical index (ie. Nth element).
+	 * NOTE: This is slow, do not use this for iteration!
+	 */
+	int32 FindLogicalIndex(int32 InternalIdx) const
+	{
+		return WithScriptMap([this, InternalIdx](auto* Map) -> int32
+		{
+			if( !IsValidIndex(InternalIdx) )
+			{
+				return INDEX_NONE;
+			}
+
+			// if map is compact, use random access
+			if (GetMaxIndex() == Num())
+			{
+				return  InternalIdx;
+			}
+
+			int32 LogicalIndex = InternalIdx;
+			for (int i = 0; i < InternalIdx; ++i)
+			{
+				if (!IsValidIndex(i))
+				{
+					LogicalIndex--;
+				}
+			}
+
+			return LogicalIndex;
+		});
+	}
+
+
 	/**
 	 * Finds the index of an element in a map which matches the key in another pair.
 	 *
