removed some printf debugging

diff --git a/Engine/Source/Runtime/CoreUObject/Tests/PluginManagerTest.cpp b/Engine/Source/Runtime/CoreUObject/Tests/PluginManagerTest.cpp
--- a/Engine/Source/Runtime/CoreUObject/Tests/PluginManagerTest.cpp
+++ b/Engine/Source/Runtime/CoreUObject/Tests/PluginManagerTest.cpp
@@ -118,7 +118,6 @@ namespace UE::CoreUObject::Private::Tests
 
 	void FProcessor::HandleTokenStreamObjectReference(FGCArrayStruct& ObjectsToSerializeStruct, UObject* ReferencingObject, UObject*& Object, UE::GC::FTokenId, EGCTokenType, bool)
 	{
-		printf("");
 	}
 
 	TEST_CASE("UE::CoreUObject::PluginHandler::LeakDetection")
