fix for failing unit test when late resolved off
https://p4-swarm.epicgames.net/reviews/29812692

diff --git a/Engine/Source/Runtime/CoreUObject/Private/UObject/PropertyObject.cpp b/Engine/Source/Runtime/CoreUObject/Private/UObject/PropertyObject.cpp
--- a/Engine/Source/Runtime/CoreUObject/Private/UObject/PropertyObject.cpp
+++ b/Engine/Source/Runtime/CoreUObject/Private/UObject/PropertyObject.cpp
@@ -218,10 +218,6 @@ void FObjectProperty::SerializeItem(FStructuredArchive::FSlot Slot, void* Value,
 #endif
 		{
 			TObjectPtr<UObject> ObjectValuePtr = GetObjectPtrPropertyValue(Value);
-			if (!ObjectValuePtr.IsResolved())
-			{
-				printf("");
-			}
 			check(ObjectValuePtr.IsResolved());
 			UObject* ObjectValue = UE::CoreUObject::Private::ReadObjectHandlePointerNoCheck(ObjectValuePtr.GetHandle());
 
diff --git a/Engine/Source/Runtime/CoreUObject/Tests/ObjectPropertyTest.cpp b/Engine/Source/Runtime/CoreUObject/Tests/ObjectPropertyTest.cpp
--- a/Engine/Source/Runtime/CoreUObject/Tests/ObjectPropertyTest.cpp
+++ b/Engine/Source/Runtime/CoreUObject/Tests/ObjectPropertyTest.cpp
@@ -151,23 +151,33 @@ class FMockArchive : public FArchive
 {
 public:
 	
-	FObjectPtr ObjectPtrValue;
+	FMockArchive(UObject* Obj)
+		: ArchiveValue(Obj)
+	{
+	}
+
+	union Value
+	{
+		FObjectPtr ObjectPtrValue;
+		UObject* ObjectValue;
+		Value(UObject* Obj) : ObjectValue(Obj) {}
+	} ArchiveValue;
+
 	virtual FArchive& operator<<(FObjectPtr& Value) override
 	{
-		Value = ObjectPtrValue;
+		Value = ArchiveValue.ObjectPtrValue;
 		return *this;
 	}
 
-	UObject* ObjectValue;
 	virtual FArchive& operator<<(UObject*& Value) override
 	{
-		Value = ObjectValue;
+		Value = ArchiveValue.ObjectValue;
 		return *this;
 	}
 };
 
 template<typename T>
-static void TestSerializeItem(FName ObjectName, TFunctionRef<void(FMockArchive&, UObject*)> SetArchiveValue)
+static void TestSerializeItem(FName ObjectName)
 {
 	UClass* Class = T::StaticClass();
 	FObjectProperty* Property = CastField<FObjectProperty>(Class->FindPropertyByName(TEXT("ObjectPtr")));
@@ -198,10 +208,7 @@ static void TestSerializeItem(FName ObjectName, TFunctionRef<void(FMockArchive&,
 #endif
 	{
 		//verify that if the property is null no reads are triggered
-		FMockArchive MockArchive;
-		SetArchiveValue(MockArchive, PlaceHolderExport);
-		
-
+		FMockArchive MockArchive(PlaceHolderExport);
 		FBinaryArchiveFormatter Formatter(MockArchive);
 		FStructuredArchive Ar(Formatter);
 		FStructuredArchiveSlot Slot = Ar.Open();
@@ -212,8 +219,7 @@ static void TestSerializeItem(FName ObjectName, TFunctionRef<void(FMockArchive&,
 	}
 	{
 		//verify that if the property is not null no reads are triggered
-		FMockArchive MockArchive;
-		SetArchiveValue(MockArchive, PlaceHolderClass);
+		FMockArchive MockArchive(PlaceHolderClass);
 		FBinaryArchiveFormatter Formatter(MockArchive);
 		FStructuredArchive Ar(Formatter);
 		FStructuredArchiveSlot Slot = Ar.Open();
@@ -226,21 +232,14 @@ static void TestSerializeItem(FName ObjectName, TFunctionRef<void(FMockArchive&,
 
 TEST_CASE("UE::CoreUObject::FObjectPtrProperty::StaticSerializeItem")
 {
-	TestSerializeItem<UObjectPtrTestClassWithRef>(TEXT("Object1"), [](FMockArchive& Ar, UObject* Value)
-		{
-			Ar.ObjectPtrValue = Value;
-		});
+	TestSerializeItem<UObjectPtrTestClassWithRef>(TEXT("Object1"));
 }
 
 TEST_CASE("UE::CoreUObject::FObjectProperty::StaticSerializeItem")
 {
-	TestSerializeItem<UObjectWithRawProperty>(TEXT("Object2"), [](FMockArchive& Ar, UObject* Value)
-		{
-			Ar.ObjectValue = Value;
-		});
+	TestSerializeItem<UObjectWithRawProperty>(TEXT("Object2"));
 }
 
-
 class MockAssetRegistryInterface : public IAssetRegistryInterface
 {
 public:
