Move clearing of dirty flag and file size update out of InnerSave
This allows delegate broadcast for dirty state update to happen out "IsSavingPackage" scope

diff --git a/Engine/Source/Runtime/CoreUObject/Private/UObject/SavePackage2.cpp b/Engine/Source/Runtime/CoreUObject/Private/UObject/SavePackage2.cpp
--- a/Engine/Source/Runtime/CoreUObject/Private/UObject/SavePackage2.cpp
+++ b/Engine/Source/Runtime/CoreUObject/Private/UObject/SavePackage2.cpp
@@ -2907,7 +2907,7 @@ ESavePackageResult InnerSave(FSaveContext& SaveContext)
 	SaveContext.SetEDLCookChecker(&FEDLCookCheckerThreadState::Get());
 
 	// Create slow task dialog if needed
-	const int32 TotalSaveSteps = 4;
+	const int32 TotalSaveSteps = 3;
 	FScopedSlowTask SlowTask(TotalSaveSteps, FText(), SaveContext.IsUsingSlowTask());
 	SlowTask.MakeDialogDelayed(3.0f, SaveContext.IsFromAutoSave());
 
@@ -2951,25 +2951,10 @@ ESavePackageResult InnerSave(FSaveContext& SaveContext)
 			return SaveContext.Result;
 		}
 	}
+	// Mark the Package RF_Loaded after its been serialized. 
+	// This was already done for each object in the package in SaveHarvestedRealms
+	SaveContext.GetPackage()->SetFlags(RF_WasLoaded | RF_LoadCompleted);
 
-	SlowTask.EnterProgressFrame();
-	{
-		// Mark the Package RF_Loaded after its been serialized. 
-		// This was already done for each object in the package in SaveHarvestedRealms
-		SaveContext.GetPackage()->SetFlags(RF_WasLoaded | RF_LoadCompleted);
-
-		// Clear dirty flag if desired
-		if (!SaveContext.IsKeepDirty())
-		{
-			SaveContext.GetPackage()->SetDirtyFlag(false);
-		}
-
-		// Update package FileSize value
-		if (SaveContext.IsUpdatingLoadedPath())
-		{
-			SaveContext.UpdatePackageFileSize(SaveContext.PackageHeaderAndExportSize);
-		}
-	}
 	return SaveContext.Result;
 }
 
@@ -3030,7 +3015,7 @@ FSavePackageResultStruct UPackage::Save2(UPackage* InPackage, UObject* InAsset,
 	FSaveContext SaveContext(InPackage, InAsset, InFilename, SaveArgs);
 
 	// Create the slow task dialog if needed
-	const int32 TotalSaveSteps = 7;
+	const int32 TotalSaveSteps = 8;
 	FScopedSlowTask SlowTask(TotalSaveSteps, GetSlowTaskStatusMessage(SaveContext), SaveContext.IsUsingSlowTask());
 	SlowTask.MakeDialogDelayed(3.0f, SaveContext.IsFromAutoSave());
 
@@ -3108,6 +3093,23 @@ FSavePackageResultStruct UPackage::Save2(UPackage* InPackage, UObject* InAsset,
 		// in case of failure or cancellation, do not exit here, still run cleanup (e.g. PostSaveRoot)
 	}
 
+	SlowTask.EnterProgressFrame();
+	// If the save was successful, update dirty flag and file size if desired
+	if (SaveContext.Result == ESavePackageResult::Success)
+	{
+		// Clear dirty flag if desired
+		if (!SaveContext.IsKeepDirty())
+		{
+			SaveContext.GetPackage()->SetDirtyFlag(false);
+		}
+
+		// Update package FileSize value
+		if (SaveContext.IsUpdatingLoadedPath())
+		{
+			SaveContext.UpdatePackageFileSize(SaveContext.PackageHeaderAndExportSize);
+		}
+	}
+
 	// PostSave Asset
 	SlowTask.EnterProgressFrame();
 	if (SaveContext.GetPostSaveRootRequired() && SaveContext.GetAsset())
@@ -3195,6 +3197,22 @@ ESavePackageResult UPackage::SaveConcurrent(TArrayView<FPackageSaveInfo> InPacka
 		SCOPED_SAVETIMER(UPackage_SaveConcurrent_PostSave);
 		for (FSaveContext& SaveContext : PackageSaveContexts)
 		{
+			// If the save was successful, update dirty flag and file size if desired
+			if (SaveContext.Result == ESavePackageResult::Success)
+			{
+				// Clear dirty flag if desired
+				if (!SaveContext.IsKeepDirty())
+				{
+					SaveContext.GetPackage()->SetDirtyFlag(false);
+				}
+
+				// Update package FileSize value
+				if (SaveContext.IsUpdatingLoadedPath())
+				{
+					SaveContext.UpdatePackageFileSize(SaveContext.PackageHeaderAndExportSize);
+				}
+			}
+
 			// PostSave Asset
 			if (SaveContext.GetAsset())
 			{
