Updated FScriptSetHelper and FScriptMapHelper iterators and updated API to reduce potential indexing issues
- now holds both sparse index and logical index and get be accessed via GetInternalIndex and GetLogicalIndex respectively
- deprecated operator*
- can be created by specifying the desired logical index
- added new accessors in FScriptSetHelper and FScriptMapHelper using the iterator.

diff --git a/Engine/Source/Editor/Kismet/Private/BlueprintNamespaceUtilities.cpp b/Engine/Source/Editor/Kismet/Private/BlueprintNamespaceUtilities.cpp
--- a/Engine/Source/Editor/Kismet/Private/BlueprintNamespaceUtilities.cpp
+++ b/Engine/Source/Editor/Kismet/Private/BlueprintNamespaceUtilities.cpp
@@ -209,7 +209,7 @@ void FBlueprintNamespaceUtilities::GetPropertyValueNamespaces(const FProperty* I
 			FScriptSetHelper SetHelper(SetProperty, ValuePtr);
 			for (FScriptSetHelper::FIterator SetIt = SetHelper.CreateIterator(); SetIt; ++SetIt)
 			{
-				GetPropertyValueNamespaces(SetProperty->ElementProp, SetHelper.GetElementPtr(*SetIt), OutNamespaces);
+				GetPropertyValueNamespaces(SetProperty->ElementProp, SetHelper.GetElementPtr(SetIt), OutNamespaces);
 			}
 		}
 		else if (const FMapProperty* MapProperty = CastField<FMapProperty>(InProperty))
@@ -217,7 +217,7 @@ void FBlueprintNamespaceUtilities::GetPropertyValueNamespaces(const FProperty* I
 			FScriptMapHelper MapHelper(MapProperty, ValuePtr);
 			for (FScriptMapHelper::FIterator MapIt = MapHelper.CreateIterator(); MapIt; ++MapIt)
 			{
-				const uint8* MapValuePtr = MapHelper.GetPairPtr(*MapIt);
+				const uint8* MapValuePtr = MapHelper.GetPairPtr(MapIt);
 				GetPropertyValueNamespaces(MapProperty->KeyProp, MapValuePtr, OutNamespaces);
 				GetPropertyValueNamespaces(MapProperty->ValueProp, MapValuePtr, OutNamespaces);
 			}
diff --git a/Engine/Source/Editor/UnrealEd/Private/Kismet2/KismetDebugUtilities.cpp b/Engine/Source/Editor/UnrealEd/Private/Kismet2/KismetDebugUtilities.cpp
--- a/Engine/Source/Editor/UnrealEd/Private/Kismet2/KismetDebugUtilities.cpp
+++ b/Engine/Source/Editor/UnrealEd/Private/Kismet2/KismetDebugUtilities.cpp
@@ -2274,11 +2274,10 @@ void FPropertyInstanceInfo::PopulateChildren(FPropertyInstance PropertyInstance)
 	{
 		FScriptMapHelper MapHelper(MapProperty, PropertyInstance.Value);
 
-		int32 Index = 0;
-		for (FScriptMapHelper::FIterator Itr = MapHelper.CreateIterator(); Itr; ++Itr, ++Index)
+		for (FScriptMapHelper::FIterator Itr = MapHelper.CreateIterator(); Itr; ++Itr)
 		{
-			uint8* KeyData = MapHelper.GetKeyPtr(*Itr);
-			uint8* ValData = MapHelper.GetValuePtr(*Itr);
+			uint8* KeyData = MapHelper.GetKeyPtr(Itr);
+			uint8* ValData = MapHelper.GetValuePtr(Itr);
 			FPropertyInstance ChildProperty = {
 				MapProperty->ValueProp,
 				ValData
@@ -2297,7 +2296,7 @@ void FPropertyInstanceInfo::PopulateChildren(FPropertyInstance PropertyInstance)
 			NameStr += TEXT("] ");
 			ChildInfo->DisplayName = FText::FromString(NameStr);
 			ChildInfo->bIsInContainer = true;
-			ChildInfo->ContainerIndex = Index;
+			ChildInfo->ContainerIndex = Itr.GetLogicalIndex();
 				
 			Children.Add(ChildInfo);
 		}
@@ -2305,19 +2304,18 @@ void FPropertyInstanceInfo::PopulateChildren(FPropertyInstance PropertyInstance)
 	else if (const FSetProperty* SetProperty = CastField<FSetProperty>(Property.Get()))
 	{
 		FScriptSetHelper SetHelper(SetProperty, PropertyInstance.Value);
-		int32 Index = 0;
-		for (FScriptSetHelper::FIterator Itr = SetHelper.CreateIterator(); Itr; ++Itr, ++Index)
+		for (FScriptSetHelper::FIterator Itr = SetHelper.CreateIterator(); Itr; ++Itr)
 		{
-			uint8* PropData = SetHelper.GetElementPtr(*Itr);
+			uint8* PropData = SetHelper.GetElementPtr(Itr);
 			FPropertyInstance ChildProperty = {
 				SetProperty->ElementProp,
 				PropData
 			};
 			const TSharedPtr<FPropertyInstanceInfo> ChildInfo = Make(ChildProperty, AsShared());
 				
-			ChildInfo->DisplayName = FText::Format(LOCTEXT("SetIndexName", "[{0}]"), FText::AsNumber(Index));
+			ChildInfo->DisplayName = FText::Format(LOCTEXT("SetIndexName", "[{0}]"), FText::AsNumber(Itr.GetLogicalIndex()));
 			ChildInfo->bIsInContainer = true;
-			ChildInfo->ContainerIndex = Index;
+			ChildInfo->ContainerIndex = Itr.GetLogicalIndex();
 
 			Children.Add(ChildInfo);
 		}
diff --git a/Engine/Source/Editor/UnrealEd/Private/MergeUtils.cpp b/Engine/Source/Editor/UnrealEd/Private/MergeUtils.cpp
--- a/Engine/Source/Editor/UnrealEd/Private/MergeUtils.cpp
+++ b/Engine/Source/Editor/UnrealEd/Private/MergeUtils.cpp
@@ -363,7 +363,7 @@ struct FPropertyInstance
 			FScriptSetHelper Helper(SetProperty, SetProperty->ContainerPtrToValuePtr<void>(Object));
 			for (FScriptSetHelper::FIterator It = Helper.CreateIterator(); It; ++It)
 			{
-				TryAddChild(Helper.GetElementPtr(*It), SetProperty->ElementProp);
+				TryAddChild(Helper.GetElementPtr(It), SetProperty->ElementProp);
 			}
 		}
 		else if (const FMapProperty* MapProperty = CastField<FMapProperty>(ValProperty))
@@ -371,7 +371,7 @@ struct FPropertyInstance
 			FScriptMapHelper MapHelper(MapProperty, MapProperty->ContainerPtrToValuePtr<void>(Object));
 			for (FScriptMapHelper::FIterator MapIt = MapHelper.CreateIterator(); MapIt; ++MapIt)
 			{
-				TryAddChild(MapHelper.GetPairPtr(*MapIt), MapProperty->ValueProp, MapProperty->KeyProp);
+				TryAddChild(MapHelper.GetPairPtr(MapIt), MapProperty->ValueProp, MapProperty->KeyProp);
 			}
 		}
 
diff --git a/Engine/Source/Runtime/CoreUObject/Public/UObject/UnrealType.h b/Engine/Source/Runtime/CoreUObject/Public/UObject/UnrealType.h
--- a/Engine/Source/Runtime/CoreUObject/Public/UObject/UnrealType.h
+++ b/Engine/Source/Runtime/CoreUObject/Public/UObject/UnrealType.h
@@ -4369,9 +4369,105 @@ public:
 	}
 };
 
+/**
+ * Templated iterator to go through script helper containers that may contain invalid entries
+ * that are not part of the valid number of elements (i.e. GetMaxIndex() != Num() ).
+ * The iterator
+ *  - will advance to the first valid entry on creation and when incremented
+ *  - can be dereferenced to an internal index to be used with methods like Get<Item>Ptr or Get<Item>PtrWithoutCheck
+ *  - can also be used directly with methods like Get<Item>PtrChecked
+ *  - can return the associated logical index (number of valid visited entries) by calling GetLogicalIndex()
+ */
+template<typename ContainerType>
+struct TScriptContainerIterator
+{
+	explicit TScriptContainerIterator(const ContainerType& InContainer) : Container(InContainer)
+	{
+		Advance();
+	}
+
+	explicit TScriptContainerIterator(const ContainerType& InContainer, const int32 InLogicalIndex) : Container(InContainer)
+	{
+		const int32 MaxIndex = Container.GetMaxIndex();
+		if (MaxIndex == Container.Num())
+		{
+			InternalIndex = InLogicalIndex;
+			LogicalIndex = InLogicalIndex;
+			return;
+		}
+
+		do
+		{
+			Advance();
+		}
+		while (LogicalIndex < InLogicalIndex && InternalIndex < MaxIndex);
+	}
+
+	TScriptContainerIterator& operator++()
+	{
+		Advance();
+		return *this;
+	}
+
+	TScriptContainerIterator operator++(int)
+	{
+		const TScriptContainerIterator Temp(*this);
+		Advance();
+		return Temp;
+	}
+
+	explicit operator bool() const
+	{
+		return Container.IsValidIndex(InternalIndex);
+	}
+
+	int32 GetInternalIndex() const
+	{
+		return InternalIndex;
+	}
+
+	int32 GetLogicalIndex() const
+	{
+		return LogicalIndex;
+	}
+
+	UE_DEPRECATED(5.4, "Use Iterator directly, GetInternalIndex or GetLogicalIndex instead.")
+	int32 operator*() const
+	{
+		return InternalIndex;
+	}
+
+private:
+	const ContainerType& Container;
+	int32 InternalIndex = INDEX_NONE;
+	int32 LogicalIndex = INDEX_NONE;
+
+	void Advance()
+	{
+		++InternalIndex;
+		const int32 MaxIndex = Container.GetMaxIndex();
+		while (InternalIndex < MaxIndex && !Container.IsValidIndex(InternalIndex))
+		{
+			++InternalIndex;
+		}
+
+		++LogicalIndex;
+	}
+};
 
 /**
  * FScriptMapHelper: Pseudo dynamic map. Used to work with map properties in a sensible way.
+ * Note that map can contain invalid entries some number of valid entries (i.e. Num() ) can
+ * be smaller that the actual number of elements (i.e. GetMaxIndex() ).
+ *
+ * Internal index naming is used to identify the actual index in the container which can point to
+ * an invalid entry. It can be used for methods like Get<Item>Ptr, Get<Item>PtrWithoutCheck or IsValidIndex.
+ *
+ * Logical index naming is used to identify only valid entries in the container so it can be smaller than the
+ * internal index in case we skipped invalid entries to reach the next valid one. This index is used on method
+ * like FindNth<Item>Ptr or FindInternalIndex.
+ * This is also the type of index we receive from most editor events (e.g. property change events) so it is
+ * strongly suggested to rely on FScriptMapHelper::FIterator to iterate or convert to internal index.
  */
 class FScriptMapHelper
 {
@@ -4404,6 +4500,18 @@ public:
 	{
 	}
 
+	using FIterator = TScriptContainerIterator<FScriptMapHelper>;
+
+	FIterator CreateIterator() const
+	{
+		return FIterator(*this);
+	}
+
+	FIterator CreateIterator(const int32 InLogicalIndex) const
+	{
+		return FIterator(*this, InLogicalIndex);
+	}
+	
 	/**
 	 * Index range check
 	 *
@@ -4456,15 +4564,27 @@ public:
 		{
 			if (Map->Num() == 0)
 			{
-				checkSlow(!InternalIndex);
+				checkf(InternalIndex == 0, TEXT("Legacy implementation was only allowing requesting InternalIndex 0 on an empty container."));
 				return nullptr;
 			}
 
-			checkSlow(Map->IsValidIndex(InternalIndex));
+			checkf(IsValidIndex(InternalIndex), TEXT("Invalid internal index. Use IsValidIndex before calling this method."));
 			return (uint8*)Map->GetData(InternalIndex, MapLayout);
 		});
 	}
 
+	/**
+	 * Returns a uint8 pointer to the pair in the map.
+	 *
+	 * @param InternalIndex index of the item to return a pointer to.
+	 *
+	 * @return Pointer to the pair, or nullptr if the map is empty.
+	 */
+	FORCEINLINE const uint8* GetPairPtr(const int32 InternalIndex) const
+	{
+		return const_cast<FScriptMapHelper*>(this)->GetPairPtr(InternalIndex);
+	}
+
 	/**
 	 * Returns a uint8 pointer to the Key (first element) in the map. Currently 
 	 * identical to GetPairPtr, but provides clarity of purpose and avoids exposing
@@ -4480,11 +4600,11 @@ public:
 		{
 			if (Map->Num() == 0)
 			{
-				checkSlow(!InternalIndex);
+				checkf(InternalIndex == 0, TEXT("Legacy implementation was only allowing requesting InternalIndex 0 on an empty container."));
 				return nullptr;
 			}
-		
-			checkSlow(Map->IsValidIndex(InternalIndex));
+
+			checkf(IsValidIndex(InternalIndex), TEXT("Invalid internal index. Use IsValidIndex before calling this method."));
 			return (uint8*)Map->GetData(InternalIndex, MapLayout);
 		});
 	}
@@ -4502,27 +4622,103 @@ public:
 		{
 			if (Map->Num() == 0)
 			{
-				checkSlow(!InternalIndex);
+				checkf(InternalIndex == 0, TEXT("Legacy implementation was only allowing requesting InternalIndex 0 on an empty container."));
 				return nullptr;
 			}
-		
-			checkSlow(Map->IsValidIndex(InternalIndex));
+
+			checkf(IsValidIndex(InternalIndex), TEXT("Invalid internal index. Use IsValidIndex before calling this method."));
 			return (uint8*)Map->GetData(InternalIndex, MapLayout) + MapLayout.ValueOffset;
 		});
 	}
 
+	/**
+	 * Returns a uint8 pointer to the pair in the map
+	 *
+	 * @param Iterator A valid iterator of the item to return a pointer to.
+	 *
+	 * @return Pointer to the pair, or will fail a check if an invalid iterator is provided.
+	 */
+	FORCEINLINE uint8* GetPairPtr(const FIterator Iterator)
+	{
+		return WithScriptMap([this, Iterator](auto* Map) -> uint8*
+		{
+			checkf(Iterator, TEXT("Invalid Iterator. Test Iterator before calling this method."));
+			return (uint8*)Map->GetData(Iterator.GetInternalIndex(), MapLayout);
+		});
+	}
+
 	/**
 	 * Returns a uint8 pointer to the pair in the map.
 	 *
-	 * @param InternalIndex index of the item to return a pointer to.
+	 * @param Iterator A valid iterator of the item to return a pointer to.
 	 *
-	 * @return Pointer to the pair, or nullptr if the map is empty.
+	 * @return Pointer to the pair, or will fail a check if an invalid iterator is provided.
 	 */
-	FORCEINLINE const uint8* GetPairPtr(int32 InternalIndex) const
+	FORCEINLINE const uint8* GetPairPtr(const FIterator Iterator) const
 	{
-		return const_cast<FScriptMapHelper*>(this)->GetPairPtr(InternalIndex);
+		return const_cast<FScriptMapHelper*>(this)->GetPairPtr(Iterator);
 	}
-	
+
+	/**
+	 * Returns a uint8 pointer to the Key (first element) in the map. Currently
+	 * identical to GetPairPtr, but provides clarity of purpose and avoids exposing
+	 * implementation details of TMap.
+	 *
+	 * @param Iterator A valid iterator of the item to return a pointer to.
+	 *
+	 * @return Pointer to the key, or will fail a check if an invalid iterator is provided.
+	 */
+	FORCEINLINE uint8* GetKeyPtr(const FIterator Iterator)
+	{
+		return WithScriptMap([this, Iterator](auto* Map) -> uint8*
+		{
+			checkf(Iterator, TEXT("Invalid Iterator. Test Iterator before calling this method."));
+			return (uint8*)Map->GetData(Iterator.GetInternalIndex(), MapLayout);
+		});
+	}
+
+	/**
+	 * Returns a const uint8 pointer to the Key (first element) in the map. Currently
+	 * identical to GetPairPtr, but provides clarity of purpose and avoids exposing
+	 * implementation details of TMap.
+	 *
+	 * @param Iterator A valid iterator of the item to return a pointer to.
+	 *
+	 * @return Pointer to the key, or will fail a check if an invalid iterator is provided.
+	 */
+	FORCEINLINE const uint8* GetKeyPtr(const FIterator Iterator) const
+	{
+		return const_cast<FScriptMapHelper*>(this)->GetKeyPtr(Iterator);
+	}
+
+	/**
+	 * Returns a uint8 pointer to the Value (second element) in the map.
+	 *
+	 * @param Iterator A valid iterator of the item to return a pointer to.
+	 *
+	 * @return Pointer to the value, or will fail a check if an invalid iterator is provided.
+	 */
+	FORCEINLINE uint8* GetValuePtr(const FIterator Iterator)
+	{
+		return WithScriptMap([this, Iterator](auto* Map) -> uint8*
+		{
+			checkf(Iterator, TEXT("Invalid Iterator. Test Iterator before calling this method."));
+			return (uint8*)Map->GetData(Iterator.GetInternalIndex(), MapLayout) + MapLayout.ValueOffset;
+		});
+	}
+
+	/**
+	 * Returns a const uint8 pointer to the Value (second element) in the map.
+	 *
+	 * @param Iterator A valid iterator of the item to return a pointer to.
+	 *
+	 * @return Pointer to the value, or will fail a check if an invalid iterator is provided.
+	 */
+	FORCEINLINE const uint8* GetValuePtr(const FIterator Iterator) const
+	{
+		return const_cast<FScriptMapHelper*>(this)->GetValuePtr(Iterator);
+	}
+
 	/**
 	* Returns a uint8 pointer to the the Nth valid pair in the map (skipping invalid entries).
 	* NOTE: This is slow, do not use this for iteration! Use CreateIterator() instead.
@@ -4532,9 +4728,7 @@ public:
 	uint8* FindNthPairPtr(int32 N)
 	{
 		const int32 InternalIndex = FindInternalIndex(N);
-		
-		checkSlow(IsValidIndex(InternalIndex));
-		return (InternalIndex != INDEX_NONE) ? GetPairPtr(InternalIndex) : nullptr;
+		return (InternalIndex != INDEX_NONE) ? GetPairPtrWithoutCheck(InternalIndex) : nullptr;
 	}
 	
 	/**
@@ -4546,7 +4740,7 @@ public:
 	uint8* FindNthKeyPtr(int32 N)
 	{
 		const int32 InternalIndex = FindInternalIndex(N);
-		return (InternalIndex != INDEX_NONE) ? GetKeyPtr(InternalIndex) : nullptr;
+		return (InternalIndex != INDEX_NONE) ? GetKeyPtrWithoutCheck(InternalIndex) : nullptr;
 	}
 	
 	/**
@@ -4558,7 +4752,7 @@ public:
 	uint8* FindNthValuePtr(int32 N)
 	{
 		const int32 InternalIndex = FindInternalIndex(N);
-		return (InternalIndex != INDEX_NONE) ? GetValuePtr(InternalIndex) : nullptr;
+		return (InternalIndex != INDEX_NONE) ? GetValuePtrWithoutCheck(InternalIndex) : nullptr;
 	}
 	
 	/**
@@ -4570,7 +4764,7 @@ public:
 	const uint8* FindNthPairPtr(int32 N) const
 	{
 		const int32 InternalIndex = FindInternalIndex(N);
-		return (InternalIndex != INDEX_NONE) ? GetPairPtr(InternalIndex) : nullptr;
+		return (InternalIndex != INDEX_NONE) ? GetPairPtrWithoutCheck(InternalIndex) : nullptr;
 	}
 
 	/**
@@ -4829,15 +5023,15 @@ public:
 	 */
 	FORCEINLINE uint8* FindMapPairPtrWithKey(const void* PairWithKeyToFind, int32 IndexHint = 0)
 	{
-		int32 InternalIndex = FindMapIndexWithKey(PairWithKeyToFind, IndexHint);
-		uint8* Result = (InternalIndex >= 0) ? GetPairPtr(InternalIndex) : nullptr;
+		const int32 InternalIndex = FindMapIndexWithKey(PairWithKeyToFind, IndexHint);
+		uint8* Result = (InternalIndex >= 0) ? GetPairPtrWithoutCheck(InternalIndex) : nullptr;
 		return Result;
 	}
 
 	/** Finds the associated pair from hash, rather than linearly searching */
 	int32 FindMapPairIndexFromHash(const void* KeyPtr)
 	{
-		int32 InternalIndex = WithScriptMap([this, KeyPtr, LocalKeyPropForCapture = this->KeyProp](auto* Map)
+		const int32 InternalIndex = WithScriptMap([this, KeyPtr, LocalKeyPropForCapture = this->KeyProp](auto* Map)
 		{
 			return Map->FindPairIndex(
 				KeyPtr,
@@ -4852,8 +5046,8 @@ public:
 	/** Finds the associated pair from hash, rather than linearly searching */
 	uint8* FindMapPairPtrFromHash(const void* KeyPtr)
 	{
-		int32 InternalIndex = FindMapPairIndexFromHash(KeyPtr);
-		uint8* Result = (InternalIndex >= 0) ? GetPairPtr(InternalIndex) : nullptr;
+		const int32 InternalIndex = FindMapPairIndexFromHash(KeyPtr);
+		uint8* Result = (InternalIndex >= 0) ? GetPairPtrWithoutCheck(InternalIndex) : nullptr;
 		return Result;
 	}
 
@@ -5009,40 +5203,6 @@ public:
 		);
 	}
 
-	class FIterator
-	{
-	public:
-		explicit FIterator(const FScriptMapHelper& InMap) :
-			Map(InMap),
-			CurrentIndex(-1)
-		{
-			Advance();
-		}
-
-		FIterator& operator++() { Advance(); return *this; }
-		FIterator operator++(int) { const FIterator Temp(*this); Advance(); return Temp; }
-		explicit operator bool() const { return Map.IsValidIndex(CurrentIndex); }
-		int32 operator*() const { return CurrentIndex; }
-
-	private:
-		const FScriptMapHelper& Map;
-		int32 CurrentIndex;
-
-		void Advance()
-		{
-			++CurrentIndex;
-			while (CurrentIndex < Map.GetMaxIndex() && !Map.IsValidIndex(CurrentIndex))
-			{
-				++CurrentIndex;
-			}
-		}
-	};
-
-	FScriptMapHelper::FIterator CreateIterator() const
-	{
-		return FIterator(*this);
-	}
-
 private:
 	FORCEINLINE FScriptMapHelper(EInternal, FProperty* InKeyProp, FProperty* InValueProp, const void* InMap, const FScriptMapLayout& InMapLayout, EMapPropertyFlags InMapFlags)
 		: KeyProp  (InKeyProp)
@@ -5184,6 +5344,54 @@ private:
 		return const_cast<FScriptMapHelper*>(this)->GetPairPtrWithoutCheck(InternalIndex);
 	}
 
+	/**
+	 * Returns a uint8 pointer to the key in the array without checking the index.
+	 *
+	 * @param InternalIndex index of the key to return a pointer to.
+	 *
+	 * @return Pointer to the key, or nullptr if the map is empty.
+	 */
+	FORCEINLINE uint8* GetKeyPtrWithoutCheck(int32 InternalIndex)
+	{
+		return WithScriptMap([this, InternalIndex](auto* Map) { return (uint8*)Map->GetData(InternalIndex, MapLayout); });
+	}
+
+	/**
+	 * Returns a const uint8 pointer to the pair in the array without checking the index.
+	 *
+	 * @param InternalIndex index of the item to return a pointer to.
+	 *
+	 * @return Pointer to the pair, or nullptr if the map is empty.
+	 */
+	FORCEINLINE const uint8* GetKeyPtrWithoutCheck(int32 InternalIndex) const
+	{
+		return const_cast<FScriptMapHelper*>(this)->GetKeyPtrWithoutCheck(InternalIndex);
+	}
+
+	/**
+	 * Returns a uint8 pointer to the pair in the array without checking the index.
+	 *
+	 * @param InternalIndex index of the item to return a pointer to.
+	 *
+	 * @return Pointer to the pair, or nullptr if the map is empty.
+	 */
+	FORCEINLINE uint8* GetValuePtrWithoutCheck(int32 InternalIndex)
+	{
+		return WithScriptMap([this, InternalIndex](auto* Map) { return (uint8*)Map->GetData(InternalIndex, MapLayout) + MapLayout.ValueOffset; });
+	}
+
+	/**
+	 * Returns a const uint8 pointer to the pair in the array without checking the index.
+	 *
+	 * @param InternalIndex index of the item to return a pointer to.
+	 *
+	 * @return Pointer to the pair, or nullptr if the map is empty.
+	 */
+	FORCEINLINE const uint8* GetValuePtrWithoutCheck(int32 InternalIndex) const
+	{
+		return const_cast<FScriptMapHelper*>(this)->GetValuePtrWithoutCheck(InternalIndex);
+	}
+
 public:
 	FProperty*        KeyProp;
 	FProperty*        ValueProp;
@@ -5206,13 +5414,37 @@ public:
 };
 
 /**
-* FScriptSetHelper: Pseudo dynamic Set. Used to work with Set properties in a sensible way.
-*/
+ * FScriptSetHelper: Pseudo dynamic Set. Used to work with Set properties in a sensible way.
+ * Note that the set can contain invalid entries some number of valid entries (i.e. Num() ) can
+ * be smaller that the actual number of elements (i.e. GetMaxIndex() ).
+ *
+ * Internal index naming is used to identify the actual index in the container which can point to
+ * an invalid entry. It can be used for methods like Get<Item>Ptr, Get<Item>PtrWithoutCheck or IsValidIndex.
+ *
+ * Logical index naming is used to identify only valid entries in the container so it can be smaller than the
+ * internal index in case we skipped invalid entries to reach the next valid one. This index is used on method
+ * like FindNth<Item>Ptr or FindInternalIndex.
+ * This is also the type of index we receive from most editor events (e.g. property change events) so it is
+ * strongly suggested to rely on FScriptSetHelper::FIterator to iterate or convert to internal index.
+ */
 class FScriptSetHelper
 {
 	friend class FSetProperty;
 
 public:
+
+	using FIterator = TScriptContainerIterator<FScriptSetHelper>;
+
+	FIterator CreateIterator() const
+	{
+		return FIterator(*this);
+	}
+
+	FIterator CreateIterator(const int32 InLogicalIndex) const
+	{
+		return FIterator(*this, InLogicalIndex);
+	}
+
 	/**
 	* Constructor, brings together a property and an instance of the property located in memory
 	*
@@ -5288,11 +5520,11 @@ public:
 	{
 		if (Num() == 0)
 		{
-			checkSlow(!InternalIndex);
+			checkf(InternalIndex == 0, TEXT("Legacy implementation was only allowing requesting InternalIndex 0 on an empty container."));
 			return nullptr;
 		}
 
-		checkSlow(IsValidIndex(InternalIndex));
+		checkf(IsValidIndex(InternalIndex), TEXT("Invalid internal index. Use IsValidIndex before calling this method."));
 		return (uint8*)Set->GetData(InternalIndex, SetLayout);
 	}
 
@@ -5308,6 +5540,31 @@ public:
 		return const_cast<FScriptSetHelper*>(this)->GetElementPtr(InternalIndex);
 	}
 
+	/**
+	 * Returns a uint8 pointer to the element in the set.
+	 *
+	 * @param Iterator A valid iterator of the item to return a pointer to.
+	 *
+	 * @return Pointer to the element, or will fail a check if an invalid iterator is provided.
+	 */
+	FORCEINLINE uint8* GetElementPtr(const FIterator Iterator)
+	{
+		checkf(Iterator, TEXT("Invalid Iterator. Test Iterator before calling this method."));
+		return (uint8*)Set->GetData(Iterator.GetInternalIndex(), SetLayout);
+	}
+
+	/**
+	 * Returns a uint8 pointer to the element in the set.
+	 *
+	 * @param Iterator A valid iterator of the item to return a pointer to.
+	 *
+	 * @return Pointer to the element, or will fail a check if an invalid iterator is provided.
+	 */
+	FORCEINLINE const uint8* GetElementPtr(const FIterator Iterator) const
+	{
+		return const_cast<FScriptSetHelper*>(this)->GetElementPtr(Iterator);
+	}
+
 	/**
 	* Returns a uint8 pointer to the the Nth valid element in the set (skipping invalid entries).
 	* NOTE: This is slow, do not use this for iteration! Use CreateIterator() instead.
@@ -5317,7 +5574,7 @@ public:
 	uint8* FindNthElementPtr(int32 N)
 	{
 		const int32 InternalIndex = FindInternalIndex(N);
-		return (InternalIndex != INDEX_NONE) ? GetElementPtr(InternalIndex) : nullptr;
+		return (InternalIndex != INDEX_NONE) ? GetElementPtrWithoutCheck(InternalIndex) : nullptr;
 	}
 
 	/**
@@ -5329,7 +5586,7 @@ public:
 	const uint8* FindNthElementPtr(int32 N) const
 	{
 		const int32 InternalIndex = FindInternalIndex(N);
-		return (InternalIndex != INDEX_NONE) ? GetElementPtr(InternalIndex) : nullptr;
+		return (InternalIndex != INDEX_NONE) ? GetElementPtrWithoutCheck(InternalIndex) : nullptr;
 	}
 
 	/**
@@ -5520,7 +5777,7 @@ public:
 	FORCEINLINE uint8* FindElementPtr(const void* ElementToFind, int32 IndexHint = 0)
 	{
 		const int32 InternalIndex = FindElementIndex(ElementToFind, IndexHint);
-		uint8* Result = (InternalIndex >= 0 ? GetElementPtr(InternalIndex) : nullptr);
+		uint8* Result = (InternalIndex >= 0 ? GetElementPtrWithoutCheck(InternalIndex) : nullptr);
 		return Result;
 	}
 
@@ -5540,7 +5797,7 @@ public:
 	FORCEINLINE uint8* FindElementPtrFromHash(const void* ElementToFind)
 	{
 		const int32 InternalIndex = FindElementIndexFromHash(ElementToFind);
-		uint8* Result = (InternalIndex >= 0 ? GetElementPtr(InternalIndex) : nullptr);
+		uint8* Result = (InternalIndex >= 0 ? GetElementPtrWithoutCheck(InternalIndex) : nullptr);
 		return Result;
 	}
 
@@ -5613,40 +5870,6 @@ public:
 		return ScriptSetHelper;
 	}
 
-	class FIterator
-	{
-	public:
-		explicit FIterator(const FScriptSetHelper& InSet) :
-			Set(InSet),
-			CurrentIndex(-1)
-		{
-			Advance();
-		}
-
-		FIterator& operator++() { Advance(); return *this; }
-		FIterator operator++(int) { const FIterator Temp(*this); Advance(); return Temp; }
-		explicit operator bool() const { return Set.IsValidIndex(CurrentIndex); }
-		int32 operator*() const { return CurrentIndex; }
-
-	private:
-		const FScriptSetHelper& Set;
-		int32 CurrentIndex;
-
-		void Advance()
-		{
-			++CurrentIndex;
-			while (CurrentIndex < Set.GetMaxIndex() && !Set.IsValidIndex(CurrentIndex))
-			{
-				++CurrentIndex;
-			}
-		}
-	};
-
-	FScriptSetHelper::FIterator CreateIterator() const
-	{
-		return FIterator(*this);
-	}
-
 private: 
 	FScriptSetHelper()
 		: ElementProp(nullptr)
