PackageName: Add some ObjectPath functions needed by CoreRedirects.
#rnx

diff --git a/Engine/Source/Runtime/CoreUObject/Private/Misc/PackageName.cpp b/Engine/Source/Runtime/CoreUObject/Private/Misc/PackageName.cpp
--- a/Engine/Source/Runtime/CoreUObject/Private/Misc/PackageName.cpp
+++ b/Engine/Source/Runtime/CoreUObject/Private/Misc/PackageName.cpp
@@ -31,6 +31,8 @@
 #include "Modules/ModuleManager.h"
 #include "ProfilingDebugging/LoadTimeTracker.h"
 #include "Stats/Stats.h"
+#include "String/Find.h"
+#include "String/ParseTokens.h"
 #include "Templates/UniquePtr.h"
 #include "UObject/Linker.h"
 #include "UObject/Package.h"
@@ -2703,6 +2705,93 @@ FWideStringView FPackageName::ObjectPathToObjectName(FWideStringView InObjectPat
 	return ObjectPathToObjectNameImpl(InObjectPath);
 }
 
+template<class CharType>
+static void ObjectPathSplitFirstNameImpl(TStringView<CharType> Text, TStringView<CharType>& OutFirst,
+	TStringView<CharType>& OutRemainder)
+{
+	int32 DelimiterIndex = UE::String::FindFirstOfAnyChar(Text, { CharType(':'), CharType('.') });
+	if (DelimiterIndex < 0)
+	{
+		OutFirst = Text;
+		OutRemainder.Reset();
+		return;
+	}
+	OutFirst = Text.Left(DelimiterIndex);
+	OutRemainder = Text.RightChop(DelimiterIndex + 1);
+}
+
+void FPackageName::ObjectPathSplitFirstName(FWideStringView Text, FWideStringView& OutFirst,
+	FWideStringView& OutRemainder)
+{
+	ObjectPathSplitFirstNameImpl(Text, OutFirst, OutRemainder);
+}
+
+void FPackageName::ObjectPathSplitFirstName(FAnsiStringView Text, FAnsiStringView& OutFirst, FAnsiStringView& OutRemainder)
+{
+	ObjectPathSplitFirstNameImpl(Text, OutFirst, OutRemainder);
+}
+
+void FPackageName::ObjectPathAppend(FStringBuilderBase& ObjectPath, FStringView NextName)
+{
+	if (ObjectPath.Len() == 0)
+	{
+		ObjectPath << NextName;
+		return;
+	}
+	if (NextName.IsEmpty())
+	{
+		return;
+	}
+	if (NextName[0] == '/')
+	{
+		ObjectPath.Reset();
+		ObjectPath << NextName;
+		return;
+	}
+
+	int32 NumObjectDelimiters = 0;
+	{
+		int32 LastSlash;
+		FStringView ObjectPathView(ObjectPath);
+		ObjectPathView.FindLastChar('/', LastSlash);
+		if (LastSlash == INDEX_NONE)
+		{
+			// Not a full object path. Always append with '.'
+			NumObjectDelimiters = 2;
+		}
+		else
+		{
+			ObjectPathView.RightChopInline(LastSlash + 1);
+			for (NumObjectDelimiters = 0;
+				NumObjectDelimiters < 2; // Stop counting after the second delimiter since behavior no longer changes
+				++NumObjectDelimiters)
+			{
+				int32 NextDelimiter = UE::String::FindFirstOfAnyChar(ObjectPathView, { TCHAR('.'), TCHAR(':') });
+				if (NextDelimiter == INDEX_NONE)
+				{
+					break;
+				}
+				ObjectPathView.RightChopInline(NextDelimiter + 1);
+			}
+		}
+	}
+
+	UE::String::ParseTokensMultiple(NextName, { TCHAR(':'), TCHAR('.') },
+	[&ObjectPath, &NumObjectDelimiters](FStringView NextSingleName)
+	{
+		TCHAR Delimiter = NumObjectDelimiters++ == 1 ? ':' : '.';
+		ObjectPath << Delimiter << NextSingleName;
+	},
+	UE::String::EParseTokensOptions::SkipEmpty);
+}
+
+FString FPackageName::ObjectPathCombine(FStringView ObjectPath, FStringView NextName)
+{
+	TStringBuilder<256> Base(InPlace, ObjectPath);
+	ObjectPathAppend(Base, NextName);
+	return FString(Base);
+}
+
 bool FPackageName::IsVersePackage(FStringView InPackageName)
 {
 	return InPackageName.Contains(FLongPackagePathsSingleton::Get().VerseSubPath);
@@ -3002,6 +3091,55 @@ bool FPackageNameTests::RunTest(const FString& Parameters)
 			};
 			RunObjectPathTests(TEXT("ObjectPathToObjectName"), ExpectedOutputPaths, [](FStringView ObjectPath) { return FPackageName::ObjectPathToObjectName(ObjectPath); });
 		}
+
+		// ObjectPathSplitFirstName
+		{
+			auto ObjectPathSplitFirstNameTest = [this](FStringView Text, FStringView ExpectedFirst, FStringView ExpectedRemainder)
+				{
+					FStringView ActualFirst;
+					FStringView ActualRemainder;
+					FPackageName::ObjectPathSplitFirstName(Text, ActualFirst, ActualRemainder);
+					if (!ActualFirst.Equals(ExpectedFirst, ESearchCase::CaseSensitive) ||
+						!ActualRemainder.Equals(ExpectedRemainder, ESearchCase::CaseSensitive))
+					{
+						AddError(*WriteToString<256>(TEXT("ObjectPathSplitFirstName"), TEXT(": Expected {'"),
+							ExpectedFirst, TEXT("', '"), ExpectedRemainder, TEXT("'} but got {'"),
+							ActualFirst, TEXT("', '"), ActualRemainder, TEXT("'} for input '"), Text, TEXT("'")));
+					}
+				};
+			ObjectPathSplitFirstNameTest(TEXT("/Game/MyAsset.MyAsset:SubObject.AnotherObject"), TEXT("/Game/MyAsset"), TEXT("MyAsset:SubObject.AnotherObject"));
+			ObjectPathSplitFirstNameTest(TEXT("/Game/MyAsset.MyAsset:SubObject"), TEXT("/Game/MyAsset"), TEXT("MyAsset:SubObject"));
+			ObjectPathSplitFirstNameTest(TEXT("/Game/MyAsset.MyAsset"), TEXT("/Game/MyAsset"), TEXT("MyAsset"));
+			ObjectPathSplitFirstNameTest(TEXT("/Game/MyAsset"), TEXT("/Game/MyAsset"), TEXT(""));
+			ObjectPathSplitFirstNameTest(TEXT("MyAsset:SubObject"), TEXT("MyAsset"), TEXT("SubObject"));
+			ObjectPathSplitFirstNameTest(TEXT("MyAsset.SubObject"), TEXT("MyAsset"), TEXT("SubObject"));
+			ObjectPathSplitFirstNameTest(TEXT("MyAsset"), TEXT("MyAsset"), TEXT(""));
+			ObjectPathSplitFirstNameTest(TEXT(""), TEXT(""), TEXT(""));
+		}
+
+		// ObjectPathAppend
+		{
+			auto ObjectPathAppendTest = [this](FStringView A, FStringView B, FStringView Expected)
+				{
+					TStringBuilder<256> Base(InPlace, A);
+					FPackageName::ObjectPathAppend(Base, B);
+					if (!Base.ToView().Equals(Expected, ESearchCase::CaseSensitive))
+					{
+						AddError(*WriteToString<256>(TEXT("ObjectPathAppend"), TEXT(": Expected '"),
+							Expected, TEXT("' but got '"), Base, TEXT("' for input {'"),
+							A, TEXT("', '"), B, TEXT("'}")));
+					}
+				};
+			ObjectPathAppendTest(TEXT("/Package"), TEXT("Object"), TEXT("/Package.Object"));
+			ObjectPathAppendTest(TEXT("/Package.Object"), TEXT("SubObject"), TEXT("/Package.Object:SubObject"));
+			ObjectPathAppendTest(TEXT("/Package.Object:SubObject"), TEXT("NextSubObject"), TEXT("/Package.Object:SubObject.NextSubObject"));
+			ObjectPathAppendTest(TEXT("/Package"), TEXT("Object.SubObject"), TEXT("/Package.Object:SubObject"));
+			ObjectPathAppendTest(TEXT("/Package"), TEXT("/OtherPackage.Object:SubObject"), TEXT("/OtherPackage.Object:SubObject"));
+			ObjectPathAppendTest(TEXT("/Package"), TEXT(""), TEXT("/Package"));
+			ObjectPathAppendTest(TEXT(""), TEXT("/Package.Object:SubObject"), TEXT("/Package.Object:SubObject"));
+			ObjectPathAppendTest(TEXT(""), TEXT("Object"), TEXT("Object"));
+			ObjectPathAppendTest(TEXT(""), TEXT(""), TEXT(""));
+		}
 	}
 
 	// TryConvertToMountedPath
diff --git a/Engine/Source/Runtime/CoreUObject/Public/Misc/PackageName.h b/Engine/Source/Runtime/CoreUObject/Public/Misc/PackageName.h
--- a/Engine/Source/Runtime/CoreUObject/Public/Misc/PackageName.h
+++ b/Engine/Source/Runtime/CoreUObject/Public/Misc/PackageName.h
@@ -727,6 +727,33 @@ public:
 	static COREUOBJECT_API FAnsiStringView ObjectPathToObjectName(FAnsiStringView InObjectPath);
 	static COREUOBJECT_API FString ObjectPathToObjectName(const FString& InObjectPath);
 
+	/**
+	 * Splits an ObjectPath string into the first component and the remainder. 
+	 *
+	 * "/Path/To/A/Package.Object:SubObject" -> { "/Path/To/A/Package", "Object:SubObject" }
+	 * "Object:SubObject" -> { "Object", "SubObject" }
+	 * "Object.SubObject" -> { "Object", "SubObject" }
+	 */
+	static COREUOBJECT_API void ObjectPathSplitFirstName(FWideStringView Text,
+		FWideStringView& OutFirst, FWideStringView& OutRemainder);
+	static COREUOBJECT_API void ObjectPathSplitFirstName(FAnsiStringView Text,
+		FAnsiStringView& OutFirst, FAnsiStringView& OutRemainder);
+
+	/**
+	 * Combines an ObjectPath with an ObjectName.
+	 * { "/Package", "Object" } -> "/Package.Object"
+	 * { "/Package.Object", "SubObject" } -> "/Package.Object:SubObject"
+	 * { "/Package.Object:SubObject", "NextSubObject" } -> "/Package.Object:SubObject.NextSubObject"
+	 * { "/Package", "Object.SubObject" } -> "/Package.Object:SubObject"
+	 * { "/Package", "/OtherPackage.Object:SubObject" } -> "/OtherPackage.Object:SubObject"
+	*/
+	static COREUOBJECT_API void ObjectPathAppend(FStringBuilderBase& ObjectPath, FStringView NextName);
+	/**
+	 * Combines an ObjectPath with an ObjectName, the same as ObjectPathAppend but returns the result rather
+	 * than modifying the input argument.
+	 */
+	static COREUOBJECT_API FString ObjectPathCombine(FStringView ObjectPath, FStringView NextName);
+
 	/**
 	 * Checks the package's path to see if it's a Verse package
 	 */
