Make sure all BP CDO subobjects are correctly loaded before the end of deferred BP loading instead of only native-added default subobjects
This fixes inheritance of subobject properties in complicated loading conditions

diff --git a/Engine/Source/Runtime/CoreUObject/Private/Blueprint/BlueprintSupport.cpp b/Engine/Source/Runtime/CoreUObject/Private/Blueprint/BlueprintSupport.cpp
--- a/Engine/Source/Runtime/CoreUObject/Private/Blueprint/BlueprintSupport.cpp
+++ b/Engine/Source/Runtime/CoreUObject/Private/Blueprint/BlueprintSupport.cpp
@@ -2265,17 +2265,33 @@ void FLinkerLoad::ResolveDeferredExports(UClass* LoadClass)
 		//       be cleared (and this will do nothing the 2nd time around)
 		Preload(BlueprintCDO);
 
-		// Ensure that all default subobject exports belonging to the CDO have been created. DSOs may no longer be
-		// referenced by a tagged property and thus may not get created and registered until after class regeneration.
-		// This can cause invalid subobjects to register themselves with a regenerated CDO if the native parent class
-		// has been changed to inherit from an entirely different type since the last time the class asset was saved.
-		// By constructing them here, we make sure that LoadAllObjects() won't construct them after class regeneration.
+		// Ensure that all subobject exports belonging to the CDO have been created. This is often handled by 
+		// PreloadSubobjects in CreateExport, but they can get skipped. Subobjects need to be created here so
+		// they can be correctly inherited by any child classes and they are correctly registered for any later
+		// deferred fixups related to native class changes.
 		for (int32 ExportIndex = 0; ExportIndex < ExportMap.Num(); ++ExportIndex)
 		{
 			FObjectExport& Export = ExportMap[ExportIndex];
-			if((Export.ObjectFlags & RF_DefaultSubObject) != 0 && Export.OuterIndex.IsExport() && Export.OuterIndex.ToExport() == DeferredCDOIndex)
+			FPackageIndex CheckOuterIndex = Export.OuterIndex;
+			bool bInsideCDO = false;
+			while (CheckOuterIndex.IsExport())
 			{
-				if (Export.Object == nullptr && Export.OuterIndex.IsExport())
+				int32 OuterExportIndex = CheckOuterIndex.ToExport();
+				if (OuterExportIndex == DeferredCDOIndex)
+				{
+					bInsideCDO = true;
+					break;
+				}
+				else
+				{
+					// Handle nested subobjects
+					CheckOuterIndex = ExportMap[OuterExportIndex].OuterIndex;
+				}
+			}
+
+			if (bInsideCDO)
+			{
+				if (Export.Object == nullptr)
 				{
 					CreateExport(ExportIndex);
 				}
diff --git a/Engine/Source/Runtime/CoreUObject/Private/UObject/LinkerLoad.cpp b/Engine/Source/Runtime/CoreUObject/Private/UObject/LinkerLoad.cpp
--- a/Engine/Source/Runtime/CoreUObject/Private/UObject/LinkerLoad.cpp
+++ b/Engine/Source/Runtime/CoreUObject/Private/UObject/LinkerLoad.cpp
@@ -5249,7 +5249,7 @@ UObject* FLinkerLoad::CreateExport( int32 Index )
 
 
 		// If we are about to create a CDO, we need to ensure that all parent sub-objects are loaded
-		// to get default value initialization to work.
+		// to get default value initialization to work. This matches code in ResolveDeferredExports
 		if ((ObjectLoadFlags & RF_ClassDefaultObject) != 0)
 		{
 			TArray<UObject*> SubObjects;
