Integrate FPropertyBag to gather unknown properties from TPS

diff --git a/Engine/Source/Runtime/CoreUObject/Private/UObject/Class.cpp b/Engine/Source/Runtime/CoreUObject/Private/UObject/Class.cpp
--- a/Engine/Source/Runtime/CoreUObject/Private/UObject/Class.cpp
+++ b/Engine/Source/Runtime/CoreUObject/Private/UObject/Class.cpp
@@ -35,6 +35,8 @@
 #include "UObject/Interface.h"
 #include "UObject/LinkerPlaceholderClass.h"
 #include "UObject/LinkerPlaceholderFunction.h"
+#include "UObject/PropertyBag.h"
+#include "UObject/PropertyBagRepository.h"
 #include "UObject/PropertyOptional.h"
 #include "UObject/StructOnScope.h"
 #include "UObject/StructScriptLoader.h"
@@ -1305,7 +1307,10 @@ const FBlake3Hash& UStruct::GetSchemaHash(bool bSkipEditorOnly) const
 
 void UStruct::SerializeVersionedTaggedProperties(FStructuredArchive::FSlot Slot, uint8* Data, UStruct* DefaultsStruct, uint8* Defaults, const UObject* BreakRecursionIfFullyLoad) const
 {
+	using namespace UE;
+
 	FArchive& UnderlyingArchive = Slot.GetUnderlyingArchive();
+	FUObjectSerializeContext* LoadContext = UnderlyingArchive.GetSerializeContext();
 	//SCOPED_LOADTIMER(SerializeTaggedPropertiesTime);
 
 	// Determine if this struct supports optional property guid's (UBlueprintGeneratedClasses Only)
@@ -1322,19 +1327,36 @@ void UStruct::SerializeVersionedTaggedProperties(FStructuredArchive::FSlot Slot,
 		else
 #endif // WITH_TEXT_ARCHIVE_SUPPORT
 		{
-		// Load tagged properties.
-		FStructuredArchive::FStream PropertiesStream = Slot.EnterStream();
+			auto TryFindPropertyBag = [PropertyBag = (FPropertyBag*)nullptr, bSearched = false, LoadContext]() mutable -> FPropertyBag*
+			{
+				if (bSearched)
+				{
+					return PropertyBag;
+				}
+				bSearched = true;
+				if (LoadContext && LoadContext->bSerializeUnknownProperty)
+				{
+					if (UObject* Object = LoadContext->SerializedObject)
+					{
+						PropertyBag = FPropertyBagRepository::Get().CreateOuterBag(Object);
+					}
+				}
+				return PropertyBag;
+			};
 
-		// This code assumes that properties are loaded in the same order they are saved in. This removes a n^2 search 
-		// and makes it an O(n) when properties are saved in the same order as they are loaded (default case). In the 
-		// case that a property was reordered the code falls back to a slower search.
+			// Load tagged properties.
+			FStructuredArchive::FStream PropertiesStream = Slot.EnterStream();
+
+			// This code assumes that properties are loaded in the same order they are saved in. This removes a n^2 search 
+			// and makes it an O(n) when properties are saved in the same order as they are loaded (default case). In the 
+			// case that a property was reordered the code falls back to a slower search.
 			FProperty*	Property = PropertyLink;
-		bool		bAdvanceProperty	= false;
-		int32		RemainingArrayDim	= Property ? Property->ArrayDim : 0;
+			bool		bAdvanceProperty	= false;
+			int32		RemainingArrayDim	= Property ? Property->ArrayDim : 0;
 
-		// Load all stored properties, potentially skipping unknown ones.
-		while (true)
-		{
+			// Load all stored properties, potentially skipping unknown ones.
+			while (true)
+			{
 				FStructuredArchive::FRecord PropertyRecord = PropertiesStream.EnterElement().EnterRecord();
 
 				FPropertyTag Tag;
@@ -1422,6 +1444,32 @@ void UStruct::SerializeVersionedTaggedProperties(FStructuredArchive::FSlot Slot,
 					Property = CustomFindProperty(Tag.Name);
 				}
 
+				if (LoadContext)
+				{
+					const FName Name = Property ? Property->GetFName() : Tag.Name;
+					const int32 Index = Tag.ArrayIndex > 0 || (Property && Property->ArrayDim > 1) ? Tag.ArrayIndex : INDEX_NONE;
+
+					FName Type;
+					if (!Tag.InnerType.IsNone() && Tag.ValueType.IsNone())
+					{
+						Type = Tag.InnerType;
+					}
+					else if (!Tag.StructName.IsNone())
+					{
+						Type = Tag.StructName;
+					}
+					else if (!Tag.EnumName.IsNone())
+					{
+						Type = Tag.EnumName;
+					}
+					else
+					{
+						Type = Tag.Type;
+					}
+
+					LoadContext->SerializedPropertyPath.Push({Name, Type, Index});
+				}
+
 				if (Property)
 				{
 					Tag.Prop = Property;
@@ -1480,12 +1528,24 @@ void UStruct::SerializeVersionedTaggedProperties(FStructuredArchive::FSlot Slot,
 						{
 							case EConvertFromTypeResult::Converted:
 								bAdvanceProperty = true;
+								if (FPropertyBag* PropertyBag = TryFindPropertyBag())
+								{
+									UnderlyingArchive.Seek(StartOfProperty);
+									FStructuredArchive::FSlot CopySlot = PropertyRecord.EnterField(TEXT("Value"));
+									Tag.Prop = nullptr;
+									PropertyBag->LoadPropertyByTag(LoadContext->SerializedPropertyPath, Tag, CopySlot, Property->ContainerPtrToValuePtrForDefaults<uint8>(DefaultsStruct, Defaults, Tag.ArrayIndex));
+								}
 								break;
 
 							case EConvertFromTypeResult::UseSerializeItem:
 								if (Tag.Type != PropID)
 								{
 									UE_LOG(LogClass, Warning, TEXT("Type mismatch in %s of %s - Previous (%s) Current(%s) for package:  %s"), *Tag.Name.ToString(), *GetName(), *Tag.Type.ToString(), *PropID.ToString(), *UnderlyingArchive.GetArchiveName());
+									if (FPropertyBag* PropertyBag = TryFindPropertyBag())
+									{
+										Tag.Prop = nullptr;
+										PropertyBag->LoadPropertyByTag(LoadContext->SerializedPropertyPath, Tag, ValueSlot, Property->ContainerPtrToValuePtrForDefaults<uint8>(DefaultsStruct, Defaults, Tag.ArrayIndex));
+									}
 								}
 								else
 								{
@@ -1499,6 +1559,11 @@ void UStruct::SerializeVersionedTaggedProperties(FStructuredArchive::FSlot Slot,
 								break;
 
 							case EConvertFromTypeResult::CannotConvert:
+								if (FPropertyBag* PropertyBag = TryFindPropertyBag())
+								{
+									Tag.Prop = nullptr;
+									PropertyBag->LoadPropertyByTag(LoadContext->SerializedPropertyPath, Tag, ValueSlot, Property->ContainerPtrToValuePtrForDefaults<uint8>(DefaultsStruct, Defaults, Tag.ArrayIndex));
+								}
 								break;
 
 							default:
@@ -1506,6 +1571,17 @@ void UStruct::SerializeVersionedTaggedProperties(FStructuredArchive::FSlot Slot,
 						}
 					}
 				}
+				else if (FPropertyBag* PropertyBag = TryFindPropertyBag())
+				{
+					// TODO: Might we find defaults in a property bag for Defaults?
+					FStructuredArchive::FSlot ValueSlot = PropertyRecord.EnterField(TEXT("Value"));
+					PropertyBag->LoadPropertyByTag(LoadContext->SerializedPropertyPath, Tag, ValueSlot);
+				}
+
+				if (LoadContext)
+				{
+					LoadContext->SerializedPropertyPath.Pop();
+				}
 
 				int64 Loaded = UnderlyingArchive.Tell() - StartOfProperty;
 
diff --git a/Engine/Source/Runtime/CoreUObject/Private/UObject/LinkerLoad.cpp b/Engine/Source/Runtime/CoreUObject/Private/UObject/LinkerLoad.cpp
--- a/Engine/Source/Runtime/CoreUObject/Private/UObject/LinkerLoad.cpp
+++ b/Engine/Source/Runtime/CoreUObject/Private/UObject/LinkerLoad.cpp
@@ -4553,9 +4553,11 @@ void FLinkerLoad::Preload( UObject* Object )
 #endif // USE_CIRCULAR_DEPENDENCY_LOAD_DEFERRING
 						
 						check(CurrentLoadContext);
-						// Maintain the current SerializedObjects.
+						// Maintain the current SerializedObject and SerializedPropertyPath.
 						UObject* PrevSerializedObject = CurrentLoadContext->SerializedObject;
+						UE::FPropertyPathName PrevSerializedPropertyPath = MoveTemp(CurrentLoadContext->SerializedPropertyPath);
 						CurrentLoadContext->SerializedObject = Object;
+						CurrentLoadContext->SerializedPropertyPath.Reset();
 
 #if WITH_EDITOR && WITH_TEXT_ARCHIVE_SUPPORT
 						if (IsTextFormat())
@@ -4580,6 +4582,7 @@ void FLinkerLoad::Preload( UObject* Object )
 						}
 
 						Object->SetFlags(RF_LoadCompleted);
+						CurrentLoadContext->SerializedPropertyPath = MoveTemp(PrevSerializedPropertyPath);
 						CurrentLoadContext->SerializedObject = PrevSerializedObject;
 					}
 					else
@@ -4597,9 +4600,11 @@ void FLinkerLoad::Preload( UObject* Object )
 						TRACE_CPUPROFILER_EVENT_SCOPE_TEXT_ON_CHANNEL(*Object->GetFullName(), AssetLoadTimeChannel);
 #endif
 
-						// Maintain the current SerializedObjects.
+						// Maintain the current SerializedObject and SerializedPropertyPath.
 						UObject* PrevSerializedObject = CurrentLoadContext->SerializedObject;
+						UE::FPropertyPathName PrevSerializedPropertyPath = MoveTemp(CurrentLoadContext->SerializedPropertyPath);
 						CurrentLoadContext->SerializedObject = Object;
+						CurrentLoadContext->SerializedPropertyPath.Reset();
 
 #if WITH_EDITOR && WITH_TEXT_ARCHIVE_SUPPORT
 						if (IsTextFormat())
@@ -4625,6 +4630,7 @@ void FLinkerLoad::Preload( UObject* Object )
 						}
 
 						Object->SetFlags(RF_LoadCompleted);
+						CurrentLoadContext->SerializedPropertyPath = MoveTemp(PrevSerializedPropertyPath);
 						CurrentLoadContext->SerializedObject = PrevSerializedObject;
 					}
 				}
diff --git a/Engine/Source/Runtime/CoreUObject/Private/UObject/PropertyArray.cpp b/Engine/Source/Runtime/CoreUObject/Private/UObject/PropertyArray.cpp
--- a/Engine/Source/Runtime/CoreUObject/Private/UObject/PropertyArray.cpp
+++ b/Engine/Source/Runtime/CoreUObject/Private/UObject/PropertyArray.cpp
@@ -4,11 +4,13 @@
 #include "Misc/ScopeExit.h"
 #include "UObject/ObjectMacros.h"
 #include "Templates/Casts.h"
+#include "UObject/PropertyPathName.h"
 #include "UObject/PropertyTag.h"
 #include "UObject/UnrealType.h"
 #include "UObject/UnrealTypePrivate.h"
 #include "UObject/LinkerLoad.h"
 #include "UObject/PropertyHelper.h"
+#include "UObject/UObjectThreadContext.h"
 
 /*-----------------------------------------------------------------------------
 	FArrayProperty.
@@ -304,6 +306,8 @@ void FArrayProperty::SerializeItem(FStructuredArchive::FSlot Slot, void* Value,
 		}
 	}
 
+	FUObjectSerializeContext* Context = UnderlyingArchive.GetSerializeContext();
+
 	// need to know how much data this call to SerializeItem consumes, so mark where we are
 	int64 DataOffset = UnderlyingArchive.Tell();
 
@@ -341,6 +345,10 @@ void FArrayProperty::SerializeItem(FStructuredArchive::FSlot Slot, void* Value,
 
 				// Serialize the item at this array index
 				i = PropertyNode->ArrayIndex;
+				if (Context)
+				{
+					Context->SerializedPropertyPath.SetIndex(i);
+				}
 				SerializeContainerItem(Array.EnterElement(), ArrayHelper.GetRawPtr(i));
 				PropertyNode = PropertyNode->PropertyListNext;
 
@@ -365,6 +373,10 @@ void FArrayProperty::SerializeItem(FStructuredArchive::FSlot Slot, void* Value,
 			NAME_UArraySerializeCount.SetNumber(i);
 			FArchive::FScopeAddDebugData P(UnderlyingArchive, NAME_UArraySerializeCount);
 #endif
+			if (Context)
+			{
+				Context->SerializedPropertyPath.SetIndex(i);
+			}
 			SerializeContainerItem(Array.EnterElement(), ArrayHelper.GetRawPtr(i++));
 		}
 
@@ -392,6 +404,11 @@ void FArrayProperty::SerializeItem(FStructuredArchive::FSlot Slot, void* Value,
 			UnderlyingArchive.Seek(DataOffset);
 		}
 	}
+
+	if (Context)
+	{
+		Context->SerializedPropertyPath.SetIndex(INDEX_NONE);
+	}
 }
 
 bool FArrayProperty::NetSerializeItem( FArchive& Ar, UPackageMap* Map, void* Data, TArray<uint8> * MetaData ) const
diff --git a/Engine/Source/Runtime/CoreUObject/Private/UObject/PropertyBag.cpp b/Engine/Source/Runtime/CoreUObject/Private/UObject/PropertyBag.cpp
--- a/Engine/Source/Runtime/CoreUObject/Private/UObject/PropertyBag.cpp
+++ b/Engine/Source/Runtime/CoreUObject/Private/UObject/PropertyBag.cpp
@@ -19,32 +19,12 @@ void FPropertyBag::Empty()
 	Properties.Empty();
 }
 
-void FPropertyBag::Copy(const UStruct* Struct, void* Data, EPropertyFlags SkipFlags)
-{
-	EFieldIterationFlags IterationFlags = EFieldIterationFlags::IncludeAll;
-	if (SkipFlags & CPF_Deprecated)
-	{
-		IterationFlags &= ~EFieldIterationFlags::IncludeDeprecated;
-	}
-
-	for (TFieldIterator<FProperty> It(Struct, IterationFlags); It; ++It)
-	{
-		if (FProperty* Property = *It; !(Property->GetPropertyFlags() & SkipFlags))
-		{
-			for (int32 ArrayIndex = 0, ArrayDim = Property->ArrayDim; ArrayIndex < ArrayDim; ++ArrayIndex)
-			{
-				Add(Property, Data, ArrayIndex);
-			}
-		}
-	}
-}
-
-void FPropertyBag::Add(FProperty* Property, void* Data, int32 ArrayIndex)
+void FPropertyBag::Add(const FPropertyPathName& Path, FProperty* Property, void* Data, int32 ArrayIndex)
 {
 	// TODO: NullAr is a workaround to FPropertyTag requiring an archive to assert on versioned property serialization.
 	FNullArchive NullAr;
 
-	FPropertyBag::FValue& Value = Properties.FindOrAdd(Property->GetFName());
+	FValue& Value = FindOrCreateValue(Path);
 
 	const bool bPropertyChanged = Value.Tag.Prop != Property;
 
@@ -54,26 +34,53 @@ void FPropertyBag::Add(FProperty* Property, void* Data, int32 ArrayIndex)
 		Value.Tag = FPropertyTag(NullAr, Property, INDEX_NONE, (uint8*)Data, nullptr);
 	}
 
-	Value.EnsurePropertyData(ArrayIndex);
+	Value.AllocateAndInitializeValue();
 
 	// Subtract the offset because the property will re-add it.
-	void* TargetBase = (uint8*)Value.Data.Memory - Property->GetOffset_ForInternal();
+	void* TargetBase = (uint8*)Value.Data - Property->GetOffset_ForInternal();
 
-	void* Target = Property->ContainerPtrToValuePtr<void>(TargetBase, ArrayIndex);
+	void* Target = Property->ContainerPtrToValuePtr<void>(TargetBase);
 	void* Source = Property->ContainerPtrToValuePtr<void>(Data, ArrayIndex);
 	Property->CopySingleValue(Target, Source);
 }
 
-void FPropertyBag::Remove(FName Name)
+void FPropertyBag::Remove(const FPropertyPathName& Path)
 {
-	Properties.Remove(Name);
+	TArray<FNodeMap*, TInlineAllocator<8>> Nodes;
+	Nodes.Add(&Properties);
+
+	// Find the chain of nodes to the leaf property.
+	for (int32 SegmentIndex = 0, SegmentCount = Path.GetSegmentCount() - 1; SegmentIndex < SegmentCount; ++SegmentIndex)
+	{
+		const FPropertyPathNameSegment Segment = Path.GetSegment(SegmentIndex);
+		if (TUniquePtr<FNode>* Node = Nodes[SegmentIndex]->Find(Segment.PackNameWithIndex()))
+		{
+			if (FNodeMap* NodeMap = (*Node)->ValueOrNodes.TryGet<FNodeMap>())
+			{
+				Nodes.Add(NodeMap);
+				continue;
+			}
+		}
+		return;
+	}
+
+	// Remove the leaf property and any empty nodes in the tail of the property chain.
+	for (int32 SegmentIndex = Path.GetSegmentCount() - 1; SegmentIndex >= 0; --SegmentIndex)
+	{
+		const FPropertyPathNameSegment Segment = Path.GetSegment(SegmentIndex);
+		FNodeMap* SegmentNodes = Nodes[SegmentIndex];
+		if (SegmentNodes->Remove(Segment.PackNameWithIndex()) == 0 || !SegmentNodes->IsEmpty())
+		{
+			return;
+		}
+	}
 }
 
-void FPropertyBag::LoadPropertyByTag(const FPropertyTag& Tag, FStructuredArchiveSlot& ValueSlot, const void* Defaults)
+void FPropertyBag::LoadPropertyByTag(const FPropertyPathName& Path, const FPropertyTag& Tag, FStructuredArchiveSlot& ValueSlot, const void* Defaults)
 {
 	FArchive& UnderlyingArchive = ValueSlot.GetUnderlyingArchive();
 
-	FPropertyBag::FValue& Value = Properties.FindOrAdd(Tag.Name);
+	FValue& Value = FindOrCreateValue(Path);
 
 	const bool bPropertyChanged =
 		(Value.Tag.Prop != Tag.Prop && Value.Tag.Prop && Tag.Prop) ||
@@ -90,89 +97,54 @@ void FPropertyBag::LoadPropertyByTag(const FPropertyTag& Tag, FStructuredArchive
 	{
 		Value.Destroy();
 		Value.Tag = Tag;
-		Value.Tag.Size = 0;
 		Value.Tag.ArrayIndex = INDEX_NONE;
 	}
 
-	if (FProperty* Property = Value.Tag.Prop)
-	{
-		Value.EnsurePropertyData(Tag.ArrayIndex);
-
-		// Subtract the offset because the property will re-add it.
-		void* TargetBase = (uint8*)Value.Data.Memory - Property->GetOffset_ForInternal();
-
-		void* Target = Property->ContainerPtrToValuePtr<void>(TargetBase, Tag.ArrayIndex);
-		Tag.SerializeTaggedProperty(ValueSlot, Property, (uint8*)Target, (const uint8*)Defaults);
-
-		return;
-	}
-
-	const FValue::FData Data = {FMemory::Malloc(Tag.Size), Tag.Size};
-	UnderlyingArchive.Serialize(Data.Memory, Data.Size);
+	Value.AllocateAndInitializeValue();
 
-	if (LIKELY(Tag.ArrayIndex == 0))
+	if (FProperty* Property = Value.Tag.Prop)
 	{
-		Value.Data = Data;
+		Tag.SerializeTaggedProperty(ValueSlot, Property, (uint8*)Value.Data, (const uint8*)Defaults);
 	}
 	else
 	{
-		Value.EnsureArrayData(Tag.ArrayIndex);
-
-		FBitReference AllocationFlag = (*Value.ArrayAllocationFlags)[Tag.ArrayIndex];
-		if (UNLIKELY(AllocationFlag))
-		{
-			FMemory::Free((*Value.ArrayData)[Tag.ArrayIndex].Memory);
-		}
-		else
-		{
-			AllocationFlag = true;
-		}
-		Value.ArrayData->Add(Tag.ArrayIndex, Data);
+		UnderlyingArchive.Serialize(Value.Data, Value.Tag.Size);
 	}
 }
 
-void FPropertyBag::AssignElement(FPropertyBagElement& Element, const FName& Name, const FValue& Value)
+FPropertyBag::FValue& FPropertyBag::FindOrCreateValue(const FPropertyPathName& Path)
 {
-	FProperty* Property = Value.Tag.Prop;
-	Element.Name = Name;
-	Element.Property = Property;
+	const int32 SegmentCount = Path.GetSegmentCount();
 
-	if (Property)
-	{
-		Element.ArrayIndex = 0;
-		Element.ArrayDim = Property->ArrayDim;
-		Element.Value = (uint8*)Value.Data.Memory;
-	}
-	else if (Value.ArrayData)
-	{
-		Element.ArrayIndex = Value.ArrayAllocationFlags->Find(/*bValue*/ true);
-		Element.ArrayDim = Value.ArrayAllocationFlags->Num();
-		Element.Value = (*Value.ArrayData)[Element.ArrayIndex].Memory;
-	}
-	else
+	FNodeMap* NodeMap = &Properties;
+	for (int32 SegmentIndex = 0; SegmentIndex < SegmentCount - 1; ++SegmentIndex)
 	{
-		Element.ArrayIndex = 0;
-		Element.ArrayDim = 1;
-		Element.Value = Value.Data.Memory;
+		FPropertyPathNameSegment Segment = Path.GetSegment(SegmentIndex);
+		TUniquePtr<FNode>& Node = NodeMap->FindOrAdd(Segment.PackNameWithIndex());
+		if (!Node)
+		{
+			Node = MakeUnique<FNode>();
+		}
+		if (!Node->ValueOrNodes.IsType<FNodeMap>())
+		{
+			Node->ValueOrNodes.Emplace<FNodeMap>();
+		}
+		Node->Type = Segment.Type;
+		NodeMap = &Node->ValueOrNodes.Get<FNodeMap>();
 	}
-}
 
-bool FPropertyBag::AdvanceElement(FPropertyBagElement& Element, const FValue& Value)
-{
-	if (const FProperty* Property = Element.Property)
+	FPropertyPathNameSegment LastSegment = Path.GetSegment(SegmentCount - 1);
+	TUniquePtr<FNode>& Node = NodeMap->FindOrAdd(LastSegment.PackNameWithIndex());
+	if (!Node)
 	{
-		Element.Value = (uint8*)Element.Value + Property->ElementSize;
-		return true;
+		Node = MakeUnique<FNode>();
 	}
-
-	Element.ArrayIndex = Value.ArrayAllocationFlags->FindFrom(/*bValue*/ true, Element.ArrayIndex);
-	if (Element.ArrayIndex >= 0)
+	if (!Node->ValueOrNodes.IsType<FValue>())
 	{
-		Element.Value = (*Value.ArrayData)[Element.ArrayIndex].Memory;
-		return true;
+		Node->ValueOrNodes.Emplace<FValue>();
 	}
-
-	return false;
+	Node->Type = LastSegment.Type;
+	return Node->ValueOrNodes.Get<FValue>();
 }
 
 ///////////////////////////////////////////////////////////////////////////////////////////////////////////////////////
@@ -182,79 +154,109 @@ FPropertyBag::FValue::~FValue()
 	Destroy();
 }
 
-void FPropertyBag::FValue::Destroy()
+void FPropertyBag::FValue::AllocateAndInitializeValue()
 {
-	const auto DestroyValue = [Property = Tag.Prop](const FData& D)
+	if (Data)
 	{
-		if (Property)
-		{
-			Property->DestroyAndFreeValue(D.Memory);
-		}
-		else
-		{
-			FMemory::Free(D.Memory);
-		}
-	};
+		return;
+	}
 
-	if (Data.Memory)
+	if (const FProperty* Property = Tag.Prop)
 	{
-		DestroyValue(Data);
-		Data = {};
+		// TODO: Need to allocate only one element for arrays.
+		Data = Property->AllocateAndInitializeValue();
 	}
-
-	if (ArrayData)
+	else
 	{
-		for (TTuple<int32, FData>& IndexAndData : *ArrayData)
-		{
-			DestroyValue(IndexAndData.Value);
-		}
-		delete ArrayData;
-		ArrayData = nullptr;
+		Data = FMemory::Malloc(Tag.Size);
 	}
-
-	delete ArrayAllocationFlags;
-	ArrayAllocationFlags = nullptr;
 }
 
-void FPropertyBag::FValue::EnsurePropertyData(int32 ArrayIndex)
+void FPropertyBag::FValue::Destroy()
 {
-	const FProperty* const Property = Tag.Prop;
-	const int32 ArrayDim = Property->ArrayDim;
-
-	if (!Data.Memory)
+	if (Data)
 	{
-		Data.Memory = Property->AllocateAndInitializeValue();
-		Data.Size = Property->GetSize();
-
-		if (ArrayDim > 1)
+		if (const FProperty* Property = Tag.Prop)
 		{
-			ArrayAllocationFlags = new TBitArray<>(/*bValue*/ false, ArrayDim);
+			// TODO: Need to destroy and free only one element for arrays.
+			Property->DestroyAndFreeValue(Data);
 		}
+		else
+		{
+			FMemory::Free(Data);
+		}
+		Data = nullptr;
 	}
+}
 
-	if (ArrayDim > 1)
+int32 FPropertyBag::FValue::GetSize() const
+{
+	if (!Data)
 	{
-		(*ArrayAllocationFlags)[ArrayIndex] = true;
+		return 0;
 	}
+	if (!Tag.Prop)
+	{
+		return Tag.Size;
+	}
+	// TODO: Need to include the size of only one element for arrays.
+	return Tag.Prop->GetSize();
 }
 
-void FPropertyBag::FValue::EnsureArrayData(int32 ArrayIndex)
+///////////////////////////////////////////////////////////////////////////////////////////////////////////////////////
+
+FPropertyBag::FConstIterator::FConstIterator(const FNodeIterator& NodeIt)
 {
-	if (ArrayAllocationFlags)
-	{
-		ArrayAllocationFlags->PadToNum(ArrayIndex + 1, /*bValue*/ false);
-		return;
-	}
+	NodeIterators.Push(NodeIt);
+	EnterNode();
+}
 
-	ArrayData = new TMap<int32, FValue::FData>();
-	ArrayAllocationFlags = new TBitArray<>(/*bValue*/ false, ArrayIndex + 1);
+FPropertyBag::FConstIterator& FPropertyBag::FConstIterator::operator++()
+{
+	++NodeIterators.Last();
+	EnterNode();
+	return *this;
+}
 
-	// Move the scalar value into array storage as the first element.
-	if (Data.Memory)
+inline void FPropertyBag::FConstIterator::EnterNode()
+{
+	for (;;)
 	{
-		(*ArrayAllocationFlags)[0] = true;
-		ArrayData->Add(0, Data);
-		Data = {};
+		if (CurrentPath.GetSegmentCount() == NodeIterators.Num())
+		{
+			CurrentPath.Pop();
+		}
+
+		if (FNodeIterator& NodeIt = NodeIterators.Last())
+		{
+			const TUniquePtr<FNode>& Node = NodeIt.Value();
+			CurrentPath.Push(FPropertyPathNameSegment().SetNameWithIndex(NodeIt.Key()).SetType(Node->Type));
+
+			if (FValue* Value = Node->ValueOrNodes.TryGet<FValue>())
+			{
+				CurrentValue = Value;
+				return;
+			}
+			else if (const FNodeMap* NodeMap = Node->ValueOrNodes.TryGet<FNodeMap>())
+			{
+				NodeIterators.Push(NodeMap->CreateConstIterator());
+			}
+			else
+			{
+				checkNoEntry();
+			}
+		}
+		else if (NodeIterators.Num() > 1)
+		{
+			NodeIterators.Pop();
+			++NodeIterators.Last();
+		}
+		else
+		{
+			// End of the root node.
+			CurrentValue = nullptr;
+			return;
+		}
 	}
 }
 
diff --git a/Engine/Source/Runtime/CoreUObject/Private/UObject/PropertyMap.cpp b/Engine/Source/Runtime/CoreUObject/Private/UObject/PropertyMap.cpp
--- a/Engine/Source/Runtime/CoreUObject/Private/UObject/PropertyMap.cpp
+++ b/Engine/Source/Runtime/CoreUObject/Private/UObject/PropertyMap.cpp
@@ -3,11 +3,13 @@
 #include "CoreMinimal.h"
 #include "UObject/ObjectMacros.h"
 #include "Templates/Casts.h"
+#include "UObject/PropertyPathName.h"
 #include "UObject/PropertyTag.h"
 #include "UObject/UnrealType.h"
 #include "UObject/UnrealTypePrivate.h"
 #include "UObject/LinkerLoad.h"
 #include "UObject/PropertyHelper.h"
+#include "UObject/UObjectThreadContext.h"
 #include "Misc/ScopeExit.h"
 #include "Serialization/ArchiveUObjectFromStructuredArchive.h"
 #include "UObject/UObjectThreadContext.h"
@@ -332,6 +334,8 @@ void FMapProperty::SerializeItem(FStructuredArchive::FSlot Slot, void* Value, co
 
 	if (UnderlyingArchive.IsLoading())
 	{
+		FUObjectSerializeContext* Context = UnderlyingArchive.GetSerializeContext();
+
 		if (Defaults)
 		{
 			CopyValuesInternal(Value, Defaults, 1);
@@ -345,6 +349,12 @@ void FMapProperty::SerializeItem(FStructuredArchive::FSlot Slot, void* Value, co
 		{
 			if (NumKeysToRemove)
 			{
+				TOptional<TGuardValue<bool>> SerializeUnknownProperty;
+				if (Context)
+				{
+					SerializeUnknownProperty.Emplace(Context->bSerializeUnknownProperty, false);
+				}
+
 				// Load and discard keys to remove, map is empty
 				void* TempKeyValueStorage = FMemory::Malloc(MapLayout.SetLayout.Size);
 				KeyProp->InitializeValue(TempKeyValueStorage);
@@ -359,6 +369,13 @@ void FMapProperty::SerializeItem(FStructuredArchive::FSlot Slot, void* Value, co
 				FMemory::Free(TempKeyValueStorage);
 			}
 
+			// Disable serialization of unknown properties until the TODO in the loop is addressed.
+			TOptional<TGuardValue<bool>> SerializeUnknownProperty;
+			if (Context)
+			{
+				SerializeUnknownProperty.Emplace(Context->bSerializeUnknownProperty, false);
+			}
+
 			int32 NumEntries = 0;
 			FStructuredArchive::FArray EntriesArray = Record.EnterArray(TEXT("Entries"), NumEntries);
 
@@ -368,7 +385,12 @@ void FMapProperty::SerializeItem(FStructuredArchive::FSlot Slot, void* Value, co
 			{
 				FStructuredArchive::FRecord EntryRecord = EntriesArray.EnterElement().EnterRecord();
 				int32 Index = MapHelper.AddDefaultValue_Invalid_NeedsRehash();
+				if (Context)
+				{
+					Context->SerializedPropertyPath.SetIndex(Index);
+				}
 
+				// TODO: Need a way to indicate that we are serializing a key and value into the property bag. Push Key/Value names?
 				{
 					FSerializedPropertyScope SerializedProperty(UnderlyingArchive, KeyProp, this);
 					KeyProp->SerializeItem(EntryRecord.EnterField(TEXT("Key")), MapHelper.GetKeyPtr(Index));
@@ -395,6 +417,12 @@ void FMapProperty::SerializeItem(FStructuredArchive::FSlot Slot, void* Value, co
 
 			if (NumKeysToRemove)
 			{
+				TOptional<TGuardValue<bool>> SerializeUnknownProperty;
+				if (Context)
+				{
+					SerializeUnknownProperty.Emplace(Context->bSerializeUnknownProperty, false);
+				}
+
 				TempKeyValueStorage = (uint8*)FMemory::Malloc(MapLayout.SetLayout.Size);
 				KeyProp->InitializeValue(TempKeyValueStorage);
 
@@ -412,6 +440,13 @@ void FMapProperty::SerializeItem(FStructuredArchive::FSlot Slot, void* Value, co
 				}
 			}
 
+			// Disable serialization of unknown properties until the TODO in the loop is addressed.
+			TOptional<TGuardValue<bool>> SerializeUnknownProperty;
+			if (Context)
+			{
+				SerializeUnknownProperty.Emplace(Context->bSerializeUnknownProperty, false);
+			}
+
 			int32 NumEntries = 0;
 			FStructuredArchive::FArray EntriesArray = Record.EnterArray(TEXT("Entries"), NumEntries);
 
@@ -425,8 +460,12 @@ void FMapProperty::SerializeItem(FStructuredArchive::FSlot Slot, void* Value, co
 			// Read remaining items into container
 			for (; NumEntries; --NumEntries)
 			{
+				// TODO: SetIndex on Context->SerializedPropertyPath and remove the element from the bag later if it existed.
+
 				FStructuredArchive::FRecord EntryRecord = EntriesArray.EnterElement().EnterRecord();
 
+				// TODO: Need a way to indicate that we are serializing a key and value into the property bag. Push Key/Value names?
+
 				// Read key into temporary storage
 				{
 					FSerializedPropertyScope SerializedProperty(UnderlyingArchive, KeyProp, this);
diff --git a/Engine/Source/Runtime/CoreUObject/Private/UObject/PropertyPathName.cpp b/Engine/Source/Runtime/CoreUObject/Private/UObject/PropertyPathName.cpp
new file mode 100644
--- /dev/null
+++ b/Engine/Source/Runtime/CoreUObject/Private/UObject/PropertyPathName.cpp
@@ -0,0 +1,106 @@
+// Copyright Epic Games, Inc. All Rights Reserved.
+
+#include "UObject/PropertyPathName.h"
+
+#include "Misc/StringBuilder.h"
+#include "Templates/TypeHash.h"
+
+namespace UE
+{
+
+inline bool FPropertyPathName::FSegment::operator==(const FSegment& Segment) const
+{
+	return NameWithIndex == Segment.NameWithIndex && Type == Segment.Type;
+}
+
+inline bool FPropertyPathName::FSegment::operator<(const FSegment& Segment) const
+{
+	return Compare(Segment) < 0;
+}
+
+inline int32 FPropertyPathName::FSegment::Compare(const FSegment& Segment) const
+{
+	if (int32 CompareNameWithIndex = NameWithIndex.Compare(Segment.NameWithIndex))
+	{
+		return CompareNameWithIndex;
+	}
+	return Type.CompareLexical(Segment.Type);
+}
+
+bool FPropertyPathName::operator==(const FPropertyPathName& Path) const
+{
+	const int32 SegmentCount = Segments.Num();
+	if (SegmentCount != Path.Segments.Num())
+	{
+		return false;
+	}
+
+	for (int32 SegmentIndex = SegmentCount - 1; SegmentIndex >= 0; --SegmentIndex)
+	{
+		if (!(Segments[SegmentIndex] == Path.Segments[SegmentIndex]))
+		{
+			return false;
+		}
+	}
+
+	return true;
+}
+
+bool FPropertyPathName::operator<(const FPropertyPathName& Path) const
+{
+	const int32 SegmentCountA = Segments.Num();
+	const int32 SegmentCountB = Path.Segments.Num();
+	const int32 SegmentCountMin = FPlatformMath::Min(SegmentCountA, SegmentCountB);
+
+	for (int32 SegmentIndex = 0; SegmentIndex < SegmentCountMin; ++SegmentIndex)
+	{
+		if (const int32 Compare = Segments[SegmentIndex].Compare(Path.Segments[SegmentIndex]))
+		{
+			return Compare < 0;
+		}
+	}
+
+	return SegmentCountA < SegmentCountB;
+}
+
+void FPropertyPathName::ToString(FStringBuilderBase& Out, FStringView Separator) const
+{
+	bool bFirst = true;
+	for (const FSegment& Segment : Segments)
+	{
+		if (bFirst)
+		{
+			bFirst = false;
+		}
+		else
+		{
+			Out.Append(Separator);
+		}
+
+		FPropertyPathNameSegment UnpackedSegment = Segment.Unpack();
+		Out << UnpackedSegment.Name;
+
+		if (const int32 Index = UnpackedSegment.Index; Index != INDEX_NONE)
+		{
+			Out << TEXT('[') << Index << TEXT(']');
+		}
+
+		if (!UnpackedSegment.Type.IsNone())
+		{
+			Out << TEXTVIEW(" (") << UnpackedSegment.Type << TEXT(')');
+		}
+	}
+}
+
+uint32 GetTypeHash(const FPropertyPathName& Path)
+{
+	uint32 Hash = 0;
+	for (const FPropertyPathName::FSegment& Segment : Path.Segments)
+	{
+		Hash = HashCombineFast(Hash, GetTypeHash(Segment.NameWithIndex));
+		Hash = HashCombineFast(Hash, GetTypeHash(Segment.Type));
+	}
+	return Hash;
+}
+
+} // UE
diff --git a/Engine/Source/Runtime/CoreUObject/Private/UObject/PropertySet.cpp b/Engine/Source/Runtime/CoreUObject/Private/UObject/PropertySet.cpp
--- a/Engine/Source/Runtime/CoreUObject/Private/UObject/PropertySet.cpp
+++ b/Engine/Source/Runtime/CoreUObject/Private/UObject/PropertySet.cpp
@@ -2,11 +2,13 @@
 
 #include "CoreMinimal.h"
 #include "UObject/ObjectMacros.h"
+#include "UObject/PropertyPathName.h"
 #include "UObject/PropertyTag.h"
 #include "UObject/UnrealType.h"
 #include "UObject/UnrealTypePrivate.h"
 #include "UObject/LinkerLoad.h"
 #include "UObject/PropertyHelper.h"
+#include "UObject/UObjectThreadContext.h"
 #include "Misc/ScopeExit.h"
 #include "Serialization/ArchiveUObjectFromStructuredArchive.h"
 
@@ -287,6 +289,8 @@ void FSetProperty::SerializeItem(FStructuredArchive::FSlot Slot, void* Value, co
 
 	if (UnderlyingArchive.IsLoading())
 	{
+		FUObjectSerializeContext* Context = UnderlyingArchive.GetSerializeContext();
+
 		if (Defaults)
 		{
 			CopyValuesInternal(Value, Defaults, 1);
@@ -300,6 +304,12 @@ void FSetProperty::SerializeItem(FStructuredArchive::FSlot Slot, void* Value, co
 		{
 			if (NumElementsToRemove)
 			{
+				TOptional<TGuardValue<bool>> SerializeUnknownProperty;
+				if (Context)
+				{
+					SerializeUnknownProperty.Emplace(Context->bSerializeUnknownProperty, false);
+				}
+
 				// Load and discard elements to remove, set is empty
 				void* TempElementStorage = FMemory::Malloc(SetLayout.Size);
 				ElementProp->InitializeValue(TempElementStorage);
@@ -322,6 +332,10 @@ void FSetProperty::SerializeItem(FStructuredArchive::FSlot Slot, void* Value, co
 			for (; Num; --Num)
 			{
 				int32 Index = SetHelper.AddDefaultValue_Invalid_NeedsRehash();
+				if (Context)
+				{
+					Context->SerializedPropertyPath.SetIndex(Index);
+				}
 				ElementProp->SerializeItem(ElementsArray.EnterElement(), SetHelper.GetElementPtrWithoutCheck(Index));
 			}
 		}
@@ -339,6 +353,12 @@ void FSetProperty::SerializeItem(FStructuredArchive::FSlot Slot, void* Value, co
 
 			if (NumElementsToRemove)
 			{
+				TOptional<TGuardValue<bool>> SerializeUnknownProperty;
+				if (Context)
+				{
+					SerializeUnknownProperty.Emplace(Context->bSerializeUnknownProperty, false);
+				}
+
 				TempElementStorage = (uint8*)FMemory::Malloc(SetLayout.Size);
 				ElementProp->InitializeValue(TempElementStorage);
 
@@ -367,10 +387,19 @@ void FSetProperty::SerializeItem(FStructuredArchive::FSlot Slot, void* Value, co
 				ElementProp->InitializeValue(TempElementStorage);
 			}
 
+			// Disable serialization of unknown properties until the TODO in the loop is addressed.
+			TOptional<TGuardValue<bool>> SerializeUnknownProperty;
+			if (Context)
+			{
+				SerializeUnknownProperty.Emplace(Context->bSerializeUnknownProperty, false);
+			}
+
 			FSerializedPropertyScope SerializedProperty(UnderlyingArchive, ElementProp, this);
 			// Read remaining items into container
 			for (; Num; --Num)
 			{
+				// TODO: SetIndex on Context->SerializedPropertyPath and remove the element from the bag later if it existed.
+
 				// Read key into temporary storage
 				ElementProp->SerializeItem(ElementsArray.EnterElement(), TempElementStorage);
 
@@ -386,6 +415,11 @@ void FSetProperty::SerializeItem(FStructuredArchive::FSlot Slot, void* Value, co
 			}
 		}
 
+		if (Context)
+		{
+			Context->SerializedPropertyPath.SetIndex(INDEX_NONE);
+		}
+
 		SetHelper.Rehash();
 	}
 	else
diff --git a/Engine/Source/Runtime/CoreUObject/Private/UObject/UObjectThreadContext.cpp b/Engine/Source/Runtime/CoreUObject/Private/UObject/UObjectThreadContext.cpp
--- a/Engine/Source/Runtime/CoreUObject/Private/UObject/UObjectThreadContext.cpp
+++ b/Engine/Source/Runtime/CoreUObject/Private/UObject/UObjectThreadContext.cpp
@@ -44,6 +44,7 @@ FUObjectSerializeContext::FUObjectSerializeContext()
 	, SerializedImportLinker(nullptr)
 	, SerializedExportIndex(0)
 	, SerializedExportLinker(nullptr)
+	, bSerializeUnknownProperty(false)
 {}
 
 FUObjectSerializeContext::~FUObjectSerializeContext()
diff --git a/Engine/Source/Runtime/CoreUObject/Public/UObject/PropertyBag.h b/Engine/Source/Runtime/CoreUObject/Public/UObject/PropertyBag.h
--- a/Engine/Source/Runtime/CoreUObject/Public/UObject/PropertyBag.h
+++ b/Engine/Source/Runtime/CoreUObject/Public/UObject/PropertyBag.h
@@ -4,6 +4,8 @@
 
 #include "CoreTypes.h"
 #include "Containers/Map.h"
+#include "Misc/TVariant.h"
+#include "UObject/PropertyPathName.h"
 #include "UObject/PropertyTag.h"
 
 #define UE_API COREUOBJECT_API
@@ -15,31 +17,6 @@ namespace UE { class FPropertyBagElement; }
 namespace UE
 {
 
-/**
- * A property bag element is a single value from the bag.
- *
- * A fixed-size array consists of multiple elements. These are not guaranteed to be contiguous in
- * memory, and iteration of the property bag skips elements that do not have a value.
- */
-class FPropertyBagElement
-{
-public:
-	/** Name, which is always valid even when the property is null. */
-	FName Name;
-
-	/** Property, which can be null if the property type is not available. */
-	const FProperty* Property = nullptr;
-
-	/** Pointer to the start of the value. This is NOT a container pointer. */
-	void* Value = nullptr;
-
-	/** Array index for fixed-size arrays. Always 0 for scalar values. */
-	int32 ArrayIndex = 0;
-
-	/** Array dimension for fixed-size arrays. Always 1 for scalar values. */
-	int32 ArrayDim = 0;
-};
-
 /**
  * A property bag is an unordered collection of properties and their values.
  */
@@ -49,36 +26,26 @@ class FPropertyBag
 	{
 		~FValue();
 
+		void AllocateAndInitializeValue();
 		void Destroy();
 
-		void EnsurePropertyData(int32 ArrayIndex);
-		void EnsureArrayData(int32 ArrayIndex);
-
-		struct FData
-		{
-			void* Memory = nullptr;
-			int32 Size = 0;
-		};
+		UE_API int32 GetSize() const;
 
 		FPropertyTag Tag;
+		void* Data = nullptr;
+	};
+
+	struct FNode;
+	using FNodeMap = TMap<FName, TUniquePtr<FNode>>;
 
-		// A value is stored in one of several ways based on what we know of the property.
-		//
-		// A scalar is stored in Data.
-		// An array is stored in Data if it is of known dimension.
-		// An array is stored in Data if it is of unknown dimension and only index 0 is known.
-		// An array is stored in ArrayData if it is of unknown dimension.
-		//
-		// Arrays use ArrayAllocationFlags to track which values are known.
-		//
-		// A value is stored in its native form if it has a known property or a property can be created dynamically.
-		// A value is otherwise stored in its serialized form.
-
-		FData Data;
-		TMap<int32, FData>* ArrayData = nullptr;
-		TBitArray<>* ArrayAllocationFlags = nullptr;
+	struct FNode
+	{
+		FName Type;
+		TVariant<TYPE_OF_NULLPTR, FValue, FNodeMap> ValueOrNodes;
 	};
 
+	FValue& FindOrCreateValue(const FPropertyPathName& Path);
+
 public:
 	UE_API FPropertyBag();
 	UE_API ~FPropertyBag();
@@ -97,24 +64,17 @@ public:
 	 */
 	UE_API void Empty();
 
-	/**
-	 * Copy every non-skipped property of the struct into the property bag.
-	 *
-	 * Any property in the bag with the same name will be overwritten.
-	 */
-	UE_API void Copy(const UStruct* Struct, void* Data, EPropertyFlags SkipFlags);
-
 	/**
 	 * Add the property to the bag.
 	 *
-	 * Any property in the bag with the same name will be overwritten.
+	 * Any property in the bag with the same path will be overwritten.
 	 */
-	UE_API void Add(FProperty* Property, void* Data, int32 ArrayIndex = 0);
+	UE_API void Add(const FPropertyPathName& Path, FProperty* Property, void* Data, int32 ArrayIndex = 0);
 
 	/**
-	 * Remove a property from the bag by name.
+	 * Remove a property from the bag by path.
 	 */
-	UE_API void Remove(FName Name);
+	UE_API void Remove(const FPropertyPathName& Path);
 
 	/**
 	 * Load a property from the slot into the bag based on the tag.
@@ -122,74 +82,39 @@ public:
 	 * If the property has a known type then it will be deserialized and accessible.
 	 * If not, the serialized data will be retained to avoid losing the unrecognized property.
 	 */
-	UE_API void LoadPropertyByTag(const FPropertyTag& Tag, FStructuredArchiveSlot& ValueSlot, const void* Defaults = nullptr);
+	UE_API void LoadPropertyByTag(const FPropertyPathName& Path, const FPropertyTag& Tag, FStructuredArchiveSlot& ValueSlot, const void* Defaults = nullptr);
 
-	template <bool bConst, bool bRangedFor>
-	class TBaseIterator
+	class FConstIterator
 	{
 	public:
-		using FPropertyIterator = std::conditional_t<bConst,
-			std::conditional_t<bRangedFor, typename TMap<FName, FValue>::TRangedForConstIterator, typename TMap<FName, FValue>::TConstIterator>,
-			std::conditional_t<bRangedFor, typename TMap<FName, FValue>::TRangedForIterator, typename TMap<FName, FValue>::TIterator>>;
-
-		inline explicit TBaseIterator(const FPropertyIterator& InPropertyIt)
-			: PropertyIt(InPropertyIt)
-		{
-			if (PropertyIt)
-			{
-				AssignElement(Element, PropertyIt.Key(), PropertyIt.Value());
-			}
-		}
-
-		inline TBaseIterator& operator++()
-		{
-			if (UNLIKELY(++Element.ArrayIndex < Element.ArrayDim && AdvanceElement(Element, PropertyIt.Value())))
-			{
-			}
-			else if (++PropertyIt)
-			{
-				AssignElement(Element, PropertyIt.Key(), PropertyIt.Value());
-			}
-			else
-			{
-				Element = {};
-			}
-			return *this;
-		}
-
-		inline explicit operator bool() const { return !!Element.ArrayDim; }
-
-		inline bool operator==(const TBaseIterator& Rhs) const { return Element.Value == Rhs.Element.Value; }
-		inline bool operator!=(const TBaseIterator& Rhs) const { return Element.Value != Rhs.Element.Value; }
-
-		inline const FPropertyBagElement& operator*() const { return Element; }
-		inline const FPropertyBagElement* operator->() const { return &Element; }
+		using FNodeIterator = typename FNodeMap::TConstIterator;
 
-	private:
-		FPropertyIterator PropertyIt;
-		FPropertyBagElement Element;
-	};
+		UE_API explicit FConstIterator(const FNodeIterator& NodeIt);
 
-	using TIterator = TBaseIterator</*bConst*/ false, /*bRangedFor*/ false>;
-	using TConstIterator = TBaseIterator</*bConst*/ true, /*bRangedFor*/ false>;
+		UE_API FConstIterator& operator++();
 
-	inline TIterator CreateIterator() { return TIterator(Properties.CreateIterator()); }
-	inline TConstIterator CreateConstIterator() const { return TConstIterator(Properties.CreateConstIterator()); }
+		inline explicit operator bool() const { return !!NodeIterators.Last(); }
 
-private:
-	using TRangedForIterator = TBaseIterator</*bConst*/ false, /*bRangedFor*/ true>;
-	using TRangedForConstIterator = TBaseIterator</*bConst*/ true, /*bRangedFor*/ true>;
+		inline bool operator==(const FConstIterator& Rhs) const { return NodeIterators.Last() == Rhs.NodeIterators.Last(); }
+		inline bool operator!=(const FConstIterator& Rhs) const { return NodeIterators.Last() != Rhs.NodeIterators.Last(); }
+
+		inline const FPropertyPathName& GetPath() const { return CurrentPath; }
+		inline const FProperty* GetProperty() const { return CurrentValue->Tag.Prop; }
+		inline void* GetValue() const { return CurrentValue->Data; }
+		inline int32 GetValueSize() const { return CurrentValue->GetSize(); }
 
-	friend inline TRangedForIterator begin(FPropertyBag& Bag) { return TRangedForIterator(Bag.Properties.begin()); }
-	friend inline TRangedForIterator end(FPropertyBag& Bag) { return TRangedForIterator(Bag.Properties.end()); }
+	private:
+		void EnterNode();
 
-	friend inline TRangedForConstIterator begin(const FPropertyBag& Bag) { return TRangedForConstIterator(Bag.Properties.begin()); }
-	friend inline TRangedForConstIterator end(const FPropertyBag& Bag) { return TRangedForConstIterator(Bag.Properties.end()); }
+		TArray<FNodeIterator, TInlineAllocator<8>> NodeIterators;
+		FPropertyPathName CurrentPath;
+		FValue* CurrentValue = nullptr;
+	};
 
-	UE_API static void AssignElement(FPropertyBagElement& Element, const FName& Name, const FValue& Value);
-	UE_API static bool AdvanceElement(FPropertyBagElement& Element, const FValue& Value);
+	inline FConstIterator CreateConstIterator() const { return FConstIterator(Properties.CreateConstIterator()); }
 
-	TMap<FName, FValue> Properties;
+private:
+	FNodeMap Properties;
 };
 
 } // UE
diff --git a/Engine/Source/Runtime/CoreUObject/Public/UObject/PropertyPathName.h b/Engine/Source/Runtime/CoreUObject/Public/UObject/PropertyPathName.h
new file mode 100644
--- /dev/null
+++ b/Engine/Source/Runtime/CoreUObject/Public/UObject/PropertyPathName.h
@@ -0,0 +1,134 @@
+// Copyright Epic Games, Inc. All Rights Reserved.
+
+#pragma once
+
+#include "CoreTypes.h"
+#include "Containers/Array.h"
+#include "Containers/StringFwd.h"
+#include "UObject/NameTypes.h"
+
+#define UE_API COREUOBJECT_API
+
+namespace UE
+{
+
+struct FPropertyPathNameSegment
+{
+	FName Name;
+	FName Type;
+	int32 Index = INDEX_NONE;
+
+	[[nodiscard]] inline FName PackNameWithIndex() const
+	{
+		return FName(Name, NAME_EXTERNAL_TO_INTERNAL(Index));
+	}
+
+	[[nodiscard]] inline FPropertyPathNameSegment SetNameWithIndex(FName NameWithIndex) const
+	{
+		return {FName(NameWithIndex, NAME_NO_NUMBER_INTERNAL), Type, NAME_INTERNAL_TO_EXTERNAL(NameWithIndex.GetNumber())};
+	}
+
+	[[nodiscard]] inline FPropertyPathNameSegment SetType(FName NewType) const
+	{
+		return {Name, NewType, Index};
+	}
+};
+
+/**
+ * Represents the path to a property, by name, including struct type names and the index in each container.
+ *
+ * Sequenced containers use the index directly and associative containers use their ElementId.
+ */
+class FPropertyPathName
+{
+	static_assert(NAME_NO_NUMBER == INDEX_NONE);
+
+	struct FSegment
+	{
+		FName NameWithIndex;
+		FNameEntryId Type;
+
+		inline FPropertyPathNameSegment Unpack() const
+		{
+			return FPropertyPathNameSegment().SetNameWithIndex(NameWithIndex).SetType(FName::CreateFromDisplayId(Type, NAME_NO_NUMBER_INTERNAL));
+		}
+
+		inline static FSegment Pack(const FPropertyPathNameSegment& Segment)
+		{
+			return {Segment.PackNameWithIndex(), Segment.Type.GetDisplayIndex()};
+		}
+
+		bool operator==(const FSegment& Segment) const;
+		bool operator<(const FSegment& Segment) const;
+		int32 Compare(const FSegment& Segment) const;
+	};
+
+public:
+	/** Returns the number of segments in the path. */
+	inline int32 GetSegmentCount() const
+	{
+		return Segments.Num();
+	}
+
+	/** Returns the segment at a valid index. */
+	inline FPropertyPathNameSegment GetSegment(int32 Index) const
+	{
+		return Segments[Index].Unpack();
+	}
+
+	/** Sets the segment at a valid index. */
+	inline void SetSegment(int32 Index, const FPropertyPathNameSegment& Segment)
+	{
+		Segments[Index] = FSegment::Pack(Segment);
+	}
+
+	/** Pushes a new segment on the end of the path. */
+	inline void Push(const FPropertyPathNameSegment& Segment)
+	{
+		Segments.Emplace(FSegment::Pack(Segment));
+	}
+
+	/** Pops the last segment off the end of the path. */
+	inline void Pop()
+	{
+		Segments.Pop(/*bAllowShrinking*/ false);
+	}
+
+	/** Sets the index of the last segment of the path. Ignored if the path is empty. */
+	inline void SetIndex(int32 Index)
+	{
+		if (!Segments.IsEmpty())
+		{
+			Segments.Last().NameWithIndex.SetNumber(NAME_EXTERNAL_TO_INTERNAL(Index));
+		}
+	}
+
+	inline bool IsEmpty() const { return Segments.IsEmpty(); }
+	inline void Empty() { Segments.Empty(); }
+	inline void Reset() { Segments.Reset(); }
+
+	UE_API bool operator==(const FPropertyPathName& Path) const;
+	UE_API bool operator<(const FPropertyPathName& Path) const;
+
+	inline bool operator!=(const FPropertyPathName& Path) const { return !(*this == Path); }
+	inline bool operator<=(const FPropertyPathName& Path) const { return !(Path < *this); }
+	inline bool operator>=(const FPropertyPathName& Path) const { return !(*this < Path); }
+	inline bool operator>(const FPropertyPathName& Path) const { return (Path < *this); }
+
+	friend inline FStringBuilderBase& operator<<(FStringBuilderBase& Builder, const FPropertyPathName& Path)
+	{
+		Path.ToString(Builder);
+		return Builder;
+	}
+
+	UE_API void ToString(FStringBuilderBase& Out, FStringView Separator = TEXTVIEW(" -> ")) const;
+
+	UE_API friend uint32 GetTypeHash(const FPropertyPathName& Path);
+
+private:
+	TArray<FSegment> Segments;
+};
+
+} // UE
+
+#undef UE_API
diff --git a/Engine/Source/Runtime/CoreUObject/Public/UObject/UObjectThreadContext.h b/Engine/Source/Runtime/CoreUObject/Public/UObject/UObjectThreadContext.h
--- a/Engine/Source/Runtime/CoreUObject/Public/UObject/UObjectThreadContext.h
+++ b/Engine/Source/Runtime/CoreUObject/Public/UObject/UObjectThreadContext.h
@@ -21,6 +21,7 @@
 #include "Templates/RefCounting.h"
 #include "Templates/UnrealTemplate.h"
 #include "Trace/Detail/Channel.h"
+#include "UObject/PropertyPathName.h"
 
 class FLinkerLoad;
 class FName;
@@ -167,6 +168,10 @@ public:
 	int32 SerializedExportIndex;
 	/** Points to the most recently used Linker for serialization by CreateExport() */
 	FLinkerLoad* SerializedExportLinker;
+	/** Path to the property currently being serialized */
+	UE::FPropertyPathName SerializedPropertyPath;
+	/** True when unknown properties will be serialized to or from a property bag for the serialized object. */
+	bool bSerializeUnknownProperty;
 
 	/** Adds a new loaded object */
 	COREUOBJECT_API void AddLoadedObject(UObject* InObject);
