Added an unnecessary null check to Class.cpp to silence a false positive from static analysis

#rnx

diff --git a/Engine/Source/Runtime/CoreUObject/Private/UObject/Class.cpp b/Engine/Source/Runtime/CoreUObject/Private/UObject/Class.cpp
--- a/Engine/Source/Runtime/CoreUObject/Private/UObject/Class.cpp
+++ b/Engine/Source/Runtime/CoreUObject/Private/UObject/Class.cpp
@@ -1571,7 +1571,7 @@ void UStruct::SerializeVersionedTaggedProperties(FStructuredArchive::FSlot Slot,
 						}
 					}
 				}
-				else if (FPropertyBag* PropertyBag = TryFindPropertyBag())
+				else if (FPropertyBag* PropertyBag = TryFindPropertyBag(); PropertyBag && LoadContext)
 				{
 					// TODO: Might we find defaults in a property bag for Defaults?
 					FStructuredArchive::FSlot ValueSlot = PropertyRecord.EnterField(TEXT("Value"));
