World Partition: fix LinkerLoad support for instancing context none entries required when streaming in actors with editor-only references to actors, or references to editor-only actors.

diff --git a/Engine/Source/Runtime/CoreUObject/Private/UObject/LinkerLoad.cpp b/Engine/Source/Runtime/CoreUObject/Private/UObject/LinkerLoad.cpp
--- a/Engine/Source/Runtime/CoreUObject/Private/UObject/LinkerLoad.cpp
+++ b/Engine/Source/Runtime/CoreUObject/Private/UObject/LinkerLoad.cpp
@@ -5550,7 +5550,8 @@ UObject* FLinkerLoad::CreateImport( int32 Index )
 					if( Import.OuterIndex.IsNull() )
 					{
 						FName ObjectName = InstancingContext.RemapPackage(Import.ObjectName);
-						UPackage* Pkg = CreatePackage(*ObjectName.ToString());
+						// Instancing context supports remapping editor-only references to none, handle this case here.
+						UPackage* Pkg = !ObjectName.IsNone() ? CreatePackage(*ObjectName.ToString()) : nullptr;
 						if (IsPackageReferenceAllowed(Pkg))
 						{
 							FindObject = Pkg;
