GarbageCollection classes: comment that FReferenceCollectorArchive should not be constructed outside of COREUOBJECT; it is soft-deprecated because it is poor performance and its use will block further optimizations we want to do.
#rnx

diff --git a/Engine/Plugins/Animation/ControlRig/Source/ControlRig/Private/Rigs/RigHierarchy.cpp b/Engine/Plugins/Animation/ControlRig/Source/ControlRig/Private/Rigs/RigHierarchy.cpp
--- a/Engine/Plugins/Animation/ControlRig/Source/ControlRig/Private/Rigs/RigHierarchy.cpp
+++ b/Engine/Plugins/Animation/ControlRig/Source/ControlRig/Private/Rigs/RigHierarchy.cpp
@@ -168,7 +168,6 @@ void URigHierarchy::AddReferencedObjects(UObject* InpThis, FReferenceCollector&
 	Super::AddReferencedObjects(InpThis, Collector);
 
 	URigHierarchy* pThis = static_cast<URigHierarchy*>(InpThis);
-	FReferenceCollectorArchive Ar(pThis, Collector);
 	FScopeLock Lock(&pThis->ElementsLock);
 	for (FRigBaseElement* Element : pThis->Elements)
 	{
diff --git a/Engine/Source/Runtime/CoreUObject/Public/UObject/UObjectGlobals.h b/Engine/Source/Runtime/CoreUObject/Public/UObject/UObjectGlobals.h
--- a/Engine/Source/Runtime/CoreUObject/Public/UObject/UObjectGlobals.h
+++ b/Engine/Source/Runtime/CoreUObject/Public/UObject/UObjectGlobals.h
@@ -2219,7 +2219,7 @@ private:
 	bool					bExactClass;
 };
 
-/** Base class for reference serialization archives */
+/** Reference collecting archive created by FReferenceCollector::GetVerySlowReferenceCollectorArchive() */
 class FReferenceCollectorArchive : public FArchiveUObject
 {
 	/** Object which is performing the serialization. */
@@ -2236,7 +2236,8 @@ protected:
 	}
 
 public:
-	COREUOBJECT_API FReferenceCollectorArchive(const UObject* InSerializingObject, FReferenceCollector& InCollector);
+	// Constructor not COREUOBJECT-exported because constructing this class is for internal use only
+	FReferenceCollectorArchive(const UObject* InSerializingObject, FReferenceCollector& InCollector);
 
 	void SetSerializingObject(const UObject* InSerializingObject)
 	{
