CoreRedirects: Add support to LinkerLoad for CoreRedirects that change all 3 of packagename, outername, and object name. Remove the error and suppression of those full CoreRedirects. Add test content to exercise these full Core Redirects on Blueprint functions.

The test content is not yet automatically confirmed for correctness; to confirm it the CoreRedirectsMap has to be loaded played in PIE - if successful the text printed will contain "New" instead of "Old".
#rnx

#virtualized

diff --git a/Engine/Source/Runtime/CoreUObject/Private/UObject/CoreRedirects.cpp b/Engine/Source/Runtime/CoreUObject/Private/UObject/CoreRedirects.cpp
--- a/Engine/Source/Runtime/CoreUObject/Private/UObject/CoreRedirects.cpp
+++ b/Engine/Source/Runtime/CoreUObject/Private/UObject/CoreRedirects.cpp
@@ -1102,12 +1102,6 @@ bool FCoreRedirects::AddRedirectList(TArrayView<const FCoreRedirect> Redirects,
 			continue;
 		}
 
-		if (NewRedirect.NewName.PackageName != NewRedirect.OldName.PackageName && NewRedirect.OldName.OuterName != NAME_None)
-		{
-			UE_LOG(LogCoreRedirects, Error, TEXT("AddRedirect(%s) failed to add redirect, cannot modify package from %s to %s while specifying outer!"), *SourceString, *NewRedirect.OldName.ToString(), *NewRedirect.NewName.ToString());
-			continue;
-		}
-
 		if (NewRedirect.IsSubstringMatch())
 		{
 			UE_LOG(LogCoreRedirects, Log, TEXT("AddRedirect(%s) has substring redirect %s, these are very slow and should be resolved as soon as possible!"), *SourceString, *NewRedirect.OldName.ToString());
diff --git a/Engine/Source/Runtime/CoreUObject/Private/UObject/FieldPathProperty.cpp b/Engine/Source/Runtime/CoreUObject/Private/UObject/FieldPathProperty.cpp
--- a/Engine/Source/Runtime/CoreUObject/Private/UObject/FieldPathProperty.cpp
+++ b/Engine/Source/Runtime/CoreUObject/Private/UObject/FieldPathProperty.cpp
@@ -234,7 +234,7 @@ FString FFieldPathProperty::RedirectFieldPathName(const FString& InPathName)
 			const FCoreRedirect* FoundValueRedirect = nullptr;
 			if (FCoreRedirects::RedirectNameAndValues(RedirectFlags, OldRedirectName, NewObjectName, &FoundValueRedirect))
 			{
-				const FString NewPathName = FString::Printf(TEXT("%s.%s:%s"), *NewObjectName.PackageName.ToString(), *NewObjectName.OuterName.ToString(), *NewObjectName.ObjectName.ToString());
+				const FString NewPathName = NewObjectName.ToString();
 				UE_LOG(LogCoreRedirects, Verbose, TEXT("FFieldPathProperty: Redirected '%s' -> '%s'"), *InPathName, *NewPathName);
 				return NewPathName;
 			}
@@ -248,7 +248,7 @@ FString FFieldPathProperty::RedirectFieldPathName(const FString& InPathName)
 			const FCoreRedirect* FoundValueRedirect = nullptr;
 			if (FCoreRedirects::RedirectNameAndValues(RedirectFlags, OldRedirectName, NewObjectName, &FoundValueRedirect))
 			{
-				const FString NewPathName = FString::Printf(TEXT("%s.%s:%s"), *NewObjectName.PackageName.ToString(), *NewObjectName.ObjectName.ToString(), *OldFieldName);
+				const FString NewPathName = FPackageName::ObjectPathCombine(NewObjectName.ToString(), OldFieldName);
 				UE_LOG(LogCoreRedirects, Verbose, TEXT("FFieldPathProperty: Redirected '%s' -> '%s'"), *InPathName, *NewPathName);
 				return NewPathName;
 			}
diff --git a/Engine/Source/Runtime/CoreUObject/Private/UObject/LinkerLoad.cpp b/Engine/Source/Runtime/CoreUObject/Private/UObject/LinkerLoad.cpp
--- a/Engine/Source/Runtime/CoreUObject/Private/UObject/LinkerLoad.cpp
+++ b/Engine/Source/Runtime/CoreUObject/Private/UObject/LinkerLoad.cpp
@@ -1947,151 +1947,167 @@ FLinkerLoad::ELinkerStatus FLinkerLoad::FixupImportMap()
 		{
 			static const FName NAME_BlueprintGeneratedClass(TEXT("BlueprintGeneratedClass"));
 
-			TArray<int32> PackageIndexesToClear;
+			auto AddNewPackageImport = [this](FObjectImport*& CurrentImport, int32 CurrentIndex, FName NewPackageName)
+				{
+					int32 NewImportIndex = ImportMap.Num();
+					FObjectImport& NewImport = ImportMap.AddDefaulted_GetRef();
+					// Adding to ImportMap may have reallocated, so reassign CurrentImport
+					CurrentImport = &ImportMap[CurrentIndex];
+					NewImport.ClassName = NAME_Package;
+					NewImport.ClassPackage = GLongCoreUObjectPackageName;
+					NewImport.ObjectName = NewPackageName;
+					NewImport.OuterIndex = FPackageIndex();
+					NewImport.XObject = nullptr;
+					NewImport.SourceLinker = nullptr;
+					NewImport.SourceIndex = -1;
+					return FPackageIndex::FromImport(NewImportIndex);
+				};
+			auto AddNewObjectImport = [this](FObjectImport*& CurrentImport, int32 CurrentIndex,
+				FPackageIndex NewImportOuter, FName NewImportName)
+				{
+					int32 NewImportIndex = ImportMap.Num();
+					FObjectImport& NewImport = ImportMap.AddDefaulted_GetRef();
+					// Adding to ImportMap may have reallocated, so reassign CurrentImport
+					CurrentImport = &ImportMap[CurrentIndex];
+					NewImport.ClassName = NAME_Object; // Don't know the class, but we won't need it. Set it to UObject
+					NewImport.ClassPackage = GLongCoreUObjectPackageName;
+					NewImport.ObjectName = NewImportName;
+					NewImport.OuterIndex = NewImportOuter;
+					NewImport.XObject = nullptr;
+					NewImport.SourceLinker = nullptr;
+					NewImport.SourceIndex = -1;
+					return FPackageIndex::FromImport(NewImportIndex);
+				};
 
-			bool bDone = false;
-			while (!bDone)
+			TArray<int32> PackageIndexesToClear;
+			for (int32 i=0; i<ImportMap.Num(); i++)
 			{
-				TArray<FName> NewPackageImports;
+				FObjectImport* Import = &ImportMap[i];
 
-				bDone = true;
-				for( int32 i=0; i<ImportMap.Num(); i++ )
-				{
-					FObjectImport& Import = ImportMap[i];
+				// Compute class name first, as instance can override it
+				const FCoreRedirect* ClassValueRedirect = nullptr;
+				FCoreRedirectObjectName OldClassName(Import->ClassName, NAME_None, Import->ClassPackage), NewClassName;
 
-					// Compute class name first, as instance can override it
-					const FCoreRedirect* ClassValueRedirect = nullptr;
-					FCoreRedirectObjectName OldClassName(Import.ClassName, NAME_None, Import.ClassPackage), NewClassName;
+				FCoreRedirects::RedirectNameAndValues(ECoreRedirectFlags::Type_Class, OldClassName, NewClassName, &ClassValueRedirect);
 
-					FCoreRedirects::RedirectNameAndValues(ECoreRedirectFlags::Type_Class, OldClassName, NewClassName, &ClassValueRedirect);
-
-					if (ClassValueRedirect)
+				if (ClassValueRedirect)
+				{
+					// Apply class value redirects before other redirects, to mirror old subobject order
+					const FString* NewInstanceName = ClassValueRedirect->ValueChanges.Find(Import->ObjectName.ToString());
+					if (NewInstanceName)
 					{
-						// Apply class value redirects before other redirects, to mirror old subobject order
-						const FString* NewInstanceName = ClassValueRedirect->ValueChanges.Find(Import.ObjectName.ToString());
-						if (NewInstanceName)
-						{
-							// Rename this import directly
-							FString Was = GetImportFullName(i);
-							Import.ObjectName = FName(**NewInstanceName);
+						// Rename this import directly
+						FString Was = GetImportFullName(i);
+						Import->ObjectName = FName(**NewInstanceName);
 
-							if (Import.ObjectName != NAME_None)
-							{
-								FString Now = GetImportFullName(i);
-								UE_LOG(LogLinker, Verbose, TEXT("FLinkerLoad::FixupImportMap() - Renamed object from %s   to   %s"), *Was, *Now);
-							}
-							else
-							{
-								UE_LOG(LogLinker, Verbose, TEXT("FLinkerLoad::FixupImportMap() - Removed object %s"), *Was);
-							}
+						if (Import->ObjectName != NAME_None)
+						{
+							FString Now = GetImportFullName(i);
+							UE_LOG(LogLinker, Verbose, TEXT("FLinkerLoad::FixupImportMap() - Renamed object from %s   to   %s"), *Was, *Now);
+						}
+						else
+						{
+							UE_LOG(LogLinker, Verbose, TEXT("FLinkerLoad::FixupImportMap() - Removed object %s"), *Was);
 						}
 					}
+				}
 
-					FCoreRedirectObjectName OldObjectName(GetImportPathName(i)), NewObjectName;
-					ECoreRedirectFlags ObjectRedirectFlags = FCoreRedirects::GetFlagsForTypeName(Import.ClassPackage, Import.ClassName);
-					const FCoreRedirect* ValueRedirect = nullptr;
+				FCoreRedirectObjectName OldObjectName(GetImportPathName(i)), NewObjectName;
+				ECoreRedirectFlags ObjectRedirectFlags = FCoreRedirects::GetFlagsForTypeName(Import->ClassPackage, Import->ClassName);
+				const FCoreRedirect* ValueRedirect = nullptr;
 					
-					FCoreRedirects::RedirectNameAndValues(ObjectRedirectFlags, OldObjectName, NewObjectName, &ValueRedirect);
+				FCoreRedirects::RedirectNameAndValues(ObjectRedirectFlags, OldObjectName, NewObjectName, &ValueRedirect);
 
-					if (ValueRedirect && ValueRedirect->OverrideClassName.IsValid())
+				if (ValueRedirect && ValueRedirect->OverrideClassName.IsValid())
+				{
+					// Override class name if found, even if the name didn't actually change
+					NewClassName = ValueRedirect->OverrideClassName;
+				}
+
+				if (NewObjectName != OldObjectName)
+				{
+					if (Import->OuterIndex.IsNull())
 					{
-						// Override class name if found, even if the name didn't actually change
-						NewClassName = ValueRedirect->OverrideClassName;
+						// If this has no outer it's a package and we don't want to rename it, the subobject renames will handle creating the new package import
+						// We do need to clear these at the end so it doesn't try to load nonexistent packages
+						PackageIndexesToClear.Add(i);
 					}
-
-					if (NewObjectName != OldObjectName)
+					else
 					{
-						if (Import.OuterIndex.IsNull())
+						FPackageIndex NewPackageIndex;
+						if (!FindImportPackage(NewObjectName.PackageName, NewPackageIndex))
 						{
-							// If this has no outer it's a package and we don't want to rename it, the subobject renames will handle creating the new package import
-							// We do need to clear these at the end so it doesn't try to load nonexistent packages
-							PackageIndexesToClear.Add(i);
+							NewPackageIndex = AddNewPackageImport(Import, i, NewObjectName.PackageName);
 						}
-						else
+
+						FPackageIndex OuterIndex = NewPackageIndex;
+						if (!NewObjectName.OuterName.IsNone())
 						{
-							// If right below package and package has changed, need to swap outer
-							if (NewObjectName.OuterName == NAME_None && NewObjectName.PackageName != OldObjectName.PackageName)
+							TStringBuilder<256> OuterNameBuffer;
+							OuterNameBuffer << NewObjectName.OuterName;
+							FStringView OuterName(OuterNameBuffer);
+							while (!OuterName.IsEmpty())
 							{
-								FPackageIndex NewPackageIndex;
-
-								if (FindImportPackage(NewObjectName.PackageName, NewPackageIndex))
+								FStringView FirstOuter;
+								FStringView Remainder;
+								FPackageName::ObjectPathSplitFirstName(OuterName, FirstOuter, Remainder);
+								FPackageIndex NewOuterIndex;
+								FName FirstOuterName(FirstOuter);
+								if (!FindImport(OuterIndex, FirstOuterName, NewOuterIndex))
 								{
-									// Already in import table, set it
-									Import.OuterIndex = NewPackageIndex;
+									NewOuterIndex = AddNewObjectImport(Import, i, OuterIndex, FirstOuterName);
 								}
-								else
-								{
-									// Need to add package import and try again
-									NewPackageImports.AddUnique(NewObjectName.PackageName);
-									bDone = false;
-									break;
-								}
-							}
-#if WITH_EDITOR
-							// If this is a class, set old name here 
-							if (ObjectRedirectFlags == ECoreRedirectFlags::Type_Class)
-							{
-								Import.OldClassName = Import.ObjectName;
+								OuterName = Remainder;
+								OuterIndex = NewOuterIndex;
 							}
+						}
 
+						Import->OuterIndex = OuterIndex;
+#if WITH_EDITOR
+						// If this is a class, set old name here 
+						if (ObjectRedirectFlags == ECoreRedirectFlags::Type_Class)
+						{
+							Import->OldClassName = Import->ObjectName;
+						}
 #endif
-							// Change object name
-							Import.ObjectName = NewObjectName.ObjectName;
+						// Change object name
+						Import->ObjectName = NewObjectName.ObjectName;
 
-							UE_LOG(LogLinker, Verbose, TEXT("FLinkerLoad::FixupImportMap() - Pkg<%s> - Renamed Object %s -> %s"), *LinkerRoot->GetName(), *OldObjectName.ToString(), *NewObjectName.ToString());
-						}
+						UE_LOG(LogLinker, Verbose, TEXT("FLinkerLoad::FixupImportMap() - Pkg<%s> - Renamed Object %s -> %s"), *LinkerRoot->GetName(), *OldObjectName.ToString(), *NewObjectName.ToString());
 					}
+				}
 
-					if (NewClassName != OldClassName)
+				if (NewClassName != OldClassName)
+				{
+					// Swap class if needed
+					if (Import->ClassPackage != NewClassName.PackageName && !IsCoreUObjectPackage(NewClassName.PackageName))
 					{
-						// Swap class if needed
-						if (Import.ClassPackage != NewClassName.PackageName && !IsCoreUObjectPackage(NewClassName.PackageName))
-						{
-							FPackageIndex NewPackageIndex;
+						FPackageIndex NewPackageIndex;
 
-							if (!FindImportPackage(NewClassName.PackageName, NewPackageIndex))
-							{
-								// Need to add package import and try again
-								NewPackageImports.AddUnique(NewClassName.PackageName);
-								bDone = false;
-								break;
-							}
+						if (!FindImportPackage(NewClassName.PackageName, NewPackageIndex))
+						{
+							NewPackageIndex = AddNewPackageImport(Import, i, NewClassName.PackageName);
 						}
+					}
 #if WITH_EDITOR
-						Import.OldClassName = Import.ClassName;
+					Import->OldClassName = Import->ClassName;
 #endif
-						// Change class name/package
-						Import.ClassPackage = NewClassName.PackageName;
-						Import.ClassName = NewClassName.ObjectName;
-
-						// Also change CDO name if needed
-						FString NewDefaultObjectName = Import.ObjectName.ToString();
+					// Change class name/package
+					Import->ClassPackage = NewClassName.PackageName;
+					Import->ClassName = NewClassName.ObjectName;
 
-						if (NewDefaultObjectName.StartsWith(DEFAULT_OBJECT_PREFIX))
-						{
-							NewDefaultObjectName = FString(DEFAULT_OBJECT_PREFIX);
-							NewDefaultObjectName += NewClassName.ObjectName.ToString();
-							Import.ObjectName = FName(*NewDefaultObjectName);
-						}
+					// Also change CDO name if needed
+					FString NewDefaultObjectName = Import->ObjectName.ToString();
 
-						UE_LOG(LogLinker, Verbose, TEXT("FLinkerLoad::FixupImportMap() - Pkg<%s> - Renamed Class %s -> %s"), *LinkerRoot->GetName(), *OldClassName.ToString(), *NewClassName.ToString());
-					}	
-				}
-
-				// Add new packages, after loop iteration for safety
-				for (FName NewPackage : NewPackageImports)
-				{
-					// We are adding a new import to the map as we need the new package dependency added to the works
-					FObjectImport& NewImport = ImportMap.AddDefaulted_GetRef();
+					if (NewDefaultObjectName.StartsWith(DEFAULT_OBJECT_PREFIX))
+					{
+						NewDefaultObjectName = FString(DEFAULT_OBJECT_PREFIX);
+						NewDefaultObjectName += NewClassName.ObjectName.ToString();
+						Import->ObjectName = FName(*NewDefaultObjectName);
+					}
 
-					NewImport.ClassName = NAME_Package;
-					NewImport.ClassPackage = GLongCoreUObjectPackageName;
-					NewImport.ObjectName = NewPackage;
-					NewImport.OuterIndex = FPackageIndex();
-					NewImport.XObject = 0;
-					NewImport.SourceLinker = 0;
-					NewImport.SourceIndex = -1;
-				}
+					UE_LOG(LogLinker, Verbose, TEXT("FLinkerLoad::FixupImportMap() - Pkg<%s> - Renamed Class %s -> %s"), *LinkerRoot->GetName(), *OldClassName.ToString(), *NewClassName.ToString());
+				}	
 			}
 
 			// Clear any packages that got renamed, once all children have been fixed up
@@ -6150,6 +6166,21 @@ bool FLinkerLoad::FindImportPackage(FName PackageName, FPackageIndex& PackageIdx
 	return false;
 }
 
+bool FLinkerLoad::FindImport(FPackageIndex OuterIndex, FName ObjectName, FPackageIndex& OutObjectIndex)
+{
+	for (int32 ImportMapIdx = 0; ImportMapIdx < ImportMap.Num(); ImportMapIdx++)
+	{
+		if (ImportMap[ImportMapIdx].ObjectName == ObjectName && ImportMap[ImportMapIdx].OuterIndex == OuterIndex)
+		{
+			OutObjectIndex = FPackageIndex::FromImport(ImportMapIdx);
+			return true;
+		}
+	}
+
+	OutObjectIndex = FPackageIndex();
+	return false;
+}
+
 /**
  * Locates the class adjusted index and its package adjusted index for a given class name in the import map
  */
diff --git a/Engine/Source/Runtime/CoreUObject/Private/UObject/Property.cpp b/Engine/Source/Runtime/CoreUObject/Private/UObject/Property.cpp
--- a/Engine/Source/Runtime/CoreUObject/Private/UObject/Property.cpp
+++ b/Engine/Source/Runtime/CoreUObject/Private/UObject/Property.cpp
@@ -2037,12 +2037,27 @@ FName FProperty::FindRedirectedPropertyName(const UStruct* ObjectStruct, FName O
 	FCoreRedirectObjectName OldRedirectName(OldName, OuterName, PackageName);
 	FCoreRedirectObjectName NewRedirectName = FCoreRedirects::GetRedirectedName(ECoreRedirectFlags::Type_Property, OldRedirectName);
 
-	if (NewRedirectName != OldRedirectName)
+	if (NewRedirectName == OldRedirectName)
 	{
-		return NewRedirectName.ObjectName;
+		return NAME_None;
 	}
 
-	return NAME_None;
+	if ((!NewRedirectName.PackageName.IsNone() && NewRedirectName.PackageName != PackageName) ||
+		(!NewRedirectName.OuterName.IsNone() && NewRedirectName.OuterName != OuterName))
+	{
+		// We don't handle changing PackageName and OuterName this in general, but we do handle it in the case
+		// that the redirected field name is a field from a superclass. So if the new ObjectName exists on the struct,
+		// allow the redirect to that ObjectName and ignore the new PackageName and OuterName.
+		if (!ObjectStruct->FindPropertyByName(NewRedirectName.ObjectName))
+		{
+			// If it doesn't exist then give a warning and ignore the redirect.
+			UE_LOG(LogProperty, Warning, TEXT("FindRedirectedPropertyName failed for %s -> %s. Changing Package or Outer is not supported. Redirect will be ignored."),
+				*OldRedirectName.ToString(), *NewRedirectName.ToString());
+			return NAME_None;
+		}
+	}
+		
+	return NewRedirectName.ObjectName;
 }
 
 /**
diff --git a/Engine/Source/Runtime/CoreUObject/Public/UObject/LinkerLoad.h b/Engine/Source/Runtime/CoreUObject/Public/UObject/LinkerLoad.h
--- a/Engine/Source/Runtime/CoreUObject/Public/UObject/LinkerLoad.h
+++ b/Engine/Source/Runtime/CoreUObject/Public/UObject/LinkerLoad.h
@@ -564,6 +564,8 @@ public:
 	 * Locates package index for a UPackage import
 	 */
 	COREUOBJECT_API bool FindImportPackage(FName PackageName, FPackageIndex& PackageIdx);
+	/* Locates package index for a given name in an outer. */
+	COREUOBJECT_API bool FindImport(FPackageIndex OuterIndex, FName ObjectName, FPackageIndex& OutObjectIndex);
 
 	/**
 	 * Locates the class adjusted index and its package adjusted index for a given class name in the import map
