CoreRedirects: Revert the new warning and early exit from FProperty::FindRedirectedPropertyName when the outer has changed, even though the function does not handle the change in outer. The redirect can point from an old field on a base class to a new field on a subclass, and the ObjectStruct passed in from UStruct::SerializeVersionedTaggedProperties points to the base class rather than to the subclass, so even though the new name does not exist on the passed-in ObjectStruct, the caller will still be able to find it. Example: ShooterGame's redirect of Pawn.Health to ShooterCharacter.Health.
#rnx

diff --git a/Engine/Source/Runtime/CoreUObject/Private/UObject/Property.cpp b/Engine/Source/Runtime/CoreUObject/Private/UObject/Property.cpp
--- a/Engine/Source/Runtime/CoreUObject/Private/UObject/Property.cpp
+++ b/Engine/Source/Runtime/CoreUObject/Private/UObject/Property.cpp
@@ -2037,27 +2037,12 @@ FName FProperty::FindRedirectedPropertyName(const UStruct* ObjectStruct, FName O
 	FCoreRedirectObjectName OldRedirectName(OldName, OuterName, PackageName);
 	FCoreRedirectObjectName NewRedirectName = FCoreRedirects::GetRedirectedName(ECoreRedirectFlags::Type_Property, OldRedirectName);
 
-	if (NewRedirectName == OldRedirectName)
+	if (NewRedirectName != OldRedirectName)
 	{
-		return NAME_None;
+		return NewRedirectName.ObjectName;
 	}
 
-	if ((!NewRedirectName.PackageName.IsNone() && NewRedirectName.PackageName != PackageName) ||
-		(!NewRedirectName.OuterName.IsNone() && NewRedirectName.OuterName != OuterName))
-	{
-		// We don't handle changing PackageName and OuterName this in general, but we do handle it in the case
-		// that the redirected field name is a field from a superclass. So if the new ObjectName exists on the struct,
-		// allow the redirect to that ObjectName and ignore the new PackageName and OuterName.
-		if (!ObjectStruct->FindPropertyByName(NewRedirectName.ObjectName))
-		{
-			// If it doesn't exist then give a warning and ignore the redirect.
-			UE_LOG(LogProperty, Warning, TEXT("FindRedirectedPropertyName failed for %s -> %s. Changing Package or Outer is not supported. Redirect will be ignored."),
-				*OldRedirectName.ToString(), *NewRedirectName.ToString());
-			return NAME_None;
-		}
-	}
-		
-	return NewRedirectName.ObjectName;
+	return NAME_None;
 }
 
 /**
