Fix CVar gc.VerifyAssumptionsOnFullPurge to bind to GShouldVerifyGCAssumptionsOnFullPurge

diff --git a/Engine/Source/Runtime/CoreUObject/Private/UObject/GarbageCollectionVerification.cpp b/Engine/Source/Runtime/CoreUObject/Private/UObject/GarbageCollectionVerification.cpp
--- a/Engine/Source/Runtime/CoreUObject/Private/UObject/GarbageCollectionVerification.cpp
+++ b/Engine/Source/Runtime/CoreUObject/Private/UObject/GarbageCollectionVerification.cpp
@@ -44,7 +44,7 @@ static FAutoConsoleVariableRef CVarShouldVerifyGCAssumptions(
 COREUOBJECT_API bool	GShouldVerifyGCAssumptionsOnFullPurge = !UE_BUILD_SHIPPING && !WITH_EDITOR;
 static FAutoConsoleVariableRef CVarShouldVerifyGCAssumptionsOnFullPurge(
 	TEXT("gc.VerifyAssumptionsOnFullPurge"),
-	GShouldVerifyGCAssumptions,
+	GShouldVerifyGCAssumptionsOnFullPurge,
 	TEXT("Whether to verify GC assumptions (disregard for GC, clustering) on full purge GCs."),
 	ECVF_Default
 );
