[FScriptSetHelper and FScriptMapHelper iterators]
simple index to iterator replacements

diff --git a/Engine/Plugins/Animation/GameplayInsights/Source/GameplayInsights/Private/ObjectPropertyTrace.cpp b/Engine/Plugins/Animation/GameplayInsights/Source/GameplayInsights/Private/ObjectPropertyTrace.cpp
--- a/Engine/Plugins/Animation/GameplayInsights/Source/GameplayInsights/Private/ObjectPropertyTrace.cpp
+++ b/Engine/Plugins/Animation/GameplayInsights/Source/GameplayInsights/Private/ObjectPropertyTrace.cpp
@@ -130,31 +130,24 @@ namespace ObjectPropertyTrace
 			
 			InFunction(InProperty->GetCPPType(), SizeString, MapRootId, InParentId, InPropertyNameBuilder);
 			
-			int32 Num = Helper.Num();
-			int32 MapIndex = 0;
-			for (int32 DynamicIndex = 0; Num; ++DynamicIndex)
+			for (FScriptMapHelper::FIterator It(Helper); It; ++It)
 			{
-				if (Helper.IsValidIndex(DynamicIndex))
-				{
-					int32 MapEntryId = ++InId;
-					FString KeyString = FString::Printf(TEXT("[%d]"), MapIndex++);
-					FString TypeString = FString::Printf(TEXT("{%s, %s}"), *MapProperty->KeyProp->GetCPPType(), *MapProperty->ValueProp->GetCPPType());
-
-					// Build pair parent name id
-					FScopedPropertyName ScopedParentPairName(InPropertyNameBuilder);
-					ScopedParentPairName.Builder.AppendChar(TEXT('.'));
-					ScopedParentPairName.Builder.Append(KeyString);
-					
-					InFunction(TypeString, TEXT("{...}"), MapEntryId, MapRootId, InPropertyNameBuilder);
-					
-					const void* PairPtr = Helper.GetPairPtr(DynamicIndex);
+				int32 MapEntryId = ++InId;
+				FString KeyString = FString::Printf(TEXT("[%d]"), It.GetLogicalIndex());
+				FString TypeString = FString::Printf(TEXT("{%s, %s}"), *MapProperty->KeyProp->GetCPPType(), *MapProperty->ValueProp->GetCPPType());
+
+				// Build pair parent name id
+				FScopedPropertyName ScopedParentPairName(InPropertyNameBuilder);
+				ScopedParentPairName.Builder.AppendChar(TEXT('.'));
+				ScopedParentPairName.Builder.Append(KeyString);
 
-					IteratePropertiesRecursive(MapProperty->KeyProp, PairPtr, MapProperty->KeyProp->GetName(), InFunction, InId, MapEntryId, InPropertyNameBuilder);
+				InFunction(TypeString, TEXT("{...}"), MapEntryId, MapRootId, InPropertyNameBuilder);
 
-					IteratePropertiesRecursive(MapProperty->ValueProp, PairPtr, MapProperty->ValueProp->GetName(), InFunction, InId, MapEntryId, InPropertyNameBuilder);
+				const void* PairPtr = Helper.GetPairPtr(It);
 
-					--Num;
-				}
+				IteratePropertiesRecursive(MapProperty->KeyProp, PairPtr, MapProperty->KeyProp->GetName(), InFunction, InId, MapEntryId, InPropertyNameBuilder);
+
+				IteratePropertiesRecursive(MapProperty->ValueProp, PairPtr, MapProperty->ValueProp->GetName(), InFunction, InId, MapEntryId, InPropertyNameBuilder);
 			}
 		}
 		else if (const FSetProperty* SetProperty = CastField<FSetProperty>(InProperty))
@@ -173,19 +166,12 @@ namespace ObjectPropertyTrace
 			InFunction(InProperty->GetCPPType(), SizeString, SetRootId, InParentId, InPropertyNameBuilder);
 
 			// Trace container items
-			int32 Num = Helper.Num();
-			int32 SetIndex = 0;
-			for (int32 DynamicIndex = 0; Num; ++DynamicIndex)
+			for (FScriptSetHelper::FIterator It(Helper); It; ++It)
 			{
-				if (Helper.IsValidIndex(DynamicIndex))
-				{
-					const void* ValuePtr = Helper.GetElementPtr(DynamicIndex);
-					FString KeyString = FString::Printf(TEXT("[%d]"), SetIndex++);
-					
-					IteratePropertiesRecursive(SetProperty->ElementProp, ValuePtr, KeyString, InFunction, InId, SetRootId, InPropertyNameBuilder);
+				const void* ValuePtr = Helper.GetElementPtr(It);
+				FString KeyString = FString::Printf(TEXT("[%d]"), It.GetLogicalIndex());
 
-					--Num;
-				}
+				IteratePropertiesRecursive(SetProperty->ElementProp, ValuePtr, KeyString, InFunction, InId, SetRootId, InPropertyNameBuilder);
 			}
 		}
 		else if (const FStructProperty* StructProperty = CastField<FStructProperty>(InProperty))
diff --git a/Engine/Plugins/Enterprise/GLTFExporter/Source/GLTFExporter/Private/Exporters/GLTFExporterAnalytics.cpp b/Engine/Plugins/Enterprise/GLTFExporter/Source/GLTFExporter/Private/Exporters/GLTFExporterAnalytics.cpp
--- a/Engine/Plugins/Enterprise/GLTFExporter/Source/GLTFExporter/Private/Exporters/GLTFExporterAnalytics.cpp
+++ b/Engine/Plugins/Enterprise/GLTFExporter/Source/GLTFExporter/Private/Exporters/GLTFExporterAnalytics.cpp
@@ -110,21 +110,14 @@ void FGLTFExporterAnalytics::GetAttributesFromProperty(const FProperty* Property
 	else if (const FMapProperty* MapProperty = CastField<FMapProperty>(Property))
 	{
 		FScriptMapHelper Helper(MapProperty, ValuePtr);
-		int32 RemainingPairs = Helper.Num();
+		OutAttributes.Emplace(AttributeName, Helper.Num() > 0);
 
-		OutAttributes.Emplace(AttributeName, RemainingPairs > 0);
-
-		for (int32 Index = 0; RemainingPairs > 0; ++Index)
+		for (FScriptMapHelper::FIterator It(Helper); It; ++It)
 		{
-			if (Helper.IsValidIndex(Index))
-			{
-				FString KeyString;
-				MapProperty->KeyProp->ExportTextItem_Direct(KeyString, Helper.GetKeyPtr(Index), nullptr, nullptr, PPF_None);
+			FString KeyString;
+			MapProperty->KeyProp->ExportTextItem_Direct(KeyString, Helper.GetKeyPtr(It), nullptr, nullptr, PPF_None);
 
-				GetAttributesFromProperty(MapProperty->ValueProp, Helper.GetValuePtr(Index), AttributeName + TEXT(".") + KeyString, OutAttributes);
-
-				--RemainingPairs;
-			}
+			GetAttributesFromProperty(MapProperty->ValueProp, Helper.GetValuePtr(It), AttributeName + TEXT(".") + KeyString, OutAttributes);
 		}
 	}
 	else if (const FStructProperty* StructProperty = CastField<FStructProperty>(Property))
diff --git a/Engine/Plugins/Experimental/PCG/Source/PCG/Private/Helpers/PCGHelpers.cpp b/Engine/Plugins/Experimental/PCG/Source/PCG/Private/Helpers/PCGHelpers.cpp
--- a/Engine/Plugins/Experimental/PCG/Source/PCG/Private/Helpers/PCGHelpers.cpp
+++ b/Engine/Plugins/Experimental/PCG/Source/PCG/Private/Helpers/PCGHelpers.cpp
@@ -378,34 +378,22 @@ namespace PCGHelpers
 		else if (FMapProperty* MapProperty = CastField<FMapProperty>(Property))
 		{
 			FScriptMapHelper_InContainer Helper(MapProperty, InContainer);
-			int32 Num = Helper.Num();
-			for (int32 DynamicIndex = 0; Num; ++DynamicIndex)
+			for (FScriptMapHelper::FIterator It(Helper); It; ++It)
 			{
-				if (Helper.IsValidIndex(DynamicIndex))
-				{
 					// Key and Value are stored next to each other in memory.
 					// ValueProp has an offset, so we should use the same starting address for both.
-					const void* PairKeyValuePtr = Helper.GetKeyPtr(DynamicIndex);
+					const void* PairKeyValuePtr = Helper.GetKeyPtr(It);
 					GatherDependencies(MapProperty->KeyProp, PairKeyValuePtr, OutDependencies, MaxDepth, InExcludedClasses);
 					GatherDependencies(MapProperty->ValueProp, PairKeyValuePtr, OutDependencies, MaxDepth, InExcludedClasses);
-
-					--Num;
-				}
 			}
 		}
 		else if (FSetProperty* SetProperty = CastField<FSetProperty>(Property))
 		{
 			FScriptSetHelper_InContainer Helper(SetProperty, InContainer);
-			int32 Num = Helper.Num();
-			for (int32 DynamicIndex = 0; Num; ++DynamicIndex)
+			for (FScriptSetHelper::FIterator It(Helper); It; ++It)
 			{
-				if (Helper.IsValidIndex(DynamicIndex))
-				{
-					const void* ValuePtr = Helper.GetElementPtr(DynamicIndex);
-					GatherDependencies(SetProperty->ElementProp, ValuePtr, OutDependencies, MaxDepth, InExcludedClasses);
-
-					--Num;
-				}
+				const void* ValuePtr = Helper.GetElementPtr(It);
+				GatherDependencies(SetProperty->ElementProp, ValuePtr, OutDependencies, MaxDepth, InExcludedClasses);
 			}
 		}
 	}
diff --git a/Engine/Plugins/Experimental/PythonScriptPlugin/Source/PythonScriptPlugin/Private/PyGenUtil.cpp b/Engine/Plugins/Experimental/PythonScriptPlugin/Source/PythonScriptPlugin/Private/PyGenUtil.cpp
--- a/Engine/Plugins/Experimental/PythonScriptPlugin/Source/PythonScriptPlugin/Private/PyGenUtil.cpp
+++ b/Engine/Plugins/Experimental/PythonScriptPlugin/Source/PythonScriptPlugin/Private/PyGenUtil.cpp
@@ -2315,16 +2315,13 @@ void PythonizeValueImpl(const FProperty* InProp, const void* InPropValue, const
 				: TEXT("[");
 			{
 				FScriptSetHelper ScriptSetHelper(SetProperty, PropArrValue);
-				for (int32 SparseElementIndex = 0, ElementIndex = 0; SparseElementIndex < ScriptSetHelper.GetMaxIndex(); ++SparseElementIndex)
+				for (FScriptSetHelper::FIterator It(ScriptSetHelper); It; ++It)
 				{
-					if (ScriptSetHelper.IsValidIndex(SparseElementIndex))
+					if (It.GetLogicalIndex() > 0)
 					{
-						if (ElementIndex++ > 0)
-						{
-							OutPythonDefaultValue += TEXT(", ");
-						}
-						PythonizeValueImpl(ScriptSetHelper.GetElementProperty(), ScriptSetHelper.GetElementPtr(SparseElementIndex), InFlags, OutPythonDefaultValue);
+						OutPythonDefaultValue += TEXT(", ");
 					}
+					PythonizeValueImpl(ScriptSetHelper.GetElementProperty(), ScriptSetHelper.GetElementPtr(It), InFlags, OutPythonDefaultValue);
 				}
 			}
 			OutPythonDefaultValue += bUseStrictTyping
@@ -2338,18 +2335,15 @@ void PythonizeValueImpl(const FProperty* InProp, const void* InPropValue, const
 				: TEXT("{");
 			{
 				FScriptMapHelper ScriptMapHelper(MapProperty, PropArrValue);
-				for (int32 SparseElementIndex = 0, ElementIndex = 0; SparseElementIndex < ScriptMapHelper.GetMaxIndex(); ++SparseElementIndex)
+				for (FScriptMapHelper::FIterator It(ScriptMapHelper); It; ++It)
 				{
-					if (ScriptMapHelper.IsValidIndex(SparseElementIndex))
+					if (It.GetLogicalIndex() > 0)
 					{
-						if (ElementIndex++ > 0)
-						{
-							OutPythonDefaultValue += TEXT(", ");
-						}
-						PythonizeValueImpl(ScriptMapHelper.GetKeyProperty(), ScriptMapHelper.GetKeyPtr(SparseElementIndex), InFlags, OutPythonDefaultValue);
-						OutPythonDefaultValue += TEXT(": ");
-						PythonizeValueImpl(ScriptMapHelper.GetValueProperty(), ScriptMapHelper.GetValuePtr(SparseElementIndex), InFlags, OutPythonDefaultValue);
+						OutPythonDefaultValue += TEXT(", ");
 					}
+					PythonizeValueImpl(ScriptMapHelper.GetKeyProperty(), ScriptMapHelper.GetKeyPtr(It), InFlags, OutPythonDefaultValue);
+					OutPythonDefaultValue += TEXT(": ");
+					PythonizeValueImpl(ScriptMapHelper.GetValueProperty(), ScriptMapHelper.GetValuePtr(It), InFlags, OutPythonDefaultValue);
 				}
 			}
 			OutPythonDefaultValue += bUseStrictTyping
diff --git a/Engine/Plugins/Experimental/PythonScriptPlugin/Source/PythonScriptPlugin/Private/PyReferenceCollector.cpp b/Engine/Plugins/Experimental/PythonScriptPlugin/Source/PythonScriptPlugin/Private/PyReferenceCollector.cpp
--- a/Engine/Plugins/Experimental/PythonScriptPlugin/Source/PythonScriptPlugin/Private/PyReferenceCollector.cpp
+++ b/Engine/Plugins/Experimental/PythonScriptPlugin/Source/PythonScriptPlugin/Private/PyReferenceCollector.cpp
@@ -345,13 +345,9 @@ void FPyReferenceCollector::AddReferencedObjectsFromPropertyInternal(FReferenceC
 		{
 			bool bSetValuesChanged = false;
 			FScriptSetHelper_InContainer ScriptSetHelper(CastProp, InBaseAddr, ArrIndex);
-
-			for (int32 SparseElementIndex = 0; SparseElementIndex < ScriptSetHelper.GetMaxIndex(); ++SparseElementIndex)
+			for (FScriptSetHelper::FIterator It(ScriptSetHelper); It; ++It)
 			{
-				if (ScriptSetHelper.IsValidIndex(SparseElementIndex))
-				{
-						AddReferencedObjectsFromPropertyInternal(InCollector, ScriptSetHelper.GetElementProperty(), ScriptSetHelper.GetElementPtr(SparseElementIndex), InFlags, bSetValuesChanged);
-				}
+				AddReferencedObjectsFromPropertyInternal(InCollector, ScriptSetHelper.GetElementProperty(), ScriptSetHelper.GetElementPtr(It), InFlags, bSetValuesChanged);
 			}
 
 			if (bSetValuesChanged)
@@ -374,14 +370,11 @@ void FPyReferenceCollector::AddReferencedObjectsFromPropertyInternal(FReferenceC
 			bool bMapValuesChanged = false;
 			FScriptMapHelper_InContainer ScriptMapHelper(CastProp, InBaseAddr, ArrIndex);
 
-			for (int32 SparseElementIndex = 0; SparseElementIndex < ScriptMapHelper.GetMaxIndex(); ++SparseElementIndex)
+			for (FScriptMapHelper::FIterator It(ScriptMapHelper); It; ++It)
 			{
-				if (ScriptMapHelper.IsValidIndex(SparseElementIndex))
-				{
 					// Note: We use the pair pointer below as AddReferencedObjectsFromPropertyInternal expects a base address and the key/value property will apply the correct offset from the base
-					AddReferencedObjectsFromPropertyInternal(InCollector, ScriptMapHelper.GetKeyProperty(), ScriptMapHelper.GetPairPtr(SparseElementIndex), InFlags, bMapKeysChanged);
-					AddReferencedObjectsFromPropertyInternal(InCollector, ScriptMapHelper.GetValueProperty(), ScriptMapHelper.GetPairPtr(SparseElementIndex), InFlags, bMapValuesChanged);
-				}
+					AddReferencedObjectsFromPropertyInternal(InCollector, ScriptMapHelper.GetKeyProperty(), ScriptMapHelper.GetPairPtr(It), InFlags, bMapKeysChanged);
+					AddReferencedObjectsFromPropertyInternal(InCollector, ScriptMapHelper.GetValueProperty(), ScriptMapHelper.GetPairPtr(It), InFlags, bMapValuesChanged);
 			}
 
 			if (bMapKeysChanged || bMapValuesChanged)
diff --git a/Engine/Plugins/Experimental/PythonScriptPlugin/Source/PythonScriptPlugin/Private/PyWrapperMap.cpp b/Engine/Plugins/Experimental/PythonScriptPlugin/Source/PythonScriptPlugin/Private/PyWrapperMap.cpp
--- a/Engine/Plugins/Experimental/PythonScriptPlugin/Source/PythonScriptPlugin/Private/PyWrapperMap.cpp
+++ b/Engine/Plugins/Experimental/PythonScriptPlugin/Source/PythonScriptPlugin/Private/PyWrapperMap.cpp
@@ -222,17 +222,17 @@ struct FPyWrapperMapItemIterator : public TPyWrapperMapIterator<FPyWrapperMapIte
 		return &PyWrapperMapItemIteratorType;
 	}
 
-	static PyObject* GetItem(FPyWrapperMapItemIterator* InSelf, FScriptMapHelper& InScriptMapHelper, int32 InElementIndex)
+	static PyObject* GetItem(FPyWrapperMapItemIterator* InSelf, FScriptMapHelper& InScriptMapHelper, const int32 InValidatedInternalIndex)
 	{
 		FPyObjectPtr PyKeyObj;
-		if (!PyConversion::PythonizeProperty(InScriptMapHelper.GetKeyProperty(), InScriptMapHelper.GetKeyPtr(InElementIndex), PyKeyObj.Get()))
+		if (!PyConversion::PythonizeProperty(InScriptMapHelper.GetKeyProperty(), InScriptMapHelper.GetKeyPtr(InValidatedInternalIndex), PyKeyObj.Get()))
 		{
 			PyUtil::SetPythonError(PyExc_TypeError, InSelf, *FString::Printf(TEXT("Failed to convert key property '%s' (%s)"), *InScriptMapHelper.GetKeyProperty()->GetName(), *InScriptMapHelper.GetKeyProperty()->GetClass()->GetName()));
 			return nullptr;
 		}
 
 		FPyObjectPtr PyValueObj;
-		if (!PyConversion::PythonizeProperty(InScriptMapHelper.GetValueProperty(), InScriptMapHelper.GetValuePtr(InElementIndex), PyValueObj.Get()))
+		if (!PyConversion::PythonizeProperty(InScriptMapHelper.GetValueProperty(), InScriptMapHelper.GetValuePtr(InValidatedInternalIndex), PyValueObj.Get()))
 		{
 			PyUtil::SetPythonError(PyExc_TypeError, InSelf, *FString::Printf(TEXT("Failed to convert value property '%s' (%s)"), *InScriptMapHelper.GetValueProperty()->GetName(), *InScriptMapHelper.GetValueProperty()->GetClass()->GetName()));
 			return nullptr;
@@ -250,10 +250,10 @@ struct FPyWrapperMapKeyIterator : public TPyWrapperMapIterator<FPyWrapperMapKeyI
 		return &PyWrapperMapKeyIteratorType;
 	}
 
-	static PyObject* GetItem(FPyWrapperMapKeyIterator* InSelf, FScriptMapHelper& InScriptMapHelper, int32 InElementIndex)
+	static PyObject* GetItem(FPyWrapperMapKeyIterator* InSelf, FScriptMapHelper& InScriptMapHelper, const int32 InValidatedInternalIndex)
 	{
 		PyObject* PyItemObj = nullptr;
-		if (!PyConversion::PythonizeProperty(InScriptMapHelper.GetKeyProperty(), InScriptMapHelper.GetKeyPtr(InElementIndex), PyItemObj))
+		if (!PyConversion::PythonizeProperty(InScriptMapHelper.GetKeyProperty(), InScriptMapHelper.GetKeyPtr(InValidatedInternalIndex), PyItemObj))
 		{
 			PyUtil::SetPythonError(PyExc_TypeError, InSelf, *FString::Printf(TEXT("Failed to convert key property '%s' (%s)"), *InScriptMapHelper.GetKeyProperty()->GetName(), *InScriptMapHelper.GetKeyProperty()->GetClass()->GetName()));
 			return nullptr;
@@ -270,10 +270,10 @@ struct FPyWrapperMapValueIterator : public TPyWrapperMapIterator<FPyWrapperMapVa
 		return &PyWrapperMapValueIteratorType;
 	}
 
-	static PyObject* GetItem(FPyWrapperMapValueIterator* InSelf, FScriptMapHelper& InScriptMapHelper, int32 InElementIndex)
+	static PyObject* GetItem(FPyWrapperMapValueIterator* InSelf, FScriptMapHelper& InScriptMapHelper, const int32 InValidatedInternalIndex)
 	{
 		PyObject* PyItemObj = nullptr;
-		if (!PyConversion::PythonizeProperty(InScriptMapHelper.GetValueProperty(), InScriptMapHelper.GetValuePtr(InElementIndex), PyItemObj))
+		if (!PyConversion::PythonizeProperty(InScriptMapHelper.GetValueProperty(), InScriptMapHelper.GetValuePtr(InValidatedInternalIndex), PyItemObj))
 		{
 			PyUtil::SetPythonError(PyExc_TypeError, InSelf, *FString::Printf(TEXT("Failed to convert value property '%s' (%s)"), *InScriptMapHelper.GetValueProperty()->GetName(), *InScriptMapHelper.GetValueProperty()->GetClass()->GetName()));
 			return nullptr;
@@ -661,47 +661,38 @@ FPyWrapperMap* FPyWrapperMap::CastPyObject(PyObject* InPyObject, PyTypeObject* I
 
 		// Attempt to convert the entries in the map to the native format of the new map
 		{
-			FScriptMapHelper SelfScriptMapHelper(Self->MapProp, Self->MapInstance);
+			const FScriptMapHelper SelfScriptMapHelper(Self->MapProp, Self->MapInstance);
 			FScriptMapHelper NewScriptMapHelper(NewMap->MapProp, NewMap->MapInstance);
 
-			const int32 ElementCount = SelfScriptMapHelper.Num();
-
 			FString ExportedKey;
 			FString ExportedValue;
-			for (int32 ElementIndex = 0, SparseIndex = 0; ElementIndex < ElementCount; ++SparseIndex)
+			for (FScriptMapHelper::FIterator It(SelfScriptMapHelper); It; ++It)
 			{
-				if (!SelfScriptMapHelper.IsValidIndex(SparseIndex))
-				{
-					continue;
-				}
-
 				ExportedKey.Reset();
-				if (!SelfScriptMapHelper.GetKeyProperty()->ExportText_Direct(ExportedKey, SelfScriptMapHelper.GetKeyPtr(SparseIndex), SelfScriptMapHelper.GetKeyPtr(SparseIndex), nullptr, PPF_None))
+				if (!SelfScriptMapHelper.GetKeyProperty()->ExportText_Direct(ExportedKey, SelfScriptMapHelper.GetKeyPtr(It), SelfScriptMapHelper.GetKeyPtr(It), nullptr, PPF_None))
 				{
-					PyUtil::SetPythonError(PyExc_Exception, Self, *FString::Printf(TEXT("Failed to export text for key property '%s' (%s) at index %d"), *SelfScriptMapHelper.GetKeyProperty()->GetName(), *SelfScriptMapHelper.GetKeyProperty()->GetClass()->GetName(), ElementIndex));
+					PyUtil::SetPythonError(PyExc_Exception, Self, *FString::Printf(TEXT("Failed to export text for key property '%s' (%s) at index %d"), *SelfScriptMapHelper.GetKeyProperty()->GetName(), *SelfScriptMapHelper.GetKeyProperty()->GetClass()->GetName(), It.GetLogicalIndex()));
 					return nullptr;
 				}
 
 				ExportedValue.Reset();
-				if (!SelfScriptMapHelper.GetValueProperty()->ExportText_Direct(ExportedValue, SelfScriptMapHelper.GetValuePtr(SparseIndex), SelfScriptMapHelper.GetValuePtr(SparseIndex), nullptr, PPF_None))
+				if (!SelfScriptMapHelper.GetValueProperty()->ExportText_Direct(ExportedValue, SelfScriptMapHelper.GetValuePtr(It), SelfScriptMapHelper.GetValuePtr(It), nullptr, PPF_None))
 				{
-					PyUtil::SetPythonError(PyExc_Exception, Self, *FString::Printf(TEXT("Failed to export text for value property '%s' (%s) at index %d"), *SelfScriptMapHelper.GetValueProperty()->GetName(), *SelfScriptMapHelper.GetValueProperty()->GetClass()->GetName(), ElementIndex));
+					PyUtil::SetPythonError(PyExc_Exception, Self, *FString::Printf(TEXT("Failed to export text for value property '%s' (%s) at index %d"), *SelfScriptMapHelper.GetValueProperty()->GetName(), *SelfScriptMapHelper.GetValueProperty()->GetClass()->GetName(), It.GetLogicalIndex()));
 					return nullptr;
 				}
 
 				const int32 NewElementIndex = NewScriptMapHelper.AddDefaultValue_Invalid_NeedsRehash();
 				if (!NewScriptMapHelper.GetKeyProperty()->ImportText_Direct(*ExportedKey, NewScriptMapHelper.GetKeyPtr(NewElementIndex), nullptr, PPF_None))
 				{
-					PyUtil::SetPythonError(PyExc_Exception, Self, *FString::Printf(TEXT("Failed to import text '%s' key for property '%s' (%s) at index %d"), *ExportedKey, *NewScriptMapHelper.GetKeyProperty()->GetName(), *NewScriptMapHelper.GetKeyProperty()->GetClass()->GetName(), ElementIndex));
+					PyUtil::SetPythonError(PyExc_Exception, Self, *FString::Printf(TEXT("Failed to import text '%s' key for property '%s' (%s) at index %d"), *ExportedKey, *NewScriptMapHelper.GetKeyProperty()->GetName(), *NewScriptMapHelper.GetKeyProperty()->GetClass()->GetName(), NewElementIndex));
 					return nullptr;
 				}
 				if (!NewScriptMapHelper.GetValueProperty()->ImportText_Direct(*ExportedValue, NewScriptMapHelper.GetValuePtr(NewElementIndex), nullptr, PPF_None))
 				{
-					PyUtil::SetPythonError(PyExc_Exception, Self, *FString::Printf(TEXT("Failed to import text '%s' value for property '%s' (%s) at index %d"), *ExportedValue, *NewScriptMapHelper.GetValueProperty()->GetName(), *NewScriptMapHelper.GetValueProperty()->GetClass()->GetName(), ElementIndex));
+					PyUtil::SetPythonError(PyExc_Exception, Self, *FString::Printf(TEXT("Failed to import text '%s' value for property '%s' (%s) at index %d"), *ExportedValue, *NewScriptMapHelper.GetValueProperty()->GetName(), *NewScriptMapHelper.GetValueProperty()->GetClass()->GetName(), NewElementIndex));
 					return nullptr;
 				}
-
-				++ElementIndex;
 			}
 
 			NewScriptMapHelper.Rehash();
@@ -1041,39 +1032,30 @@ PyObject* FPyWrapperMap::PopItem(FPyWrapperMap* InSelf)
 	}
 
 	FScriptMapHelper SelfScriptMapHelper(InSelf->MapProp, InSelf->MapInstance);
-	const int32 SelfElementCount = SelfScriptMapHelper.Num();
-
-	if (SelfElementCount == 0)
+	const FScriptMapHelper::FIterator It(SelfScriptMapHelper);
+	if (It)
 	{
-		PyUtil::SetPythonError(PyExc_KeyError, InSelf, TEXT("Cannot pop from an empty map"));
-		return nullptr;
-	}
-
-	for (int32 SelfSparseIndex = 0; ; ++SelfSparseIndex)
-	{
-		if (!SelfScriptMapHelper.IsValidIndex(SelfSparseIndex))
-		{
-			continue;
-		}
-
 		FPyObjectPtr PyReturnKey;
-		if (!PyConversion::PythonizeProperty(SelfScriptMapHelper.GetKeyProperty(), SelfScriptMapHelper.GetKeyPtr(SelfSparseIndex), PyReturnKey.Get()))
+		if (!PyConversion::PythonizeProperty(SelfScriptMapHelper.GetKeyProperty(), SelfScriptMapHelper.GetKeyPtr(It), PyReturnKey.Get()))
 		{
 			PyUtil::SetPythonError(PyExc_TypeError, InSelf, *FString::Printf(TEXT("Failed to convert key property '%s' (%s) at index 0"), *SelfScriptMapHelper.GetKeyProperty()->GetName(), *SelfScriptMapHelper.GetKeyProperty()->GetClass()->GetName()));
 			return nullptr;
 		}
 
 		FPyObjectPtr PyReturnValue;
-		if (!PyConversion::PythonizeProperty(SelfScriptMapHelper.GetValueProperty(), SelfScriptMapHelper.GetValuePtr(SelfSparseIndex), PyReturnValue.Get()))
+		if (!PyConversion::PythonizeProperty(SelfScriptMapHelper.GetValueProperty(), SelfScriptMapHelper.GetValuePtr(It), PyReturnValue.Get()))
 		{
 			PyUtil::SetPythonError(PyExc_TypeError, InSelf, *FString::Printf(TEXT("Failed to convert value property '%s' (%s) at index 0"), *SelfScriptMapHelper.GetValueProperty()->GetName(), *SelfScriptMapHelper.GetValueProperty()->GetClass()->GetName()));
 			return nullptr;
 		}
 
-		SelfScriptMapHelper.RemoveAt(SelfSparseIndex);
+		SelfScriptMapHelper.RemoveAt(It.GetInternalIndex());
 		
 		return PyTuple_Pack(2, PyReturnKey.Get(), PyReturnValue.Get());
 	}
+
+	PyUtil::SetPythonError(PyExc_KeyError, InSelf, TEXT("Cannot pop from an empty map"));
+	return nullptr;
 }
 
 int FPyWrapperMap::Update(FPyWrapperMap* InSelf, PyObject* InOther)
@@ -1093,17 +1075,11 @@ int FPyWrapperMap::Update(FPyWrapperMap* InSelf, PyObject* InOther)
 	}
 
 	FScriptMapHelper SelfScriptMapHelper(InSelf->MapProp, InSelf->MapInstance);
-	FScriptMapHelper OtherScriptMapHelper(Other->MapProp, Other->MapInstance);
+	const FScriptMapHelper OtherScriptMapHelper(Other->MapProp, Other->MapInstance);
 
-	const int32 OtherSparseCount = OtherScriptMapHelper.Num();
-	for (int32 OtherSparseIndex = 0; OtherSparseIndex < OtherSparseCount; ++OtherSparseIndex)
+	for (FScriptMapHelper::FIterator It(OtherScriptMapHelper); It; ++It)
 	{
-		if (!OtherScriptMapHelper.IsValidIndex(OtherSparseIndex))
-		{
-			continue;
-		}
-
-		SelfScriptMapHelper.AddPair(OtherScriptMapHelper.GetKeyPtr(OtherSparseIndex), OtherScriptMapHelper.GetValuePtr(OtherSparseIndex));
+		SelfScriptMapHelper.AddPair(OtherScriptMapHelper.GetKeyPtr(It), OtherScriptMapHelper.GetValuePtr(It));
 	}
 
 	return 0;
@@ -1238,25 +1214,18 @@ PyTypeObject InitializePyWrapperMapType()
 				return nullptr;
 			}
 
-			FScriptMapHelper SelfScriptMapHelper(InSelf->MapProp, InSelf->MapInstance);
-			const int32 ElementCount = SelfScriptMapHelper.Num();
+			const FScriptMapHelper SelfScriptMapHelper(InSelf->MapProp, InSelf->MapInstance);
 
 			FString ExportedMap;
-			for (int32 ElementIndex = 0, SparseIndex = 0; ElementIndex < ElementCount; ++SparseIndex)
+			for (FScriptMapHelper::FIterator It(SelfScriptMapHelper); It; ++It)
 			{
-				if (!SelfScriptMapHelper.IsValidIndex(SparseIndex))
-				{
-					continue;
-				}
-
-				if (ElementIndex > 0)
+				if (It.GetLogicalIndex() > 0)
 				{
 					ExportedMap += TEXT(", ");
 				}
-				ExportedMap += PyUtil::GetFriendlyPropertyValue(SelfScriptMapHelper.GetKeyProperty(), SelfScriptMapHelper.GetKeyPtr(SparseIndex), PPF_Delimited | PPF_IncludeTransient);
+				ExportedMap += PyUtil::GetFriendlyPropertyValue(SelfScriptMapHelper.GetKeyProperty(), SelfScriptMapHelper.GetKeyPtr(It), PPF_Delimited | PPF_IncludeTransient);
 				ExportedMap += TEXT(": ");
-				ExportedMap += PyUtil::GetFriendlyPropertyValue(SelfScriptMapHelper.GetValueProperty(), SelfScriptMapHelper.GetValuePtr(SparseIndex), PPF_Delimited | PPF_IncludeTransient);
-				++ElementIndex;
+				ExportedMap += PyUtil::GetFriendlyPropertyValue(SelfScriptMapHelper.GetValueProperty(), SelfScriptMapHelper.GetValuePtr(It), PPF_Delimited | PPF_IncludeTransient);
 			}
 			return PyUnicode_FromFormat("{%s}", TCHAR_TO_UTF8(*ExportedMap));
 		}
diff --git a/Engine/Plugins/Experimental/PythonScriptPlugin/Source/PythonScriptPlugin/Private/PyWrapperSet.cpp b/Engine/Plugins/Experimental/PythonScriptPlugin/Source/PythonScriptPlugin/Private/PyWrapperSet.cpp
--- a/Engine/Plugins/Experimental/PythonScriptPlugin/Source/PythonScriptPlugin/Private/PyWrapperSet.cpp
+++ b/Engine/Plugins/Experimental/PythonScriptPlugin/Source/PythonScriptPlugin/Private/PyWrapperSet.cpp
@@ -381,31 +381,22 @@ FPyWrapperSet* FPyWrapperSet::CastPyObject(PyObject* InPyObject, PyTypeObject* I
 			FScriptSetHelper SelfScriptSetHelper(Self->SetProp, Self->SetInstance);
 			FScriptSetHelper NewScriptSetHelper(NewSet->SetProp, NewSet->SetInstance);
 
-			const int32 ElementCount = SelfScriptSetHelper.Num();
-
 			FString ExportedEntry;
-			for (int32 ElementIndex = 0, SparseIndex = 0; ElementIndex < ElementCount; ++SparseIndex)
+			for (FScriptSetHelper::FIterator It(SelfScriptSetHelper); It; ++It)
 			{
-				if (!SelfScriptSetHelper.IsValidIndex(SparseIndex))
-				{
-					continue;
-				}
-
 				ExportedEntry.Reset();
-				if (!SelfScriptSetHelper.GetElementProperty()->ExportText_Direct(ExportedEntry, SelfScriptSetHelper.GetElementPtr(SparseIndex), SelfScriptSetHelper.GetElementPtr(SparseIndex), nullptr, PPF_None))
+				if (!SelfScriptSetHelper.GetElementProperty()->ExportText_Direct(ExportedEntry, SelfScriptSetHelper.GetElementPtr(It), SelfScriptSetHelper.GetElementPtr(It), nullptr, PPF_None))
 				{
-					PyUtil::SetPythonError(PyExc_Exception, Self, *FString::Printf(TEXT("Failed to export text for element property '%s' (%s) at index %d"), *SelfScriptSetHelper.GetElementProperty()->GetName(), *SelfScriptSetHelper.GetElementProperty()->GetClass()->GetName(), ElementIndex));
+					PyUtil::SetPythonError(PyExc_Exception, Self, *FString::Printf(TEXT("Failed to export text for element property '%s' (%s) at index %d"), *SelfScriptSetHelper.GetElementProperty()->GetName(), *SelfScriptSetHelper.GetElementProperty()->GetClass()->GetName(), It.GetLogicalIndex()));
 					return nullptr;
 				}
 
 				const int32 NewElementIndex = NewScriptSetHelper.AddDefaultValue_Invalid_NeedsRehash();
 				if (!NewScriptSetHelper.GetElementProperty()->ImportText_Direct(*ExportedEntry, NewScriptSetHelper.GetElementPtr(NewElementIndex), nullptr, PPF_None))
 				{
-					PyUtil::SetPythonError(PyExc_Exception, Self, *FString::Printf(TEXT("Failed to import text '%s' element for property '%s' (%s) at index %d"), *ExportedEntry, *NewScriptSetHelper.GetElementProperty()->GetName(), *NewScriptSetHelper.GetElementProperty()->GetClass()->GetName(), ElementIndex));
+					PyUtil::SetPythonError(PyExc_Exception, Self, *FString::Printf(TEXT("Failed to import text '%s' element for property '%s' (%s) at index %d"), *ExportedEntry, *NewScriptSetHelper.GetElementProperty()->GetName(), *NewScriptSetHelper.GetElementProperty()->GetClass()->GetName(), NewElementIndex));
 					return nullptr;
 				}
-
-				++ElementIndex;
 			}
 
 			NewScriptSetHelper.Rehash();
@@ -563,32 +554,23 @@ PyObject* FPyWrapperSet::Pop(FPyWrapperSet* InSelf)
 	}
 
 	FScriptSetHelper SelfScriptSetHelper(InSelf->SetProp, InSelf->SetInstance);
-	const int32 SelfElementCount = SelfScriptSetHelper.Num();
-
-	if (SelfElementCount == 0)
+	const FScriptSetHelper::FIterator It(SelfScriptSetHelper);
+	if (It)
 	{
-		PyUtil::SetPythonError(PyExc_KeyError, InSelf, TEXT("Cannot pop from an empty set"));
-		return nullptr;
-	}
-
-	for (int32 SelfSparseIndex = 0; ; ++SelfSparseIndex)
-	{
-		if (!SelfScriptSetHelper.IsValidIndex(SelfSparseIndex))
-		{
-			continue;
-		}
-
 		PyObject* PyReturnValue = nullptr;
-		if (!PyConversion::PythonizeProperty(SelfScriptSetHelper.GetElementProperty(), SelfScriptSetHelper.GetElementPtr(SelfSparseIndex), PyReturnValue))
+		if (!PyConversion::PythonizeProperty(SelfScriptSetHelper.GetElementProperty(), SelfScriptSetHelper.GetElementPtr(It), PyReturnValue))
 		{
 			PyUtil::SetPythonError(PyExc_TypeError, InSelf, *FString::Printf(TEXT("Failed to convert element property '%s' (%s) at index 0"), *SelfScriptSetHelper.GetElementProperty()->GetName(), *SelfScriptSetHelper.GetElementProperty()->GetClass()->GetName()));
 			return nullptr;
 		}
 
-		SelfScriptSetHelper.RemoveAt(SelfSparseIndex);
+		SelfScriptSetHelper.RemoveAt(It.GetInternalIndex());
 
 		return PyReturnValue;
 	}
+
+	PyUtil::SetPythonError(PyExc_KeyError, InSelf, TEXT("Cannot pop from an empty set"));
+	return nullptr;
 }
 
 int FPyWrapperSet::Clear(FPyWrapperSet* InSelf)
@@ -639,17 +621,10 @@ int FPyWrapperSet::DifferenceUpdate(FPyWrapperSet* InSelf, PyObject* InOthers)
 			return false;
 		}
 
-		FScriptSetHelper OtherScriptSetHelper(Other->SetProp, Other->SetInstance);
-		const int32 OtherSparseCount = OtherScriptSetHelper.GetMaxIndex();
-
-		for (int32 OtherSparseIndex = 0; OtherSparseIndex < OtherSparseCount; ++OtherSparseIndex)
+		const FScriptSetHelper OtherScriptSetHelper(Other->SetProp, Other->SetInstance);
+		for (FScriptSetHelper::FIterator It(OtherScriptSetHelper); It; ++It)
 		{
-			if (!OtherScriptSetHelper.IsValidIndex(OtherSparseIndex))
-			{
-				continue;
-			}
-
-			const void* OtherElementPtr = OtherScriptSetHelper.GetElementPtr(OtherSparseIndex);
+			const void* OtherElementPtr = OtherScriptSetHelper.GetElementPtr(It);
 			SelfScriptSetHelper.RemoveElement(OtherElementPtr);
 		}
 
@@ -795,18 +770,12 @@ int FPyWrapperSet::SymmetricDifferenceUpdate(FPyWrapperSet* InSelf, PyObject* In
 		return -1;
 	}
 
-	FScriptSetHelper OtherScriptSetHelper(Other->SetProp, Other->SetInstance);
-	const int32 OtherSparseCount = OtherScriptSetHelper.GetMaxIndex();
+	const FScriptSetHelper OtherScriptSetHelper(Other->SetProp, Other->SetInstance);
 
 	// We need to go through the other set and remove any values from Self that are present in Other, and add any values from Other that aren't present in Self
-	for (int32 OtherSparseIndex = 0; OtherSparseIndex < OtherSparseCount; ++OtherSparseIndex)
+	for (FScriptSetHelper::FIterator It(OtherScriptSetHelper); It; ++It)
 	{
-		if (!OtherScriptSetHelper.IsValidIndex(OtherSparseIndex))
-		{
-			continue;
-		}
-
-		const void* OtherElementPtr = OtherScriptSetHelper.GetElementPtr(OtherSparseIndex);
+		const void* OtherElementPtr = OtherScriptSetHelper.GetElementPtr(It);
 		if (SelfScriptSetHelper.FindElementIndexFromHash(OtherElementPtr) == INDEX_NONE)
 		{
 			SelfScriptSetHelper.AddElement(OtherElementPtr);
@@ -855,17 +824,10 @@ int FPyWrapperSet::Update(FPyWrapperSet* InSelf, PyObject* InOthers)
 			return false;
 		}
 
-		FScriptSetHelper OtherScriptSetHelper(Other->SetProp, Other->SetInstance);
-		const int32 OtherSparseCount = OtherScriptSetHelper.GetMaxIndex();
-
-		for (int32 OtherSparseIndex = 0; OtherSparseIndex < OtherSparseCount; ++OtherSparseIndex)
+		const FScriptSetHelper OtherScriptSetHelper(Other->SetProp, Other->SetInstance);
+		for (FScriptSetHelper::FIterator It(OtherScriptSetHelper); It; ++It)
 		{
-			if (!OtherScriptSetHelper.IsValidIndex(OtherSparseIndex))
-			{
-				continue;
-			}
-
-			const void* OtherElementPtr = OtherScriptSetHelper.GetElementPtr(OtherSparseIndex);
+			const void* OtherElementPtr = OtherScriptSetHelper.GetElementPtr(It);
 			SelfScriptSetHelper.AddElement(OtherElementPtr);
 		}
 
@@ -934,19 +896,12 @@ int FPyWrapperSet::IsSubset(FPyWrapperSet* InSelf, PyObject* InOther)
 		return -1;
 	}
 
-	FScriptSetHelper SelfScriptSetHelper(InSelf->SetProp, InSelf->SetInstance);
-	const int32 SelfSparseCount = SelfScriptSetHelper.Num();
+	const FScriptSetHelper SelfScriptSetHelper(InSelf->SetProp, InSelf->SetInstance);
+	const FScriptSetHelper OtherScriptSetHelper(Other->SetProp, Other->SetInstance);
 
-	FScriptSetHelper OtherScriptSetHelper(Other->SetProp, Other->SetInstance);
-
-	for (int32 SelfSparseIndex = 0; SelfSparseIndex < SelfSparseCount; ++SelfSparseIndex)
+	for (FScriptSetHelper::FIterator It(SelfScriptSetHelper); It; ++It)
 	{
-		if (!SelfScriptSetHelper.IsValidIndex(SelfSparseIndex))
-		{
-			continue;
-		}
-
-		const void* SelfElementPtr = SelfScriptSetHelper.GetElementPtr(SelfSparseIndex);
+		const void* SelfElementPtr = SelfScriptSetHelper.GetElementPtr(It);
 		if (OtherScriptSetHelper.FindElementIndexFromHash(SelfElementPtr) == INDEX_NONE)
 		{
 			return 0;
@@ -971,19 +926,12 @@ int FPyWrapperSet::IsSuperset(FPyWrapperSet* InSelf, PyObject* InOther)
 		return -1;
 	}
 
-	FScriptSetHelper SelfScriptSetHelper(InSelf->SetProp, InSelf->SetInstance);
+	const FScriptSetHelper SelfScriptSetHelper(InSelf->SetProp, InSelf->SetInstance);
+	const FScriptSetHelper OtherScriptSetHelper(Other->SetProp, Other->SetInstance);
 
-	FScriptSetHelper OtherScriptSetHelper(Other->SetProp, Other->SetInstance);
-	const int32 OtherSparseCount = OtherScriptSetHelper.Num();
-
-	for (int32 OtherSparseIndex = 0; OtherSparseIndex < OtherSparseCount; ++OtherSparseIndex)
+	for (FScriptSetHelper::FIterator It(OtherScriptSetHelper); It; ++It)
 	{
-		if (!OtherScriptSetHelper.IsValidIndex(OtherSparseIndex))
-		{
-			continue;
-		}
-
-		const void* OtherElementPtr = OtherScriptSetHelper.GetElementPtr(OtherSparseIndex);
+		const void* OtherElementPtr = OtherScriptSetHelper.GetElementPtr(It);
 		if (SelfScriptSetHelper.FindElementIndexFromHash(OtherElementPtr) == INDEX_NONE)
 		{
 			return 0;
@@ -1034,23 +982,16 @@ PyTypeObject InitializePyWrapperSetType()
 				return nullptr;
 			}
 
-			FScriptSetHelper SelfScriptSetHelper(InSelf->SetProp, InSelf->SetInstance);
-			const int32 ElementCount = SelfScriptSetHelper.Num();
+			const FScriptSetHelper SelfScriptSetHelper(InSelf->SetProp, InSelf->SetInstance);
 
 			FString ExportedSet;
-			for (int32 ElementIndex = 0, SparseIndex = 0; ElementIndex < ElementCount; ++SparseIndex)
+			for (FScriptSetHelper::FIterator It(SelfScriptSetHelper); It; ++It)
 			{
-				if (!SelfScriptSetHelper.IsValidIndex(SparseIndex))
-				{
-					continue;
-				}
-
-				if (ElementIndex > 0)
+				if (It.GetLogicalIndex() > 0)
 				{
 					ExportedSet += TEXT(", ");
 				}
-				ExportedSet += PyUtil::GetFriendlyPropertyValue(SelfScriptSetHelper.GetElementProperty(), SelfScriptSetHelper.GetElementPtr(SparseIndex), PPF_Delimited | PPF_IncludeTransient);
-				++ElementIndex;
+				ExportedSet += PyUtil::GetFriendlyPropertyValue(SelfScriptSetHelper.GetElementProperty(), SelfScriptSetHelper.GetElementPtr(It), PPF_Delimited | PPF_IncludeTransient);
 			}
 			return PyUnicode_FromFormat("set([%s])", TCHAR_TO_UTF8(*ExportedSet));
 		}
diff --git a/Engine/Source/Editor/KismetCompiler/Private/KismetCompilerVMBackend.cpp b/Engine/Source/Editor/KismetCompiler/Private/KismetCompilerVMBackend.cpp
--- a/Engine/Source/Editor/KismetCompiler/Private/KismetCompilerVMBackend.cpp
+++ b/Engine/Source/Editor/KismetCompiler/Private/KismetCompilerVMBackend.cpp
@@ -1018,15 +1018,10 @@ public:
 				Writer << InnerProp;
 				Writer << ElementNum;
 
-				for (int32 ElemIdx = 0, SparseIndex = 0; ElemIdx < ElementNum; ++SparseIndex)
+				for (FScriptSetHelper::FIterator It(ScriptSetHelper); It; ++It)
 				{
-					if (ScriptSet.IsValidIndex(SparseIndex))
-					{
-						uint8* RawElemData = ScriptSetHelper.GetElementPtr(SparseIndex);
-						EmitInnerElementExpr(Term, InnerProp, RawElemData);
-
-						++ElemIdx;
-					}
+					uint8* RawElemData = ScriptSetHelper.GetElementPtr(It);
+					EmitInnerElementExpr(Term, InnerProp, RawElemData);
 				}
 				Writer << EX_EndSetConst;
 			}
@@ -1047,15 +1042,10 @@ public:
 				Writer << ValProp;
 				Writer << ElementNum;
 
-				for (int32 ElemIdx = 0, SparseIndex = 0; ElemIdx < ElementNum; ++SparseIndex)
+				for (FScriptMapHelper::FIterator It(ScriptMapHelper); It; ++It)
 				{
-					if (ScriptMap.IsValidIndex(SparseIndex))
-					{
-						EmitInnerElementExpr(Term, KeyProp, ScriptMapHelper.GetKeyPtr(SparseIndex));
-						EmitInnerElementExpr(Term, ValProp, ScriptMapHelper.GetValuePtr(SparseIndex));
-
-						++ElemIdx;
-					}
+					EmitInnerElementExpr(Term, KeyProp, ScriptMapHelper.GetKeyPtr(It));
+					EmitInnerElementExpr(Term, ValProp, ScriptMapHelper.GetValuePtr(It));
 				}
 				Writer << EX_EndMapConst;
 			}
diff --git a/Engine/Source/Editor/PropertyEditor/Private/PropertyHandleImpl.cpp b/Engine/Source/Editor/PropertyEditor/Private/PropertyHandleImpl.cpp
--- a/Engine/Source/Editor/PropertyEditor/Private/PropertyHandleImpl.cpp
+++ b/Engine/Source/Editor/PropertyEditor/Private/PropertyHandleImpl.cpp
@@ -366,18 +366,13 @@ FPropertyAccess::Result FPropertyValueImpl::ImportText( const TArray<FObjectBase
 							ElementProp->DestroyAndFreeValue(TempElementStorage);
 						};
 
-						for (int32 Index = 0, ItemsLeft = Helper.Num(); ItemsLeft > 0; ++Index)
+						for (FScriptSetHelper::FIterator It(Helper); It; ++It)
 						{
-							if (Helper.IsValidIndex(Index))
-							{
-								--ItemsLeft;
-
-								const uint8* Element = Helper.GetElementPtr(Index);
+							const uint8* Element = Helper.GetElementPtr(It);
 
-								if (Element != InBaseAddress && ElementProp->ImportText_Direct(*InElementValue, TempElementStorage, nullptr, 0) && ElementProp->Identical(Element, TempElementStorage))
-								{
-									return true;
-								}
+							if (Element != InBaseAddress && ElementProp->ImportText_Direct(*InElementValue, TempElementStorage, nullptr, 0) && ElementProp->Identical(Element, TempElementStorage))
+							{
+								return true;
 							}
 						}
 
@@ -422,19 +417,14 @@ FPropertyAccess::Result FPropertyValueImpl::ImportText( const TArray<FObjectBase
 							KeyProp->DestroyAndFreeValue(TempKeyStorage);
 						};
 
-						for (int32 Index = 0, ItemsLeft = Helper.Num(); ItemsLeft > 0; ++Index)
+						for (FScriptMapHelper::FIterator It(Helper); It; ++It)
 						{
-							if (Helper.IsValidIndex(Index))
-							{
-								--ItemsLeft;
-
-								const uint8* PairPtr = Helper.GetPairPtr(Index);
-								const uint8* KeyPtr = KeyProp->ContainerPtrToValuePtr<const uint8>(PairPtr);
+							const uint8* PairPtr = Helper.GetPairPtr(It);
+							const uint8* KeyPtr = KeyProp->ContainerPtrToValuePtr<const uint8>(PairPtr);
 
-								if (KeyPtr != InBaseAddress && KeyProp->ImportText_Direct(*InKeyValue, TempKeyStorage, nullptr, 0) && KeyProp->Identical(KeyPtr, TempKeyStorage))
-								{
-									return true;
-								}
+							if (KeyPtr != InBaseAddress && KeyProp->ImportText_Direct(*InKeyValue, TempKeyStorage, nullptr, 0) && KeyProp->Identical(KeyPtr, TempKeyStorage))
+							{
+								return true;
 							}
 						}
 
@@ -1339,20 +1329,13 @@ void FPropertyValueImpl::ClearChildren()
 							FObjectProperty* ElementObjectProperty = CastField<FObjectProperty>(SetProperty->ElementProp);
 							if (ShouldOwnInstance(ElementObjectProperty, PropertyNodePin.Get()))
 							{
-								int32 ElementsToRemove = SetHelper.Num();
-								int32 Index = 0;
-								while (ElementsToRemove > 0)
+								for (FScriptSetHelper::FIterator It(SetHelper); It; ++It)
 								{
-									if (SetHelper.IsValidIndex(Index))
+									if (UObject* InstancedObject = *reinterpret_cast<UObject**>(SetHelper.GetElementPtr(It)))
 									{
-										if (UObject* InstancedObject = *reinterpret_cast<UObject**>(SetHelper.GetElementPtr(Index)))
-										{
-											InstancedObject->Modify();
-											InstancedObject->Rename(nullptr, GetTransientPackage(), REN_DontCreateRedirectors);
-										}
-										--ElementsToRemove;
+										InstancedObject->Modify();
+										InstancedObject->Rename(nullptr, GetTransientPackage(), REN_DontCreateRedirectors);
 									}
-									++Index;
 								}
 							}
 
@@ -1367,20 +1350,13 @@ void FPropertyValueImpl::ClearChildren()
 							FObjectProperty* ValueObjectProperty = CastField<FObjectProperty>(MapProperty->ValueProp);
 							if (ShouldOwnInstance(ValueObjectProperty, PropertyNodePin.Get()))
 							{
-								int32 ElementsToRemove = MapHelper.Num();
-								int32 Index = 0;
-								while (ElementsToRemove > 0)
+								for (FScriptMapHelper::FIterator It(MapHelper); It; ++It)
 								{
-									if (MapHelper.IsValidIndex(Index))
+									if (UObject* InstancedObject = *reinterpret_cast<UObject**>(MapHelper.GetValuePtr(It)))
 									{
-										if (UObject* InstancedObject = *reinterpret_cast<UObject**>(MapHelper.GetValuePtr(Index)))
-										{
-											InstancedObject->Modify();
-											InstancedObject->Rename(nullptr, GetTransientPackage(), REN_DontCreateRedirectors);
-										}
-										--ElementsToRemove;
+										InstancedObject->Modify();
+										InstancedObject->Rename(nullptr, GetTransientPackage(), REN_DontCreateRedirectors);
 									}
-									++Index;
 								}
 							}
 
diff --git a/Engine/Source/Editor/PropertyEditor/Private/PropertyNode.cpp b/Engine/Source/Editor/PropertyEditor/Private/PropertyNode.cpp
--- a/Engine/Source/Editor/PropertyEditor/Private/PropertyNode.cpp
+++ b/Engine/Source/Editor/PropertyEditor/Private/PropertyNode.cpp
@@ -2114,19 +2114,12 @@ private:
 		if (MapProp != NULL)
 		{
 			FScriptMapHelper MapHelper(MapProp, PropertyValueAddress);
-
-			int32 ItemsLeft = MapHelper.Num();
-			for (int32 Index = 0; ItemsLeft > 0; ++Index)
+			for (FScriptMapHelper::FIterator It(MapHelper); It; ++It)
 			{
-				if (MapHelper.IsValidIndex(Index))
-				{
-					--ItemsLeft;
-
-					uint8* Data = MapHelper.GetPairPtr(Index);
+				uint8* Data = MapHelper.GetPairPtr(It);
 
-					ProcessProperty(MapProp->KeyProp, MapProp->KeyProp->ContainerPtrToValuePtr<uint8>(Data));
-					ProcessProperty(MapProp->ValueProp, MapProp->ValueProp->ContainerPtrToValuePtr<uint8>(Data));
-				}
+				ProcessProperty(MapProp->KeyProp, MapProp->KeyProp->ContainerPtrToValuePtr<uint8>(Data));
+				ProcessProperty(MapProp->ValueProp, MapProp->ValueProp->ContainerPtrToValuePtr<uint8>(Data));
 			}
 
 			bResult = true;
diff --git a/Engine/Source/Editor/UnrealEd/Private/Kismet2/ComponentEditorUtils.cpp b/Engine/Source/Editor/UnrealEd/Private/Kismet2/ComponentEditorUtils.cpp
--- a/Engine/Source/Editor/UnrealEd/Private/Kismet2/ComponentEditorUtils.cpp
+++ b/Engine/Source/Editor/UnrealEd/Private/Kismet2/ComponentEditorUtils.cpp
@@ -241,17 +241,14 @@ FProperty* FComponentEditorUtils::GetPropertyForEditableNativeComponent(const UA
 				}
 
 				FScriptMapHelper MapHelper(TestProperty, MapPropInstAddress);
-				for (int32 MapSparseIndex = 0; MapSparseIndex < MapHelper.GetMaxIndex(); ++MapSparseIndex)
+				for (FScriptMapHelper::FIterator It(MapHelper); It; ++It)
 				{
 					// For each value in the map (don't bother checking the keys, they won't be what the user can edit in this case)
-					if (MapHelper.IsValidIndex(MapSparseIndex))
+					const uint8* MapValueData = MapHelper.GetValuePtr(It);
+					UObject* ValueElement = MapValProp->GetObjectPropertyValue(MapValueData);
+					if (ValueElement != nullptr && ValueElement->GetFName() == NativeComponent->GetFName())
 					{
-						const uint8* MapValueData = MapHelper.GetValuePtr(MapSparseIndex);						
-						UObject* ValueElement = MapValProp->GetObjectPropertyValue(MapValueData);
-						if (ValueElement != nullptr && ValueElement->GetFName() == NativeComponent->GetFName())
-						{
-							return MapValProp;
-						}
+						return MapValProp;
 					}
 				}
 			}
@@ -271,17 +268,14 @@ FProperty* FComponentEditorUtils::GetPropertyForEditableNativeComponent(const UA
 				
 				// For each item in the set
 				FScriptSetHelper SetHelper(TestProperty, SetPropInstAddress);
-				for (int32 i = 0; i < SetHelper.Num(); ++i)
+				for (FScriptSetHelper::FIterator It(SetHelper); It; ++It)
 				{
-					if (SetHelper.IsValidIndex(i))
+					const uint8* SetValData = SetHelper.GetElementPtr(It);
+					UObject* SetValueElem = SetValProp->GetObjectPropertyValue(SetValData);
+
+					if (SetValueElem != nullptr && SetValueElem->GetFName() == NativeComponent->GetFName())
 					{
-						const uint8* SetValData = SetHelper.GetElementPtr(i);
-						UObject* SetValueElem = SetValProp->GetObjectPropertyValue(SetValData);
-						
-						if (SetValueElem != nullptr && SetValueElem->GetFName() == NativeComponent->GetFName())
-						{
-							return SetValProp;
-						}
+						return SetValProp;
 					}
 				}
 			}
diff --git a/Engine/Source/Runtime/CoreUObject/Private/Internationalization/StabilizeLocalizationKeys.cpp b/Engine/Source/Runtime/CoreUObject/Private/Internationalization/StabilizeLocalizationKeys.cpp
--- a/Engine/Source/Runtime/CoreUObject/Private/Internationalization/StabilizeLocalizationKeys.cpp
+++ b/Engine/Source/Runtime/CoreUObject/Private/Internationalization/StabilizeLocalizationKeys.cpp
@@ -109,23 +109,15 @@ void StabilizeLocalizationKeys::StabilizeLocalizationKeysForProperty(FProperty*
 		{
 			FScriptSetHelper ScriptSetHelper(SetProp, InPropData);
 
-			const int32 ElementCount = ScriptSetHelper.Num();
-			for (int32 RawIndex = 0, ElementIndex = 0; ElementIndex < ElementCount; ++RawIndex)
+			for (FScriptSetHelper::FIterator It(ScriptSetHelper); It; ++It)
 			{
-				if (!ScriptSetHelper.IsValidIndex(RawIndex))
-				{
-					continue;
-				}
-
 				StabilizeLocalizationKeysForProperty(
 					SetProp->ElementProp,
-					ScriptSetHelper.GetElementPtr(RawIndex),
+					ScriptSetHelper.GetElementPtr(It),
 					InNamespace,
-					FString::Printf(TEXT("%s_Index%d"), *PropKeyRoot, ElementIndex),
+					FString::Printf(TEXT("%s_Index%d"), *PropKeyRoot, It.GetLogicalIndex()),
 					/*bAppendPropertyNameToKey*/false
 					);
-
-				++ElementIndex;
 			}
 
 			ScriptSetHelper.Rehash();
@@ -139,21 +131,15 @@ void StabilizeLocalizationKeys::StabilizeLocalizationKeysForProperty(FProperty*
 		{
 			FScriptMapHelper ScriptMapHelper(MapProp, InPropData);
 
-			const int32 ElementCount = ScriptMapHelper.Num();
-			for (int32 RawIndex = 0, ElementIndex = 0; ElementIndex < ElementCount; ++RawIndex)
+			for (FScriptMapHelper::FIterator It(ScriptMapHelper); It; ++It)
 			{
-				if (!ScriptMapHelper.IsValidIndex(RawIndex))
-				{
-					continue;
-				}
-
 				if (ShouldStabilizeInnerProperty(MapProp->KeyProp))
 				{
 					StabilizeLocalizationKeysForProperty(
 						MapProp->KeyProp,
-						ScriptMapHelper.GetKeyPtr(RawIndex),
+						ScriptMapHelper.GetKeyPtr(It),
 						InNamespace,
-						FString::Printf(TEXT("%s_KeyIndex%d"), *PropKeyRoot, ElementIndex),
+						FString::Printf(TEXT("%s_KeyIndex%d"), *PropKeyRoot, It.GetLogicalIndex()),
 						/*bAppendPropertyNameToKey*/false
 						);
 				}
@@ -162,14 +148,12 @@ void StabilizeLocalizationKeys::StabilizeLocalizationKeysForProperty(FProperty*
 				{
 					StabilizeLocalizationKeysForProperty(
 						MapProp->ValueProp,
-						ScriptMapHelper.GetValuePtr(RawIndex),
+						ScriptMapHelper.GetValuePtr(It),
 						InNamespace,
-						FString::Printf(TEXT("%s_ValueIndex%d"), *PropKeyRoot, ElementIndex),
+						FString::Printf(TEXT("%s_ValueIndex%d"), *PropKeyRoot, It.GetLogicalIndex()),
 						/*bAppendPropertyNameToKey*/false
 						);
 				}
-
-				++ElementIndex;
 			}
 
 			if (ShouldStabilizeInnerProperty(MapProp->KeyProp))
diff --git a/Engine/Source/Runtime/CoreUObject/Private/Serialization/PropertyLocalizationDataGathering.cpp b/Engine/Source/Runtime/CoreUObject/Private/Serialization/PropertyLocalizationDataGathering.cpp
--- a/Engine/Source/Runtime/CoreUObject/Private/Serialization/PropertyLocalizationDataGathering.cpp
+++ b/Engine/Source/Runtime/CoreUObject/Private/Serialization/PropertyLocalizationDataGathering.cpp
@@ -443,15 +443,9 @@ void FPropertyLocalizationDataGatherer::GatherLocalizationDataFromChildTextPrope
 
 			// Iterate over all elements of the map.
 			FScriptMapHelper ScriptMapHelper(MapProperty, ElementValueAddress);
-			const int32 ElementCount = ScriptMapHelper.Num();
-			for(int32 j = 0, ElementIndex = 0; ElementIndex < ElementCount; ++j)
+			for (FScriptMapHelper::FIterator It(ScriptMapHelper); It; ++It)
 			{
-				if (!ScriptMapHelper.IsValidIndex(j))
-				{
-					continue;
-				}
-
-				const uint8* MapPairPtr = ScriptMapHelper.GetPairPtr(j);
+				const uint8* MapPairPtr = ScriptMapHelper.GetPairPtr(It);
 
 				if (bGatherMapKey)
 				{
@@ -460,7 +454,7 @@ void FPropertyLocalizationDataGatherer::GatherLocalizationDataFromChildTextPrope
 						PathToInnerElement.Reserve(PathToElement.Len() + 20); // +20 for some slack for the number, braces, and description
 						PathToInnerElement += PathToElement;
 						PathToInnerElement += TEXT('(');
-						PathToInnerElement.AppendInt(ElementIndex);
+						PathToInnerElement.AppendInt(It.GetLogicalIndex());
 						PathToInnerElement += TEXT(" - Key)");
 					}
 
@@ -475,15 +469,13 @@ void FPropertyLocalizationDataGatherer::GatherLocalizationDataFromChildTextPrope
 						PathToInnerElement.Reserve(PathToElement.Len() + 20); // +20 for some slack for the number, braces, and description
 						PathToInnerElement += PathToElement;
 						PathToInnerElement += TEXT('(');
-						PathToInnerElement.AppendInt(ElementIndex);
+						PathToInnerElement.AppendInt(It.GetLogicalIndex());
 						PathToInnerElement += TEXT(" - Value)");
 					}
 
 					const uint8* MapValuePtr = MapPairPtr + MapProperty->MapLayout.ValueOffset;
 					GatherLocalizationDataFromChildTextProperties(PathToInnerElement, MapProperty->ValueProp, MapValuePtr, nullptr, ElementChildPropertyGatherTextFlags);
 				}
-
-				++ElementIndex;
 			}
 		}
 		// Property is a set property.
@@ -491,26 +483,19 @@ void FPropertyLocalizationDataGatherer::GatherLocalizationDataFromChildTextPrope
 		{
 			// Iterate over all elements of the Set.
 			FScriptSetHelper ScriptSetHelper(SetProperty, ElementValueAddress);
-			const int32 ElementCount = ScriptSetHelper.Num();
-			for(int32 j = 0, ElementIndex = 0; ElementIndex < ElementCount; ++j)
+			for (FScriptSetHelper::FIterator It(ScriptSetHelper); It; ++It)
 			{
-				if (!ScriptSetHelper.IsValidIndex(j))
-				{
-					continue;
-				}
-
 				FString PathToInnerElement;
 				{
 					PathToInnerElement.Reserve(PathToElement.Len() + 10); // +10 for some slack for the number and braces
 					PathToInnerElement += PathToElement;
 					PathToInnerElement += TEXT('(');
-					PathToInnerElement.AppendInt(ElementIndex);
+					PathToInnerElement.AppendInt(It.GetLogicalIndex());
 					PathToInnerElement += TEXT(')');
 				}
 
-				const uint8* ElementPtr = ScriptSetHelper.GetElementPtr(j);
+				const uint8* ElementPtr = ScriptSetHelper.GetElementPtr(It);
 				GatherLocalizationDataFromChildTextProperties(PathToInnerElement, SetProperty->ElementProp, ElementPtr, nullptr, ElementChildPropertyGatherTextFlags);
-				++ElementIndex;
 			}
 		}
 		// Property is a struct property.
diff --git a/Engine/Source/Runtime/CoreUObject/Private/UObject/LinkerPlaceholderBase.cpp b/Engine/Source/Runtime/CoreUObject/Private/UObject/LinkerPlaceholderBase.cpp
--- a/Engine/Source/Runtime/CoreUObject/Private/UObject/LinkerPlaceholderBase.cpp
+++ b/Engine/Source/Runtime/CoreUObject/Private/UObject/LinkerPlaceholderBase.cpp
@@ -129,15 +129,10 @@ int32 FLinkerPlaceholderObjectImpl::ResolvePlaceholderValues(const TArray<FField
 			// because we can't know which set entry was set with a reference 
 			// to this object, we have to comb through them all
 			FScriptSetHelper SetHelper(SetProperty, ValueAddress);
-			int32 Num = SetHelper.Num();
-			for (int32 SetIndex = 0; Num; ++SetIndex)
+			for (FScriptSetHelper::FIterator It(SetHelper); It; ++It)
 			{
-				if (SetHelper.IsValidIndex(SetIndex))
-				{
-					--Num;
-					uint8* ElementAddress = SetHelper.GetElementPtr(SetIndex);
-					ReplacementCount += ResolvePlaceholderValues(PropertyChain, PropertyIndex - 1, ElementAddress, OldValue, ReplacementValue);
-				}
+				uint8* ElementAddress = SetHelper.GetElementPtr(It);
+				ReplacementCount += ResolvePlaceholderValues(PropertyChain, PropertyIndex - 1, ElementAddress, OldValue, ReplacementValue);
 			}
 
 			// the above recursive call chewed through the rest of the
@@ -158,23 +153,17 @@ int32 FLinkerPlaceholderObjectImpl::ResolvePlaceholderValues(const TArray<FField
 			// because we can't know which map entry was set with a reference 
 			// to this object, we have to comb through them all
 			FScriptMapHelper MapHelper(MapProperty, ValueAddress);
-			int32 Num = MapHelper.Num();
-			for (int32 MapIndex = 0; Num; ++MapIndex)
+			for (FScriptMapHelper::FIterator It(MapHelper); It; ++It)
 			{
-				if (MapHelper.IsValidIndex(MapIndex))
+				if (NextProperty == MapProperty->KeyProp)
 				{
-					--Num;
-
-					if (NextProperty == MapProperty->KeyProp)
-					{
-						uint8* KeyAddress = MapHelper.GetKeyPtr(MapIndex);
-						ReplacementCount += ResolvePlaceholderValues(PropertyChain, PropertyIndex - 1, KeyAddress, OldValue, ReplacementValue);
-					}
-					else if (NextProperty == MapProperty->ValueProp)
-					{
-						uint8* MapValueAddress = MapHelper.GetValuePtr(MapIndex);
-						ReplacementCount += ResolvePlaceholderValues(PropertyChain, PropertyIndex - 1, MapValueAddress, OldValue, ReplacementValue);
-					}
+					uint8* KeyAddress = MapHelper.GetKeyPtr(It);
+					ReplacementCount += ResolvePlaceholderValues(PropertyChain, PropertyIndex - 1, KeyAddress, OldValue, ReplacementValue);
+				}
+				else if (NextProperty == MapProperty->ValueProp)
+				{
+					uint8* MapValueAddress = MapHelper.GetValuePtr(It);
+					ReplacementCount += ResolvePlaceholderValues(PropertyChain, PropertyIndex - 1, MapValueAddress, OldValue, ReplacementValue);
 				}
 			}
 
diff --git a/Engine/Source/Runtime/CoreUObject/Private/UObject/PropertyMap.cpp b/Engine/Source/Runtime/CoreUObject/Private/UObject/PropertyMap.cpp
--- a/Engine/Source/Runtime/CoreUObject/Private/UObject/PropertyMap.cpp
+++ b/Engine/Source/Runtime/CoreUObject/Private/UObject/PropertyMap.cpp
@@ -19,45 +19,41 @@ namespace UEMapProperty_Private
 	/**
 	 * Checks if any of the pairs in the map compare equal to the one passed.
 	 *
-	 * @param  MapHelper  The map to search through.
-	 * @param  Index      The index in the map to start searching from.
-	 * @param  Num        The number of elements to compare.
+	 * @param  MapHelper    The map to search through.
+	 * @param  LogicalIndex The index in the map to start searching from.
+	 * @param  Num          The number of elements to compare.
 	 */
-	bool AnyEqual(const FScriptMapHelper& MapHelper, int32 Index, int32 Num, const uint8* PairToCompare, uint32 PortFlags)
+	bool AnyEqual(const FScriptMapHelper& MapHelper, const int32 LogicalIndex, int32 Num, const uint8* PairToCompare, const uint32 PortFlags)
 	{
-		FProperty* KeyProp   = MapHelper.GetKeyProperty();
-		FProperty* ValueProp = MapHelper.GetValueProperty();
+		const FProperty* KeyProp   = MapHelper.GetKeyProperty();
+		const FProperty* ValueProp = MapHelper.GetValueProperty();
+		const int32 ValueOffset = MapHelper.MapLayout.ValueOffset;
 
-		int32 ValueOffset = MapHelper.MapLayout.ValueOffset;
-
-		for (; Num; --Num)
+		FScriptMapHelper::FIterator IteratorA(MapHelper, LogicalIndex);
+		for (; IteratorA && Num; --Num, ++IteratorA)
 		{
-			while (!MapHelper.IsValidIndex(Index))
-			{
-				++Index;
-			}
-
-			if (KeyProp->Identical(MapHelper.GetPairPtr(Index), PairToCompare, PortFlags) && ValueProp->Identical(MapHelper.GetPairPtr(Index) + ValueOffset, PairToCompare + ValueOffset, PortFlags))
+			if (KeyProp->Identical(MapHelper.GetPairPtr(IteratorA), PairToCompare, PortFlags) && ValueProp->Identical(MapHelper.GetPairPtr(IteratorA) + ValueOffset, PairToCompare + ValueOffset, PortFlags))
 			{
 				return true;
 			}
-
-			++Index;
 		}
 
 		return false;
 	}
 
-	bool RangesContainSameAmountsOfVal(const FScriptMapHelper& MapHelperA, int32 IndexA, const FScriptMapHelper& MapHelperB, int32 IndexB, int32 Num, const uint8* PairToCompare, uint32 PortFlags)
+	bool RangesContainSameAmountsOfVal(const FScriptMapHelper& MapHelperA, const int32 LogicalIndexA, const FScriptMapHelper& MapHelperB, const int32 LogicalIndexB, int32 Num, const uint8* PairToCompare, const uint32 PortFlags)
 	{
-		FProperty* KeyProp   = MapHelperA.GetKeyProperty();
-		FProperty* ValueProp = MapHelperA.GetValueProperty();
+		const FProperty* KeyProp   = MapHelperA.GetKeyProperty();
+		const FProperty* ValueProp = MapHelperA.GetValueProperty();
 
 		// Ensure that both maps are the same type
 		check(KeyProp   == MapHelperB.GetKeyProperty());
 		check(ValueProp == MapHelperB.GetValueProperty());
 
-		int32 ValueOffset = MapHelperA.MapLayout.ValueOffset;
+		const int32 ValueOffset = MapHelperA.MapLayout.ValueOffset;
+
+		FScriptMapHelper::FIterator IteratorA(MapHelperA, LogicalIndexA);
+		FScriptMapHelper::FIterator IteratorB(MapHelperB, LogicalIndexB);
 
 		int32 CountA = 0;
 		int32 CountB = 0;
@@ -68,18 +64,8 @@ namespace UEMapProperty_Private
 				return CountA == CountB;
 			}
 
-			while (!MapHelperA.IsValidIndex(IndexA))
-			{
-				++IndexA;
-			}
-
-			while (!MapHelperB.IsValidIndex(IndexB))
-			{
-				++IndexB;
-			}
-
-			const uint8* PairA = MapHelperA.GetPairPtr(IndexA);
-			const uint8* PairB = MapHelperB.GetPairPtr(IndexB);
+			const uint8* PairA = MapHelperA.GetPairPtr(IteratorA);
+			const uint8* PairB = MapHelperB.GetPairPtr(IteratorB);
 			if (PairA == PairToCompare || (KeyProp->Identical(PairA, PairToCompare, PortFlags) && ValueProp->Identical(PairA + ValueOffset, PairToCompare + ValueOffset, PortFlags)))
 			{
 				++CountA;
@@ -90,16 +76,16 @@ namespace UEMapProperty_Private
 				++CountB;
 			}
 
-			++IndexA;
-			++IndexB;
+			++IteratorA;
+			++IteratorB;
 			--Num;
 		}
 	}
 
-	bool IsPermutation(const FScriptMapHelper& MapHelperA, const FScriptMapHelper& MapHelperB, uint32 PortFlags)
+	bool IsPermutation(const FScriptMapHelper& MapHelperA, const FScriptMapHelper& MapHelperB, const uint32 PortFlags)
 	{
-		FProperty* KeyProp   = MapHelperA.GetKeyProperty();
-		FProperty* ValueProp = MapHelperA.GetValueProperty();
+		const FProperty* KeyProp   = MapHelperA.GetKeyProperty();
+		const FProperty* ValueProp = MapHelperA.GetValueProperty();
 
 		// Ensure that both maps are the same type
 		check(KeyProp   == MapHelperB.GetKeyProperty());
@@ -111,11 +97,11 @@ namespace UEMapProperty_Private
 			return false;
 		}
 
-		int32 ValueOffset = MapHelperA.MapLayout.ValueOffset;
+		const int32 ValueOffset = MapHelperA.MapLayout.ValueOffset;
 
 		// Skip over common initial sequence
-		int32 IndexA = 0;
-		int32 IndexB = 0;
+		FScriptMapHelper::FIterator IteratorA(MapHelperA);
+		FScriptMapHelper::FIterator IteratorB(MapHelperB);
 		for (;;)
 		{
 			if (Num == 0)
@@ -123,18 +109,8 @@ namespace UEMapProperty_Private
 				return true;
 			}
 
-			while (!MapHelperA.IsValidIndex(IndexA))
-			{
-				++IndexA;
-			}
-
-			while (!MapHelperB.IsValidIndex(IndexB))
-			{
-				++IndexB;
-			}
-
-			const uint8* PairA = MapHelperA.GetPairPtr(IndexA);
-			const uint8* PairB = MapHelperB.GetPairPtr(IndexB);
+			const uint8* PairA = MapHelperA.GetPairPtr(IteratorA);
+			const uint8* PairB = MapHelperB.GetPairPtr(IteratorB);
 			if (!KeyProp->Identical(PairA, PairB, PortFlags))
 			{
 				break;
@@ -145,17 +121,17 @@ namespace UEMapProperty_Private
 				break;
 			}
 
-			++IndexA;
-			++IndexB;
+			++IteratorA;
+			++IteratorB;
 			--Num;
 		}
 
-		int32 FirstIndexA = IndexA;
-		int32 FirstIndexB = IndexB;
-		int32 FirstNum    = Num;
+		const int32 FirstIndexA = IteratorA.GetLogicalIndex();
+		const int32 FirstIndexB = IteratorB.GetLogicalIndex();
+		const int32 FirstNum    = Num;
 		for (;;)
 		{
-			const uint8* PairA = MapHelperA.GetPairPtr(IndexA);
+			const uint8* PairA = MapHelperA.GetPairPtr(IteratorA);
 			if (!AnyEqual(MapHelperA, FirstIndexA, FirstNum - Num, PairA, PortFlags) && !RangesContainSameAmountsOfVal(MapHelperA, FirstIndexA, MapHelperB, FirstIndexB, FirstNum, PairA, PortFlags))
 			{
 				return false;
@@ -167,11 +143,7 @@ namespace UEMapProperty_Private
 				return true;
 			}
 
-			++IndexA;
-			while (!MapHelperA.IsValidIndex(IndexA))
-			{
-				++IndexA;
-			}
+			++IteratorA;
 		}
 	}
 }
@@ -492,18 +464,12 @@ void FMapProperty::SerializeItem(FStructuredArchive::FSlot Slot, void* Value, co
 		// Determine how many keys are missing from the object
 		if (Defaults)
 		{
-			for (int32 Index = 0, Count = DefaultsHelper.Num(); Count; ++Index)
+			for (FScriptMapHelper::FIterator Iterator(DefaultsHelper); Iterator; ++Iterator)
 			{
-				uint8* DefaultPairPtr = DefaultsHelper.GetPairPtrWithoutCheck(Index);
-
-				if (DefaultsHelper.IsValidIndex(Index))
+				uint8* DefaultPairPtr = DefaultsHelper.GetPairPtr(Iterator);
+				if (!MapHelper.FindMapPairPtrWithKey(DefaultPairPtr))
 				{
-					if (!MapHelper.FindMapPairPtrWithKey(DefaultPairPtr))
-					{
-						Indices.Add(Index);
-					}
-
-					--Count;
+					Indices.Add(Iterator.GetInternalIndex());
 				}
 			}
 		}
@@ -523,19 +489,14 @@ void FMapProperty::SerializeItem(FStructuredArchive::FSlot Slot, void* Value, co
 		if (Defaults)
 		{
 			Indices.Empty(Indices.Num());
-			for (int32 Index = 0, Count = MapHelper.Num(); Count; ++Index)
+			for (FScriptMapHelper::FIterator Iterator(MapHelper); Iterator; ++Iterator)
 			{
-				if (MapHelper.IsValidIndex(Index))
-				{
-					uint8* ValuePairPtr   = MapHelper.GetPairPtrWithoutCheck(Index);
-					uint8* DefaultPairPtr = DefaultsHelper.FindMapPairPtrWithKey(ValuePairPtr);
+				uint8* ValuePairPtr   = MapHelper.GetPairPtr(Iterator);
+				uint8* DefaultPairPtr = DefaultsHelper.FindMapPairPtrWithKey(ValuePairPtr);
 
-					if (!DefaultPairPtr || !ValueProp->Identical(ValuePairPtr + MapLayout.ValueOffset, DefaultPairPtr + MapLayout.ValueOffset))
-					{
-						Indices.Add(Index);
-					}
-
-					--Count;
+				if (!DefaultPairPtr || !ValueProp->Identical(ValuePairPtr + MapLayout.ValueOffset, DefaultPairPtr + MapLayout.ValueOffset))
+				{
+					Indices.Add(Iterator.GetInternalIndex());
 				}
 			}
 
@@ -562,24 +523,18 @@ void FMapProperty::SerializeItem(FStructuredArchive::FSlot Slot, void* Value, co
 			int32 Num = MapHelper.Num();
 			FStructuredArchive::FArray EntriesArray = Record.EnterArray(TEXT("Entries"), Num);
 
-			for (int32 Index = 0; Num; ++Index)
+			for (FScriptMapHelper::FIterator Iterator(MapHelper); Iterator; ++Iterator)
 			{
-				if (MapHelper.IsValidIndex(Index))
-				{
-					FStructuredArchive::FRecord EntryRecord = EntriesArray.EnterElement().EnterRecord();
-
-					uint8* ValuePairPtr = MapHelper.GetPairPtrWithoutCheck(Index);
-
-					{
-						FSerializedPropertyScope SerializedProperty(UnderlyingArchive, KeyProp, this);
-						KeyProp->SerializeItem(EntryRecord.EnterField(TEXT("Key")), ValuePairPtr);
-					}
-					{
-						FSerializedPropertyScope SerializedProperty(UnderlyingArchive, ValueProp, this);
-						ValueProp->SerializeItem(EntryRecord.EnterField(TEXT("Value")), ValuePairPtr + MapLayout.ValueOffset);
-					}
+				FStructuredArchive::FRecord EntryRecord = EntriesArray.EnterElement().EnterRecord();
 
-					--Num;
+				uint8* ValuePairPtr = MapHelper.GetPairPtr(Iterator);
+				{
+					FSerializedPropertyScope SerializedProperty(UnderlyingArchive, KeyProp, this);
+					KeyProp->SerializeItem(EntryRecord.EnterField(TEXT("Key")), ValuePairPtr);
+				}
+				{
+					FSerializedPropertyScope SerializedProperty(UnderlyingArchive, ValueProp, this);
+					ValueProp->SerializeItem(EntryRecord.EnterField(TEXT("Value")), ValuePairPtr + MapLayout.ValueOffset);
 				}
 			}
 		}
@@ -1023,7 +978,7 @@ bool FMapProperty::PassCPPArgsByRef() const
  * 
  * @param	Data				pointer to the address of the instanced object referenced by this UComponentProperty
  * @param	DefaultData			pointer to the address of the default value of the instanced object referenced by this UComponentProperty
- * @param	Owner				the object that contains this property's data
+ * @param	InOwner				the object that contains this property's data
  * @param	InstanceGraph		contains the mappings of instanced objects and components to their templates
  */
 void FMapProperty::InstanceSubobjects(void* Data, void const* DefaultData, UObject* InOwner, FObjectInstancingGraph* InstanceGraph)
@@ -1033,8 +988,8 @@ void FMapProperty::InstanceSubobjects(void* Data, void const* DefaultData, UObje
 		return;
 	}
 
-	bool bInstancedKey   = KeyProp  ->ContainsInstancedObjectProperty();
-	bool bInstancedValue = ValueProp->ContainsInstancedObjectProperty();
+	const bool bInstancedKey   = KeyProp  ->ContainsInstancedObjectProperty();
+	const bool bInstancedValue = ValueProp->ContainsInstancedObjectProperty();
 
 	if (!bInstancedKey && !bInstancedValue)
 	{
@@ -1046,48 +1001,36 @@ void FMapProperty::InstanceSubobjects(void* Data, void const* DefaultData, UObje
 	if (DefaultData)
 	{
 		FScriptMapHelper DefaultMapHelper(this, DefaultData);
-		int32            DefaultNum = DefaultMapHelper.Num();
-
-		for (int32 Index = 0, Num = MapHelper.Num(); Num; ++Index)
+		for (FScriptMapHelper::FIterator It(MapHelper.CreateIterator()); It; ++It)
 		{
-			if (MapHelper.IsValidIndex(Index))
-			{
-				uint8* PairPtr        = MapHelper.GetPairPtr(Index);
-				uint8* DefaultPairPtr = DefaultMapHelper.FindMapPairPtrWithKey(PairPtr, Index);
+			uint8* PairPtr = MapHelper.GetPairPtr(It);
+			const uint8* DefaultPairPtr = DefaultMapHelper.FindMapPairPtrWithKey(PairPtr, /*IndexHint*/ It.GetLogicalIndex());
 
-				if (bInstancedKey)
-				{
-					KeyProp->InstanceSubobjects(PairPtr, DefaultPairPtr, InOwner, InstanceGraph);
-				}
-
-				if (bInstancedValue)
-				{
-					ValueProp->InstanceSubobjects(PairPtr + MapLayout.ValueOffset, DefaultPairPtr ? DefaultPairPtr + MapLayout.ValueOffset : nullptr, InOwner, InstanceGraph);
-				}
+			if (bInstancedKey)
+			{
+				KeyProp->InstanceSubobjects(PairPtr, DefaultPairPtr, InOwner, InstanceGraph);
+			}
 
-				--Num;
+			if (bInstancedValue)
+			{
+				ValueProp->InstanceSubobjects(PairPtr + MapLayout.ValueOffset, DefaultPairPtr ? DefaultPairPtr + MapLayout.ValueOffset : nullptr, InOwner, InstanceGraph);
 			}
 		}
 	}
 	else
 	{
-		for (int32 Index = 0, Num = MapHelper.Num(); Num; ++Index)
+		for (FScriptMapHelper::FIterator It(MapHelper.CreateIterator()); It; ++It)
 		{
-			if (MapHelper.IsValidIndex(Index))
-			{
-				uint8* PairPtr = MapHelper.GetPairPtr(Index);
+			uint8* PairPtr = MapHelper.GetPairPtr(It);
 
-				if (bInstancedKey)
-				{
-					KeyProp->InstanceSubobjects(PairPtr, nullptr, InOwner, InstanceGraph);
-				}
-
-				if (bInstancedValue)
-				{
-					ValueProp->InstanceSubobjects(PairPtr + MapLayout.ValueOffset, nullptr, InOwner, InstanceGraph);
-				}
+			if (bInstancedKey)
+			{
+				KeyProp->InstanceSubobjects(PairPtr, nullptr, InOwner, InstanceGraph);
+			}
 
-				--Num;
+			if (bInstancedValue)
+			{
+				ValueProp->InstanceSubobjects(PairPtr + MapLayout.ValueOffset, nullptr, InOwner, InstanceGraph);
 			}
 		}
 	}
@@ -1350,29 +1293,20 @@ void FMapProperty::GetInnerFields(TArray<FField*>& OutFields)
 	}
 }
 
-void* FMapProperty::GetValueAddressAtIndex_Direct(const FProperty* Inner, void* InValueAddress, int32 Index) const
+void* FMapProperty::GetValueAddressAtIndex_Direct(const FProperty* Inner, void* InValueAddress, const int32 LogicalIndex) const
 {
-	FScriptMapHelper MapHelper(this, InValueAddress);
 	checkf(Inner == KeyProp || Inner == ValueProp, TEXT("Inner property must be either KeyProp or ValueProp"));
 
-	for (int32 MapIndex = 0, Num = MapHelper.Num(), LocalIndex = Index; LocalIndex >= 0 && Num > 0; ++MapIndex)
+	FScriptMapHelper MapHelper(this, InValueAddress);
+	const int32 InternalIndex = MapHelper.FindInternalIndex(LogicalIndex);
+	if (InternalIndex != INDEX_NONE)
 	{
-		if (MapHelper.IsValidIndex(MapIndex))
+		if (Inner == KeyProp)
 		{
-			if (LocalIndex == 0)
-			{
-				if (Inner == KeyProp)
-				{
-					return MapHelper.GetKeyPtr(MapIndex);
-				}
-				else
-				{
-					return MapHelper.GetValuePtr(MapIndex);
-				}
-			}
-			LocalIndex--;
-			Num--;
+			return MapHelper.GetKeyPtr(InternalIndex);
 		}
+
+		return MapHelper.GetValuePtr(InternalIndex);
 	}
 	return nullptr;
 }
diff --git a/Engine/Source/Runtime/CoreUObject/Private/UObject/PropertySet.cpp b/Engine/Source/Runtime/CoreUObject/Private/UObject/PropertySet.cpp
--- a/Engine/Source/Runtime/CoreUObject/Private/UObject/PropertySet.cpp
+++ b/Engine/Source/Runtime/CoreUObject/Private/UObject/PropertySet.cpp
@@ -17,39 +17,36 @@ namespace UESetProperty_Private
 	/**
 	 * Checks if any of the elements in the set compare equal to the one passed.
 	 *
-	 * @param  SetHelper  The set to search through.
-	 * @param  Index      The index in the set to start searching from.
-	 * @param  Num        The number of elements to compare.
+	 * @param  SetHelper    The set to search through.
+	 * @param  LogicalIndex The index in the set to start searching from.
+	 * @param  Num          The number of elements to compare.
 	 */
-	bool AnyEqual(const FScriptSetHelper& SetHelper, int32 Index, int32 Num, const uint8* ElementToCompare, uint32 PortFlags)
+	bool AnyEqual(const FScriptSetHelper& SetHelper, const int32 LogicalIndex, int32 Num, const uint8* ElementToCompare, uint32 PortFlags)
 	{
-		FProperty* ElementProp = SetHelper.GetElementProperty();
+		const FProperty* ElementProp = SetHelper.GetElementProperty();
 
-		for (; Num; --Num)
+		FScriptSetHelper::FIterator Iterator(SetHelper, LogicalIndex);
+		for (; Iterator && Num; --Num, ++Iterator)
 		{
-			while (!SetHelper.IsValidIndex(Index))
-			{
-				++Index;
-			}
-
-			if (ElementProp->Identical(SetHelper.GetElementPtr(Index), ElementToCompare, PortFlags))
+			if (ElementProp->Identical(SetHelper.GetElementPtr(Iterator), ElementToCompare, PortFlags))
 			{
 				return true;
 			}
-
-			++Index;
 		}
 
 		return false;
 	}
 
-	bool RangesContainSameAmountsOfVal(const FScriptSetHelper& SetHelperA, int32 IndexA, const FScriptSetHelper& SetHelperB, int32 IndexB, int32 Num, const uint8* ElementToCompare, uint32 PortFlags)
+	bool RangesContainSameAmountsOfVal(const FScriptSetHelper& SetHelperA, const int32 LogicalIndexA, const FScriptSetHelper& SetHelperB, const int32 LogicalIndexB, int32 Num, const uint8* ElementToCompare, uint32 PortFlags)
 	{
-		FProperty* ElementProp = SetHelperA.GetElementProperty();
+		const FProperty* ElementProp = SetHelperA.GetElementProperty();
 
 		// Ensure that both sets are the same type
 		check(ElementProp == SetHelperB.GetElementProperty());
 
+		FScriptSetHelper::FIterator IteratorA(SetHelperA, LogicalIndexA);
+		FScriptSetHelper::FIterator IteratorB(SetHelperB, LogicalIndexB);
+
 		int32 CountA = 0;
 		int32 CountB = 0;
 		for (;;)
@@ -59,18 +56,8 @@ namespace UESetProperty_Private
 				return CountA == CountB;
 			}
 
-			while (!SetHelperA.IsValidIndex(IndexA))
-			{
-				++IndexA;
-			}
-
-			while (!SetHelperB.IsValidIndex(IndexB))
-			{
-				++IndexB;
-			}
-
-			const uint8* ElementA = SetHelperA.GetElementPtr(IndexA);
-			const uint8* ElementB = SetHelperB.GetElementPtr(IndexB);
+			const uint8* ElementA = SetHelperA.GetElementPtr(IteratorA);
+			const uint8* ElementB = SetHelperB.GetElementPtr(IteratorB);
 			if (ElementProp->Identical(ElementA, ElementToCompare, PortFlags))
 			{
 				++CountA;
@@ -81,15 +68,15 @@ namespace UESetProperty_Private
 				++CountB;
 			}
 
-			++IndexA;
-			++IndexB;
+			++IteratorA;
+			++IteratorB;
 			--Num;
 		}
 	}
 
-	bool IsPermutation(const FScriptSetHelper& SetHelperA, const FScriptSetHelper& SetHelperB, uint32 PortFlags)
+	bool IsPermutation(const FScriptSetHelper& SetHelperA, const FScriptSetHelper& SetHelperB, const uint32 PortFlags)
 	{
-		FProperty* ElementProp = SetHelperA.GetElementProperty();
+		const FProperty* ElementProp = SetHelperA.GetElementProperty();
 
 		// Ensure that both maps are the same type
 		check(ElementProp == SetHelperB.GetElementProperty());
@@ -101,8 +88,8 @@ namespace UESetProperty_Private
 		}
 
 		// Skip over common initial sequence
-		int32 IndexA = 0;
-		int32 IndexB = 0;
+		FScriptSetHelper::FIterator IteratorA(SetHelperA);
+		FScriptSetHelper::FIterator IteratorB(SetHelperB);
 		for (;;)
 		{
 			if (Num == 0)
@@ -110,34 +97,24 @@ namespace UESetProperty_Private
 				return true;
 			}
 
-			while (!SetHelperA.IsValidIndex(IndexA))
-			{
-				++IndexA;
-			}
-
-			while (!SetHelperB.IsValidIndex(IndexB))
-			{
-				++IndexB;
-			}
-
-			const uint8* ElementA = SetHelperA.GetElementPtr(IndexA);
-			const uint8* ElementB = SetHelperB.GetElementPtr(IndexB);
+			const uint8* ElementA = SetHelperA.GetElementPtr(IteratorA);
+			const uint8* ElementB = SetHelperB.GetElementPtr(IteratorB);
 			if (!ElementProp->Identical(ElementA, ElementB, PortFlags))
 			{
 				break;
 			}
 
-			++IndexA;
-			++IndexB;
+			++IteratorA;
+			++IteratorB;
 			--Num;
 		}
 
-		int32 FirstIndexA = IndexA;
-		int32 FirstIndexB = IndexB;
-		int32 FirstNum    = Num;
+		const int32 FirstIndexA = IteratorA.GetLogicalIndex();
+		const int32 FirstIndexB = IteratorB.GetLogicalIndex();
+		const int32 FirstNum    = Num;
 		for (;;)
 		{
-			const uint8* ElementA = SetHelperA.GetElementPtr(IndexA);
+			const uint8* ElementA = SetHelperA.GetElementPtr(IteratorA);
 			if (!AnyEqual(SetHelperA, FirstIndexA, FirstNum - Num, ElementA, PortFlags) && !RangesContainSameAmountsOfVal(SetHelperA, FirstIndexA, SetHelperB, FirstIndexB, FirstNum, ElementA, PortFlags))
 			{
 				return false;
@@ -149,12 +126,7 @@ namespace UESetProperty_Private
 				return true;
 			}
 
-			++IndexA;
-			while (!SetHelperA.IsValidIndex(IndexA))
-			{
-				++IndexA;
-			}
-
+			++IteratorA;
 		}
 	}
 }
@@ -432,18 +404,12 @@ void FSetProperty::SerializeItem(FStructuredArchive::FSlot Slot, void* Value, co
 		// Determine how many keys are missing from the object
 		if (Defaults)
 		{
-			for (int32 Index = 0, Count = DefaultsHelper.Num(); Count; ++Index)
+			for (FScriptSetHelper::FIterator Iterator(DefaultsHelper); Iterator; ++Iterator)
 			{
-				uint8* DefaultElementPtr = DefaultsHelper.GetElementPtrWithoutCheck(Index);
-
-				if (DefaultsHelper.IsValidIndex(Index))
+				const uint8* DefaultElementPtr = DefaultsHelper.GetElementPtr(Iterator);
+				if (SetHelper.FindElementIndex(DefaultElementPtr) == INDEX_NONE)
 				{
-					if (SetHelper.FindElementIndex(DefaultElementPtr) == INDEX_NONE)
-					{
-						Indices.Add(Index);
-					}
-
-					--Count;
+					Indices.Add(Iterator.GetInternalIndex());
 				}
 			}
 		}
@@ -464,19 +430,14 @@ void FSetProperty::SerializeItem(FStructuredArchive::FSlot Slot, void* Value, co
 		if (Defaults)
 		{
 			Indices.Reset();
-			for (int32 Index = 0, Count = SetHelper.Num(); Count; ++Index)
+			for (FScriptSetHelper::FIterator Iterator(SetHelper); Iterator; ++Iterator)
 			{
-				if (SetHelper.IsValidIndex(Index))
-				{
-					uint8* ValueElement   = SetHelper.GetElementPtrWithoutCheck(Index);
-					uint8* DefaultElement = DefaultsHelper.FindElementPtr(ValueElement);
+				const uint8* ValueElement = SetHelper.GetElementPtr(Iterator);
+				const uint8* DefaultElement = DefaultsHelper.FindElementPtr(ValueElement);
 
-					if (!DefaultElement)
-					{
-						Indices.Add(Index);
-					}
-
-					--Count;
+				if (!DefaultElement)
+				{
+					Indices.Add(Iterator.GetInternalIndex());
 				}
 			}
 
@@ -498,16 +459,11 @@ void FSetProperty::SerializeItem(FStructuredArchive::FSlot Slot, void* Value, co
 			FStructuredArchive::FArray ElementsArray = Record.EnterArray(TEXT("Elements"), Num);
 
 			FSerializedPropertyScope SerializedProperty(UnderlyingArchive, ElementProp, this);
-			for (int32 Index = 0; Num; ++Index)
-			{
-				if (SetHelper.IsValidIndex(Index))
-				{
-					uint8* ElementPtr = SetHelper.GetElementPtrWithoutCheck(Index);
-
-					ElementProp->SerializeItem(ElementsArray.EnterElement(), ElementPtr);
 
-					--Num;
-				}
+			for (FScriptSetHelper::FIterator Iterator(SetHelper); Iterator; ++Iterator)
+			{
+				uint8* ElementPtr = SetHelper.GetElementPtr(Iterator);
+				ElementProp->SerializeItem(ElementsArray.EnterElement(), ElementPtr);
 			}
 		}
 	}
@@ -903,31 +859,20 @@ void FSetProperty::InstanceSubobjects(void* Data, void const* DefaultData, UObje
 	{
 		FScriptSetHelper DefaultSetHelper(this, DefaultData);
 
-		for (int32 Index = 0, Num = SetHelper.Num(); Num; ++Index)
+		for (FScriptSetHelper::FIterator It(SetHelper); It; ++It)
 		{
-			if (SetHelper.IsValidIndex(Index))
-			{
-				uint8* ElementPtr        = SetHelper.GetElementPtr(Index);
-				uint8* DefaultElementPtr = DefaultSetHelper.FindElementPtr(ElementPtr, Index);
-
-				ElementProp->InstanceSubobjects(ElementPtr, DefaultElementPtr, InOwner, InstanceGraph);
+			uint8* ElementPtr = SetHelper.GetElementPtr(It);
+			const uint8* DefaultElementPtr = DefaultSetHelper.FindElementPtr(ElementPtr, It.GetLogicalIndex());
 
-				--Num;
-			}
+			ElementProp->InstanceSubobjects(ElementPtr, DefaultElementPtr, InOwner, InstanceGraph);
 		}
 	}
 	else
 	{
-		for (int32 Index = 0, Num = SetHelper.Num(); Num; ++Index)
+		for (FScriptSetHelper::FIterator It(SetHelper); It; ++It)
 		{
-			if (SetHelper.IsValidIndex(Index))
-			{
-				uint8* ElementPtr = SetHelper.GetElementPtr(Index);
-
-				ElementProp->InstanceSubobjects(ElementPtr, nullptr, InOwner, InstanceGraph);
-
-				--Num;
-			}
+			uint8* ElementPtr = SetHelper.GetElementPtr(It);
+			ElementProp->InstanceSubobjects(ElementPtr, nullptr, InOwner, InstanceGraph);
 		}
 	}
 }
@@ -1124,22 +1069,16 @@ void FSetProperty::GetInnerFields(TArray<FField*>& OutFields)
 	}
 }
 
-void* FSetProperty::GetValueAddressAtIndex_Direct(const FProperty* Inner, void* InValueAddress, int32 Index) const
+void* FSetProperty::GetValueAddressAtIndex_Direct(const FProperty* Inner, void* InValueAddress, const int32 LogicalIndex) const
 {
-	FScriptSetHelper SetHelper(this, InValueAddress);
 	checkf(Inner == ElementProp, TEXT("Inner property must be identical to ElementProp"));
 
-	for (int32 SetIndex = 0, Num = SetHelper.Num(), LocalIndex = Index; LocalIndex >= 0 && Num > 0; ++SetIndex)
+	FScriptSetHelper SetHelper(this, InValueAddress);
+	const int32 InternalIndex = SetHelper.FindInternalIndex(LogicalIndex);
+	if (InternalIndex != INDEX_NONE)
 	{
-		if (SetHelper.IsValidIndex(SetIndex))
-		{
-			if (LocalIndex == 0)
-			{
-				return SetHelper.GetElementPtr(SetIndex);
-			}
-			LocalIndex--;
-			Num--;
-		}
+		return SetHelper.GetElementPtrWithoutCheck(InternalIndex);
 	}
+
 	return nullptr;
 }
diff --git a/Engine/Source/Runtime/Engine/Private/BlueprintMapLibrary.cpp b/Engine/Source/Runtime/Engine/Private/BlueprintMapLibrary.cpp
--- a/Engine/Source/Runtime/Engine/Private/BlueprintMapLibrary.cpp
+++ b/Engine/Source/Runtime/Engine/Private/BlueprintMapLibrary.cpp
@@ -61,17 +61,11 @@ void UBlueprintMapLibrary::GenericMap_Keys(const void* TargetMap, const FMapProp
 		FScriptArrayHelper ArrayHelper(ArrayProperty, TargetArray);
 		ArrayHelper.EmptyValues();
 
-		FProperty* InnerProp = ArrayProperty->Inner;
-
-		int32 Size = MapHelper.Num();
-		for( int32 I = 0; Size; ++I )
+		const FProperty* InnerProp = ArrayProperty->Inner;
+		for (FScriptMapHelper::FIterator It(MapHelper); It; ++It)
 		{
-			if(MapHelper.IsValidIndex(I))
-			{
-				int32 LastIndex = ArrayHelper.AddValue();
-				InnerProp->CopySingleValueToScriptVM(ArrayHelper.GetRawPtr(LastIndex), MapHelper.GetKeyPtr(I));
-				--Size;
-			}
+			const int32 LastIndex = ArrayHelper.AddValue();
+			InnerProp->CopySingleValueToScriptVM(ArrayHelper.GetRawPtr(LastIndex), MapHelper.GetKeyPtr(It));
 		}
 	}
 }
@@ -84,17 +78,11 @@ void UBlueprintMapLibrary::GenericMap_Values(const void* TargetMap, const FMapPr
 		FScriptArrayHelper ArrayHelper(ArrayProperty, TargetArray);
 		ArrayHelper.EmptyValues();
 
-		FProperty* InnerProp = ArrayProperty->Inner;
-		
-		int32 Size = MapHelper.Num();
-		for( int32 I = 0; Size; ++I )
+		const FProperty* InnerProp = ArrayProperty->Inner;
+		for (FScriptMapHelper::FIterator It(MapHelper); It; ++It)
 		{
-			if(MapHelper.IsValidIndex(I))
-			{
-				int32 LastIndex = ArrayHelper.AddValue();
-				InnerProp->CopySingleValueToScriptVM(ArrayHelper.GetRawPtr(LastIndex), MapHelper.GetValuePtr(I));
-				--Size;
-			}
+			const int32 LastIndex = ArrayHelper.AddValue();
+			InnerProp->CopySingleValueToScriptVM(ArrayHelper.GetRawPtr(LastIndex), MapHelper.GetValuePtr(It));
 		}
 	}
 }
diff --git a/Engine/Source/Runtime/Engine/Private/BlueprintSetLibrary.cpp b/Engine/Source/Runtime/Engine/Private/BlueprintSetLibrary.cpp
--- a/Engine/Source/Runtime/Engine/Private/BlueprintSetLibrary.cpp
+++ b/Engine/Source/Runtime/Engine/Private/BlueprintSetLibrary.cpp
@@ -64,15 +64,9 @@ void UBlueprintSetLibrary::GenericSet_ToArray(const void* TargetSet, const FSetP
 	if (TargetSet && TargetArray)
 	{
 		FScriptSetHelper SetHelper(SetProperty, TargetSet);
-
-		int32 Size = SetHelper.Num();
-		for (int32 I = 0; Size; ++I)
+		for (FScriptSetHelper::FIterator It(SetHelper); It; ++It)
 		{
-			if(SetHelper.IsValidIndex(I))
-			{
-				UKismetArrayLibrary::GenericArray_Add(TargetArray, ArrayProperty, SetHelper.GetElementPtr(I));
-				--Size;
-			}
+			UKismetArrayLibrary::GenericArray_Add(TargetArray, ArrayProperty, SetHelper.GetElementPtr(It));
 		}
 	}
 }
@@ -142,18 +136,12 @@ void UBlueprintSetLibrary::GenericSet_Intersect(const void* SetA, const FSetProp
 		FScriptSetHelper SetHelperResult(SetPropertyResult, SetResult);
 		
 		SetHelperResult.EmptyElements();
-
-		int32 Size = SetHelperA.Num();
-		for (int32 I = 0; Size; ++I)
+		for (FScriptSetHelper::FIterator It(SetHelperA); It; ++It)
 		{
-			if(SetHelperA.IsValidIndex(I))
+			const void* EntryInA = SetHelperA.GetElementPtr(It);
+			if (SetHelperB.FindElementIndexFromHash(EntryInA) != INDEX_NONE)
 			{
-				const void* EntryInA = SetHelperA.GetElementPtr(I);
-				if (SetHelperB.FindElementIndexFromHash(EntryInA) != INDEX_NONE)
-				{
-					SetHelperResult.AddElement(EntryInA);
-				}
-				--Size;
+				SetHelperResult.AddElement(EntryInA);
 			}
 		}
 	}
@@ -169,24 +157,14 @@ void UBlueprintSetLibrary::GenericSet_Union(const void* SetA, const FSetProperty
 		
 		SetHelperResult.EmptyElements();
 
-		int32 SizeA = SetHelperA.Num();
-		for (int32 I = 0; SizeA; ++I)
+		for (FScriptSetHelper::FIterator It(SetHelperA); It; ++It)
 		{
-			if(SetHelperA.IsValidIndex(I))
-			{
-				SetHelperResult.AddElement(SetHelperA.GetElementPtr(I));
-				--SizeA;
-			}
+			SetHelperResult.AddElement(SetHelperA.GetElementPtr(It));
 		}
 
-		int32 SizeB = SetHelperB.Num();
-		for (int32 I = 0; SizeB; ++I)
+		for (FScriptSetHelper::FIterator It(SetHelperB); It; ++It)
 		{
-			if(SetHelperB.IsValidIndex(I))
-			{
-				SetHelperResult.AddElement(SetHelperB.GetElementPtr(I));
-				--SizeB;
-			}
+			SetHelperResult.AddElement(SetHelperB.GetElementPtr(It));
 		}
 	}
 }
@@ -201,17 +179,12 @@ void UBlueprintSetLibrary::GenericSet_Difference(const void* SetA, const FSetPro
 
 		SetHelperResult.EmptyElements();
 
-		int32 Size = SetHelperA.Num();
-		for (int32 I = 0; Size; ++I)
+		for (FScriptSetHelper::FIterator It(SetHelperA); It; ++It)
 		{
-			if(SetHelperA.IsValidIndex(I))
+			const void* EntryInA = SetHelperA.GetElementPtr(It);
+			if (SetHelperB.FindElementIndexFromHash(EntryInA) == INDEX_NONE)
 			{
-				const void* EntryInA = SetHelperA.GetElementPtr(I);
-				if (SetHelperB.FindElementIndexFromHash(EntryInA) == INDEX_NONE)
-				{
-					SetHelperResult.AddElement(EntryInA);
-				}
-				--Size;
+				SetHelperResult.AddElement(EntryInA);
 			}
 		}
 	}
diff --git a/Engine/Source/Runtime/Engine/Private/DataTableJSON.cpp b/Engine/Source/Runtime/Engine/Private/DataTableJSON.cpp
--- a/Engine/Source/Runtime/Engine/Private/DataTableJSON.cpp
+++ b/Engine/Source/Runtime/Engine/Private/DataTableJSON.cpp
@@ -265,13 +265,10 @@ bool TDataTableExporterJSON<CharType>::WriteStructEntry(const void* InRowData, c
 		JsonWriter->WriteArrayStart(Identifier);
 
 		FScriptSetHelper SetHelper(SetProp, InPropertyData);
-		for (int32 SetSparseIndex = 0; SetSparseIndex < SetHelper.GetMaxIndex(); ++SetSparseIndex)
+		for (FScriptSetHelper::FIterator It(SetHelper); It; ++It)
 		{
-			if (SetHelper.IsValidIndex(SetSparseIndex))
-			{
-				const uint8* SetEntryData = SetHelper.GetElementPtr(SetSparseIndex);
-				WriteContainerEntry(SetHelper.GetElementProperty(), SetEntryData);
-			}
+			const uint8* SetEntryData = SetHelper.GetElementPtr(It);
+			WriteContainerEntry(SetHelper.GetElementProperty(), SetEntryData);
 		}
 
 		JsonWriter->WriteArrayEnd();
@@ -281,17 +278,14 @@ bool TDataTableExporterJSON<CharType>::WriteStructEntry(const void* InRowData, c
 		JsonWriter->WriteObjectStart(Identifier);
 
 		FScriptMapHelper MapHelper(MapProp, InPropertyData);
-		for (int32 MapSparseIndex = 0; MapSparseIndex < MapHelper.GetMaxIndex(); ++MapSparseIndex)
+		for (FScriptMapHelper::FIterator It(MapHelper); It; ++It)
 		{
-			if (MapHelper.IsValidIndex(MapSparseIndex))
-			{
-				const uint8* MapKeyData = MapHelper.GetKeyPtr(MapSparseIndex);
-				const uint8* MapValueData = MapHelper.GetValuePtr(MapSparseIndex);
+			const uint8* MapKeyData = MapHelper.GetKeyPtr(It);
+			const uint8* MapValueData = MapHelper.GetValuePtr(It);
 
-				// JSON object keys must always be strings
-				const FString KeyValue = DataTableUtils::GetPropertyValueAsStringDirect(MapHelper.GetKeyProperty(), (uint8*)MapKeyData, DTExportFlags);
-				WriteContainerEntry(MapHelper.GetValueProperty(), MapValueData, &KeyValue);
-			}
+			// JSON object keys must always be strings
+			const FString KeyValue = DataTableUtils::GetPropertyValueAsStringDirect(MapHelper.GetKeyProperty(), (uint8*)MapKeyData, DTExportFlags);
+			WriteContainerEntry(MapHelper.GetValueProperty(), MapValueData, &KeyValue);
 		}
 
 		JsonWriter->WriteObjectEnd();
diff --git a/Engine/Source/Runtime/Engine/Private/DataTableUtils.cpp b/Engine/Source/Runtime/Engine/Private/DataTableUtils.cpp
--- a/Engine/Source/Runtime/Engine/Private/DataTableUtils.cpp
+++ b/Engine/Source/Runtime/Engine/Private/DataTableUtils.cpp
@@ -142,21 +142,17 @@ void GetPropertyValueAsStringDirect(const FProperty* InProp, const uint8* InData
 
 			OutString.AppendChar('(');
 
-			int32 NumWrittenSetEntries = 0;
 			FScriptSetHelper SetHelper(SetProp, InData);
-			for (int32 SetSparseIndex = 0; SetSparseIndex < SetHelper.GetMaxIndex(); ++SetSparseIndex)
+			for (FScriptSetHelper::FIterator It(SetHelper); It; ++It)
 			{
-				if (SetHelper.IsValidIndex(SetSparseIndex))
-				{
-					if (NumWrittenSetEntries++ > 0)
+					if (It.GetLogicalIndex() > 0)
 					{
 						OutString.AppendChar(',');
 						OutString.AppendChar(' ');
 					}
 
-					const uint8* SetEntryData = SetHelper.GetElementPtr(SetSparseIndex);
+					const uint8* SetEntryData = SetHelper.GetElementPtr(It);
 					OutString.Append(ExportStructAsJson(StructInner->Struct, SetEntryData));
-				}
 			}
 
 			OutString.AppendChar(')');
@@ -166,36 +162,32 @@ void GetPropertyValueAsStringDirect(const FProperty* InProp, const uint8* InData
 		{
 			OutString.AppendChar('(');
 
-			int32 NumWrittenMapEntries = 0;
 			FScriptMapHelper MapHelper(MapProp, InData);
-			for (int32 MapSparseIndex = 0; MapSparseIndex < MapHelper.GetMaxIndex(); ++MapSparseIndex)
+			for (FScriptMapHelper::FIterator It(MapHelper); It; ++It)
 			{
-				if (MapHelper.IsValidIndex(MapSparseIndex))
+				if (It.GetLogicalIndex() > 0)
 				{
-					if (NumWrittenMapEntries++ > 0)
-					{
-						OutString.AppendChar(',');
-						OutString.AppendChar(' ');
-					}
+					OutString.AppendChar(',');
+					OutString.AppendChar(' ');
+				}
 
-					const uint8* MapKeyData = MapHelper.GetKeyPtr(MapSparseIndex);
-					const uint8* MapValueData = MapHelper.GetValuePtr(MapSparseIndex);
+				const uint8* MapKeyData = MapHelper.GetKeyPtr(It);
+				const uint8* MapValueData = MapHelper.GetValuePtr(It);
 
-					OutString.AppendChar('"');
-					GetPropertyValueAsStringDirect(MapHelper.GetKeyProperty(), MapKeyData, InPortFlags, InDTExportFlags, OutString);
-					OutString.AppendChar('"');
+				OutString.AppendChar('"');
+				GetPropertyValueAsStringDirect(MapHelper.GetKeyProperty(), MapKeyData, InPortFlags, InDTExportFlags, OutString);
+				OutString.AppendChar('"');
 
-					OutString.Append(TEXT(" = "));
+				OutString.Append(TEXT(" = "));
 
-					if (MapHelper.GetValueProperty()->IsA<FStructProperty>() && EnumHasAnyFlags(InDTExportFlags, EDataTableExportFlags::UseJsonObjectsForStructs))
-					{
-						const FStructProperty* StructMapValue = CastFieldChecked<const FStructProperty>(MapHelper.GetValueProperty());
-						OutString.Append(ExportStructAsJson(StructMapValue->Struct, MapValueData));
-					}
-					else
-					{
-						GetPropertyValueAsStringDirect(MapHelper.GetValueProperty(), MapValueData, InPortFlags, InDTExportFlags, OutString);
-					}
+				if (MapHelper.GetValueProperty()->IsA<FStructProperty>() && EnumHasAnyFlags(InDTExportFlags, EDataTableExportFlags::UseJsonObjectsForStructs))
+				{
+					const FStructProperty* StructMapValue = CastFieldChecked<const FStructProperty>(MapHelper.GetValueProperty());
+					OutString.Append(ExportStructAsJson(StructMapValue->Struct, MapValueData));
+				}
+				else
+				{
+					GetPropertyValueAsStringDirect(MapHelper.GetValueProperty(), MapValueData, InPortFlags, InDTExportFlags, OutString);
 				}
 			}
 
diff --git a/Engine/Source/Runtime/Engine/Private/InstancedReferenceSubobjectHelper.cpp b/Engine/Source/Runtime/Engine/Private/InstancedReferenceSubobjectHelper.cpp
--- a/Engine/Source/Runtime/Engine/Private/InstancedReferenceSubobjectHelper.cpp
+++ b/Engine/Source/Runtime/Engine/Private/InstancedReferenceSubobjectHelper.cpp
@@ -160,26 +160,20 @@ void FFindInstancedReferenceSubobjectHelper::ForEachInstancedSubObject(FInstance
 			return;
 		}
 
-		int32 LogicalIndex = 0;
 		FScriptMapHelper MapHelper(MapProperty, ContainerAddress);
-		for (int32 ElementIndex = 0; ElementIndex < MapHelper.GetMaxIndex(); ++ElementIndex)
+		for (FScriptMapHelper::FIterator It(MapHelper); It; ++It)
 		{
-			if (MapHelper.IsValidIndex(ElementIndex))
-			{
-				T KeyAddress = MapHelper.GetKeyPtr(ElementIndex);
-				T ValueAddress = MapHelper.GetValuePtr(ElementIndex);
-
-				// Note: Keep these as the logical (Nth) index in case the map changes internally after we construct the path or in case we resolve using a different object.
-				PropertyPath.Push(MapProperty->KeyProp, LogicalIndex);
-				ForEachInstancedSubObject(PropertyPath, KeyAddress, ObjRefFunc);
-				PropertyPath.Pop();
+			T KeyAddress = MapHelper.GetKeyPtr(It);
+			T ValueAddress = MapHelper.GetValuePtr(It);
 
-				PropertyPath.Push(MapProperty->ValueProp, LogicalIndex, true);
-				ForEachInstancedSubObject(PropertyPath, ValueAddress, ObjRefFunc);
-				PropertyPath.Pop();
+			// Note: Keep these as the logical (Nth) index in case the map changes internally after we construct the path or in case we resolve using a different object.
+			PropertyPath.Push(MapProperty->KeyProp, It.GetLogicalIndex());
+			ForEachInstancedSubObject(PropertyPath, KeyAddress, ObjRefFunc);
+			PropertyPath.Pop();
 
-				++LogicalIndex;
-			}
+			PropertyPath.Push(MapProperty->ValueProp, It.GetLogicalIndex(), true);
+			ForEachInstancedSubObject(PropertyPath, ValueAddress, ObjRefFunc);
+			PropertyPath.Pop();
 		}
 	}
 	else if (const FSetProperty* SetProperty = CastField<const FSetProperty>(TargetProp))
@@ -190,21 +184,15 @@ void FFindInstancedReferenceSubobjectHelper::ForEachInstancedSubObject(FInstance
 			return;
 		}
 
-		int32 LogicalIndex = 0;
 		FScriptSetHelper SetHelper(SetProperty, ContainerAddress);
-		for (int32 ElementIndex = 0; ElementIndex < SetHelper.GetMaxIndex(); ++ElementIndex)
+		for (FScriptSetHelper::FIterator It(SetHelper); It; ++It)
 		{
-			if (SetHelper.IsValidIndex(ElementIndex))
-			{
-				T ValueAddress = SetHelper.GetElementPtr(ElementIndex);
-
-				// Note: Keep this as the logical (Nth) index in case the set changes internally after we construct the path or in case we resolve using a different object.
-				PropertyPath.Push(SetProperty->ElementProp, LogicalIndex);
-				ForEachInstancedSubObject(PropertyPath, ValueAddress, ObjRefFunc);
-				PropertyPath.Pop();
+			T ValueAddress = SetHelper.GetElementPtr(It);
 
-				++LogicalIndex;
-			}
+			// Note: Keep this as the logical (Nth) index in case the set changes internally after we construct the path or in case we resolve using a different object.
+			PropertyPath.Push(SetProperty->ElementProp, It.GetLogicalIndex());
+			ForEachInstancedSubObject(PropertyPath, ValueAddress, ObjRefFunc);
+			PropertyPath.Pop();
 		}
 	}
 	else if (const FOptionalProperty* OptionalProperty = CastField<FOptionalProperty>(TargetProp))
diff --git a/Engine/Source/Runtime/Engine/Private/ObjectEditorUtils.cpp b/Engine/Source/Runtime/Engine/Private/ObjectEditorUtils.cpp
--- a/Engine/Source/Runtime/Engine/Private/ObjectEditorUtils.cpp
+++ b/Engine/Source/Runtime/Engine/Private/ObjectEditorUtils.cpp
@@ -202,25 +202,19 @@ namespace FObjectEditorUtils
 
 			TargetMapHelper.EmptyValues();
 
-			int32 Num = SourceMapHelper.Num();
-			for ( int32 Index = 0; Num; Index++ )
+			for (FScriptMapHelper::FIterator It(SourceMapHelper); It; ++It)
 			{
-				if ( SourceMapHelper.IsValidIndex(Index) )
-				{
-					uint8* SrcPairPtr = SourceMapHelper.GetPairPtr(Index);
-
-					int32 NewIndex = TargetMapHelper.AddDefaultValue_Invalid_NeedsRehash();
-					TargetMapHelper.Rehash();
+				uint8* SrcPairPtr = SourceMapHelper.GetPairPtr(It);
 
-					uint8* PairPtr = TargetMapHelper.GetPairPtr(NewIndex);
+				int32 NewIndex = TargetMapHelper.AddDefaultValue_Invalid_NeedsRehash();
+				TargetMapHelper.Rehash();
 
-					CopySinglePropertyRecursive(SourceObject, SrcPairPtr, SrcMapProperty->KeyProp, PairPtr, InDestinationObject, DestMapProperty->KeyProp);
-					CopySinglePropertyRecursive(SourceObject, SrcPairPtr, SrcMapProperty->ValueProp, PairPtr, InDestinationObject, DestMapProperty->ValueProp);
+				uint8* PairPtr = TargetMapHelper.GetPairPtr(NewIndex);
 
-					TargetMapHelper.Rehash();
+				CopySinglePropertyRecursive(SourceObject, SrcPairPtr, SrcMapProperty->KeyProp, PairPtr, InDestinationObject, DestMapProperty->KeyProp);
+				CopySinglePropertyRecursive(SourceObject, SrcPairPtr, SrcMapProperty->ValueProp, PairPtr, InDestinationObject, DestMapProperty->ValueProp);
 
-					--Num;
-				}
+				TargetMapHelper.Rehash();
 			}
 
 			bNeedsShallowCopy = false;
@@ -235,23 +229,17 @@ namespace FObjectEditorUtils
 
 			TargetSetHelper.EmptyElements();
 
-			int32 Num = SourceSetHelper.Num();
-			for ( int32 Index = 0; Num; Index++ )
+			for (FScriptSetHelper::FIterator It(SourceSetHelper); It; ++It)
 			{
-				if ( SourceSetHelper.IsValidIndex(Index) )
-				{
-					uint8* SrcPtr = SourceSetHelper.GetElementPtr(Index);
-
-					int32 NewIndex = TargetSetHelper.AddDefaultValue_Invalid_NeedsRehash();
-					TargetSetHelper.Rehash();
+				uint8* SrcPtr = SourceSetHelper.GetElementPtr(It);
 
-					uint8* TargetPtr = TargetSetHelper.GetElementPtr(NewIndex);
-					CopySinglePropertyRecursive(SourceObject, SrcPtr, SrcSetProperty->ElementProp, TargetPtr, InDestinationObject, DestSetProperty->ElementProp);
+				int32 NewIndex = TargetSetHelper.AddDefaultValue_Invalid_NeedsRehash();
+				TargetSetHelper.Rehash();
 
-					TargetSetHelper.Rehash();
+				uint8* TargetPtr = TargetSetHelper.GetElementPtr(NewIndex);
+				CopySinglePropertyRecursive(SourceObject, SrcPtr, SrcSetProperty->ElementProp, TargetPtr, InDestinationObject, DestSetProperty->ElementProp);
 
-					--Num;
-				}
+				TargetSetHelper.Rehash();
 			}
 
 			bNeedsShallowCopy = false;
diff --git a/Engine/Source/Runtime/Serialization/Private/StructSerializer.cpp b/Engine/Source/Runtime/Serialization/Private/StructSerializer.cpp
--- a/Engine/Source/Runtime/Serialization/Private/StructSerializer.cpp
+++ b/Engine/Source/Runtime/Serialization/Private/StructSerializer.cpp
@@ -22,18 +22,18 @@ namespace StructSerializer
 	template <>
 	struct TScriptHelper_InContainer<FMapProperty> { using Type = FScriptMapHelper_InContainer; };
 
-	FStructSerializerState CreateItemState(FScriptArrayHelper& InHelper, FArrayProperty* InArrayProperty, int32 InElementIndex, const FStructSerializerPolicies& Policies, EStructSerializerStateFlags InFlags)
+	FStructSerializerState CreateItemState(FScriptArrayHelper& InHelper, FArrayProperty* InArrayProperty, int32 InValidatedInternalElementIndex, const FStructSerializerPolicies& Policies, EStructSerializerStateFlags InFlags)
 	{
-		return FStructSerializerState(InHelper.GetRawPtr(InElementIndex), InArrayProperty->Inner, InFlags);
+		return FStructSerializerState(InHelper.GetRawPtr(InValidatedInternalElementIndex), InArrayProperty->Inner, InFlags);
 	};
 
-	FStructSerializerState CreateItemState(FScriptMapHelper& InHelper, FMapProperty* InMapProperty, int32 InElementIndex, const FStructSerializerPolicies& Policies, EStructSerializerStateFlags InFlags)
+	FStructSerializerState CreateItemState(FScriptMapHelper& InHelper, FMapProperty* InMapProperty, int32 InValidatedInternalElementIndex, const FStructSerializerPolicies& Policies, EStructSerializerStateFlags InFlags)
 	{
-		FStructSerializerState State = FStructSerializerState(InHelper.GetPairPtr(InElementIndex), InMapProperty->ValueProp, InFlags);
+		FStructSerializerState State = FStructSerializerState(InHelper.GetPairPtr(InValidatedInternalElementIndex), InMapProperty->ValueProp, InFlags);
 
 		if (Policies.MapSerialization == EStructSerializerMapPolicies::KeyValuePair)
 		{
-			State.KeyData = InHelper.GetPairPtr(InElementIndex);
+			State.KeyData = InHelper.GetPairPtr(InValidatedInternalElementIndex);
 			State.ValueData = State.KeyData;
 			State.KeyProperty = InMapProperty->KeyProp;
 		}
@@ -41,9 +41,9 @@ namespace StructSerializer
 		return State;
 	}
 
-	FStructSerializerState CreateItemState(FScriptSetHelper& InHelper, FSetProperty* InSetProperty, int32 InElementIndex, const FStructSerializerPolicies& Policies, EStructSerializerStateFlags InFlags)
+	FStructSerializerState CreateItemState(FScriptSetHelper& InHelper, FSetProperty* InSetProperty, int32 InValidatedInternalElementIndex, const FStructSerializerPolicies& Policies, EStructSerializerStateFlags InFlags)
 	{
-		return FStructSerializerState(InHelper.GetElementPtr(InElementIndex), InSetProperty->ElementProp, InFlags);
+		return FStructSerializerState(InHelper.GetElementPtr(InValidatedInternalElementIndex), InSetProperty->ElementProp, InFlags);
 	}
 
 	bool IsArrayLike(FProperty* Property, const FStructSerializerPolicies& Policies)
